[Upgrade Notes Critical! Licensee build scripts and DeviceProfiles may need updating]

- Formalizing all TargetPlatforms to have a Client version, rename WindowsNoEditor to Windows, and removing DDPI specification of TargetPlatforms, and generate them programmatically
- Updated names DeviceProfiles and Build scripts, as above
- Some PlatformInfo class cleanup
- Added a TNonDesktopTargetPlatformBase class to make most TargetPlatforms simpler
- Added "No Compiled Support" to the Turnkey LaunchOn menu when the TargetPlatforms aren't compiled in (to show that even if you install an SDK, you will need to compile before you can LaunchOn)\
- Starting the transition away from PlatformInfo::FPlatformInfo to FDDPI

diff --git a/Engine/Build/InstalledEngineBuild.xml b/Engine/Build/InstalledEngineBuild.xml
--- a/Engine/Build/InstalledEngineBuild.xml
+++ b/Engine/Build/InstalledEngineBuild.xml
@@ -547,14 +547,14 @@
 		</Node>
 	</Agent>
 
-	<Property Name="DDCPlatformsWin64" Value="Windows"/>
-	<Property Name="DDCPlatformsWin64" Value="$(DDCPlatformsWin64)+WindowsNoEditor" If="'$(WithWin64)' == true"/>
+	<Property Name="DDCPlatformsWin64" Value="WindowsEditor"/>
+	<Property Name="DDCPlatformsWin64" Value="$(DDCPlatformsWin64)+Windows" If="'$(WithWin64)' == true"/>
 	<Do If="'$(HostPlatformDDCOnly)' == false">
 		<Property Name="DDCPlatformsWin64" Value="$(DDCPlatformsWin64)+Android_ATC" If="'$(WithAndroid)' == true"/>
 		<Property Name="DDCPlatformsWin64" Value="$(DDCPlatformsWin64)+IOS" If="'$(WithIOS)' == true"/>
 		<Property Name="DDCPlatformsWin64" Value="$(DDCPlatformsWin64)+TVOS" If="'$(WithTVOS)' == true"/>
-		<Property Name="DDCPlatformsWin64" Value="$(DDCPlatformsWin64)+LinuxNoEditor" If="'$(WithLinux)' == true"/>
-		<Property Name="DDCPlatformsWin64" Value="$(DDCPlatformsWin64)+LinuxAArch64NoEditor" If="'$(WithLinuxAArch64)' == true"/>
+		<Property Name="DDCPlatformsWin64" Value="$(DDCPlatformsWin64)+Linux" If="'$(WithLinux)' == true"/>
+		<Property Name="DDCPlatformsWin64" Value="$(DDCPlatformsWin64)+LinuxAArch64" If="'$(WithLinuxAArch64)' == true"/>
 		<Property Name="DDCPlatformsWin64" Value="$(DDCPlatformsWin64)+Lumin" If="'$(WithLumin)' == true"/>
 		<Property Name="DDCPlatformsWin64" Value="$(DDCPlatformsWin64)+HoloLens" If="'$(WithHoloLens)' == true"/>
 		<ForEach Name="Platform" Values="$(ExtensionPlatforms)">
@@ -629,8 +629,8 @@
 		</Node>
 	</Agent>
 
-	<Property Name="DDCPlatformsMac" Value="Mac"/>
-	<Property Name="DDCPlatformsMac" Value="$(DDCPlatformsMac)+MacNoEditor" If="'$(WithMac)' == true"/>
+	<Property Name="DDCPlatformsMac" Value="MacEditor"/>
+	<Property Name="DDCPlatformsMac" Value="$(DDCPlatformsMac)+Mac" If="'$(WithMac)' == true"/>
 	<Do If="'$(HostPlatformDDCOnly)' == false">
 		<Property Name="DDCPlatformsMac" Value="$(DDCPlatformsMac)+Android_ATC" If="'$(WithAndroid)' == true"/>
 		<Property Name="DDCPlatformsMac" Value="$(DDCPlatformsMac)+IOS" If="'$(WithIOS)' == true"/>
@@ -666,10 +666,10 @@
 		</Node>
 	</Agent>
 
-	<Property Name="DDCPlatformsLinux" Value="Linux"/>
-	<Property Name="DDCPlatformsLinux" Value="$(DDCPlatformsLinux)+LinuxNoEditor" If="'$(WithLinux)' == true"/>
+	<Property Name="DDCPlatformsLinux" Value="LinuxEditor"/>
+	<Property Name="DDCPlatformsLinux" Value="$(DDCPlatformsLinux)+Linux" If="'$(WithLinux)' == true"/>
 	<Do If="'$(HostPlatformDDCOnly)' == false">
-		<Property Name="DDCPlatformsLinux" Value="$(DDCPlatformsLinux)+LinuxAArch64NoEditor" If="'$(WithLinuxAArch64)' == true"/>
+		<Property Name="DDCPlatformsLinux" Value="$(DDCPlatformsLinux)+LinuxAArch64" If="'$(WithLinuxAArch64)' == true"/>
 	</Do>
 
 	<Agent Name="DDC Group Linux" Type="Linux;Win64">
diff --git a/Engine/Config/Android/DataDrivenPlatformInfo.ini b/Engine/Config/Android/DataDrivenPlatformInfo.ini
--- a/Engine/Config/Android/DataDrivenPlatformInfo.ini
+++ b/Engine/Config/Android/DataDrivenPlatformInfo.ini
@@ -14,73 +14,3 @@ bUsesHostCompiler=false
 bUATClosesAfterLaunch=true
 PlatformGroupName=Mobile
 
-[PlatformInfo Android]
-TargetPlatformName=Android
-DisplayName=Android
-PlatformType=Game
-PlatformFlags=None
-UATCommandLine=
-
-[PlatformInfo Android_Multi]
-TargetPlatformName=Android_Multi
-DisplayName=Android (Multi)
-PlatformType=Game
-PlatformFlags=CookFlavor
-UATCommandLine=-targetplatform=Android -cookflavor=Multi
-
-[PlatformInfo Android_DXT]
-TargetPlatformName=Android_DXT
-DisplayName=Android (DXT)
-PlatformType=Game
-PlatformFlags=CookFlavor
-UATCommandLine=-targetplatform=Android -cookflavor=DXT
-
-[PlatformInfo Android_ETC2]
-TargetPlatformName=Android_ETC2
-DisplayName=Android (ETC2)
-PlatformType=Game
-PlatformFlags=CookFlavor
-UATCommandLine=-targetplatform=Android -cookflavor=ETC2
-
-[PlatformInfo Android_ASTC]
-TargetPlatformName=Android_ASTC
-DisplayName=Android (ASTC)
-PlatformType=Game
-PlatformFlags=CookFlavor
-UATCommandLine=-targetplatform=Android -cookflavor=ASTC
-
-
-[PlatformInfo AndroidClient]
-TargetPlatformName=AndroidClient
-DisplayName=Android (Client-only)
-PlatformType=Client
-PlatformFlags=None
-UATCommandLine=-client
-
-[PlatformInfo Android_MultiClient]
-TargetPlatformName=Android_MultiClient
-DisplayName=Android (Multi) (Client-only)
-PlatformType=Client
-PlatformFlags=CookFlavor
-UATCommandLine=-client -targetplatform=Android -cookflavor=Multi
-
-[PlatformInfo Android_DXTClient]
-TargetPlatformName=Android_DXTClient
-DisplayName=Android (DXT) (Client-only)
-PlatformType=Client
-PlatformFlags=CookFlavor
-UATCommandLine=-client -targetplatform=Android -cookflavor=DXT
-
-[PlatformInfo Android_ETC2Client]
-TargetPlatformName=Android_ETC2Client
-DisplayName=Android (ETC2) (Client-only)
-PlatformType=Client
-PlatformFlags=CookFlavor
-UATCommandLine=-client -targetplatform=Android -cookflavor=ETC2
-
-[PlatformInfo Android_ASTCClient]
-TargetPlatformName=Android_ASTCClient
-DisplayName=Android (ASTC) (Client-only)
-PlatformType=Client
-PlatformFlags=CookFlavor
-UATCommandLine=-client -targetplatform=Android -cookflavor=ASTC
diff --git a/Engine/Config/BaseDeviceProfiles.ini b/Engine/Config/BaseDeviceProfiles.ini
--- a/Engine/Config/BaseDeviceProfiles.ini
+++ b/Engine/Config/BaseDeviceProfiles.ini
@@ -4,8 +4,8 @@
 *TextureLODGroups=Group
 
 [DeviceProfiles]
++DeviceProfileNameAndTypes=WindowsEditor,Windows
 +DeviceProfileNameAndTypes=Windows,Windows
-+DeviceProfileNameAndTypes=WindowsNoEditor,Windows
 +DeviceProfileNameAndTypes=WindowsServer,Windows
 +DeviceProfileNameAndTypes=WindowsClient,Windows
 +DeviceProfileNameAndTypes=IOS,IOS
@@ -52,13 +52,13 @@
 +DeviceProfileNameAndTypes=AppleTV,IOS
 +DeviceProfileNameAndTypes=AppleTV4K,IOS
 +DeviceProfileNameAndTypes=TVOS,AppleTV
-+DeviceProfileNameAndTypes=Mac,Mac
++DeviceProfileNameAndTypes=MacEditor,Mac
 +DeviceProfileNameAndTypes=MacClient,Mac
-+DeviceProfileNameAndTypes=MacNoEditor,Mac
++DeviceProfileNameAndTypes=Mac,Mac
 +DeviceProfileNameAndTypes=MacServer,Mac
++DeviceProfileNameAndTypes=LinuxEditor,Linux
 +DeviceProfileNameAndTypes=Linux,Linux
-+DeviceProfileNameAndTypes=LinuxNoEditor,Linux
-+DeviceProfileNameAndTypes=LinuxAArch64NoEditor,Linux
++DeviceProfileNameAndTypes=LinuxAArch64,Linux
 +DeviceProfileNameAndTypes=LinuxClient,Linux
 +DeviceProfileNameAndTypes=LinuxAArch64Client,Linux
 +DeviceProfileNameAndTypes=LinuxServer,Linux
@@ -128,11 +128,11 @@ TextureLODGroups=(Group=TEXTUREGROUP_World,MinLODSize=1,MaxLODSize=8192,LODBias=
 [Windows DeviceProfile]
 DeviceType=Windows
 BaseProfileName=
-; +CVars=r.Vulkan.EnableValidation=0
 
-[WindowsNoEditor DeviceProfile]
+[WindowsEditor DeviceProfile]
 DeviceType=Windows
 BaseProfileName=Windows
+; +CVars=r.Vulkan.EnableValidation=0
 
 [WindowsClient DeviceProfile]
 DeviceType=Windows
@@ -698,7 +698,7 @@ BaseProfileName=
 DeviceType=Mac
 BaseProfileName=Mac
 
-[MacNoEditor DeviceProfile]
+[MacEditor DeviceProfile]
 DeviceType=Mac
 BaseProfileName=Mac
 
@@ -712,29 +712,25 @@ BaseProfileName=
 MeshLODSettings=
 TextureLODSettings=
 
-[LinuxNoEditor DeviceProfile]
+[LinuxEditor DeviceProfile]
 DeviceType=Linux
 BaseProfileName=Linux
 MeshLODSettings=
 TextureLODSettings=
-+CVars=r.RHICmdBypass=0
 
-[LinuxAArch64NoEditor DeviceProfile]
+[LinuxAArch6 DeviceProfile]
 DeviceType=Linux
 BaseProfileName=Linux
 MeshLODSettings=
 TextureLODSettings=
-+CVars=r.RHICmdBypass=0
 
 [LinuxClient DeviceProfile]
 DeviceType=Linux
 BaseProfileName=Linux
-+CVars=r.RHICmdBypass=0
 
 [LinuxAArch64Client DeviceProfile]
 DeviceType=Linux
 BaseProfileName=Linux
-+CVars=r.RHICmdBypass=0
 
 [LinuxServer DeviceProfile]
 DeviceType=Linux
@@ -750,4 +746,4 @@ TextureLODSettings=
 
 [HoloLens DeviceProfile]
 DeviceType=HoloLens
-BaseProfileName=WindowsNoEditor
+BaseProfileName=Windows
diff --git a/Engine/Config/HoloLens/DataDrivenPlatformInfo.ini b/Engine/Config/HoloLens/DataDrivenPlatformInfo.ini
--- a/Engine/Config/HoloLens/DataDrivenPlatformInfo.ini
+++ b/Engine/Config/HoloLens/DataDrivenPlatformInfo.ini
@@ -3,8 +3,8 @@ bIsConfidential=false
 Freezing_MaxFieldAlignment=8
 Freezing_bAlignBases=true
 
-NormalIconPath=Launcher/Windows/Platform_WindowsNoEditor_24x
-LargeIconPath=Launcher/Windows/Platform_WindowsNoEditor_128x
+NormalIconPath=Launcher/Windows/Platform_Windows_24x
+LargeIconPath=Launcher/Windows/Platform_Windows_128x
 XLargeIconPath=
 AutoSDKPath=
 TutorialPath=/Engine/Tutorial/Installation/InstallingVisualStudioTutorial.InstallingVisualStudioTutorial
@@ -16,9 +16,3 @@ bCanUseCrashReporter=false
 bUATClosesAfterLaunch=false
 PlatformGroupName=Mobile
 
-[PlatformInfo HoloLens]
-TargetPlatformName=HoloLens
-DisplayName=HoloLens
-PlatformType=Game
-PlatformFlags=None
-UATCommandLine=-targetplatform=HoloLens
diff --git a/Engine/Config/IOS/DataDrivenPlatformInfo.ini b/Engine/Config/IOS/DataDrivenPlatformInfo.ini
--- a/Engine/Config/IOS/DataDrivenPlatformInfo.ini
+++ b/Engine/Config/IOS/DataDrivenPlatformInfo.ini
@@ -15,17 +15,3 @@ Linux:bIsEnabled=false
 bUsesHostCompiler=false
 bUATClosesAfterLaunch=true
 PlatformGroupName=Mobile
-
-[PlatformInfo IOS]
-TargetPlatformName=IOS
-DisplayName=iOS
-PlatformType=Game
-PlatformFlags=None
-UATCommandLine=
-
-[PlatformInfo IOSClient]
-TargetPlatformName=IOSClient
-DisplayName=iOS (Client-only)
-PlatformType=Client
-PlatformFlags=None
-UATCommandLine=-client
diff --git a/Engine/Config/Linux/DataDrivenPlatformInfo.ini b/Engine/Config/Linux/DataDrivenPlatformInfo.ini
--- a/Engine/Config/Linux/DataDrivenPlatformInfo.ini
+++ b/Engine/Config/Linux/DataDrivenPlatformInfo.ini
@@ -15,30 +15,3 @@ bUATClosesAfterLaunch=true
 PlatformGroupName=Desktop
 PlatformSubMenu=Linux
 
-[PlatformInfo Linux]
-TargetPlatformName=Linux
-DisplayName=Linux (Editor)
-PlatformType=Editor
-PlatformFlags=None
-UATCommandLine=
-
-[PlatformInfo LinuxNoEditor]
-TargetPlatformName=LinuxNoEditor
-DisplayName=Linux
-PlatformType=Game
-PlatformFlags=None
-UATCommandLine=-targetplatform=Linux
-
-[PlatformInfo LinuxClient]
-TargetPlatformName=LinuxClient
-DisplayName=Linux (Client-only)
-PlatformType=Client
-PlatformFlags=None
-UATCommandLine=-client
-
-[PlatformInfo LinuxServer]
-TargetPlatformName=LinuxServer
-DisplayName=Linux (Dedicated Server)
-PlatformType=Server
-PlatformFlags=None
-UATCommandLine=
diff --git a/Engine/Config/LinuxAArch64/DataDrivenPlatformInfo.ini b/Engine/Config/LinuxAArch64/DataDrivenPlatformInfo.ini
--- a/Engine/Config/LinuxAArch64/DataDrivenPlatformInfo.ini
+++ b/Engine/Config/LinuxAArch64/DataDrivenPlatformInfo.ini
@@ -16,23 +16,3 @@ bUATClosesAfterLaunch=true
 PlatformGroupName=Mobile
 PlatformSubMenu=Linux
 
-[PlatformInfo LinuxAArch64]
-TargetPlatformName=LinuxAArch64
-DisplayName=Linux (AArch64)
-PlatformType=Game
-PlatformFlags=BuildFlavor
-UATCommandLine=-targetplatform=LinuxAArch64
-
-[PlatformInfo LinuxAArch64Client]
-TargetPlatformName=LinuxAArch64Client
-DisplayName=LinuxAArch64 (Client-only)
-PlatformType=Client
-PlatformFlags=None
-UATCommandLine=-client
-
-[PlatformInfo LinuxAArch64Server]
-TargetPlatformName=LinuxAArch64Server
-DisplayName=LinuxAArch64 (Dedicated Server)
-PlatformType=Server
-PlatformFlags=None
-UATCommandLine=
diff --git a/Engine/Config/Lumin/DataDrivenPlatformInfo.ini b/Engine/Config/Lumin/DataDrivenPlatformInfo.ini
--- a/Engine/Config/Lumin/DataDrivenPlatformInfo.ini
+++ b/Engine/Config/Lumin/DataDrivenPlatformInfo.ini
@@ -14,16 +14,3 @@ bUsesHostCompiler=false
 bUATClosesAfterLaunch=true
 PlatformGroupName=Mobile
 
-[PlatformInfo Lumin]
-TargetPlatformName=Lumin
-DisplayName=Lumin
-PlatformType=Game
-PlatformFlags=None
-UATCommandLine=
-
-[PlatformInfo LuminClient]
-TargetPlatformName=LuminClient
-DisplayName=Lumin (Client-only)
-PlatformType=Client
-PlatformFlags=None
-UATCommandLine=-client
diff --git a/Engine/Config/Mac/DataDrivenPlatformInfo.ini b/Engine/Config/Mac/DataDrivenPlatformInfo.ini
--- a/Engine/Config/Mac/DataDrivenPlatformInfo.ini
+++ b/Engine/Config/Mac/DataDrivenPlatformInfo.ini
@@ -14,30 +14,3 @@ bUsesHostCompiler=true
 bUATClosesAfterLaunch=true
 PlatformGroupName=Desktop
 
-[PlatformInfo MacNoEditor]
-TargetPlatformName=MacNoEditor
-DisplayName=Mac
-PlatformType=Game
-PlatformFlags=None
-UATCommandLine=-targetplatform=Mac
-
-[PlatformInfo Mac]
-TargetPlatformName=Mac
-DisplayName=Mac (Editor)
-PlatformType=Editor
-PlatformFlags=None
-UATCommandLine=
-
-[PlatformInfo MacClient]
-TargetPlatformName=MacClient
-DisplayName=Mac (Client-only)
-PlatformType=Client
-PlatformFlags=None
-UATCommandLine=-client
-
-[PlatformInfo MacServer]
-TargetPlatformName=MacServer
-DisplayName=Mac (Dedicated Server)
-PlatformType=Server
-PlatformFlags=None
-UATCommandLine=
diff --git a/Engine/Config/TVOS/DataDrivenPlatformInfo.ini b/Engine/Config/TVOS/DataDrivenPlatformInfo.ini
--- a/Engine/Config/TVOS/DataDrivenPlatformInfo.ini
+++ b/Engine/Config/TVOS/DataDrivenPlatformInfo.ini
@@ -15,16 +15,3 @@ bUsesHostCompiler=false
 bUATClosesAfterLaunch=true
 PlatformGroupName=Mobile
 
-[PlatformInfo TVOS]
-TargetPlatformName=TVOS
-DisplayName=tvOS
-PlatformType=Game
-PlatformFlags=None
-UATCommandLine=
-
-[PlatformInfo TVOSClient]
-TargetPlatformName=TVOSClient
-DisplayName=tvOS (Client-only)
-PlatformType=Client
-PlatformFlags=None
-UATCommandLine=-client
diff --git a/Engine/Config/Windows/DataDrivenPlatformInfo.ini b/Engine/Config/Windows/DataDrivenPlatformInfo.ini
--- a/Engine/Config/Windows/DataDrivenPlatformInfo.ini
+++ b/Engine/Config/Windows/DataDrivenPlatformInfo.ini
@@ -7,8 +7,8 @@ Freezing_MaxFieldAlignment32=4
 Freezing_MaxFieldAlignment64=8
 Freezing_bAlignBases=true
 
-NormalIconPath=Launcher/Windows/Platform_WindowsNoEditor_24x
-LargeIconPath=Launcher/Windows/Platform_WindowsNoEditor_128x
+NormalIconPath=Launcher/Windows/Platform_Windows_24x
+LargeIconPath=Launcher/Windows/Platform_Windows_128x
 XLargeIconPath=
 AutoSDKPath=
 TutorialPath=/Engine/Tutorial/Installation/InstallingVisualStudioTutorial.InstallingVisualStudioTutorial
@@ -21,35 +21,6 @@ bUATClosesAfterLaunch=false
 PlatformGroupName=Desktop
 UBTPlatformName=Win64
 
-
-[PlatformInfo WindowsNoEditor]
-TargetPlatformName=WindowsNoEditor
-DisplayName=Windows
-PlatformType=Game
-PlatformFlags=None
-UATCommandLine=-targetplatform=Win64
-
-[PlatformInfo Windows]
-TargetPlatformName=Windows
-DisplayName=Windows (Editor)
-PlatformType=Editor
-PlatformFlags=None
-UATCommandLine=
-
-[PlatformInfo WindowsClient]
-TargetPlatformName=WindowsClient
-DisplayName=Windows (Client-only)
-PlatformType=Client
-PlatformFlags=None
-UATCommandLine=-client
-
-[PlatformInfo WindowsServer]
-TargetPlatformName=WindowsServer
-DisplayName=Windows (Dedicated Server)
-PlatformType=Server
-PlatformFlags=None
-UATCommandLine=
-
 [ShaderPlatform PCD3D_SM5]
 bSupportsRayTracing=true
 bSupportsRayTracingMissShaderBindings=true
diff --git a/Engine/Plugins/Editor/MobileLauncherProfileWizard/Source/MobileLauncherProfileWizard/Private/SProfileWizardUI.cpp b/Engine/Plugins/Editor/MobileLauncherProfileWizard/Source/MobileLauncherProfileWizard/Private/SProfileWizardUI.cpp
--- a/Engine/Plugins/Editor/MobileLauncherProfileWizard/Source/MobileLauncherProfileWizard/Private/SProfileWizardUI.cpp
+++ b/Engine/Plugins/Editor/MobileLauncherProfileWizard/Source/MobileLauncherProfileWizard/Private/SProfileWizardUI.cpp
@@ -742,7 +742,7 @@ void SProfileWizardUI::CacheCookFlavorsList()
 	const PlatformInfo::FTargetPlatformInfo* Info = FindPlatformInfo(TargetPlatformName);
 	for (const FTargetPlatformInfo* PlatformFlaworInfo : Info->Flavors)
 	{
-		DLCFlavorList.Add(MakeShareable(new FString(PlatformFlaworInfo->PlatformInfoName.ToString())));
+		DLCFlavorList.Add(MakeShareable(new FString(PlatformFlaworInfo->Name.ToString())));
 	}
 
 	// In case no flavors, add vanilla name as cooking target
diff --git a/Engine/Plugins/Enterprise/DatasmithImporter/Source/DatasmithImporter/Private/Utility/DatasmithImporterImpl.cpp b/Engine/Plugins/Enterprise/DatasmithImporter/Source/DatasmithImporter/Private/Utility/DatasmithImporterImpl.cpp
--- a/Engine/Plugins/Enterprise/DatasmithImporter/Source/DatasmithImporter/Private/Utility/DatasmithImporterImpl.cpp
+++ b/Engine/Plugins/Enterprise/DatasmithImporter/Source/DatasmithImporter/Private/Utility/DatasmithImporterImpl.cpp
@@ -931,7 +931,7 @@ bool FDatasmithImporterImpl::CheckAssetPersistenceValidity(const FString& Packag
 	}
 
 	const FString AbsoluteGamePath = FPaths::ConvertRelativePathToFull(FPaths::ProjectDir());
-	const FString AbsoluteGameCookPath = AbsoluteGamePath / TEXT("Saved") / TEXT("Cooked") / TEXT("WindowsNoEditor") / GameName;
+	const FString AbsoluteGameCookPath = AbsoluteGamePath / TEXT("Saved") / TEXT("Cooked") / TEXT("Windows") / GameName;
 
 	FString AssetPathWithinCookDir = AbsolutePathToAsset;
 	FPaths::RemoveDuplicateSlashes(AssetPathWithinCookDir);
diff --git a/Engine/Plugins/Media/MediaPlayerEditor/Source/MediaPlayerEditor/Private/Customizations/PlatformMediaSourceCustomization.cpp b/Engine/Plugins/Media/MediaPlayerEditor/Source/MediaPlayerEditor/Private/Customizations/PlatformMediaSourceCustomization.cpp
--- a/Engine/Plugins/Media/MediaPlayerEditor/Source/MediaPlayerEditor/Private/Customizations/PlatformMediaSourceCustomization.cpp
+++ b/Engine/Plugins/Media/MediaPlayerEditor/Source/MediaPlayerEditor/Private/Customizations/PlatformMediaSourceCustomization.cpp
@@ -77,7 +77,7 @@ TSharedRef<SWidget> FPlatformMediaSourceCustomization::MakePlatformMediaSourcesV
 	{
 		if (PlatformInfo->PlatformType == EBuildTargetType::Game)
 		{
-			if (PlatformInfo->PlatformInfoName == TEXT("TVOS"))
+			if (PlatformInfo->Name == TEXT("TVOS"))
 			{
 				continue; // tvOS is just iOS for now
 			}
diff --git a/Engine/Plugins/Runtime/AR/Microsoft/HoloLensAR/Source/HoloLensTargetPlatform/Private/HoloLensTargetPlatform.cpp b/Engine/Plugins/Runtime/AR/Microsoft/HoloLensAR/Source/HoloLensTargetPlatform/Private/HoloLensTargetPlatform.cpp
--- a/Engine/Plugins/Runtime/AR/Microsoft/HoloLensAR/Source/HoloLensTargetPlatform/Private/HoloLensTargetPlatform.cpp
+++ b/Engine/Plugins/Runtime/AR/Microsoft/HoloLensAR/Source/HoloLensTargetPlatform/Private/HoloLensTargetPlatform.cpp
@@ -16,10 +16,9 @@
 
 DEFINE_LOG_CATEGORY(LogHoloLensTargetPlatform);
 
-FHoloLensTargetPlatform::FHoloLensTargetPlatform()
+FHoloLensTargetPlatform::FHoloLensTargetPlatform(bool bIsClientOnly)
+	: TNonDesktopTargetPlatformBase(bIsClientOnly)
 {
-	PlatformInfo = ::PlatformInfo::FindPlatformInfo(FName("HoloLens"));
-
 #if WITH_ENGINE
 	FConfigCacheIni::LoadLocalIniFile(EngineSettings, TEXT("Engine"), true, *PlatformName());
 	TextureLODSettings = nullptr; // These are registered by the device profile system.
diff --git a/Engine/Plugins/Runtime/AR/Microsoft/HoloLensAR/Source/HoloLensTargetPlatform/Private/HoloLensTargetPlatform.h b/Engine/Plugins/Runtime/AR/Microsoft/HoloLensAR/Source/HoloLensTargetPlatform/Private/HoloLensTargetPlatform.h
--- a/Engine/Plugins/Runtime/AR/Microsoft/HoloLensAR/Source/HoloLensTargetPlatform/Private/HoloLensTargetPlatform.h
+++ b/Engine/Plugins/Runtime/AR/Microsoft/HoloLensAR/Source/HoloLensTargetPlatform/Private/HoloLensTargetPlatform.h
@@ -27,14 +27,14 @@
  * FHoloLensTargetPlatform, abstraction for cooking HoloLens platforms
  */
 class HOLOLENSTARGETPLATFORM_API FHoloLensTargetPlatform
-	: public TTargetPlatformBase<FHoloLensPlatformProperties>
+	: public TNonDesktopTargetPlatformBase<FHoloLensPlatformProperties>
 {
 public:
 
 	/**
 	 * Default constructor.
 	 */
-	FHoloLensTargetPlatform();
+	FHoloLensTargetPlatform(bool bIsClient);
 
 	/**
 	 * Destructor.
@@ -59,8 +59,6 @@ public:
 
 	virtual bool GenerateStreamingInstallManifest(const TMultiMap<FString, int32>& ChunkMap, const TSet<int32>& ChunkIDsInUse) const override { return true; }
 
-	virtual bool IsRunningPlatform() const override { return false; }
-
 	virtual bool SupportsFeature(ETargetPlatformFeatures Feature) const override;
 
 	virtual bool SupportsBuildTarget(EBuildTargetType BuildTarget) const override;
@@ -114,21 +112,11 @@ public:
 		return true;
 	}
 
-	virtual FText GetVariantTitle() const override
-	{
-		return LOCTEXT("HoloLensVariantTitle", "Build Type");
-	}
-
 	virtual FString PlatformName() const override
 	{
 		return TEXT("HoloLens");
 	}
 
-	virtual FText GetVariantDisplayName() const override
-	{
-		return LOCTEXT("HoloLensVariantDisplayName", "HoloLens");
-	}
-
 	virtual float GetVariantPriority() const override
 	{
 		return 1.0f;
diff --git a/Engine/Plugins/Runtime/AR/Microsoft/HoloLensAR/Source/HoloLensTargetPlatform/Private/HoloLensTargetPlatformModule.cpp b/Engine/Plugins/Runtime/AR/Microsoft/HoloLensAR/Source/HoloLensTargetPlatform/Private/HoloLensTargetPlatformModule.cpp
--- a/Engine/Plugins/Runtime/AR/Microsoft/HoloLensAR/Source/HoloLensTargetPlatform/Private/HoloLensTargetPlatformModule.cpp
+++ b/Engine/Plugins/Runtime/AR/Microsoft/HoloLensAR/Source/HoloLensTargetPlatform/Private/HoloLensTargetPlatformModule.cpp
@@ -25,7 +25,8 @@ public:
 		if (FHoloLensTargetPlatform::IsUsable())
 		{
 			//@todo HoloLens: Check for SDK?
-			TargetPlatforms.Add(new FHoloLensTargetPlatform());
+			TargetPlatforms.Add(new FHoloLensTargetPlatform(false));
+			TargetPlatforms.Add(new FHoloLensTargetPlatform(true));
 		}
 	}
 };
diff --git a/Engine/Plugins/Runtime/WindowsDeviceProfileSelector/Source/WindowsDeviceProfileSelector/Private/WindowsDeviceProfileSelectorModule.cpp b/Engine/Plugins/Runtime/WindowsDeviceProfileSelector/Source/WindowsDeviceProfileSelector/Private/WindowsDeviceProfileSelectorModule.cpp
--- a/Engine/Plugins/Runtime/WindowsDeviceProfileSelector/Source/WindowsDeviceProfileSelector/Private/WindowsDeviceProfileSelectorModule.cpp
+++ b/Engine/Plugins/Runtime/WindowsDeviceProfileSelector/Source/WindowsDeviceProfileSelector/Private/WindowsDeviceProfileSelectorModule.cpp
@@ -21,7 +21,7 @@ void FWindowsDeviceProfileSelectorModule::ShutdownModule()
 
 const FString FWindowsDeviceProfileSelectorModule::GetRuntimeDeviceProfileName()
 {
-	// Windows, WindowsNoEditor, WindowsClient, or WindowsServer
+	// Windows, WindowsEditor, WindowsClient, or WindowsServer
 	FString ProfileName = FPlatformProperties::PlatformName();
 
 	if (FApp::CanEverRender())
diff --git a/Engine/Source/Developer/Android/AndroidTargetPlatform/Private/AndroidTargetPlatform.cpp b/Engine/Source/Developer/Android/AndroidTargetPlatform/Private/AndroidTargetPlatform.cpp
--- a/Engine/Source/Developer/Android/AndroidTargetPlatform/Private/AndroidTargetPlatform.cpp
+++ b/Engine/Source/Developer/Android/AndroidTargetPlatform/Private/AndroidTargetPlatform.cpp
@@ -190,8 +190,8 @@ static bool HasLicense()
 	return false;
 }
 
-FAndroidTargetPlatform::FAndroidTargetPlatform(bool bInIsClient )
-	: bIsClient(bInIsClient)
+FAndroidTargetPlatform::FAndroidTargetPlatform(bool bInIsClient, const TCHAR* FlavorName, const TCHAR* OverrideIniPlatformName)
+	: TNonDesktopTargetPlatformBase(bInIsClient, FlavorName, OverrideIniPlatformName)
 	, DeviceDetection(nullptr)
 
 {
@@ -310,11 +310,6 @@ ITargetDevicePtr FAndroidTargetPlatform::GetDevice( const FTargetDeviceId& Devic
 	return nullptr;
 }
 
-bool FAndroidTargetPlatform::IsRunningPlatform( ) const
-{
-	return false; // This platform never runs the target platform framework
-}
-
 
 bool FAndroidTargetPlatform::IsSdkInstalled(bool bProjectHasCode, FString& OutDocumentationPath) const
 {
@@ -658,10 +653,6 @@ bool FAndroidTargetPlatform::SupportsVariants() const
 	return true;
 }
 
-FText FAndroidTargetPlatform::GetVariantTitle() const
-{
-	return LOCTEXT("AndroidVariantTitle", "Texture Format");
-}
 
 /* FAndroidTargetPlatform implementation
  *****************************************************************************/
diff --git a/Engine/Source/Developer/Android/AndroidTargetPlatform/Private/AndroidTargetPlatform.h b/Engine/Source/Developer/Android/AndroidTargetPlatform/Private/AndroidTargetPlatform.h
--- a/Engine/Source/Developer/Android/AndroidTargetPlatform/Private/AndroidTargetPlatform.h
+++ b/Engine/Source/Developer/Android/AndroidTargetPlatform/Private/AndroidTargetPlatform.h
@@ -83,14 +83,14 @@ enum class EAndroidTextureFormatCategory
 /**
  * FAndroidTargetPlatform, abstraction for cooking Android platforms
  */
-class ANDROIDTARGETPLATFORM_API FAndroidTargetPlatform : public TTargetPlatformBase<FAndroidPlatformProperties>
+class ANDROIDTARGETPLATFORM_API FAndroidTargetPlatform : public TNonDesktopTargetPlatformBase<FAndroidPlatformProperties>
 {
 public:
 
 	/**
 	 * Default constructor.
 	 */
-	FAndroidTargetPlatform(bool bInIsClient);
+	FAndroidTargetPlatform(bool bIsClient, const TCHAR* FlavorName, const TCHAR* OverrideIniPlatformName=nullptr);
 
 	/**
 	 * Destructor
@@ -104,31 +104,9 @@ public:
 	 *
 	 * @param Variant name.
 	 */
-	virtual FString GetAndroidVariantName() const
+	FString GetAndroidVariantName() const
 	{
-		return FString();	
-	}
-
-
-	virtual FString IniPlatformName() const override
-	{
-		return "Android";
-	}
-
-	virtual FString PlatformName() const override
-	{
-		FString PlatformName = TEXT("Android");
-		FString Variant = GetAndroidVariantName();
-		if (Variant.Len() > 0)
-		{
-			PlatformName += FString(TEXT("_")) + Variant;
-		}
-		if (bIsClient)
-		{
-			PlatformName += TEXT("Client");
-		}
-
-		return PlatformName;
+		return PlatformInfo->PlatformFlavor.ToString();
 	}
 
 public:
@@ -153,18 +131,6 @@ public:
 
 	virtual ITargetDevicePtr GetDevice( const FTargetDeviceId& DeviceId ) override;
 
-	virtual bool IsRunningPlatform( ) const override;
-
-	virtual bool IsServerOnly( ) const override
-	{
-		return false;
-	}
-
-	virtual bool IsClientOnly() const override
-	{
-		return bIsClient;
-	}
-
 	virtual bool IsSdkInstalled(bool bProjectHasCode, FString& OutDocumentationPath) const override;
 
 	virtual int32 CheckRequirements(bool bProjectHasCode, EBuildConfiguration Configuration, bool bRequiresAssetNativization, FString& OutTutorialPath, FString& OutDocumentationPath, FText& CustomizedLogMessage) const override;
@@ -215,8 +181,6 @@ public:
 
 	virtual bool SupportsVariants() const override;
 
-	virtual FText GetVariantTitle() const override;
-
 	virtual void GetBuildProjectSettingKeys(FString& OutSection, TArray<FString>& InBoolKeys, TArray<FString>& InIntKeys, TArray<FString>& InStringKeys) const override
 	{
 		OutSection = TEXT("/Script/AndroidRuntimeSettings.AndroidRuntimeSettings");
@@ -275,9 +239,6 @@ protected:
 
 	virtual FAndroidTargetDeviceRef CreateNewDevice(const FAndroidDeviceInfo &DeviceInfo);
 
-	// true if this is a client TP
-	bool bIsClient;
-
 	// Holds a map of valid devices.
 	TMap<FString, FAndroidTargetDevicePtr> Devices;
 
@@ -308,19 +269,8 @@ protected:
 class FAndroid_DXTTargetPlatform : public FAndroidTargetPlatform
 {
 public:
-	FAndroid_DXTTargetPlatform(bool bIsClient) : FAndroidTargetPlatform(bIsClient)
-	{
-		this->PlatformInfo = PlatformInfo::FindPlatformInfo("Android_DXT");
-	}
-
-	virtual FString GetAndroidVariantName() const override
-	{
-		return TEXT("DXT");
-	}
-
-	virtual FText DisplayName() const override
+	FAndroid_DXTTargetPlatform(bool bIsClient) : FAndroidTargetPlatform(bIsClient, TEXT("DXT"))
 	{
-		return LOCTEXT("Android_DXT", "Android (DXT)");
 	}
 
 	virtual bool SupportsTextureFormat(FName Format) const override
@@ -344,11 +294,6 @@ public:
 		return (ExtensionsString.Contains(TEXT("GL_NV_texture_compression_s3tc")) || ExtensionsString.Contains(TEXT("GL_EXT_texture_compression_s3tc")));
 	}
 
-	virtual FText GetVariantDisplayName() const override
-	{
-		return LOCTEXT("Android_DXT_ShortName", "DXT");
-	}
-
 	virtual float GetVariantPriority() const override
 	{
 		float Priority;
@@ -361,19 +306,8 @@ public:
 class FAndroid_ASTCTargetPlatform : public FAndroidTargetPlatform
 {
 public:
-	FAndroid_ASTCTargetPlatform(bool bIsClient) : FAndroidTargetPlatform(bIsClient) 
-	{
-		this->PlatformInfo = PlatformInfo::FindPlatformInfo("Android_ASTC");
-	}
-
-	virtual FString GetAndroidVariantName() const override
+	FAndroid_ASTCTargetPlatform(bool bIsClient) : FAndroidTargetPlatform(bIsClient, TEXT("ASTC")) 
 	{
-		return TEXT("ASTC");
-	}
-
-	virtual FText DisplayName() const override
-	{
-		return LOCTEXT("Android_ASTC", "Android (ASTC)");
 	}
 
 	virtual bool SupportsTextureFormat(FName Format) const override
@@ -470,11 +404,6 @@ public:
 		return ExtensionsString.Contains(TEXT("GL_KHR_texture_compression_astc_ldr"));
 	}
 
-	virtual FText GetVariantDisplayName() const override
-	{
-		return LOCTEXT("Android_ASTC_ShortName", "ASTC");
-	}
-
 	virtual float GetVariantPriority() const override
 	{
 		float Priority;
@@ -488,19 +417,8 @@ class FAndroid_ETC2TargetPlatform : public FAndroidTargetPlatform
 {
 public:
 
-	FAndroid_ETC2TargetPlatform(bool bIsClient) : FAndroidTargetPlatform(bIsClient)
+	FAndroid_ETC2TargetPlatform(bool bIsClient) : FAndroidTargetPlatform(bIsClient, TEXT("ETC2"))
 	{
-		this->PlatformInfo = PlatformInfo::FindPlatformInfo("Android_ETC2");
-	}
-
-	virtual FText DisplayName() const override
-	{
-		return LOCTEXT("Android_ETC2", "Android (ETC2)");
-	}
-
-	virtual FString GetAndroidVariantName() const override
-	{
-		return TEXT("ETC2");
 	}
 
 	virtual bool SupportsTextureFormat(FName Format) const override
@@ -525,11 +443,6 @@ public:
 		return GLESVersion >= 0x30000;
 	}
 
-	virtual FText GetVariantDisplayName() const override
-	{
-		return LOCTEXT("Android_ETC2_ShortName", "ETC2");
-	}
-
 	virtual float GetVariantPriority() const override
 	{
 		float Priority;
@@ -544,9 +457,8 @@ class FAndroid_MultiTargetPlatform : public FAndroidTargetPlatform
 	FString FormatTargetString;
 
 public:
-	FAndroid_MultiTargetPlatform(bool bIsClient) : FAndroidTargetPlatform(bIsClient)
+	FAndroid_MultiTargetPlatform(bool bIsClient) : FAndroidTargetPlatform(bIsClient, TEXT("Multi"))
 	{
-		this->PlatformInfo = PlatformInfo::FindPlatformInfo("Android_Multi");
 	}
 
 	// set up all of the multiple formats together into this one
@@ -594,14 +506,10 @@ public:
 			}
 		}
 
+		// @todo do we need this with DisplayName below?
 		PlatformInfo::UpdatePlatformDisplayName(TEXT("Android_Multi"), DisplayName());
 	}
 
-	virtual FString GetAndroidVariantName() const override
-	{
-		return TEXT("Multi");
-	}
-
 	virtual FText DisplayName() const override
 	{
 		return FText::Format(LOCTEXT("Android_Multi", "Android (Multi:{0})"), FText::FromString(FormatTargetString));
@@ -637,11 +545,6 @@ public:
 	}
 #endif	
 
-	virtual FText GetVariantDisplayName() const override
-	{
-		return LOCTEXT("Android_Multi_ShortName", "Multi");
-	}
-
 	virtual float GetVariantPriority() const override
 	{
 		// lowest priority so specific variants are chosen first
diff --git a/Engine/Source/Developer/Android/AndroidTargetPlatform/Private/AndroidTargetPlatformModule.cpp b/Engine/Source/Developer/Android/AndroidTargetPlatform/Private/AndroidTargetPlatformModule.cpp
--- a/Engine/Source/Developer/Android/AndroidTargetPlatform/Private/AndroidTargetPlatformModule.cpp
+++ b/Engine/Source/Developer/Android/AndroidTargetPlatform/Private/AndroidTargetPlatformModule.cpp
@@ -27,7 +27,8 @@ public:
 			for (int32 Type = 0; Type < 2; Type++)
 			{
 				bool bIsClient = Type == 1;
-				SinglePlatforms.Add(new FAndroidTargetPlatform(bIsClient));
+				// flavorless must come first
+				SinglePlatforms.Add(new FAndroidTargetPlatform(bIsClient, nullptr));
 				SinglePlatforms.Add(new FAndroid_ASTCTargetPlatform(bIsClient));
 				SinglePlatforms.Add(new FAndroid_DXTTargetPlatform(bIsClient));
 				SinglePlatforms.Add(new FAndroid_ETC2TargetPlatform(bIsClient));
diff --git a/Engine/Source/Developer/AssetTools/Private/AssetViewUtils.cpp b/Engine/Source/Developer/AssetTools/Private/AssetViewUtils.cpp
--- a/Engine/Source/Developer/AssetTools/Private/AssetViewUtils.cpp
+++ b/Engine/Source/Developer/AssetTools/Private/AssetViewUtils.cpp
@@ -1334,12 +1334,13 @@ int32 AssetViewUtils::GetPackageLengthForCooking(const FString& PackageName, boo
 		GameNamePadded += TEXT(" ");
 	}
 
-	// We use "WindowsNoEditor" below as it's the longest platform name, so will also prove that any shorter platform names will validate correctly
+	// We use "LinuxAArch64Server" below as it's probably the longest platform name, so will also prove that any shorter platform names will validate correctly
+	const TCHAR* LongPlatformName = TEXT("LinuxAArch64Server");
 	const FString AbsoluteRootPath = FPaths::ConvertRelativePathToFull(FPaths::RootDir());
 	const FString AbsoluteGamePath = FPaths::ConvertRelativePathToFull(FPaths::ProjectDir());
 	const FString AbsoluteEnginePath = FPaths::ConvertRelativePathToFull(FPaths::EngineDir());
-	const FString AbsoluteEngineCookPath = AbsoluteGamePath / TEXT("Saved") / TEXT("Cooked") / TEXT("WindowsNoEditor") / TEXT("Engine");
-	const FString AbsoluteGameCookPath = AbsoluteGamePath / TEXT("Saved") / TEXT("Cooked") / TEXT("WindowsNoEditor") / GameName;
+	const FString AbsoluteEngineCookPath = AbsoluteGamePath / TEXT("Saved") / TEXT("Cooked") / LongPlatformName / TEXT("Engine");
+	const FString AbsoluteGameCookPath = AbsoluteGamePath / TEXT("Saved") / TEXT("Cooked") / LongPlatformName / GameName;
 
 	EPluginLoadedFrom PluginLoadedFrom;
 	const bool bIsPluginAsset = IsPluginFolder(PackageName, &PluginLoadedFrom);
@@ -1417,7 +1418,7 @@ bool AssetViewUtils::IsValidPackageForCooking(const FString& PackageName, FText&
 		// The BuildFarm uses a verbose path to encode extra information to provide more information when things fail, however this makes the path limitation a problem.
 		//	- We assume a base path of D:/BuildFarm/buildmachine_++depot+UE4-Releases+4.10/
 		//	- We assume the game name is 20 characters (the maximum allowed) to make sure that content can be ported between projects
-		//	- We calculate the cooked game path relative to the game root (eg, Showcases/Infiltrator/Saved/Cooked/WindowsNoEditor/Infiltrator)
+		//	- We calculate the cooked game path relative to the game root (eg, Showcases/Infiltrator/Saved/Cooked/Windows/Infiltrator)
 		//	- We calculate the asset path relative to (and including) the Content directory (eg, Content/Environment/Infil1/Infil1_Underground/Infrastructure/Model/SM_Infil1_Tunnel_Ceiling_Pipes_1xEntryCurveOuter_Double.uasset)
 		if (FEngineBuildSettings::IsInternalBuild())
 		{
diff --git a/Engine/Source/Developer/BlueprintNativeCodeGen/Private/BlueprintNativeCodeGenUtils.cpp b/Engine/Source/Developer/BlueprintNativeCodeGen/Private/BlueprintNativeCodeGenUtils.cpp
--- a/Engine/Source/Developer/BlueprintNativeCodeGen/Private/BlueprintNativeCodeGenUtils.cpp
+++ b/Engine/Source/Developer/BlueprintNativeCodeGen/Private/BlueprintNativeCodeGenUtils.cpp
@@ -125,7 +125,7 @@ static bool BlueprintNativeCodeGenUtilsImpl::GeneratePluginDescFile(const FBluep
 		const FName PlatformName = TargetPaths.GetTargetPlatformName();
 		for (const PlatformInfo::FTargetPlatformInfo* PlatformInfo : PlatformInfo::GetPlatformInfoArray())
 		{
-			if (PlatformInfo->TargetPlatformName == PlatformName)
+			if (PlatformInfo->Name == PlatformName)
 			{
 				// We use the 'UBTPlatformName' because this white-list expects the 
 				// string to correspond to UBT's UnrealTargetPlatform enum (and by proxy, FPlatformMisc::GetUBTPlatform)
diff --git a/Engine/Source/Developer/DesktopPlatform/Private/PlatformInfo.cpp b/Engine/Source/Developer/DesktopPlatform/Private/PlatformInfo.cpp
--- a/Engine/Source/Developer/DesktopPlatform/Private/PlatformInfo.cpp
+++ b/Engine/Source/Developer/DesktopPlatform/Private/PlatformInfo.cpp
@@ -17,179 +17,102 @@ namespace PlatformInfo
 	TArray<FName> AllPlatformGroupNames;
 	TArray<FName> AllVanillaPlatformNames;
 
-namespace
-{
-
-TArray<FTargetPlatformInfo*> AllPlatformInfoArray;
-TArray<FTargetPlatformInfo*> VanillaPlatformInfoArray;
-
-
-// we don't need any of this without the editor, although we would ideally not even compile this outside of the editor
-// @todo platplug: Figure out why this is compiled on target devices
-
-void BuildPlatformInfo(FName InPlatformInfoName, FName InTargetPlatformName, const FText& InDisplayName, const EBuildTargetType InPlatformType, 
-	const EPlatformFlags::Flags InPlatformFlags, FName InIniPlatformName, FName InUBTPlatformName)
-{
-	// some verification before we create the object
-	FName VanillaName = NAME_None;
-	FName FlavorName = NAME_None;
-	FTargetPlatformInfo* VanillaInfo = nullptr;
-	// See if this contains a flavor
-	if (InPlatformFlags & EPlatformFlags::CookFlavor)
+	namespace
 	{
-		FString InPlatformInfoNameString = InPlatformInfoName.ToString();
-		int32 UnderscoreLoc;
-		if (InPlatformInfoNameString.FindChar(TEXT('_'), UnderscoreLoc))
-		{
-			// removing VanillaPlatformName from the TP info means we need to verify, as it assumed the vanilla platform name is 
-			// if this assumption breaks, then we can store a private copy of this name to find the vanilla platform later
-			VanillaName = *InPlatformInfoNameString.Mid(0, UnderscoreLoc);
-			FlavorName = *InPlatformInfoNameString.Mid(UnderscoreLoc + 1);
-
-			FTargetPlatformInfo** FoundInfo = AllPlatformInfoArray.FindByPredicate([VanillaName](const FTargetPlatformInfo* Item) -> bool
-			{
-				return Item->PlatformInfoName == VanillaName;
-			});
 
-			VanillaInfo = FoundInfo ? *FoundInfo : nullptr;
-		}
+		TArray<FTargetPlatformInfo*> AllPlatformInfoArray;
+		TArray<FTargetPlatformInfo*> VanillaPlatformInfoArray;
 
-		// make sure it was good
-		if (VanillaInfo == nullptr)
-		{
-			UE_LOG(LogDesktopPlatform, Error, TEXT("TargetPlatformInfo %s is a flavor, but wasn't in the form Parent_Flavor (Parent needs to be specified before flavor in the .ini file)"), *InPlatformInfoName.ToString());
-			return;
-		}
 	}
 
-	// create the platform info
-	FTargetPlatformInfo* PlatformInfo = new FTargetPlatformInfo();
-	AllPlatformInfoArray.Add(PlatformInfo);
+FTargetPlatformInfo::FTargetPlatformInfo(const FString& InIniPlatformName, EBuildTargetType InType, const FString& InCookFlavor)
+{
+	IniPlatformName = *InIniPlatformName;
+	PlatformType = InType;
 
-	PlatformInfo->PlatformInfoName = InPlatformInfoName;
-	PlatformInfo->TargetPlatformName = InTargetPlatformName;
-	PlatformInfo->DisplayName = InDisplayName;
-	PlatformInfo->PlatformType = InPlatformType;
-	PlatformInfo->PlatformFlags = InPlatformFlags;
-	PlatformInfo->IniPlatformName = InIniPlatformName;
-	PlatformInfo->PlatformFlavor = FlavorName;
+	// calculate the name of the TargetPlatform
+	FString TPName = InIniPlatformName;
 
-	// add this object to either a parent, or the list of vanilla platforms
-	if (VanillaInfo != nullptr)
+	// by default we are vanilla, so point to ourself
+	VanillaInfo = this;
+	if (InCookFlavor != TEXT("") || PlatformType != EBuildTargetType::Game)
 	{
-		PlatformInfo->VanillaInfo = VanillaInfo;
-		VanillaInfo->Flavors.Add(PlatformInfo);
-	}
-	else
-	{
-		// if we are vanilla, then point to ourself (this way we can always just use VanillaInfo without checking for null)
-		PlatformInfo->VanillaInfo = PlatformInfo;
+		// if we are non-Game or cook flavor, then we need to find the most base version (will have same name as IniPlatform)
+		FTargetPlatformInfo** FoundInfo = AllPlatformInfoArray.FindByPredicate([this](const FTargetPlatformInfo* Item) -> bool
+ 		{
+ 			return Item->Name == IniPlatformName;
+ 		});
 
-		VanillaPlatformInfoArray.Add(PlatformInfo);
-		PlatformInfo::AllVanillaPlatformNames.AddUnique(InPlatformInfoName);
-	}
+		checkf(FoundInfo != nullptr, TEXT("Creating a TargetPlatform (%s, %s, %s) that needed a 'vanilla' TP already created, but it wasn't found. Create Game TPs first, and all cook flavors last"), *InIniPlatformName, LexToString(InType), &InCookFlavor);
 
-	// build up used platform group names
-	PlatformInfo->DataDrivenPlatformInfo = &FDataDrivenPlatformInfoRegistry::GetPlatformInfo(InIniPlatformName.ToString());
-	if (PlatformInfo->DataDrivenPlatformInfo->PlatformGroupName != NAME_None)
-	{
-		PlatformInfo::AllPlatformGroupNames.AddUnique(PlatformInfo->DataDrivenPlatformInfo->PlatformGroupName);
+		VanillaInfo = *FoundInfo;
 	}
-}
-
-// Gets a string from a section, or empty string if it didn't exist
-FString GetSectionString(const FConfigSection& Section, FName Key)
-{
-	// look for a value prefixed with host:
-	FName HostKey = *FString::Printf(TEXT("%s:%s"), ANSI_TO_TCHAR(FPlatformProperties::IniPlatformName()), *Key.ToString());
-	const FConfigValue* HostValue = Section.Find(HostKey);
-	return HostValue ? HostValue->GetValue() : Section.FindRef(Key).GetValue();
-}
 
-// Gets a bool from a section, or false if it didn't exist
-bool GetSectionBool(const FConfigSection& Section, FName Key)
-{
-	return FCString::ToBool(*GetSectionString(Section, Key));
-}
+	FString DisplayString = InIniPlatformName;
 
-EPlatformFlags::Flags ConvertPlatformFlags(const FString& String)
-{
-	if (String == TEXT("") || String == TEXT("None")) { return EPlatformFlags::None; }
-	if (String == TEXT("CookFlavor")) { return EPlatformFlags::CookFlavor; }
-	if (String == TEXT("BuildFlavor")) { return EPlatformFlags::BuildFlavor; }
+	// handle cook flavors
+	PlatformFlags = EPlatformFlags::None;
+	if (InCookFlavor != TEXT(""))
+	{
+		// mark us as a cookflavor
+		PlatformFlags = EPlatformFlags::CookFlavor;
+		PlatformFlavor = *InCookFlavor;
 
-	UE_LOG(LogInit, Fatal, TEXT("Unknown platform flag %s in PlatformInfo"), *String);
-	return EPlatformFlags::None;
-}
+		// append flavor with _
+		TPName += FString::Printf(TEXT("_%s"), *InCookFlavor);
+		
+		// put the flavor in parens
+		DisplayString += FString::Printf(TEXT(" (%s)"), *InCookFlavor);
 
-namespace EPlatformSection
-{
-	const FName TargetPlatformName = FName(TEXT("TargetPlatformName"));
-	const FName DisplayName = FName(TEXT("DisplayName"));
-	const FName PlatformType = FName(TEXT("PlatformType"));
-	const FName PlatformFlags = FName(TEXT("PlatformFlags"));
-	const FName UATCommandLine = FName(TEXT("UATCommandLine"));
-	const FName UBTPlatformName = FName(TEXT("UBTPlatformName"));
-}
+		// append UAT commandline with cook flavor
+		UATCommandLine += FString::Printf(TEXT(" -cookflavor=%s"), *InCookFlavor);
 
-void ParseDataDrivenPlatformInfo(const TCHAR* Name, const FConfigSection& Section, FName IniPlatformName)
-{
-	FName TargetPlatformName = *GetSectionString(Section, EPlatformSection::TargetPlatformName);
-	FString DisplayName = GetSectionString(Section, EPlatformSection::DisplayName);
-	FString PlatformType = GetSectionString(Section, EPlatformSection::PlatformType);
-	FString PlatformFlags = GetSectionString(Section, EPlatformSection::PlatformFlags);
-	FString UATCommandLine = GetSectionString(Section, EPlatformSection::UATCommandLine);
-	FString UBTTargetValue = GetSectionString(Section, EPlatformSection::UBTPlatformName);
-	// in almost all cases the UBT platform name matches the IniPlatformName, so it's optional
-	FName UBTPlatformName = UBTTargetValue.Len() == 0 ? IniPlatformName : FName(*UBTTargetValue);
-
-	EBuildTargetType TargetType;
-	if (!LexTryParseString(TargetType, *PlatformType))
-	{
-		TargetType = EBuildTargetType::Unknown;
+		VanillaInfo->Flavors.Add(this);
 	}
 
-	BuildPlatformInfo(Name, TargetPlatformName, FText::FromString(DisplayName), TargetType, ConvertPlatformFlags(PlatformFlags), IniPlatformName, UBTPlatformName);
-}
-
-void LoadDataDrivenPlatforms()
-{
-	// look for the standard DataDriven ini files
-	int32 NumDDInfoFiles = FDataDrivenPlatformInfoRegistry::GetNumDataDrivenIniFiles();
-	for (int32 Index = 0; Index < NumDDInfoFiles; Index++)
+	// now append the build type (game type has no type suffix, all others do)
+	if (PlatformType != EBuildTargetType::Game)
 	{
-		FConfigFile IniFile;
-		FString PlatformName;
-
-		FDataDrivenPlatformInfoRegistry::LoadDataDrivenIniFile(Index, IniFile, PlatformName);
+		TPName += LexToString(PlatformType);
 
-		FName PlatformFName = *PlatformName;
+		// put the type in parens
+		DisplayString += FString::Printf(TEXT(" (%s)"), LexToString(PlatformType));
 
-		// now walk over the file, looking for ShaderPlatformInfo sections
-		for (auto Section : IniFile)
+		// client builds need to be plopped on the commandline, servers are handled differently? 
+		// @todo: not sure if needed to be handled specially, honestly
+		if (PlatformType == EBuildTargetType::Client)
 		{
-			if (Section.Key.StartsWith(TEXT("PlatformInfo ")))
-			{
-				const FString& SectionName = Section.Key;
-				ParseDataDrivenPlatformInfo(*SectionName.Mid(13), Section.Value, PlatformFName);
-			}
+			UATCommandLine += TEXT(" -client");
 		}
 	}
-}
 
-FDelayedAutoRegisterHelper GPlatformInfoInit(EDelayedRegisterRunPhase::FileSystemReady, []()
-{
-	LoadDataDrivenPlatforms();
-});
+	// now we can store the final values
+	Name = *TPName;
+	DisplayName = FText::FromString(DisplayString);
+ 	DataDrivenPlatformInfo = &FDataDrivenPlatformInfoRegistry::GetPlatformInfo(InIniPlatformName);
+
+	// update various arrays
+	if (DataDrivenPlatformInfo->PlatformGroupName != NAME_None)
+	{
+		AllPlatformGroupNames.AddUnique(DataDrivenPlatformInfo->PlatformGroupName);
+	}
 
+	AllPlatformInfoArray.Add(this);
+	if (VanillaInfo == this)
+	{
+		VanillaPlatformInfoArray.Add(this);
+		AllVanillaPlatformNames.AddUnique(Name);
+	}
+}
 
-} // anonymous namespace
 
 const FTargetPlatformInfo* FindPlatformInfo(const FName& InPlatformName)
 {
+	checkf(AllPlatformInfoArray.Num() > 0, TEXT("Querying for TargetPlatformInfo objects before they are ready!"));
+
 	for(const FTargetPlatformInfo* PlatformInfo : AllPlatformInfoArray)
 	{
-		if(PlatformInfo->PlatformInfoName == InPlatformName)
+		if(PlatformInfo->Name == InPlatformName)
 		{
 			return PlatformInfo;
 		}
@@ -202,15 +125,19 @@ const FTargetPlatformInfo* FindPlatformInfo(const FName& InPlatformName)
 
 const FTargetPlatformInfo* FindVanillaPlatformInfo(const FName& InPlatformName)
 {
+	checkf(AllPlatformInfoArray.Num() > 0, TEXT("Querying for TargetPlatformInfo objects before they are ready!"));
+
 	const FTargetPlatformInfo* const FoundInfo = FindPlatformInfo(InPlatformName);
 	return FoundInfo ? FoundInfo->VanillaInfo : nullptr;
 }
 
 void UpdatePlatformDisplayName(FString InPlatformName, FText InDisplayName)
 {
+	checkf(AllPlatformInfoArray.Num() > 0, TEXT("Querying for TargetPlatformInfo objects before they are ready!"));
+
 	for (FTargetPlatformInfo* PlatformInfo : AllPlatformInfoArray)
 	{
-		if (PlatformInfo->TargetPlatformName == FName(*InPlatformName))
+		if (PlatformInfo->Name == FName(*InPlatformName))
 		{
 			PlatformInfo->DisplayName = InDisplayName;
 		}
@@ -219,21 +146,29 @@ void UpdatePlatformDisplayName(FString InPlatformName, FText InDisplayName)
 
 const TArray<FTargetPlatformInfo*>& GetPlatformInfoArray()
 {
+	checkf(AllPlatformInfoArray.Num() > 0, TEXT("Querying for TargetPlatformInfo objects before they are ready!"));
+
 	return AllPlatformInfoArray;
 }
 
 const TArray<FTargetPlatformInfo*>& GetVanillaPlatformInfoArray()
 {
+	checkf(AllPlatformInfoArray.Num() > 0, TEXT("Querying for TargetPlatformInfo objects before they are ready!"));
+
 	return VanillaPlatformInfoArray;
 }
 
 const TArray<FName>& GetAllPlatformGroupNames()
 {
+	checkf(AllPlatformInfoArray.Num() > 0, TEXT("Querying for TargetPlatformInfo objects before they are ready!"));
+
 	return PlatformInfo::AllPlatformGroupNames;
 }
 
 const TArray<FName>& GetAllVanillaPlatformNames()
 {
+	checkf(AllPlatformInfoArray.Num() > 0, TEXT("Querying for TargetPlatformInfo objects before they are ready!"));
+
 	return PlatformInfo::AllVanillaPlatformNames;
 }
 
diff --git a/Engine/Source/Developer/DesktopPlatform/Public/PlatformInfo.h b/Engine/Source/Developer/DesktopPlatform/Public/PlatformInfo.h
--- a/Engine/Source/Developer/DesktopPlatform/Public/PlatformInfo.h
+++ b/Engine/Source/Developer/DesktopPlatform/Public/PlatformInfo.h
@@ -43,11 +43,10 @@ namespace PlatformInfo
 	/** Information about a given platform */
 	struct FTargetPlatformInfo
 	{
-		/** Name used to identify this platform, eg "Android_ETC2" */
-		FName PlatformInfoName;
+		DESKTOPPLATFORM_API FTargetPlatformInfo(const FString& InIniPlatformName, EBuildTargetType InType, const FString& InCookFlavor);
 
-		/** Name used to find the corresponding ITargetPlatform for this platform (also used by UAT) */
-		FName TargetPlatformName;
+		/** Name of the Info object as well as the ITargetPlatform that this Info describes */
+		FName Name;
 
 		/** Platform flavor, eg "ETC2" for "Android_ETC2" */
 		FName PlatformFlavor;
diff --git a/Engine/Source/Developer/IOS/IOSPlatformEditor/Private/IOSCustomIconProjectBuildMutatorFeature.cpp b/Engine/Source/Developer/IOS/IOSPlatformEditor/Private/IOSCustomIconProjectBuildMutatorFeature.cpp
--- a/Engine/Source/Developer/IOS/IOSPlatformEditor/Private/IOSCustomIconProjectBuildMutatorFeature.cpp
+++ b/Engine/Source/Developer/IOS/IOSPlatformEditor/Private/IOSCustomIconProjectBuildMutatorFeature.cpp
@@ -52,7 +52,7 @@ bool FIOSCustomIconProjectBuildMutatorFeature ::RequiresProjectBuild(const FName
 
 	if (PlatInfo->DataDrivenPlatformInfo->GetSdkStatus() == DDPIPlatformSdkStatus::Valid)
 	{
-		const ITargetPlatform* const Platform = GetTargetPlatformManager()->FindTargetPlatform(PlatInfo->TargetPlatformName.ToString());
+		const ITargetPlatform* const Platform = GetTargetPlatformManager()->FindTargetPlatform(PlatInfo->Name.ToString());
 		if (Platform)
 		{
 			if (InPlatformInfoName.ToString() == TEXT("IOS"))
diff --git a/Engine/Source/Developer/IOS/IOSTargetPlatform/Private/IOSTargetPlatform.cpp b/Engine/Source/Developer/IOS/IOSTargetPlatform/Private/IOSTargetPlatform.cpp
--- a/Engine/Source/Developer/IOS/IOSTargetPlatform/Private/IOSTargetPlatform.cpp
+++ b/Engine/Source/Developer/IOS/IOSTargetPlatform/Private/IOSTargetPlatform.cpp
@@ -25,24 +25,17 @@
 /* FIOSTargetPlatform structors
  *****************************************************************************/
 
-FIOSTargetPlatform::FIOSTargetPlatform(bool bInIsTVOS, bool bInIsClientOnly)
-	: bIsTVOS(bInIsTVOS)
-	, bIsClientOnly(bInIsClientOnly)
+FIOSTargetPlatform::FIOSTargetPlatform(bool bInIsTVOS, bool bIsClientOnly)
+	// override the ini name up in the base classes, which will go into the FTargetPlatformInfo
+	: TNonDesktopTargetPlatformBase(bIsClientOnly, nullptr, bInIsTVOS ? TEXT("TVOS") : nullptr)
+	, bIsTVOS(bInIsTVOS)
 {
-    if (bIsTVOS)
-    {
-        this->PlatformInfo = PlatformInfo::FindPlatformInfo("TVOS");
-    }
 #if WITH_ENGINE
 	FConfigCacheIni::LoadLocalIniFile(EngineSettings, TEXT("Engine"), true, *PlatformName());
 	TextureLODSettings = nullptr; // TextureLODSettings are registered by the device profile.
 	StaticMeshLODSettings.Initialize(EngineSettings);
 #endif // #if WITH_ENGINE
 
-	// Initialize Ticker for device discovery
-	TickDelegate = FTickerDelegate::CreateRaw(this, &FIOSTargetPlatform::HandleTicker);
-	TickDelegateHandle = FTicker::GetCoreTicker().AddTicker(TickDelegate, 10.0f);
-	
 	// initialize the connected device detector
 	DeviceHelper.OnDeviceConnected().AddRaw(this, &FIOSTargetPlatform::HandleDeviceConnected);
 	DeviceHelper.OnDeviceDisconnected().AddRaw(this, &FIOSTargetPlatform::HandleDeviceDisconnected);
@@ -52,7 +45,6 @@ FIOSTargetPlatform::FIOSTargetPlatform(bool bInIsTVOS, bool bInIsClientOnly)
 
 FIOSTargetPlatform::~FIOSTargetPlatform()
 {
-	FTicker::GetCoreTicker().RemoveTicker(TickDelegateHandle);
 }
 
 
@@ -324,41 +316,6 @@ int32 FIOSTargetPlatform::CheckRequirements(bool bProjectHasCode, EBuildConfigur
 }
 
 
-/* FIOSTargetPlatform implementation
- *****************************************************************************/
-
-void FIOSTargetPlatform::PingNetworkDevices()
-{
-    //Only put this here in case we put in auto-detection of missing stats
-    QUICK_SCOPE_CYCLE_COUNTER(STAT_FIOSTargetPlatform_PingNetworkDevices);
-
-	// disabled for now because we find IOS devices from the USB, this is a relic from ULD, but it may be needed in the future
-/*	if (!MessageEndpoint.IsValid())
-	{
-		MessageEndpoint = FMessageEndpoint::Builder("FIOSTargetPlatform")
-			.Handling<FIOSLaunchDaemonPong>(this, &FIOSTargetPlatform::HandlePongMessage);
-	}
-
-	if (MessageEndpoint.IsValid())
-	{
-		MessageEndpoint->Publish(new FIOSLaunchDaemonPing(), EMessageScope::Network);
-	}
-
-	// remove disconnected & timed out devices
-	FDateTime Now = FDateTime::UtcNow();
-
-	for (auto DeviceIt = Devices.CreateIterator(); DeviceIt; ++DeviceIt)
-	{
-		FIOSTargetDevicePtr Device = DeviceIt->Value;
-
-		if (Now > Device->LastPinged + FTimespan::FromSeconds(60.0))
-		{
-			DeviceIt.RemoveCurrent();
-			DeviceLostEvent.Broadcast(Device.ToSharedRef());
-		}
-	}*/
-}
-
 
 /* FIOSTargetPlatform callbacks
  *****************************************************************************/
@@ -437,13 +394,6 @@ void FIOSTargetPlatform::HandleDeviceDisconnected(const FIOSLaunchDaemonPong& Me
 	}
 }
 
-bool FIOSTargetPlatform::HandleTicker(float DeltaTime)
-{
-	PingNetworkDevices();
-
-	return true;
-}
-
 
 /* ITargetPlatform interface
  *****************************************************************************/
diff --git a/Engine/Source/Developer/IOS/IOSTargetPlatform/Private/IOSTargetPlatform.h b/Engine/Source/Developer/IOS/IOSTargetPlatform/Private/IOSTargetPlatform.h
--- a/Engine/Source/Developer/IOS/IOSTargetPlatform/Private/IOSTargetPlatform.h
+++ b/Engine/Source/Developer/IOS/IOSTargetPlatform/Private/IOSTargetPlatform.h
@@ -25,7 +25,7 @@
 /**
  * FIOSTargetPlatform, abstraction for cooking iOS platforms
  */
-class FIOSTargetPlatform : public TTargetPlatformBase<FIOSPlatformProperties>
+class FIOSTargetPlatform : public TNonDesktopTargetPlatformBase<FIOSPlatformProperties>
 {
 public:
 
@@ -41,39 +41,6 @@ public:
 
 public:
 
-	//~ Begin TTargetPlatformBase Interface
-
-	virtual bool IsServerOnly( ) const override
-	{
-		return false;
-	}
-
-	virtual bool IsClientOnly() const override
-	{
-		return bIsClientOnly;
-	}
-
-	//~ End TTargetPlatformBase Interface
-
-public:
-
-	//~ Begin ITargetPlatform Interface
-	
-	// this is used for cooking to a separate directory, NOT for runtime. Runtime TVOS is still "IOS"
-	virtual FString PlatformName() const override
-	{
-		if (bIsTVOS)
-		{
-			return bIsClientOnly ? "TVOSClient" : "TVOS";
-		}
-		return bIsClientOnly ? "IOSClient" : "IOS";
-	}
-
-    virtual FString IniPlatformName() const override
-    {
-        return "IOS";
-    }
-    
 	virtual void EnableDeviceCheck(bool OnOff) override;
 
 	virtual void GetAllDevices( TArray<ITargetDevicePtr>& OutDevices ) const override;
@@ -86,16 +53,7 @@ public:
 	virtual ITargetDevicePtr GetDefaultDevice( ) const override;
 
 	virtual ITargetDevicePtr GetDevice( const FTargetDeviceId& DeviceId ) override;
-
-	virtual bool IsRunningPlatform( ) const override
-	{
-		#if PLATFORM_IOS && WITH_EDITOR
-			return true;
-		#else
-			return false;
-		#endif
-	}
-
+		
 	virtual bool SupportsFeature( ETargetPlatformFeatures Feature ) const override;
 
 	virtual bool CanSupportXGEShaderCompile() const override;
@@ -153,18 +111,8 @@ public:
 
 	//~ Begin ITargetPlatform Interface
 
-protected:
-
-	/**
-	 * Sends a ping message over the network to find devices running the launch daemon.
-	 */
-	void PingNetworkDevices( );
-
 private:
 
-	// Handles when the ticker fires.
-	bool HandleTicker( float DeltaTime );
-
 	// Handles received pong messages from the LauncherDaemon.
 	void HandlePongMessage( const FIOSLaunchDaemonPong& Message, const TSharedRef<IMessageContext, ESPMode::ThreadSafe>& Context );
 
@@ -176,18 +124,9 @@ private:
 	// true if this is targeting TVOS vs IOS
 	bool bIsTVOS;
 
-	// true if this is a client-only TP
-	bool bIsClientOnly;
-
 	// Contains all discovered IOSTargetDevices over the network.
 	TMap<FTargetDeviceId, FIOSTargetDevicePtr> Devices;
 
-	// Holds a delegate to be invoked when the widget ticks.
-	FTickerDelegate TickDelegate;
-
-	// Handle to the registered TickDelegate.
-	FDelegateHandle TickDelegateHandle;
-
 	// Holds the message endpoint used for communicating with the LaunchDaemon.
 	TSharedPtr<FMessageEndpoint, ESPMode::ThreadSafe> MessageEndpoint;
 
diff --git a/Engine/Source/Developer/LauncherServices/Private/Profiles/LauncherProfile.h b/Engine/Source/Developer/LauncherServices/Private/Profiles/LauncherProfile.h
--- a/Engine/Source/Developer/LauncherServices/Private/Profiles/LauncherProfile.h
+++ b/Engine/Source/Developer/LauncherServices/Private/Profiles/LauncherProfile.h
@@ -1648,48 +1648,13 @@ public:
 				return false;
 			}
 
-			// switch server and no editor platforms to the proper type
-			if (PlatformInfo->TargetPlatformName == FName("LinuxServer"))
-			{
-				ServerPlatforms.Add(TEXT("Linux"));
-			}
-			else if (PlatformInfo->TargetPlatformName == FName("WindowsServer"))
-			{
-				ServerPlatforms.Add(TEXT("Win64"));
-			}
-			else if (PlatformInfo->TargetPlatformName == FName("MacServer"))
-			{
-				ServerPlatforms.Add(TEXT("Mac"));
-			}
-			else if (PlatformInfo->TargetPlatformName == FName("LinuxNoEditor") || PlatformInfo->TargetPlatformName == FName("LinuxClient"))
-			{
-				ClientPlatforms.Add(TEXT("Linux"));
-			}
-			else if (PlatformInfo->TargetPlatformName == FName("LinuxAArch64NoEditor") || PlatformInfo->TargetPlatformName == FName("LinuxAArch64Client"))
-			{
-				ClientPlatforms.Add(TEXT("LinuxAArch64"));
-			}
-			else if (PlatformInfo->TargetPlatformName == FName("LinuxAArch64Server"))
-			{
-				ServerPlatforms.Add(TEXT("LinuxAArch64"));
-			}
-			else if (PlatformInfo->TargetPlatformName == FName("WindowsNoEditor") || PlatformInfo->TargetPlatformName == FName("Windows"))
-			{
-				ClientPlatforms.Add(TEXT("Win64"));
-			}
-			else if (PlatformInfo->TargetPlatformName == FName("MacNoEditor"))
-			{
-				ClientPlatforms.Add(TEXT("Mac"));
-			}
-			else
-			{
-				ClientPlatforms.Add(PlatformInfo->TargetPlatformName.ToString());
-			}
+			// add the UBT platform name to the appropriate list
+			TArray<FString>& Platforms = PlatformInfo->PlatformType == EBuildTargetType::Server ? ServerPlatforms : ClientPlatforms;
+			Platforms.Add(PlatformInfo->DataDrivenPlatformInfo->UBTPlatformName.ToString());
 
 			// Append any extra UAT flags specified for this platform flavor
 			if (!PlatformInfo->UATCommandLine.IsEmpty())
 			{
-				OptionalParams += TEXT(" ");
 				OptionalParams += PlatformInfo->UATCommandLine;
 			}
 
diff --git a/Engine/Source/Developer/Linux/LinuxAArch64TargetPlatform/Private/LinuxAArch64TargetPlatformModule.cpp b/Engine/Source/Developer/Linux/LinuxAArch64TargetPlatform/Private/LinuxAArch64TargetPlatformModule.cpp
--- a/Engine/Source/Developer/Linux/LinuxAArch64TargetPlatform/Private/LinuxAArch64TargetPlatformModule.cpp
+++ b/Engine/Source/Developer/Linux/LinuxAArch64TargetPlatform/Private/LinuxAArch64TargetPlatformModule.cpp
@@ -23,21 +23,12 @@ public:
 
 	virtual void GetTargetPlatforms(TArray<ITargetPlatform*>& TargetPlatforms) override
 	{
-		// Server TP
-		if (TLinuxTargetPlatform<FLinuxPlatformProperties<false, true, false, true> >::IsUsable())
-		{
-			TargetPlatforms.Add(new TLinuxTargetPlatform<FLinuxPlatformProperties<false, true, false, true> >());
-		}
 		// Game TP
-		if (TLinuxTargetPlatform<FLinuxPlatformProperties<false, false, false, true> >::IsUsable())
-		{
-			TargetPlatforms.Add(new TLinuxTargetPlatform<FLinuxPlatformProperties<false, false, false, true> >());
-		}
+		TargetPlatforms.Add(new TLinuxTargetPlatform<FLinuxPlatformProperties<false, false, false, true> >());
+		// Server TP
+		TargetPlatforms.Add(new TLinuxTargetPlatform<FLinuxPlatformProperties<false, true, false, true> >());
 		// Client TP
-		if (TLinuxTargetPlatform<FLinuxPlatformProperties<false, false, true, true> >::IsUsable())
-		{
-			TargetPlatforms.Add(new TLinuxTargetPlatform<FLinuxPlatformProperties<false, false, true, true> >());
-		}
+		TargetPlatforms.Add(new TLinuxTargetPlatform<FLinuxPlatformProperties<false, false, true, true> >());
 	}
 };
 
diff --git a/Engine/Source/Developer/Linux/LinuxTargetPlatform/Private/LinuxTargetPlatform.h b/Engine/Source/Developer/Linux/LinuxTargetPlatform/Private/LinuxTargetPlatform.h
--- a/Engine/Source/Developer/Linux/LinuxTargetPlatform/Private/LinuxTargetPlatform.h
+++ b/Engine/Source/Developer/Linux/LinuxTargetPlatform/Private/LinuxTargetPlatform.h
@@ -358,31 +358,6 @@ public:
 		return true;
 	}
 
-	virtual FText GetVariantDisplayName() const override
-	{
-		if (TProperties::IsServerOnly())
-		{
-			return LOCTEXT("LinuxServerVariantTitle", "Dedicated Server");
-		}
-
-		if (TProperties::HasEditorOnlyData())
-		{
-			return LOCTEXT("LinuxClientEditorDataVariantTitle", "Client with Editor Data");
-		}
-
-		if (TProperties::IsClientOnly())
-		{
-			return LOCTEXT("LinuxClientOnlyVariantTitle", "Client only");
-		}
-
-		return LOCTEXT("LinuxClientVariantTitle", "Client");
-	}
-
-	virtual FText GetVariantTitle() const override
-	{
-		return LOCTEXT("LinuxVariantTitle", "Build Type");
-	}
-
 	virtual float GetVariantPriority() const override
 	{
 		return TProperties::GetVariantPriority();
diff --git a/Engine/Source/Developer/Linux/LinuxTargetPlatform/Private/LinuxTargetPlatformModule.cpp b/Engine/Source/Developer/Linux/LinuxTargetPlatform/Private/LinuxTargetPlatformModule.cpp
--- a/Engine/Source/Developer/Linux/LinuxTargetPlatform/Private/LinuxTargetPlatformModule.cpp
+++ b/Engine/Source/Developer/Linux/LinuxTargetPlatform/Private/LinuxTargetPlatformModule.cpp
@@ -27,26 +27,14 @@ class FLinuxTargetPlatformModule
 public:
 	virtual void GetTargetPlatforms(TArray<ITargetPlatform*>& TargetPlatforms) override
 	{
+		// NoEditor TP
+		TargetPlatforms.Add(new TLinuxTargetPlatform<FLinuxPlatformProperties<false, false, false, false> >());
 		// Editor TP
-		if (TLinuxTargetPlatform<FLinuxPlatformProperties<true, false, false, false> >::IsUsable())
-		{
-			TargetPlatforms.Add(new TLinuxTargetPlatform<FLinuxPlatformProperties<true, false, false, false> >());
-		}
+		TargetPlatforms.Add(new TLinuxTargetPlatform<FLinuxPlatformProperties<true, false, false, false> >());
 		// Server TP
-		if (TLinuxTargetPlatform<FLinuxPlatformProperties<false, true, false, false> >::IsUsable())
-		{
-			TargetPlatforms.Add(new TLinuxTargetPlatform<FLinuxPlatformProperties<false, true, false, false> >());
-		}
-		// NoEditor TP
-		if (TLinuxTargetPlatform<FLinuxPlatformProperties<false, false, false, false> >::IsUsable())
-		{
-			TargetPlatforms.Add(new TLinuxTargetPlatform<FLinuxPlatformProperties<false, false, false, false> >());
-		}
+		TargetPlatforms.Add(new TLinuxTargetPlatform<FLinuxPlatformProperties<false, true, false, false> >());
 		// Client TP
-		if (TLinuxTargetPlatform<FLinuxPlatformProperties<false, false, true, false> >::IsUsable())
-		{
-			TargetPlatforms.Add(new TLinuxTargetPlatform<FLinuxPlatformProperties<false, false, true, false> >());
-		}
+		TargetPlatforms.Add(new TLinuxTargetPlatform<FLinuxPlatformProperties<false, false, true, false> >());
 	}
 
 public:
diff --git a/Engine/Source/Developer/Localization/Private/LocTextHelper.cpp b/Engine/Source/Developer/Localization/Private/LocTextHelper.cpp
--- a/Engine/Source/Developer/Localization/Private/LocTextHelper.cpp
+++ b/Engine/Source/Developer/Localization/Private/LocTextHelper.cpp
@@ -31,7 +31,11 @@ const TArray<FString>& FLocTextPlatformSplitUtils::GetPlatformsToSplit(const ELo
 		{
 			static const TArray<FString> ConfidentialPlatformNames = []()
 			{
-				TArray<FString> TmpArray = FDataDrivenPlatformInfoRegistry::GetConfidentialPlatforms();
+				TArray<FString> TmpArray;
+				for (FName Name : FDataDrivenPlatformInfoRegistry::GetConfidentialPlatforms())
+				{
+					TmpArray.Add(Name.ToString());
+				}
 				TmpArray.Sort();
 				return TmpArray;
 			}();
@@ -45,8 +49,10 @@ const TArray<FString>& FLocTextPlatformSplitUtils::GetPlatformsToSplit(const ELo
 			{
 				// the Keys of the FDataDrivenPlatformInfoRegistry platforms are the known ini platform names
 				TArray<FString> TmpArray;
-				FDataDrivenPlatformInfoRegistry::GetAllPlatformInfos().GetKeys(TmpArray);
-				TmpArray.Sort();
+				for (FName Name : FDataDrivenPlatformInfoRegistry::GetSortedPlatformNames())
+				{
+					TmpArray.Add(Name.ToString());
+				}
 				return TmpArray;
 			}();
 			return AllPlatformNames;
diff --git a/Engine/Source/Developer/Lumin/LuminTargetPlatform/Private/LuminTargetPlatform.cpp b/Engine/Source/Developer/Lumin/LuminTargetPlatform/Private/LuminTargetPlatform.cpp
--- a/Engine/Source/Developer/Lumin/LuminTargetPlatform/Private/LuminTargetPlatform.cpp
+++ b/Engine/Source/Developer/Lumin/LuminTargetPlatform/Private/LuminTargetPlatform.cpp
@@ -23,12 +23,10 @@
  *****************************************************************************/
 
 FLuminTargetPlatform::FLuminTargetPlatform(bool bIsClient)
-	: FAndroidTargetPlatform(bIsClient)
+	// override the IniPlatformName up on the FTargetPlatformInfo
+	: FAndroidTargetPlatform(bIsClient, nullptr, TEXT("Lumin"))
 {
 #if WITH_ENGINE
-	// by using the FAndroidPlatformProperties, the PlatformInfo up in TTargetPlatformBase/FTargetPlatformBase would be Android
-	this->PlatformInfo = PlatformInfo::FindPlatformInfo(FName(*PlatformName()));
-
 	RefreshSettings();
 #endif
 }
diff --git a/Engine/Source/Developer/Lumin/LuminTargetPlatform/Private/LuminTargetPlatform.h b/Engine/Source/Developer/Lumin/LuminTargetPlatform/Private/LuminTargetPlatform.h
--- a/Engine/Source/Developer/Lumin/LuminTargetPlatform/Private/LuminTargetPlatform.h
+++ b/Engine/Source/Developer/Lumin/LuminTargetPlatform/Private/LuminTargetPlatform.h
@@ -47,17 +47,6 @@ public:
 
 public:
 
-	// Begin ITargetPlatform interface
-
-	virtual FString PlatformName() const override
-	{
-		return IsClientOnly() ? TEXT("LuminClient") : TEXT("Lumin");
-	}
-
-	virtual FString IniPlatformName() const override
-	{
-		return TEXT("Lumin");
-	}
 
 	virtual bool IsSdkInstalled(bool bProjectHasCode, FString& OutDocumentationPath) const override;
 
diff --git a/Engine/Source/Developer/Mac/MacTargetPlatform/Private/GenericMacTargetPlatform.h b/Engine/Source/Developer/Mac/MacTargetPlatform/Private/GenericMacTargetPlatform.h
--- a/Engine/Source/Developer/Mac/MacTargetPlatform/Private/GenericMacTargetPlatform.h
+++ b/Engine/Source/Developer/Mac/MacTargetPlatform/Private/GenericMacTargetPlatform.h
@@ -353,33 +353,6 @@ return TSuper::SupportsFeature(Feature);
 	}
 
 
-	virtual FText GetVariantDisplayName() const override
-	{
-		if (IS_DEDICATED_SERVER)
-		{
-			return LOCTEXT("MacServerVariantTitle", "Dedicated Server");
-		}
-
-		if (HAS_EDITOR_DATA)
-		{
-			return LOCTEXT("MacClientEditorDataVariantTitle", "Client with Editor Data");
-		}
-
-		if (IS_CLIENT_ONLY)
-		{
-			return LOCTEXT("MacClientOnlyVariantTitle", "Client only");
-		}
-
-		return LOCTEXT("MacClientVariantTitle", "Client");
-	}
-
-
-	virtual FText GetVariantTitle() const override
-	{
-		return LOCTEXT("MacVariantTitle", "Build Type");
-	}
-
-
 	virtual float GetVariantPriority() const override
 	{
 		return TProperties::GetVariantPriority();
diff --git a/Engine/Source/Developer/Mac/MacTargetPlatform/Private/MacTargetPlatformModule.cpp b/Engine/Source/Developer/Mac/MacTargetPlatform/Private/MacTargetPlatformModule.cpp
--- a/Engine/Source/Developer/Mac/MacTargetPlatform/Private/MacTargetPlatformModule.cpp
+++ b/Engine/Source/Developer/Mac/MacTargetPlatform/Private/MacTargetPlatformModule.cpp
@@ -24,26 +24,14 @@ public:
 
 	virtual void GetTargetPlatforms(TArray<ITargetPlatform*>& TargetPlatforms) override
 	{
+		// Game TP
+		TargetPlatforms.Add(new TGenericMacTargetPlatform<false, false, false>());
 		// Editor TP
-		if (TGenericMacTargetPlatform<true, false, false>::IsUsable())
-		{
-			TargetPlatforms.Add(new TGenericMacTargetPlatform<true, false, false>());
-		}
+		TargetPlatforms.Add(new TGenericMacTargetPlatform<true, false, false>());
 		// Server TP
-		if (TGenericMacTargetPlatform<false, true, false>::IsUsable())
-		{
-			TargetPlatforms.Add(new TGenericMacTargetPlatform<false, true, false>());
-		}
-		// NoEditor TP
-		if (TGenericMacTargetPlatform<false, false, false>::IsUsable())
-		{
-			TargetPlatforms.Add(new TGenericMacTargetPlatform<false, false, false>());
-		}
+		TargetPlatforms.Add(new TGenericMacTargetPlatform<false, true, false>());
 		// Client TP
-		if (TGenericMacTargetPlatform<false, false, true>::IsUsable())
-		{
-			TargetPlatforms.Add(new TGenericMacTargetPlatform<false, false, true>());
-		}
+		TargetPlatforms.Add(new TGenericMacTargetPlatform<false, false, true>());
 	}
 
 
diff --git a/Engine/Source/Developer/TargetDeviceServices/Private/Services/TargetDeviceService.cpp b/Engine/Source/Developer/TargetDeviceServices/Private/Services/TargetDeviceService.cpp
--- a/Engine/Source/Developer/TargetDeviceServices/Private/Services/TargetDeviceService.cpp
+++ b/Engine/Source/Developer/TargetDeviceServices/Private/Services/TargetDeviceService.cpp
@@ -90,7 +90,7 @@ void FTargetDeviceService::AddTargetDevice(TSharedPtr<ITargetDevice, ESPMode::Th
 		const PlatformInfo::FTargetPlatformInfo& Info = InDevice->GetTargetPlatform().GetTargetPlatformInfo();
 		const PlatformInfo::FTargetPlatformInfo* VanillaInfo = Info.VanillaInfo;
 
-		DevicePlatformName = Info.PlatformInfoName;
+		DevicePlatformName = Info.Name;
 		DevicePlatformDisplayName = VanillaInfo->DisplayName.ToString();
 		
 		// Sigh the hacks... Should be able to remove if platform info gets cleaned up.... Windows doesn't have a reasonable vanilla platform.
@@ -518,7 +518,7 @@ void FTargetDeviceService::HandlePingMessage(const FTargetDeviceServicePing& InM
 
 		// Check if we should also create an aggregate (All_<platform>_devices_on_<host>) proxy
 		Message->Aggregated = DefaultDevice->IsPlatformAggregated();
-		Message->AllDevicesName = DefaultDevice->GetAllDevicesName().IsEmpty() ? VanillaInfo->PlatformInfoName.ToString() : DefaultDevice->GetAllDevicesName();
+		Message->AllDevicesName = DefaultDevice->GetAllDevicesName().IsEmpty() ? VanillaInfo->Name.ToString() : DefaultDevice->GetAllDevicesName();
 		Message->AllDevicesDefaultVariant = DefaultDevice->GetAllDevicesDefaultVariant().IsNone() ? Message->DefaultVariant : DefaultDevice->GetAllDevicesDefaultVariant();
 
 		// Add the data for all the flavors
@@ -535,8 +535,8 @@ void FTargetDeviceService::HandlePingMessage(const FTargetDeviceServicePing& InM
 			Variant.DeviceID = TargetDevice->GetId().ToString();
 			Variant.VariantName = TargetDeviceIt.Key();
 			Variant.TargetPlatformName = TargetDevice->GetTargetPlatform().PlatformName();
-			Variant.TargetPlatformId = Info.TargetPlatformName;
-			Variant.VanillaPlatformId = Info.VanillaInfo->PlatformInfoName;
+			Variant.TargetPlatformId = Info.Name;
+			Variant.VanillaPlatformId = Info.VanillaInfo->Name;
 			Variant.PlatformDisplayName = Info.DisplayName.ToString();
 		}
 
diff --git a/Engine/Source/Developer/TargetPlatform/Private/TargetPlatformBase.cpp b/Engine/Source/Developer/TargetPlatform/Private/TargetPlatformBase.cpp
--- a/Engine/Source/Developer/TargetPlatform/Private/TargetPlatformBase.cpp
+++ b/Engine/Source/Developer/TargetPlatform/Private/TargetPlatformBase.cpp
@@ -184,7 +184,7 @@ bool FTargetPlatformBase::RequiresTempTarget(bool bProjectHasCode, EBuildConfigu
 	// check to see if any projectmutator modular features are available
 	for (IProjectBuildMutatorFeature* Feature : IModularFeatures::Get().GetModularFeatureImplementations<IProjectBuildMutatorFeature>(PROJECT_BUILD_MUTATOR_FEATURE))
 	{
-		if (Feature->RequiresProjectBuild(PlatformInfo->PlatformInfoName, OutReason))
+		if (Feature->RequiresProjectBuild(PlatformInfo->Name, OutReason))
 		{
 			return true;
 		}
diff --git a/Engine/Source/Developer/TargetPlatform/Private/TargetPlatformManagerModule.cpp b/Engine/Source/Developer/TargetPlatform/Private/TargetPlatformManagerModule.cpp
--- a/Engine/Source/Developer/TargetPlatform/Private/TargetPlatformManagerModule.cpp
+++ b/Engine/Source/Developer/TargetPlatform/Private/TargetPlatformManagerModule.cpp
@@ -197,6 +197,23 @@ public:
 		return nullptr;
 	}
 
+	virtual ITargetPlatform* FindTargetPlatform(FName Name) override
+	{
+		GetTargetPlatforms(); // Populates PlatformsByName
+
+		if (ITargetPlatform** Platform = PlatformsByName.Find(Name))
+		{
+			return *Platform;
+		}
+
+		return nullptr;
+	}
+
+	virtual ITargetPlatform* FindTargetPlatform(const TCHAR* Name) override
+	{
+		return FindTargetPlatform(FName(Name));
+	}
+
 	virtual ITargetPlatform* FindTargetPlatformWithSupport(FName SupportType, FName RequiredSupportedValue)
 	{
 		const TArray<ITargetPlatform*>& TargetPlatforms = GetTargetPlatforms();
@@ -640,7 +657,7 @@ protected:
 		return false;
 	}
 
-	bool InitializeSinglePlatform(const FString& PlatformName, const FDataDrivenPlatformInfo& Info)
+	bool InitializeSinglePlatform(FName PlatformName, const FDataDrivenPlatformInfo& Info)
 	{
 		// disabled?
 		if (Info.bEnabledForUse == false)
@@ -650,7 +667,7 @@ protected:
 
 		// there are two ways targetplatform modules are setup: a single DLL per TargetPlatform, or a DLL for the platform
 		// that returns multiple TargetPlatforms. we try single first, then full platform
-		FName PlatformModuleName = *(PlatformName + TEXT("TargetPlatform"));
+		FName PlatformModuleName = *(PlatformName.ToString() + TEXT("TargetPlatform"));
 
 		ITargetPlatformModule* Module = nullptr;
 
@@ -691,7 +708,7 @@ protected:
 						goto RETRY_SETUPANDVALIDATE;
 					}
 				}
-				UE_LOG(LogTargetPlatformManager, Display, TEXT("Failed to SetupAndValidateAutoSDK for platform '%s'"), *PlatformName);
+				UE_LOG(LogTargetPlatformManager, Display, TEXT("Failed to SetupAndValidateAutoSDK for platform '%s'"), *PlatformName.ToString());
 			}
 		}
 
@@ -720,12 +737,10 @@ protected:
 
 		// find a set of valid target platform names (the platform DataDrivenPlatformInfo.ini file was found indicates support for the platform 
 		// exists on disk, so the TP is expected to work)
-		const TArray<PlatformInfo::FTargetPlatformInfo*>& PlatformInfos = PlatformInfo::GetPlatformInfoArray();
-
-		FScopedSlowTask SlowTask(PlatformInfos.Num());
+		FScopedSlowTask SlowTask(FDataDrivenPlatformInfoRegistry::GetAllPlatformInfos().Num());
 		for (auto Pair : FDataDrivenPlatformInfoRegistry::GetAllPlatformInfos())
 		{
-			const FString& PlatformName = Pair.Key;
+			FName PlatformName = Pair.Key;
 			const FDataDrivenPlatformInfo& Info = Pair.Value;
 
 			SlowTask.EnterProgressFrame(1);
@@ -1017,9 +1032,9 @@ protected:
 				FString PlatformName = PlatArray[Index+1];
 				if (PlatformName == TEXT("Win32") || PlatformName == TEXT("Win64"))
 				{
-					PlatformName = TEXT("Windows");
+					PlatformName = TEXT("WindowsEditor");
 					PlatformInfo::UpdatePlatformSDKStatus(PlatformName, Status);
-					PlatformName = TEXT("WindowsNoEditor");
+					PlatformName = TEXT("Windows");
 					PlatformInfo::UpdatePlatformSDKStatus(PlatformName, Status);
 					PlatformName = TEXT("WindowsClient");
 					PlatformInfo::UpdatePlatformSDKStatus(PlatformName, Status);
@@ -1029,7 +1044,7 @@ protected:
 				else if (PlatformName == TEXT("Mac"))
 				{
 					PlatformInfo::UpdatePlatformSDKStatus(PlatformName, Status);
-					PlatformName = TEXT("MacNoEditor");
+					PlatformName = TEXT("Mac");
 					PlatformInfo::UpdatePlatformSDKStatus(PlatformName, Status);
 					PlatformName = TEXT("MacClient");
 					PlatformInfo::UpdatePlatformSDKStatus(PlatformName, Status);
@@ -1039,7 +1054,7 @@ protected:
 				else if (PlatformName == TEXT("Linux"))
 				{
 					PlatformInfo::UpdatePlatformSDKStatus(PlatformName, Status);
-					PlatformName = TEXT("LinuxNoEditor");
+					PlatformName = TEXT("Linux");
 					PlatformInfo::UpdatePlatformSDKStatus(PlatformName, Status);
 					PlatformName = TEXT("LinuxClient");
 					PlatformInfo::UpdatePlatformSDKStatus(PlatformName, Status);
@@ -1049,7 +1064,7 @@ protected:
 				else if (PlatformName == TEXT("LinuxAArch64"))
 				{
 					PlatformInfo::UpdatePlatformSDKStatus(PlatformName, Status);
-					PlatformName = TEXT("LinuxAArch64NoEditor");
+					PlatformName = TEXT("LinuxAArch64");
 					PlatformInfo::UpdatePlatformSDKStatus(PlatformName, Status);
 					PlatformName = TEXT("LinuxAArch64Client");
 					PlatformInfo::UpdatePlatformSDKStatus(PlatformName, Status);
@@ -1076,7 +1091,7 @@ protected:
 		return true;
 	}
 
-	bool UpdateAfterSDKInstall(const FString& PlatformName)
+	bool UpdateAfterSDKInstall(FName PlatformName)
 	{
 		const FDataDrivenPlatformInfo& Info = FDataDrivenPlatformInfoRegistry::GetPlatformInfo(PlatformName);
 
diff --git a/Engine/Source/Developer/TargetPlatform/Public/Common/TargetPlatformBase.h b/Engine/Source/Developer/TargetPlatform/Public/Common/TargetPlatformBase.h
--- a/Engine/Source/Developer/TargetPlatform/Public/Common/TargetPlatformBase.h
+++ b/Engine/Source/Developer/TargetPlatform/Public/Common/TargetPlatformBase.h
@@ -118,16 +118,6 @@ public:
 		return false;
 	}
 
-	virtual FText GetVariantDisplayName() const override
-	{
-		return FText();
-	}
-
-	virtual FText GetVariantTitle() const override
-	{
-		return FText();
-	}
-
 	virtual float GetVariantPriority() const override
 	{
 		return IsClientOnly() ? 0.0f : 0.2f;
@@ -217,17 +207,36 @@ public:
 	 */
 	static bool IsUsable()
 	{
-		return PlatformInfo::FindPlatformInfo(TPlatformProperties::PlatformName()) != nullptr;
+		return true;
 	}
-	
-	/** Default constructor. */
-	TTargetPlatformBase()
-		: FTargetPlatformBase( PlatformInfo::FindPlatformInfo(TPlatformProperties::PlatformName()) )
+
+
+
+	/**
+	 * Constructor that already has a TPI (notably coming from TNonDesktopTargetPlatform)
+	 */
+	TTargetPlatformBase(PlatformInfo::FTargetPlatformInfo* PremadePlatformInfo)
+		: FTargetPlatformBase(PremadePlatformInfo)
 	{
 		// HasEditorOnlyData and RequiresCookedData are mutually exclusive.
 		check(TPlatformProperties::HasEditorOnlyData() != TPlatformProperties::RequiresCookedData());
 	}
 
+	/**
+	 * Constructor that makes a TPI based solely on TPlatformProperties
+	 */
+	TTargetPlatformBase()
+		: TTargetPlatformBase( new PlatformInfo::FTargetPlatformInfo(
+			TPlatformProperties::IniPlatformName(),
+			TPlatformProperties::HasEditorOnlyData() ? EBuildTargetType::Editor :
+				TPlatformProperties::IsServerOnly() ? EBuildTargetType::Server : 
+				TPlatformProperties::IsClientOnly() ? EBuildTargetType::Client : 
+				EBuildTargetType::Game,
+			TEXT(""))
+		)
+	{
+	}
+
 public:
 
 	// ITargetPlatform interface
@@ -254,6 +263,9 @@ public:
 
 	virtual FString PlatformName() const override
 	{
+		// we assume these match for DesktopPlatforms (NonDesktop doesn't return "FooClient", but Desktop does, for legacy reasons)
+		checkSlow(this->PlatformInfo->Name == TPlatformProperties::PlatformName());
+		
 		return FString(TPlatformProperties::PlatformName());
 	}
 
@@ -382,3 +394,66 @@ public:
 	}
 #endif // WITH_ENGINE
 };
+
+
+template<typename TPlatformProperties>
+class TNonDesktopTargetPlatformBase 
+	: public TTargetPlatformBase<TPlatformProperties>
+{
+public:
+	/**
+	 * A simplified version for TPs that never will have Editor or ServerOnly versions, potentially multiple CookFlavors, as well as IN VERY RARE CASES, 
+	 * a different runtime IniPlatformName than what is passed in here (an example being TVOS and IOS, where passing in TVOS properties is very complicated)
+	 * Note that if we delayed the Info creation, we could just use this->IniPlatformName() and override that in, say TVOS, but we can't call a virtual here,
+	 * so we pass it up into the ctor
+	 */
+	TNonDesktopTargetPlatformBase(bool bInIsClientOnly, const TCHAR* CookFlavor=nullptr, const TCHAR* OverrideIniPlatformName=nullptr)
+		: TTargetPlatformBase(new PlatformInfo::FTargetPlatformInfo(
+			OverrideIniPlatformName ? FString(OverrideIniPlatformName) : FString(TPlatformProperties::IniPlatformName()),
+			bInIsClientOnly ? EBuildTargetType::Client : EBuildTargetType::Game,
+			CookFlavor))
+		, bIsClientOnly(bInIsClientOnly)
+	{
+
+	}
+
+	virtual FString PlatformName() const override
+	{
+		// instead of TPlatformProperties (which won't have Client for non-desktop platforms), use the Info's name, which is programmatically made
+		return this->PlatformInfo->Name.ToString();
+	}
+
+	virtual FString IniPlatformName() const override
+	{
+		// we use the Info's IniPlatformName as it may have been overridden in the constructor IN RARE CASES
+		return this->PlatformInfo->IniPlatformName.ToString();
+	}
+
+	virtual bool HasEditorOnlyData() const override
+	{
+		return false;
+	}
+
+	virtual bool IsServerOnly() const override
+	{
+		return false;
+	}
+
+	virtual bool IsClientOnly() const override
+	{
+		return bIsClientOnly;
+	}
+
+	virtual bool IsRunningPlatform() const override
+	{
+		// IsRunningPlatform is only for editor platforms
+
+		return false;
+	}
+
+protected:
+	// true if this target platform is client-only, ie strips out server stuff
+	bool bIsClientOnly;
+};
+
+
diff --git a/Engine/Source/Developer/TargetPlatform/Public/Interfaces/ITargetPlatform.h b/Engine/Source/Developer/TargetPlatform/Public/Interfaces/ITargetPlatform.h
--- a/Engine/Source/Developer/TargetPlatform/Public/Interfaces/ITargetPlatform.h
+++ b/Engine/Source/Developer/TargetPlatform/Public/Interfaces/ITargetPlatform.h
@@ -549,22 +549,6 @@ public:
 	 */
 	virtual bool SupportsVariants() const = 0;
 
-	/**
-	 * Gets the variant display name of this platform.
-	 * eg. For Android: "ASTC", "ETC2", ...
-	 *
-	 * @return FText display name for this platform variant.
-	 */
-	virtual FText GetVariantDisplayName() const = 0;
-
-	/**
-	 * Gets the variant title of this platform family
-	 * eg. For Android: "Texture Format"
-	 *
-	 * @return FText title for what variants mean to this family of platforms.
-	 */
-	virtual FText GetVariantTitle() const = 0;
-
 	/**
 	 * Gets the variant priority of this platform
 	 *
diff --git a/Engine/Source/Developer/TargetPlatform/Public/Interfaces/ITargetPlatformManagerModule.h b/Engine/Source/Developer/TargetPlatform/Public/Interfaces/ITargetPlatformManagerModule.h
--- a/Engine/Source/Developer/TargetPlatform/Public/Interfaces/ITargetPlatformManagerModule.h
+++ b/Engine/Source/Developer/TargetPlatform/Public/Interfaces/ITargetPlatformManagerModule.h
@@ -66,6 +66,22 @@ public:
 	 */
 	virtual ITargetPlatform* FindTargetPlatform( FStringView Name ) = 0;
 
+	/**
+	 * Finds a target platform by name.
+	 *
+	 * @param Name The target platform's short or long name.
+	 * @return The target platform, or nullptr if not found.
+	 */
+	virtual ITargetPlatform* FindTargetPlatform(FName Name) = 0;
+
+	/**
+	 * Finds a target platform by name.
+	 *
+	 * @param Name The target platform's short or long name.
+	 * @return The target platform, or nullptr if not found.
+	 */
+	virtual ITargetPlatform* FindTargetPlatform(const TCHAR* Name) = 0;
+
 	/**
 	 * Finds a target platform by looking for one that supports a given value for a generic type of support
 	 *
@@ -172,7 +188,7 @@ public:
 	/**
 	 * After installing an SDK with Turnkey, this will refresh the TargetPlatform, find devices, etc
 	 */
-	virtual bool UpdateAfterSDKInstall(const FString& TargetPlatformName) = 0;
+	virtual bool UpdateAfterSDKInstall(FName TargetPlatformName) = 0;
 
 public:
 
diff --git a/Engine/Source/Developer/Windows/WindowsTargetPlatform/Private/GenericWindowsTargetPlatform.h b/Engine/Source/Developer/Windows/WindowsTargetPlatform/Private/GenericWindowsTargetPlatform.h
--- a/Engine/Source/Developer/Windows/WindowsTargetPlatform/Private/GenericWindowsTargetPlatform.h
+++ b/Engine/Source/Developer/Windows/WindowsTargetPlatform/Private/GenericWindowsTargetPlatform.h
@@ -443,31 +443,6 @@ public:
 		return true;
 	}
 
-	virtual FText GetVariantDisplayName() const override
-	{
-		if (IS_DEDICATED_SERVER)
-		{
-			return LOCTEXT("WindowsServerVariantTitle", "Dedicated Server");
-		}
-
-		if (HAS_EDITOR_DATA)
-		{
-			return LOCTEXT("WindowsClientEditorDataVariantTitle", "Client with Editor Data");
-		}
-
-		if (IS_CLIENT_ONLY)
-		{
-			return LOCTEXT("WindowsClientOnlyVariantTitle", "Client only");
-		}
-
-		return LOCTEXT("WindowsClientVariantTitle", "Client");
-	}
-
-	virtual FText GetVariantTitle() const override
-	{
-		return LOCTEXT("WindowsVariantTitle", "Build Type");
-	}
-
 	virtual float GetVariantPriority() const override
 	{
 		return TProperties::GetVariantPriority();
diff --git a/Engine/Source/Developer/Windows/WindowsTargetPlatform/Private/WindowsTargetPlatformModule.cpp b/Engine/Source/Developer/Windows/WindowsTargetPlatform/Private/WindowsTargetPlatformModule.cpp
--- a/Engine/Source/Developer/Windows/WindowsTargetPlatform/Private/WindowsTargetPlatformModule.cpp
+++ b/Engine/Source/Developer/Windows/WindowsTargetPlatform/Private/WindowsTargetPlatformModule.cpp
@@ -36,26 +36,14 @@ public:
 
 	virtual void GetTargetPlatforms(TArray<ITargetPlatform*>& TargetPlatforms) override
 	{
+		// Game TP
+		TargetPlatforms.Add(new TGenericWindowsTargetPlatform<false, false, false>());
 		// Editor TP
-		if (TGenericWindowsTargetPlatform<true, false, false>::IsUsable())
-		{
-			TargetPlatforms.Add(new TGenericWindowsTargetPlatform<true, false, false>());
-		}
+		TargetPlatforms.Add(new TGenericWindowsTargetPlatform<true, false, false>());
 		// Server TP
-		if (TGenericWindowsTargetPlatform<false, true, false>::IsUsable())
-		{
-			TargetPlatforms.Add(new TGenericWindowsTargetPlatform<false, true, false>());
-		}
-		// NoEditor TP
-		if (TGenericWindowsTargetPlatform<false, false, false>::IsUsable())
-		{
-			TargetPlatforms.Add(new TGenericWindowsTargetPlatform<false, false, false>());
-		}
+		TargetPlatforms.Add(new TGenericWindowsTargetPlatform<false, true, false>());
 		// Client TP
-		if (TGenericWindowsTargetPlatform<false, false, true>::IsUsable())
-		{
-			TargetPlatforms.Add(new TGenericWindowsTargetPlatform<false, false, true>());
-		}
+		TargetPlatforms.Add(new TGenericWindowsTargetPlatform<false, false, true>());
 	}
 
 public:
diff --git a/Engine/Source/Editor/GameProjectGeneration/Private/SProjectBrowser.cpp b/Engine/Source/Editor/GameProjectGeneration/Private/SProjectBrowser.cpp
--- a/Engine/Source/Editor/GameProjectGeneration/Private/SProjectBrowser.cpp
+++ b/Engine/Source/Editor/GameProjectGeneration/Private/SProjectBrowser.cpp
@@ -802,9 +802,9 @@ static TSharedPtr<FProjectItem> CreateProjectItem(const FString& ProjectFilename
 			TArray<FName> TargetPlatforms;
 			for (const PlatformInfo::FTargetPlatformInfo* PlatformInfo : PlatformInfo::GetVanillaPlatformInfoArray())
 			{
-				if (PlatformInfo->PlatformType == EBuildTargetType::Game && ProjectStatus.IsTargetPlatformSupported(PlatformInfo->PlatformInfoName))
+				if (ProjectStatus.IsTargetPlatformSupported(PlatformInfo->Name))
 				{
-					TargetPlatforms.Add(PlatformInfo->PlatformInfoName);
+					TargetPlatforms.Add(PlatformInfo->Name);
 				}
 			}
 			TargetPlatforms.Sort(FNameLexicalLess());
diff --git a/Engine/Source/Editor/GameProjectGeneration/Public/GameProjectGenerationModule.h b/Engine/Source/Editor/GameProjectGeneration/Public/GameProjectGenerationModule.h
--- a/Engine/Source/Editor/GameProjectGeneration/Public/GameProjectGenerationModule.h
+++ b/Engine/Source/Editor/GameProjectGeneration/Public/GameProjectGenerationModule.h
@@ -107,7 +107,7 @@ public:
 	 * Update the list of supported target platforms based upon the parameters provided
 	 * This will take care of checking out and saving the updated .uproject file automatically
 	 * 
-	 * @param	InPlatformName		Name of the platform to target (eg, WindowsNoEditor)
+	 * @param	InPlatformName		Name of the platform to target (eg, WindowsClient)
 	 * @param	bIsSupported		true if the platform should be supported by this project, false if it should not
 	 */
 	virtual void UpdateSupportedTargetPlatforms(const FName& InPlatformName, const bool bIsSupported);
diff --git a/Engine/Source/Editor/GameProjectGeneration/Public/GameProjectUtils.h b/Engine/Source/Editor/GameProjectGeneration/Public/GameProjectUtils.h
--- a/Engine/Source/Editor/GameProjectGeneration/Public/GameProjectUtils.h
+++ b/Engine/Source/Editor/GameProjectGeneration/Public/GameProjectUtils.h
@@ -225,7 +225,7 @@ public:
 	 * Update the list of supported target platforms based upon the parameters provided
 	 * This will take care of checking out and saving the updated .uproject file automatically
 	 *
-	 * @param	InPlatformName		Name of the platform to target (eg, WindowsNoEditor)
+	 * @param	InPlatformName		Name of the platform to target (eg, WindowsClient)
 	 * @param	bIsSupported		true if the platform should be supported by this project, false if it should not
 	 */
 	static void UpdateSupportedTargetPlatforms(const FName& InPlatformName, const bool bIsSupported);
diff --git a/Engine/Source/Editor/MainFrame/Private/Frame/MainFrameActions.cpp b/Engine/Source/Editor/MainFrame/Private/Frame/MainFrameActions.cpp
--- a/Engine/Source/Editor/MainFrame/Private/Frame/MainFrameActions.cpp
+++ b/Engine/Source/Editor/MainFrame/Private/Frame/MainFrameActions.cpp
@@ -389,7 +389,7 @@ void FMainFrameActionCallbacks::CookContent(const FName InPlatformInfoName)
 
 	FString OptionalParams;
 
-	if (!FModuleManager::LoadModuleChecked<IProjectTargetPlatformEditorModule>("ProjectTargetPlatformEditor").ShowUnsupportedTargetWarning(PlatformInfo->VanillaInfo->PlatformInfoName))
+	if (!FModuleManager::LoadModuleChecked<IProjectTargetPlatformEditorModule>("ProjectTargetPlatformEditor").ShowUnsupportedTargetWarning(PlatformInfo->VanillaInfo->Name))
 	{
 		return;
 	}
@@ -398,7 +398,7 @@ void FMainFrameActionCallbacks::CookContent(const FName InPlatformInfoName)
 	if (PlatformInfo->DataDrivenPlatformInfo->GetSdkStatus() != DDPIPlatformSdkStatus::Valid)
 	{
 		IMainFrameModule& MainFrameModule = FModuleManager::GetModuleChecked<IMainFrameModule>(TEXT("MainFrame"));
-		MainFrameModule.BroadcastMainFrameSDKNotInstalled(PlatformInfo->TargetPlatformName.ToString(), PlatformInfo->DataDrivenPlatformInfo->SDKTutorial);
+		MainFrameModule.BroadcastMainFrameSDKNotInstalled(PlatformInfo->Name.ToString(), PlatformInfo->DataDrivenPlatformInfo->SDKTutorial);
 		return;
 	}
 
@@ -411,7 +411,7 @@ void FMainFrameActionCallbacks::CookContent(const FName InPlatformInfoName)
 	else
 	{
 		OptionalParams += TEXT(" -targetplatform=");
-		OptionalParams += *PlatformInfo->TargetPlatformName.ToString();
+		OptionalParams += *PlatformInfo->Name.ToString();
 	}
 
 	OptionalParams += GetCookingOptionalParams();
@@ -499,10 +499,10 @@ void FMainFrameActionCallbacks::PackageProject( const FName InPlatformInfoName )
 	if (PlatformInfo->DataDrivenPlatformInfo->GetSdkStatus() != DDPIPlatformSdkStatus::Valid || (bProjectHasCode && PlatformInfo->DataDrivenPlatformInfo->bUsesHostCompiler && !FSourceCodeNavigation::IsCompilerAvailable()))
 	{
 		IMainFrameModule& MainFrameModule = FModuleManager::GetModuleChecked<IMainFrameModule>(TEXT("MainFrame"));
-		MainFrameModule.BroadcastMainFrameSDKNotInstalled(PlatformInfo->TargetPlatformName.ToString(), PlatformInfo->DataDrivenPlatformInfo->SDKTutorial);
+		MainFrameModule.BroadcastMainFrameSDKNotInstalled(PlatformInfo->Name.ToString(), PlatformInfo->DataDrivenPlatformInfo->SDKTutorial);
 		TArray<FAnalyticsEventAttribute> ParamArray;
 		ParamArray.Add(FAnalyticsEventAttribute(TEXT("Time"), 0.0));
-		FEditorAnalytics::ReportEvent(TEXT("Editor.Package.Failed"), PlatformInfo->TargetPlatformName.ToString(), bProjectHasCode, EAnalyticsErrorCodes::SDKNotFound, ParamArray);
+		FEditorAnalytics::ReportEvent(TEXT("Editor.Package.Failed"), PlatformInfo->Name.ToString(), bProjectHasCode, EAnalyticsErrorCodes::SDKNotFound, ParamArray);
 		return;
 	}
 
@@ -510,7 +510,7 @@ void FMainFrameActionCallbacks::PackageProject( const FName InPlatformInfoName )
 	const UProjectPackagingSettings::FConfigurationInfo& ConfigurationInfo = UProjectPackagingSettings::ConfigurationInfo[PackagingSettings->BuildConfiguration];
 	bool bAssetNativizationEnabled = (PackagingSettings->BlueprintNativizationMethod != EProjectPackagingBlueprintNativizationMethod::Disabled);
 
-	const ITargetPlatform* const Platform = GetTargetPlatformManager()->FindTargetPlatform(PlatformInfo->TargetPlatformName.ToString());
+	const ITargetPlatform* const Platform = GetTargetPlatformManager()->FindTargetPlatform(PlatformInfo->Name.ToString());
 	{
 		if (Platform)
 		{
@@ -521,7 +521,7 @@ void FMainFrameActionCallbacks::PackageProject( const FName InPlatformInfoName )
 			int32 Result = Platform->CheckRequirements(bProjectHasCode, ConfigurationInfo.Configuration, bAssetNativizationEnabled, NotInstalledTutorialLink, DocumentationLink, CustomizedLogMessage);
 
 			// report to analytics
-			FEditorAnalytics::ReportBuildRequirementsFailure(TEXT("Editor.Package.Failed"), PlatformInfo->TargetPlatformName.ToString(), bProjectHasCode, Result);
+			FEditorAnalytics::ReportBuildRequirementsFailure(TEXT("Editor.Package.Failed"), PlatformInfo->Name.ToString(), bProjectHasCode, Result);
 
 			// report to main frame
 			bool UnrecoverableError = false;
@@ -614,7 +614,7 @@ void FMainFrameActionCallbacks::PackageProject( const FName InPlatformInfoName )
 		}
 	}
 
-	if (!FModuleManager::LoadModuleChecked<IProjectTargetPlatformEditorModule>("ProjectTargetPlatformEditor").ShowUnsupportedTargetWarning(PlatformInfo->VanillaInfo->PlatformInfoName))
+	if (!FModuleManager::LoadModuleChecked<IProjectTargetPlatformEditorModule>("ProjectTargetPlatformEditor").ShowUnsupportedTargetWarning(PlatformInfo->VanillaInfo->Name))
 	{
 		return;
 	}
@@ -705,18 +705,6 @@ void FMainFrameActionCallbacks::PackageProject( const FName InPlatformInfoName )
 		OptionalParams += TEXT(" -manifests");
 	}
 
-	bool bTargetPlatformCanUseCrashReporter = PlatformInfo->DataDrivenPlatformInfo->bCanUseCrashReporter;
-	if (bTargetPlatformCanUseCrashReporter && PlatformInfo->TargetPlatformName == FName("WindowsNoEditor") && PlatformInfo->PlatformFlavor == TEXT("Win32"))
-	{
-		FString MinumumSupportedWindowsOS;
-		GConfig->GetString(TEXT("/Script/WindowsTargetPlatform.WindowsTargetSettings"), TEXT("MinimumOSVersion"), MinumumSupportedWindowsOS, GEngineIni);
-		if (MinumumSupportedWindowsOS == TEXT("MSOS_XP"))
-		{
-			OptionalParams += TEXT(" -SpecifiedArchitecture=_xp");
-			bTargetPlatformCanUseCrashReporter = false;
-		}
-	}
-
 	// Append any extra UAT flags specified for this platform flavor
 	if (!PlatformInfo->UATCommandLine.IsEmpty())
 	{
@@ -726,7 +714,7 @@ void FMainFrameActionCallbacks::PackageProject( const FName InPlatformInfoName )
 	else
 	{
 		OptionalParams += TEXT(" -targetplatform=");
-		OptionalParams += *PlatformInfo->TargetPlatformName.ToString();
+		OptionalParams += *PlatformInfo->Name.ToString();
 	}
 
 	// Get the target to build
@@ -802,7 +790,7 @@ void FMainFrameActionCallbacks::PackageProject( const FName InPlatformInfoName )
 	}
 
 	// Whether to include the crash reporter.
-	if (PackagingSettings->IncludeCrashReporter && bTargetPlatformCanUseCrashReporter)
+	if (PackagingSettings->IncludeCrashReporter && PlatformInfo->DataDrivenPlatformInfo->bCanUseCrashReporter)
 	{
 		OptionalParams += TEXT( " -CrashReporter" );
 	}
diff --git a/Engine/Source/Editor/MainFrame/Private/Menus/PackageProjectMenu.h b/Engine/Source/Editor/MainFrame/Private/Menus/PackageProjectMenu.h
--- a/Engine/Source/Editor/MainFrame/Private/Menus/PackageProjectMenu.h
+++ b/Engine/Source/Editor/MainFrame/Private/Menus/PackageProjectMenu.h
@@ -157,8 +157,8 @@ protected:
 		IProjectTargetPlatformEditorModule& ProjectTargetPlatformEditorModule = FModuleManager::LoadModuleChecked<IProjectTargetPlatformEditorModule>("ProjectTargetPlatformEditor");
 
 		FUIAction Action(
-			FExecuteAction::CreateStatic(&FMainFrameActionCallbacks::PackageProject, PlatformInfo.PlatformInfoName),
-			FCanExecuteAction::CreateStatic(&FMainFrameActionCallbacks::PackageProjectCanExecute, PlatformInfo.PlatformInfoName)
+			FExecuteAction::CreateStatic(&FMainFrameActionCallbacks::PackageProject, PlatformInfo.Name),
+			FCanExecuteAction::CreateStatic(&FMainFrameActionCallbacks::PackageProjectCanExecute, PlatformInfo.Name)
 			);
 
 		// ... generate tooltip text
@@ -167,7 +167,7 @@ protected:
 		FText Tooltip = FText::Format(LOCTEXT("PackageGameForPlatformTooltip", "Build, cook and package your game for the {DisplayName} platform"), TooltipArguments);
 
 		FProjectStatus ProjectStatus;
-		if (IProjectManager::Get().QueryStatusForCurrentProject(ProjectStatus) && !ProjectStatus.IsTargetPlatformSupported(PlatformInfo.VanillaInfo->PlatformInfoName))
+		if (IProjectManager::Get().QueryStatusForCurrentProject(ProjectStatus) && !ProjectStatus.IsTargetPlatformSupported(PlatformInfo.VanillaInfo->Name))
 		{
 			FText TooltipLine2 = FText::Format(LOCTEXT("PackageUnsupportedPlatformWarning", "{DisplayName} is not listed as a target platform for this project, so may not run as expected."), TooltipArguments);
 			Tooltip = FText::Format(FText::FromString(TEXT("{0}\n\n{1}")), Tooltip, TooltipLine2);
diff --git a/Engine/Source/Editor/ProjectTargetPlatformEditor/Private/ProjectTargetPlatformEditorModule.cpp b/Engine/Source/Editor/ProjectTargetPlatformEditor/Private/ProjectTargetPlatformEditorModule.cpp
--- a/Engine/Source/Editor/ProjectTargetPlatformEditor/Private/ProjectTargetPlatformEditorModule.cpp
+++ b/Engine/Source/Editor/ProjectTargetPlatformEditor/Private/ProjectTargetPlatformEditorModule.cpp
@@ -113,7 +113,7 @@ public:
 					.HeightOverride(MenuIconSize)
 					[
 						SNew(SImage)
-						.Visibility(TAttribute<EVisibility>::Create(TAttribute<EVisibility>::FGetter::CreateStatic(&Local::IsUnsupportedPlatformWarningVisible, PlatformInfo.VanillaInfo->PlatformInfoName)))
+						.Visibility(TAttribute<EVisibility>::Create(TAttribute<EVisibility>::FGetter::CreateStatic(&Local::IsUnsupportedPlatformWarningVisible, PlatformInfo.VanillaInfo->Name)))
 						.Image(FEditorStyle::GetBrush("Launcher.Platform.Warning"))
 					]
 				]
@@ -136,7 +136,7 @@ public:
 
 		// Don't show the warning during automation testing; the dlg is modal and blocks
 		FProjectStatus ProjectStatus;
-		if(!GIsAutomationTesting && IProjectManager::Get().QueryStatusForCurrentProject(ProjectStatus) && !ProjectStatus.IsTargetPlatformSupported(PlatformInfo->VanillaInfo->PlatformInfoName))
+		if(!GIsAutomationTesting && IProjectManager::Get().QueryStatusForCurrentProject(ProjectStatus) && !ProjectStatus.IsTargetPlatformSupported(PlatformInfo->VanillaInfo->Name))
 		{
 			FFormatNamedArguments Args;
 			Args.Add(TEXT("DisplayName"), PlatformInfo->DisplayName);
diff --git a/Engine/Source/Editor/ProjectTargetPlatformEditor/Private/Widgets/SProjectTargetPlatformSettings.cpp b/Engine/Source/Editor/ProjectTargetPlatformEditor/Private/Widgets/SProjectTargetPlatformSettings.cpp
--- a/Engine/Source/Editor/ProjectTargetPlatformEditor/Private/Widgets/SProjectTargetPlatformSettings.cpp
+++ b/Engine/Source/Editor/ProjectTargetPlatformEditor/Private/Widgets/SProjectTargetPlatformSettings.cpp
@@ -42,7 +42,7 @@ void SProjectTargetPlatformSettings::Construct(const FArguments& InArgs)
 		PlatformsListBox->AddSlot()
 		.AutoHeight()
 		[
-			MakePlatformRow(AvailablePlatform->DisplayName, AvailablePlatform->PlatformInfoName, FEditorStyle::GetBrush(AvailablePlatform->GetIconStyleName(EPlatformIconSize::Normal)))
+			MakePlatformRow(AvailablePlatform->DisplayName, AvailablePlatform->Name, FEditorStyle::GetBrush(AvailablePlatform->GetIconStyleName(EPlatformIconSize::Normal)))
 		];
 	}
 
@@ -189,7 +189,7 @@ void SProjectTargetPlatformSettings::HandlePlatformCheckBoxStateChanged(ECheckBo
 			// We've deselected "All Platforms", so manually select every available platform
 			for(const PlatformInfo::FTargetPlatformInfo* AvailablePlatform : AvailablePlatforms)
 			{
-				FGameProjectGenerationModule::Get().UpdateSupportedTargetPlatforms(AvailablePlatform->TargetPlatformName, true);
+				FGameProjectGenerationModule::Get().UpdateSupportedTargetPlatforms(AvailablePlatform->Name, true);
 			}
 		}
 	}
diff --git a/Engine/Source/Editor/UnrealEd/Private/Commandlets/ExportPakDependencies.cpp b/Engine/Source/Editor/UnrealEd/Private/Commandlets/ExportPakDependencies.cpp
--- a/Engine/Source/Editor/UnrealEd/Private/Commandlets/ExportPakDependencies.cpp
+++ b/Engine/Source/Editor/UnrealEd/Private/Commandlets/ExportPakDependencies.cpp
@@ -119,7 +119,7 @@ TMap<FName,FPackage*> FPackage::NameToPackageMap;
 bool ExportDependencies(const TCHAR * PakFilename, const TCHAR* GameName, const TCHAR* OutputFilenameBase, bool bSigned)
 {
 	// Example command line used for this tool
-	// C:\Development\BB\WEX\Saved\StagedBuilds\WindowsNoEditor\WorldExplorers\Content\Paks\WorldExplorers-WindowsNoEditor.pak WorldExplorers WEX -exportdependencies=c:\dvtemp\output -debug -NoAssetRegistryCache -ForceDependsGathering
+	// C:\Development\BB\WEX\Saved\StagedBuilds\WindowsClient\WorldExplorers\Content\Paks\WorldExplorers-WindowsNoEditor.pak WorldExplorers WEX -exportdependencies=c:\dvtemp\output -debug -NoAssetRegistryCache -ForceDependsGathering
 	
 	FPakFile PakFile(&FPlatformFileManager::Get().GetPlatformFile(), PakFilename, bSigned);
 
diff --git a/Engine/Source/Editor/UnrealEd/Private/CookOnTheFlyServer.cpp b/Engine/Source/Editor/UnrealEd/Private/CookOnTheFlyServer.cpp
--- a/Engine/Source/Editor/UnrealEd/Private/CookOnTheFlyServer.cpp
+++ b/Engine/Source/Editor/UnrealEd/Private/CookOnTheFlyServer.cpp
@@ -6743,7 +6743,7 @@ void UCookOnTheFlyServer::StartCookByTheBook( const FCookByTheBookStartupOptions
 				// For example release version could be cooked as Android_ASTC flavor, but DLC can be made as Android_ETC2
 				for (const PlatformInfo::FTargetPlatformInfo* PlatformFlavorInfo : TargetPlatform->GetTargetPlatformInfo().Flavors)
 				{
-					OriginalSandboxRegistryFilename = GetBasedOnReleaseVersionAssetRegistryPath(BasedOnReleaseVersion, PlatformFlavorInfo->PlatformInfoName.ToString()) / GetAssetRegistryFilename();
+					OriginalSandboxRegistryFilename = GetBasedOnReleaseVersionAssetRegistryPath(BasedOnReleaseVersion, PlatformFlavorInfo->Name.ToString()) / GetAssetRegistryFilename();
 					bSucceeded = GetAllPackageFilenamesFromAssetRegistry(OriginalSandboxRegistryFilename, PackageList);
 					if (bSucceeded)
 					{
diff --git a/Engine/Source/Editor/UnrealEd/Private/Kismet2/DebuggerCommands.cpp b/Engine/Source/Editor/UnrealEd/Private/Kismet2/DebuggerCommands.cpp
--- a/Engine/Source/Editor/UnrealEd/Private/Kismet2/DebuggerCommands.cpp
+++ b/Engine/Source/Editor/UnrealEd/Private/Kismet2/DebuggerCommands.cpp
@@ -295,21 +295,17 @@ FPlayWorldCommands::FPlayWorldCommands()
 		FString RunningPlatformName = GetTargetPlatformManagerRef().GetRunningTargetPlatform()->PlatformName();
 		FString PlayPlatformName;
 
-		if (RunningPlatformName == TEXT("Windows"))
+		if (RunningPlatformName == TEXT("WindowsEditor"))
 		{
-			PlayPlatformName = TEXT("WindowsNoEditor");
+			PlayPlatformName = TEXT("Windows");
 		}
-		else if (RunningPlatformName == TEXT("Mac"))
+		else if (RunningPlatformName == TEXT("MacEditor"))
 		{
-			PlayPlatformName = TEXT("MacNoEditor");
+			PlayPlatformName = TEXT("Mac");
 		}
-		else if (RunningPlatformName == TEXT("Linux"))
+		else if (RunningPlatformName == TEXT("LinuxEditor"))
 		{
-			PlayPlatformName = TEXT("LinuxNoEditor");
-		}
-		else if (RunningPlatformName == TEXT("LinuxAArch64"))
-		{
-			PlayPlatformName = TEXT("LinuxAArch64NoEditor");
+			PlayPlatformName = TEXT("Linux");
 		}
 
 		if (!PlayPlatformName.IsEmpty())
@@ -1010,7 +1006,7 @@ static void MakeAllDevicesSubMenu(FMenuBuilder& InMenuBuilder, const PlatformInf
 		// for an aggregate (All_<platform>_devices_on_<host>) proxy, allow only the "Android_<texture_compression>" variants
 		const PlatformInfo::FTargetPlatformInfo* PlatformInfo = PlatformInfo::FindPlatformInfo(Variant);
 		if (DeviceProxy->IsAggregated() && PlatformInfo != NULL &&
-			(Variant == PlatformInfo->VanillaInfo->PlatformInfoName || PlatformInfo->PlatformType != EBuildTargetType::Game))
+			(Variant == PlatformInfo->VanillaInfo->Name || PlatformInfo->PlatformType != EBuildTargetType::Game))
 		{
 			continue;
 		}
@@ -1125,7 +1121,7 @@ static void TurnkeyInstallSdk(FString PlatformName, bool bPreferFull, bool bForc
 
 			// update the Sdk status
 //			FDataDrivenPlatformInfoRegistry::UpdateSdkStatus();
-			GetTargetPlatformManager()->UpdateAfterSDKInstall(PlatformName);
+			GetTargetPlatformManager()->UpdateAfterSDKInstall(*PlatformName);
 			RenderUtilsInit();
 
 
@@ -1136,7 +1132,7 @@ static void TurnkeyInstallSdk(FString PlatformName, bool bPreferFull, bool bForc
 	);
 }
 
-static TAttribute<FText> MakeSdkStatusAttribute(const PlatformInfo::FTargetPlatformInfo* PlatformInfo, TSharedPtr< ITargetDeviceProxy> DeviceProxy)
+static TAttribute<FText> MakeSdkStatusAttribute(FName IniPlatformName, TSharedPtr< ITargetDeviceProxy> DeviceProxy)
 {
 // 	FFormatNamedArguments LabelArguments;
 // 	LabelArguments.Add(TEXT("DeviceName"), FText::FromString(DeviceProxy->GetName()));
@@ -1158,57 +1154,59 @@ static TAttribute<FText> MakeSdkStatusAttribute(const PlatformInfo::FTargetPlatf
 // 	TAttribute<FText> = Make
 
 
-	FText DisplayString = DeviceProxy ? FText::FromString(DeviceProxy->GetName()) : PlatformInfo->DisplayName;
+	FString DisplayString = DeviceProxy ? DeviceProxy->GetName() : IniPlatformName.ToString();
 	FString DeviceId = DeviceProxy ? DeviceProxy->GetTargetDeviceId(NAME_None) : FString();
 
-	return TAttribute<FText>::Create(TAttribute<FText>::FGetter::CreateLambda([IniPlatformName=PlatformInfo->IniPlatformName, DisplayString, DeviceId]()
+	return TAttribute<FText>::Create(TAttribute<FText>::FGetter::CreateLambda([IniPlatformName, DisplayString, DeviceId]()
 	{
 		// get the status, or Unknown if it's not there
 		const FDataDrivenPlatformInfo& Info = FDataDrivenPlatformInfoRegistry::GetPlatformInfo(IniPlatformName);
 		DDPIPlatformSdkStatus Status = DeviceId.Len() ? Info.GetStatusForDeviceId(DeviceId) : Info.GetSdkStatus();
 
 		// @todo turnkey: Have premade FText's by SdkStatus for speed
-		FText Desc = FText::FromString(
+		const TCHAR* Desc =
 			Status == DDPIPlatformSdkStatus::Querying ? TEXT("Querying...") :
 			Status == DDPIPlatformSdkStatus::Valid ? TEXT("Valid Sdk") :
 			Status == DDPIPlatformSdkStatus::OutOfDate ? TEXT("Outdated Sdk") :
 			Status == DDPIPlatformSdkStatus::NoSdk ? TEXT("No Sdk") :
 			Status == DDPIPlatformSdkStatus::FlashValid ? TEXT("Valid Flash") :
 			Status == DDPIPlatformSdkStatus::FlashOutOfDate ? TEXT("Outdated Flash") :
-			TEXT("???"));
-
-		return FText::Format(FText::FromString(TEXT("{0} ({1})")), DisplayString, Desc);
+			TEXT("???");
+		const TCHAR* CompiledStatus = FDataDrivenPlatformInfoRegistry::HasCompiledSupportForPlatform(IniPlatformName, FDataDrivenPlatformInfoRegistry::EPlatformNameType::Ini) ? TEXT("") : TEXT(" (No Compiled Support)");
+		return FText::FromString(FString::Printf(TEXT("%s (%s)%s"), *DisplayString, Desc, CompiledStatus));
 	}));
 }
 
 
-static void MakeTurnkeyPlatformMenu(FMenuBuilder& MenuBuilder, const PlatformInfo::FTargetPlatformInfo* PlatformInfo, ITargetDeviceServicesModule* TargetDeviceServicesModule)
+static void MakeTurnkeyPlatformMenu(FMenuBuilder& MenuBuilder, FName IniPlatformName, ITargetDeviceServicesModule* TargetDeviceServicesModule)
 {
 	MenuBuilder.BeginSection("AllDevices", LOCTEXT("TurnkeySection_AllDevices", "All Devices"));
 
 	TArray<TSharedPtr<ITargetDeviceProxy>> DeviceProxies;
-	TargetDeviceServicesModule->GetDeviceProxyManager()->GetAllProxies(PlatformInfo->IniPlatformName, DeviceProxies);
+	TargetDeviceServicesModule->GetDeviceProxyManager()->GetAllProxies(IniPlatformName, DeviceProxies);
 
-	FString PlatformName = PlatformInfo->DataDrivenPlatformInfo->UBTPlatformString;
+	const FDataDrivenPlatformInfo& Info = FDataDrivenPlatformInfoRegistry::GetPlatformInfo(IniPlatformName);
+
+	FString UBTPlatformName = Info.UBTPlatformString;
 	for (const TSharedPtr<ITargetDeviceProxy> Proxy : DeviceProxies)
 	{
 		FString DeviceName = Proxy->GetName();
 		MenuBuilder.AddSubMenu(
-			MakeSdkStatusAttribute(PlatformInfo, DeviceProxies[0]),
+			MakeSdkStatusAttribute(IniPlatformName, Proxy),
 			FText(),
-			FNewMenuDelegate::CreateLambda([PlatformName, DeviceName](FMenuBuilder& SubMenuBuilder) {
+			FNewMenuDelegate::CreateLambda([UBTPlatformName, DeviceName](FMenuBuilder& SubMenuBuilder) {
 				SubMenuBuilder.AddMenuEntry(
 					LOCTEXT("Turnkey_RepairDevice", "Repair Device as Needed"),
 					LOCTEXT("TurnkeyTooltip_RepairDevice", "Perform any fixup that may be needed on this device. If up to date already, nothing will be done."),
 					FSlateIcon(),
-					FExecuteAction::CreateStatic(TurnkeyInstallSdk, PlatformName, false, false, *DeviceName)
+					FExecuteAction::CreateStatic(TurnkeyInstallSdk, UBTPlatformName, false, false, *DeviceName)
 				);
 
 				SubMenuBuilder.AddMenuEntry(
 					LOCTEXT("Turnkey_ForceRepairDevice", "Force Repair Device"),
 					LOCTEXT("TurnkeyTooltip_ForceRepairDevice", "Force repairing anything on the device needed (update firmware, etc). Will perform all steps possible, even if not needed."),
 					FSlateIcon(),
-					FExecuteAction::CreateStatic(TurnkeyInstallSdk, PlatformName, true, false, *DeviceName)
+					FExecuteAction::CreateStatic(TurnkeyInstallSdk, UBTPlatformName, true, false, *DeviceName)
 				);
 			})
 		);
@@ -1218,7 +1216,6 @@ static void MakeTurnkeyPlatformMenu(FMenuBuilder& MenuBuilder, const PlatformInf
 
 	MenuBuilder.BeginSection("AllDevices", LOCTEXT("TurnkeySection_Sdks", "Sdk Managment"));
 
-	const FDataDrivenPlatformInfo& Info = FDataDrivenPlatformInfoRegistry::GetPlatformInfo(PlatformInfo->IniPlatformName);
 	const TCHAR* NoDevice = nullptr;
 	if (Info.GetSdkStatus() == DDPIPlatformSdkStatus::OutOfDate)
 	{
@@ -1226,14 +1223,14 @@ static void MakeTurnkeyPlatformMenu(FMenuBuilder& MenuBuilder, const PlatformInf
 			LOCTEXT("Turnkey_InstallSdkMinimal", "Update Sdk (Prefer Minimal)"),
 			LOCTEXT("TurnkeyTooltip_InstallSdkMinimal", "Attempt to update an Sdk, as hosted by your studio. Will attempt to install a minimal Sdk (useful for building/running only)"),
 			FSlateIcon(),
-			FExecuteAction::CreateStatic(TurnkeyInstallSdk, PlatformName, false, false, NoDevice)
+			FExecuteAction::CreateStatic(TurnkeyInstallSdk, UBTPlatformName, false, false, NoDevice)
 		);
 
 		MenuBuilder.AddMenuEntry(
 			LOCTEXT("Turnkey_InstallSdkFull", "Update Sdk (Prefer Full)"),
 			LOCTEXT("TurnkeyTooltip_InstallSdkMinimal", "Attempt to update an Sdk, as hosted by your studio. Will attempt to install a full Sdk (useful profiling or other use cases)"),
 			FSlateIcon(),
-			FExecuteAction::CreateStatic(TurnkeyInstallSdk, PlatformName, true, false, NoDevice)
+			FExecuteAction::CreateStatic(TurnkeyInstallSdk, UBTPlatformName, true, false, NoDevice)
 		);
 	}
 	else if (Info.GetSdkStatus() == DDPIPlatformSdkStatus::Valid)
@@ -1242,14 +1239,14 @@ static void MakeTurnkeyPlatformMenu(FMenuBuilder& MenuBuilder, const PlatformInf
 			LOCTEXT("Turnkey_InstallSdkMinimal", "Force Reinstall Sdk (Prefer Minimal)"),
 			LOCTEXT("TurnkeyTooltip_InstallSdkMinimal", "Attempt to force re-install an Sdk, as hosted by your studio. Will attempt to install a minimal Sdk (useful for building/running only)"),
 			FSlateIcon(),
-			FExecuteAction::CreateStatic(TurnkeyInstallSdk, PlatformName, false, true, NoDevice)
+			FExecuteAction::CreateStatic(TurnkeyInstallSdk, UBTPlatformName, false, true, NoDevice)
 		);
 
 		MenuBuilder.AddMenuEntry(
 			LOCTEXT("Turnkey_InstallSdkFull", "Force Reinstall (Prefer Full)"),
 			LOCTEXT("TurnkeyTooltip_InstallSdkMinimal", "Attempt to force re-install an Sdk, as hosted by your studio. Will attempt to install a full Sdk (useful profiling or other use cases)"),
 			FSlateIcon(),
-			FExecuteAction::CreateStatic(TurnkeyInstallSdk, PlatformName, true, true, NoDevice)
+			FExecuteAction::CreateStatic(TurnkeyInstallSdk, UBTPlatformName, true, true, NoDevice)
 		);
 	}
 	else
@@ -1258,19 +1255,19 @@ static void MakeTurnkeyPlatformMenu(FMenuBuilder& MenuBuilder, const PlatformInf
 			LOCTEXT("Turnkey_InstallSdkMinimal", "Install Sdk (Prefer Minimal)"),
 			LOCTEXT("TurnkeyTooltip_InstallSdkMinimal", "Attempt to install an Sdk, as hosted by your studio. Will attempt to install a minimal Sdk (useful for building/running only)"),
 			FSlateIcon(),
-			FExecuteAction::CreateStatic(TurnkeyInstallSdk, PlatformName, false, false, NoDevice)
+			FExecuteAction::CreateStatic(TurnkeyInstallSdk, UBTPlatformName, false, false, NoDevice)
 		);
 
 		MenuBuilder.AddMenuEntry(
 			LOCTEXT("Turnkey_InstallSdkFull", "Install Sdk (Prefer Full)"),
 			LOCTEXT("TurnkeyTooltip_InstallSdkMinimal", "Attempt to install an Sdk, as hosted by your studio. Will attempt to install a full Sdk (useful profiling or other use cases)"),
 			FSlateIcon(),
-			FExecuteAction::CreateStatic(TurnkeyInstallSdk, PlatformName, true, false, NoDevice)
+			FExecuteAction::CreateStatic(TurnkeyInstallSdk, UBTPlatformName, true, false, NoDevice)
 		);
 	}
 }
 
-static void GenerateDeviceProxyMenuParams(TSharedPtr<ITargetDeviceProxy> DeviceProxy, const PlatformInfo::FTargetPlatformInfo* EditorPlatformInfo, FName PlatformName, FUIAction& OutAction, FText& OutTooltip)
+static void GenerateDeviceProxyMenuParams(TSharedPtr<ITargetDeviceProxy> DeviceProxy, FName PlatformName, FUIAction& OutAction, FText& OutTooltip)
 {
 	// 	// create an All_<platform>_devices_on_<host> submenu
 	// 	if (DeviceProxy->IsAggregated())
@@ -1297,7 +1294,7 @@ static void GenerateDeviceProxyMenuParams(TSharedPtr<ITargetDeviceProxy> DeviceP
 	// ... generate tooltip text
 	FFormatNamedArguments TooltipArguments;
 	TooltipArguments.Add(TEXT("DeviceID"), FText::FromString(DeviceProxy->GetName()));
-	TooltipArguments.Add(TEXT("DisplayName"), EditorPlatformInfo->DisplayName);
+	TooltipArguments.Add(TEXT("DisplayName"), FText::FromName(PlatformName));
 	OutTooltip = FText::Format(LOCTEXT("LaunchDeviceToolTipText_ThisDevice", "Launch the game on this {DisplayName} device ({DeviceID})"), TooltipArguments);
 	if (!DeviceProxy->IsAuthorized())
 	{
@@ -1333,13 +1330,9 @@ TSharedRef< SWidget > FPlayWorldCommands::GenerateTurnkeyMenuContent(TSharedRef<
 				continue;
 			}
 
-			FName PlatformName(*Pair.Key);
+			FName PlatformName = Pair.Key;
 			const FDataDrivenPlatformInfo& Info = Pair.Value;
 
-			// find a matching editor-facing FPlatformInfo object
-			const PlatformInfo::FTargetPlatformInfo* EditorPlatformInfo = PlatformInfo::FindVanillaPlatformInfo(PlatformName);
-			check(EditorPlatformInfo != nullptr);
-
 			// look for devices for all platforms, even if the platform isn't installed - Turnkey can install Sdk after selecting LaunchOn
 			TArray<TSharedPtr<ITargetDeviceProxy>> DeviceProxies;
 			TargetDeviceServicesModule->GetDeviceProxyManager()->GetAllProxies(PlatformName, DeviceProxies);
@@ -1353,14 +1346,14 @@ TSharedRef< SWidget > FPlayWorldCommands::GenerateTurnkeyMenuContent(TSharedRef<
 				// always use the first one, after sorting
 				FUIAction Action;
 				FText Tooltip;
-				GenerateDeviceProxyMenuParams(DeviceProxies[0], EditorPlatformInfo, PlatformName, Action, Tooltip);
+				GenerateDeviceProxyMenuParams(DeviceProxies[0], PlatformName, Action, Tooltip);
 
 				if (DeviceProxies.Num() == 1)
 				{
 					MenuBuilder.AddMenuEntry(
-						MakeSdkStatusAttribute(EditorPlatformInfo, DeviceProxies[0]),
+						MakeSdkStatusAttribute(PlatformName, DeviceProxies[0]),
 						Tooltip,
-						FSlateIcon(FEditorStyle::GetStyleSetName(), EditorPlatformInfo->GetIconStyleName(EPlatformIconSize::Normal)),
+						FSlateIcon(FEditorStyle::GetStyleSetName(), Pair.Value.GetIconStyleName(EPlatformIconSize::Normal)),
 						Action,
 						NAME_None,
 						EUserInterfaceActionType::Button
@@ -1369,9 +1362,9 @@ TSharedRef< SWidget > FPlayWorldCommands::GenerateTurnkeyMenuContent(TSharedRef<
 				else
 				{
 					MenuBuilder.AddSubMenu(
-						MakeSdkStatusAttribute(EditorPlatformInfo, DeviceProxies[0]),
+						MakeSdkStatusAttribute(PlatformName, DeviceProxies[0]),
 						Tooltip,
-						FNewMenuDelegate::CreateLambda([TargetDeviceServicesModule, EditorPlatformInfo, PlatformName](FMenuBuilder& SubMenuBuilder) {
+						FNewMenuDelegate::CreateLambda([TargetDeviceServicesModule, PlatformName, IconStyle=Pair.Value.GetIconStyleName(EPlatformIconSize::Normal)](FMenuBuilder& SubMenuBuilder) {
 							// re-get the proxies, just in case they changed
 							TArray<TSharedPtr<ITargetDeviceProxy>> DeviceProxies;
 							TargetDeviceServicesModule->GetDeviceProxyManager()->GetAllProxies(PlatformName, DeviceProxies);
@@ -1380,11 +1373,11 @@ TSharedRef< SWidget > FPlayWorldCommands::GenerateTurnkeyMenuContent(TSharedRef<
 							{
 								FUIAction SubAction;
 								FText SubTooltip;
-								GenerateDeviceProxyMenuParams(Proxy, EditorPlatformInfo, PlatformName, SubAction, SubTooltip);
+								GenerateDeviceProxyMenuParams(Proxy, PlatformName, SubAction, SubTooltip);
 								SubMenuBuilder.AddMenuEntry(
-									MakeSdkStatusAttribute(EditorPlatformInfo, Proxy),
+									MakeSdkStatusAttribute(PlatformName, Proxy),
 									SubTooltip,
-									FSlateIcon(FEditorStyle::GetStyleSetName(), EditorPlatformInfo->GetIconStyleName(EPlatformIconSize::Normal)),
+									FSlateIcon(FEditorStyle::GetStyleSetName(), IconStyle),
 									SubAction,
 									NAME_None,
 									EUserInterfaceActionType::Button
@@ -1395,7 +1388,7 @@ TSharedRef< SWidget > FPlayWorldCommands::GenerateTurnkeyMenuContent(TSharedRef<
 						NAME_None,
 						EUserInterfaceActionType::Check,
 						false,
-						FSlateIcon(FEditorStyle::GetStyleSetName(), EditorPlatformInfo->GetIconStyleName(EPlatformIconSize::Normal)),
+						FSlateIcon(FEditorStyle::GetStyleSetName(), Pair.Value.GetIconStyleName(EPlatformIconSize::Normal)),
 						true
 					);
 				}
@@ -1404,7 +1397,7 @@ TSharedRef< SWidget > FPlayWorldCommands::GenerateTurnkeyMenuContent(TSharedRef<
 				for (const TSharedPtr<ITargetDeviceProxy> Proxy : DeviceProxies)
 				{
 					FString DeviceId = Proxy->GetTargetDeviceId(NAME_None);
-					if (EditorPlatformInfo->DataDrivenPlatformInfo->GetStatusForDeviceId(DeviceId) == DDPIPlatformSdkStatus::Unknown)
+					if (Pair.Value.GetStatusForDeviceId(DeviceId) == DDPIPlatformSdkStatus::Unknown)
 					{
 						DeviceIdsToQuery.Add(DeviceId);
 					}
@@ -1435,17 +1428,16 @@ TSharedRef< SWidget > FPlayWorldCommands::GenerateTurnkeyMenuContent(TSharedRef<
 			continue;
 		}
 
-		FName PlatformName(*Pair.Key);
+		FName PlatformName = Pair.Key;
 		const FDataDrivenPlatformInfo& Info = Pair.Value;
-		const PlatformInfo::FTargetPlatformInfo* EditorPlatformInfo = PlatformInfo::FindVanillaPlatformInfo(PlatformName);
 
 		FString DeviceNameToCheckForStatus;
 		MenuBuilder.AddSubMenu(
-			MakeSdkStatusAttribute(EditorPlatformInfo, nullptr),
+			MakeSdkStatusAttribute(PlatformName, nullptr),
 			FText::FromString(PlatformName.ToString()),
-			FNewMenuDelegate::CreateStatic(&MakeTurnkeyPlatformMenu, EditorPlatformInfo, TargetDeviceServicesModule),
+			FNewMenuDelegate::CreateStatic(&MakeTurnkeyPlatformMenu, PlatformName, TargetDeviceServicesModule),
 			false, 
-			FSlateIcon(FEditorStyle::GetStyleSetName(), EditorPlatformInfo->GetIconStyleName(EPlatformIconSize::Normal)),
+			FSlateIcon(FEditorStyle::GetStyleSetName(), Pair.Value.GetIconStyleName(EPlatformIconSize::Normal)),
 			true
 		);
 	}
@@ -1586,15 +1578,15 @@ void PopulateLaunchMenu(UToolMenu* Menu)
 				// for each platform...
 				TArray<TSharedPtr<ITargetDeviceProxy>> DeviceProxies;
 				// the list of proxies include the "Al_Android" entry
-				TargetDeviceServicesModule->GetDeviceProxyManager()->GetAllProxies(VanillaPlatform->PlatformInfoName, DeviceProxies);
+				TargetDeviceServicesModule->GetDeviceProxyManager()->GetAllProxies(VanillaPlatform->Name, DeviceProxies);
 
 				// if this platform had no devices, but we want to show an extra option if not installed right
 				if (DeviceProxies.Num() == 0)
 				{
-					if (PlatformsWithNoDevices.Find(VanillaPlatform->PlatformInfoName) == INDEX_NONE)
+					if (PlatformsWithNoDevices.Find(VanillaPlatform->Name) == INDEX_NONE)
 					{
 						// add an entry with a no devices found
-						PlatformsWithNoDevices.Add(VanillaPlatform->PlatformInfoName);
+						PlatformsWithNoDevices.Add(VanillaPlatform->Name);
 					}
 				}
 				else
@@ -1657,7 +1649,7 @@ void PopulateLaunchMenu(UToolMenu* Menu)
 						}
 
 						FProjectStatus ProjectStatus;
-						if (IProjectManager::Get().QueryStatusForCurrentProject(ProjectStatus) && !ProjectStatus.IsTargetPlatformSupported(VanillaPlatform->PlatformInfoName))
+						if (IProjectManager::Get().QueryStatusForCurrentProject(ProjectStatus) && !ProjectStatus.IsTargetPlatformSupported(VanillaPlatform->Name))
 						{
 							FText TooltipLine2 = FText::Format(LOCTEXT("LaunchDevicePlatformWarning", "{DisplayName} is not listed as a target platform for this project, so may not run as expected."), TooltipArguments);
 							Tooltip = FText::Format(FText::FromString(TEXT("{0}\n\n{1}")), Tooltip, TooltipLine2);
@@ -1677,7 +1669,7 @@ void PopulateLaunchMenu(UToolMenu* Menu)
 			else
 			{
 				// if the platform wasn't installed, we'll add a menu item later (we never care about code in this case, since we don't compile)
-				if (PlatformsToMaybeInstallLinksFor.Find(VanillaPlatform->PlatformInfoName.ToString()) != INDEX_NONE)
+				if (PlatformsToMaybeInstallLinksFor.Find(VanillaPlatform->Name.ToString()) != INDEX_NONE)
 				{
 					PlatformsToAddInstallLinksFor.Add(VanillaPlatform);
 				}
@@ -2646,10 +2638,10 @@ bool FInternalPlayWorldCommandCallbacks::IsReadyToLaunchOnDevice(FString DeviceI
 	if (PlatformInfo->DataDrivenPlatformInfo->GetSdkStatus() != DDPIPlatformSdkStatus::Valid)
 	{
 		IMainFrameModule& MainFrameModule = FModuleManager::GetModuleChecked<IMainFrameModule>(TEXT("MainFrame"));
-		MainFrameModule.BroadcastMainFrameSDKNotInstalled(PlatformInfo->TargetPlatformName.ToString(), PlatformInfo->DataDrivenPlatformInfo->SDKTutorial);
+		MainFrameModule.BroadcastMainFrameSDKNotInstalled(PlatformInfo->Name.ToString(), PlatformInfo->DataDrivenPlatformInfo->SDKTutorial);
 		TArray<FAnalyticsEventAttribute> ParamArray;
 		ParamArray.Add(FAnalyticsEventAttribute(TEXT("Time"), 0.0));
-		FEditorAnalytics::ReportEvent(TEXT("Editor.LaunchOn.Failed"), PlatformInfo->TargetPlatformName.ToString(), bHasCode, EAnalyticsErrorCodes::SDKNotFound, ParamArray);
+		FEditorAnalytics::ReportEvent(TEXT("Editor.LaunchOn.Failed"), PlatformInfo->Name.ToString(), bHasCode, EAnalyticsErrorCodes::SDKNotFound, ParamArray);
 		return false;
 	}
 
@@ -2765,7 +2757,7 @@ bool FInternalPlayWorldCommandCallbacks::IsReadyToLaunchOnDevice(FString DeviceI
 	else
 	{
 		IMainFrameModule& MainFrameModule = FModuleManager::GetModuleChecked<IMainFrameModule>(TEXT("MainFrame"));
-		MainFrameModule.BroadcastMainFrameSDKNotInstalled(PlatformInfo->TargetPlatformName.ToString(), PlatformInfo->DataDrivenPlatformInfo->SDKTutorial);
+		MainFrameModule.BroadcastMainFrameSDKNotInstalled(PlatformInfo->Name.ToString(), PlatformInfo->DataDrivenPlatformInfo->SDKTutorial);
 		return false;
 	}
 
diff --git a/Engine/Source/Programs/AutomationTool/AutomationUtils/CommandletUtils.cs b/Engine/Source/Programs/AutomationTool/AutomationUtils/CommandletUtils.cs
--- a/Engine/Source/Programs/AutomationTool/AutomationUtils/CommandletUtils.cs
+++ b/Engine/Source/Programs/AutomationTool/AutomationUtils/CommandletUtils.cs
@@ -53,7 +53,7 @@ namespace AutomationTool
         /// <param name="CulturesToCook">List of culture names whose localized assets should be cooked, can be null (implying defaults should be used).</param>
 		/// <param name="TargetPlatform">Target platform.</param>
 		/// <param name="Parameters">List of additional parameters.</param>
-		public static void CookCommandlet(FileReference ProjectName, string UE4Exe = "UE4Editor-Cmd.exe", string[] Maps = null, string[] Dirs = null, string InternationalizationPreset = "", string[] CulturesToCook = null, string TargetPlatform = "WindowsNoEditor", string Parameters = "-Unversioned")
+		public static void CookCommandlet(FileReference ProjectName, string UE4Exe = "UE4Editor-Cmd.exe", string[] Maps = null, string[] Dirs = null, string InternationalizationPreset = "", string[] CulturesToCook = null, string TargetPlatform = "Windows", string Parameters = "-Unversioned")
 		{
             string CommandletArguments = "";
 
@@ -97,7 +97,7 @@ namespace AutomationTool
         /// <param name="Maps">List of maps to cook, can be null in which case -MapIniSection=AllMaps is used.</param>
         /// <param name="TargetPlatform">Target platform.</param>
         /// <param name="Parameters">List of additional parameters.</param>
-        public static void DDCCommandlet(FileReference ProjectName, string UE4Exe = "UE4Editor-Cmd.exe", string[] Maps = null, string TargetPlatform = "WindowsNoEditor", string Parameters = "")
+        public static void DDCCommandlet(FileReference ProjectName, string UE4Exe = "UE4Editor-Cmd.exe", string[] Maps = null, string TargetPlatform = "Windows", string Parameters = "")
         {
             string MapsToCook = "";
             if (!IsNullOrEmpty(Maps))
diff --git a/Engine/Source/Programs/AutomationTool/AutomationUtils/DeploymentContext.cs b/Engine/Source/Programs/AutomationTool/AutomationUtils/DeploymentContext.cs
--- a/Engine/Source/Programs/AutomationTool/AutomationUtils/DeploymentContext.cs
+++ b/Engine/Source/Programs/AutomationTool/AutomationUtils/DeploymentContext.cs
@@ -157,12 +157,12 @@ public class DeploymentContext //: ProjectParams
 	public string PlatformDir;
 
 	/// <summary>
-	/// Directory to put all of the files in: d:\stagedir\WindowsNoEditor
+	/// Directory to put all of the files in: d:\stagedir\Windows
 	/// </summary>
 	public DirectoryReference StageDirectory;
 
 	/// <summary>
-	/// Directory to put all of the debug files in: d:\stagedir\WindowsNoEditor
+	/// Directory to put all of the debug files in: d:\stagedir\Windows
 	/// </summary>
 	public DirectoryReference DebugStageDirectory;
 
@@ -276,7 +276,7 @@ public class DeploymentContext //: ProjectParams
 	public List<string> BlacklistLocalizationTargets = new List<string>();
 
 	/// <summary>
-	///  Directory to archive all of the files in: d:\archivedir\WindowsNoEditor
+	///  Directory to archive all of the files in: d:\archivedir\Windows
 	/// </summary>
 	public DirectoryReference ArchiveDirectory;
 
diff --git a/Engine/Source/Programs/AutomationTool/BuildGraph/Tasks/CookTask.cs b/Engine/Source/Programs/AutomationTool/BuildGraph/Tasks/CookTask.cs
--- a/Engine/Source/Programs/AutomationTool/BuildGraph/Tasks/CookTask.cs
+++ b/Engine/Source/Programs/AutomationTool/BuildGraph/Tasks/CookTask.cs
@@ -24,7 +24,7 @@ namespace BuildGraph.Tasks
 		public string Project;
 
 		/// <summary>
-		/// The cook platform to target (for example, WindowsNoEditor).
+		/// The cook platform to target (for example, Windows).
 		/// </summary>
 		[TaskParameter]
 		public string Platform;
diff --git a/Engine/Source/Programs/AutomationTool/Linux/LinuxPlatform.Automation.cs b/Engine/Source/Programs/AutomationTool/Linux/LinuxPlatform.Automation.cs
--- a/Engine/Source/Programs/AutomationTool/Linux/LinuxPlatform.Automation.cs
+++ b/Engine/Source/Programs/AutomationTool/Linux/LinuxPlatform.Automation.cs
@@ -1,4 +1,4 @@
-﻿// Copyright Epic Games, Inc. All Rights Reserved.
+// Copyright Epic Games, Inc. All Rights Reserved.
 using System;
 using System.Collections.Generic;
 using System.Linq;
@@ -147,7 +147,7 @@ public abstract class BaseLinuxPlatform : Platform
 
 	public override string GetCookPlatform(bool bDedicatedServer, bool bIsClientOnly)
 	{
-		const string NoEditorCookPlatform = "NoEditor";
+		const string NoEditorCookPlatform = "";
 		const string ServerCookPlatform = "Server";
 		const string ClientCookPlatform = "Client";
 		string PlatformStr = (TargetPlatformType == UnrealTargetPlatform.LinuxAArch64) ? "LinuxAArch64" : "Linux";
@@ -168,10 +168,10 @@ public abstract class BaseLinuxPlatform : Platform
 	{
 		if (TargetPlatformType == UnrealTargetPlatform.LinuxAArch64)
 		{
-			return "LinuxAArch64";
+			return "LinuxAArch64Editor";
 		}
 
-		return "Linux";
+		return "LinuxEditor";
 	}
 
 	/// <summary>
diff --git a/Engine/Source/Programs/AutomationTool/Mac/MacPlatform.Automation.cs b/Engine/Source/Programs/AutomationTool/Mac/MacPlatform.Automation.cs
--- a/Engine/Source/Programs/AutomationTool/Mac/MacPlatform.Automation.cs
+++ b/Engine/Source/Programs/AutomationTool/Mac/MacPlatform.Automation.cs
@@ -18,7 +18,7 @@ public class MacPlatform : Platform
 
 	public override string GetCookPlatform(bool bDedicatedServer, bool bIsClientOnly)
 	{
-		const string NoEditorCookPlatform = "MacNoEditor";
+		const string NoEditorCookPlatform = "Mac";
 		const string ServerCookPlatform = "MacServer";
 		const string ClientCookPlatform = "MacClient";
 
@@ -38,7 +38,7 @@ public class MacPlatform : Platform
 
 	public override string GetEditorCookPlatform()
 	{
-		return "Mac";
+		return "MacEditor";
 	}
 
 	private void StageAppBundle(DeploymentContext SC, DirectoryReference InPath, StagedDirectoryReference NewName)
diff --git a/Engine/Source/Programs/AutomationTool/Win/WinPlatform.Automation.cs b/Engine/Source/Programs/AutomationTool/Win/WinPlatform.Automation.cs
--- a/Engine/Source/Programs/AutomationTool/Win/WinPlatform.Automation.cs
+++ b/Engine/Source/Programs/AutomationTool/Win/WinPlatform.Automation.cs
@@ -156,7 +156,7 @@ public abstract class BaseWinPlatform : Platform
 
 	public override string GetCookPlatform(bool bDedicatedServer, bool bIsClientOnly)
 	{
-		const string NoEditorCookPlatform = "WindowsNoEditor";
+		const string NoEditorCookPlatform = "Windows";
 		const string ServerCookPlatform = "WindowsServer";
 		const string ClientCookPlatform = "WindowsClient";
 
@@ -176,7 +176,7 @@ public abstract class BaseWinPlatform : Platform
 
 	public override string GetEditorCookPlatform()
 	{
-		return "Windows";
+		return "WindowsEditor";
 	}
 	
 	public override string GetPlatformPakCommandLine(ProjectParams Params, DeploymentContext SC)
diff --git a/Engine/Source/Runtime/Core/Private/Misc/DataDrivenPlatformInfoRegistry.cpp b/Engine/Source/Runtime/Core/Private/Misc/DataDrivenPlatformInfoRegistry.cpp
--- a/Engine/Source/Runtime/Core/Private/Misc/DataDrivenPlatformInfoRegistry.cpp
+++ b/Engine/Source/Runtime/Core/Private/Misc/DataDrivenPlatformInfoRegistry.cpp
@@ -11,7 +11,12 @@
 #include "Misc/MonitoredProcess.h"
 
 
-TMap<FString, FDataDrivenPlatformInfo> FDataDrivenPlatformInfoRegistry::DataDrivenPlatforms;
+namespace 
+{
+	TMap<FName, FDataDrivenPlatformInfo> DataDrivenPlatforms;
+	TArray<FName> SortedPlatformNames;
+	TArray<const FDataDrivenPlatformInfo*> SortedPlatformInfos;
+}
 
 #if DDPI_HAS_EXTENDED_PLATFORMINFO_DATA
 // NO AUTO-KICK-OFF - let the editor request it as needed
@@ -183,7 +188,7 @@ static void DDPIGetStringArray(const FConfigFile& IniFile, const TCHAR* Key, TAr
 	IniFile.GetArray(TEXT("DataDrivenPlatformInfo"), Key, OutArray);
 }
 
-static void LoadDDPIIniSettings(const FConfigFile& IniFile, FDataDrivenPlatformInfo& Info, const FString& PlatformName)
+static void LoadDDPIIniSettings(const FConfigFile& IniFile, FDataDrivenPlatformInfo& Info, FName PlatformName)
 {
 	DDPIGetBool(IniFile, TEXT("bIsConfidential"), Info.bIsConfidential);
 	DDPIGetBool(IniFile, TEXT("bIsFakePlatform"), Info.bIsFakePlatform);
@@ -215,9 +220,10 @@ static void LoadDDPIIniSettings(const FConfigFile& IniFile, FDataDrivenPlatformI
 		Info.IconPaths.XLargePath = Info.IconPaths.XLargePath;
 	}
 
-	Info.IconPaths.NormalStyleName = *FString::Printf(TEXT("Launcher.Platform_%s"), *PlatformName);
-	Info.IconPaths.LargeStyleName = *FString::Printf(TEXT("Launcher.Platform_%s.Large"), *PlatformName);
-	Info.IconPaths.XLargeStyleName = *FString::Printf(TEXT("Launcher.Platform_%s.XLarge"), *PlatformName);
+	FString PlatformString = PlatformName.ToString();
+	Info.IconPaths.NormalStyleName = *FString::Printf(TEXT("Launcher.Platform_%s"), *PlatformString);
+	Info.IconPaths.LargeStyleName = *FString::Printf(TEXT("Launcher.Platform_%s.Large"), *PlatformString);
+	Info.IconPaths.XLargeStyleName = *FString::Printf(TEXT("Launcher.Platform_%s.XLarge"), *PlatformString);
 
 	Info.bCanUseCrashReporter = true; // not specified means true, not false
 	DDPIGetBool(IniFile, TEXT("bCanUseCrashReporter"), Info.bCanUseCrashReporter);
@@ -229,7 +235,7 @@ static void LoadDDPIIniSettings(const FConfigFile& IniFile, FDataDrivenPlatformI
 	// if unspecified, use the ini platform name (only Win64 breaks this)
 	if (Info.UBTPlatformName == NAME_None)
 	{
-		Info.UBTPlatformName = *PlatformName;
+		Info.UBTPlatformName = PlatformName;
 	}
 	Info.UBTPlatformString = Info.UBTPlatformName.ToString();
 		
@@ -245,7 +251,7 @@ static void LoadDDPIIniSettings(const FConfigFile& IniFile, FDataDrivenPlatformI
 /**
 * Get the global set of data driven platform information
 */
-const TMap<FString, FDataDrivenPlatformInfo>& FDataDrivenPlatformInfoRegistry::GetAllPlatformInfos()
+const TMap<FName, FDataDrivenPlatformInfo>& FDataDrivenPlatformInfoRegistry::GetAllPlatformInfos()
 {
 	static bool bHasSearchedForPlatforms = false;
 
@@ -261,9 +267,10 @@ const TMap<FString, FDataDrivenPlatformInfo>& FDataDrivenPlatformInfoRegistry::G
 		{
 			// load the .ini file
 			FConfigFile IniFile;
-			FString PlatformName;
-			FDataDrivenPlatformInfoRegistry::LoadDataDrivenIniFile(Index, IniFile, PlatformName);
+			FString PlatformString;
+			FDataDrivenPlatformInfoRegistry::LoadDataDrivenIniFile(Index, IniFile, PlatformString);
 
+			FName PlatformName(*PlatformString);
 			// platform info is registered by the platform name
 			if (IniFile.Contains(TEXT("DataDrivenPlatformInfo")))
 			{
@@ -274,7 +281,7 @@ const TMap<FString, FDataDrivenPlatformInfo>& FDataDrivenPlatformInfoRegistry::G
 				// get the parent to build list later
 				FString IniParent;
 				IniFile.GetString(TEXT("DataDrivenPlatformInfo"), TEXT("IniParent"), IniParent);
-				IniParents.Add(PlatformName, IniParent);
+				IniParents.Add(PlatformString, IniParent);
 			}
 		}
 
@@ -282,17 +289,48 @@ const TMap<FString, FDataDrivenPlatformInfo>& FDataDrivenPlatformInfoRegistry::G
 		for (auto& It : DataDrivenPlatforms)
 		{
 			// walk up the chain and build up the ini chain of parents
-			for (FString CurrentPlatform = IniParents.FindRef(It.Key); CurrentPlatform != TEXT(""); CurrentPlatform = IniParents.FindRef(CurrentPlatform))
+			for (FString CurrentPlatform = IniParents.FindRef(It.Key.ToString()); CurrentPlatform != TEXT(""); CurrentPlatform = IniParents.FindRef(CurrentPlatform))
 			{
 				// insert at 0 to reverse the order
 				It.Value.IniParentChain.Insert(CurrentPlatform, 0);
 			}
 		}
+
+		DataDrivenPlatforms.GetKeys(SortedPlatformNames);
+		// now sort them into arrays of keys and values
+		Algo::Sort(SortedPlatformNames, [](FName One, FName Two) -> bool
+		{
+			return One.Compare(Two) < 0;
+		});
+
+		// now build list of values from the sort
+		SortedPlatformInfos.Empty(SortedPlatformNames.Num());
+		for (int Index = 0; Index < SortedPlatformInfos.Num(); Index++)
+		{
+			SortedPlatformInfos[Index] = &DataDrivenPlatforms[SortedPlatformNames[Index]];
+		}
 	}
 
 	return DataDrivenPlatforms;
 }
 
+const TArray<FName> FDataDrivenPlatformInfoRegistry::GetSortedPlatformNames()
+{
+	// make sure we've read in the inis
+	GetAllPlatformInfos();
+
+	return SortedPlatformNames;
+}
+
+const TArray<const FDataDrivenPlatformInfo*>& FDataDrivenPlatformInfoRegistry::GetSortedPlatformInfos()
+{
+	// make sure we've read in the inis
+	GetAllPlatformInfos();
+
+	return SortedPlatformInfos;
+}
+
+
 #if DDPI_HAS_EXTENDED_PLATFORMINFO_DATA
 
 // some shared functionality
@@ -349,6 +387,9 @@ static FString ConvertToDDPIDeviceId(const FString& DeviceId)
 
 void FDataDrivenPlatformInfoRegistry::UpdateSdkStatus()
 {
+	// make sure we've read in the inis
+	GetAllPlatformInfos();
+
 	// don't run UAT from commandlets (like the cooker) that are often launched from UAT and this will go poorly
 	if (IsRunningCommandlet())
 	{
@@ -366,7 +407,7 @@ void FDataDrivenPlatformInfoRegistry::UpdateSdkStatus()
 
 	FString Command, BaseCommandline, ReportFilename;
 	PrepForTurnkeyReport(Command, BaseCommandline, ReportFilename);
-	FString Commandline = BaseCommandline + FString(TEXT(" -platform=")) + FString::JoinBy(DataDrivenPlatforms, TEXT("+"), [](TPair<FString, FDataDrivenPlatformInfo> Pair) { return ConvertToUATPlatform(Pair.Key); });
+	FString Commandline = BaseCommandline + FString(TEXT(" -platform=")) + FString::JoinBy(DataDrivenPlatforms, TEXT("+"), [](TPair<FName, FDataDrivenPlatformInfo> Pair) { return ConvertToUATPlatform(Pair.Key.ToString()); });
 
 	// reset status to unknown
 	for (auto& It : DataDrivenPlatforms)
@@ -412,7 +453,7 @@ void FDataDrivenPlatformInfoRegistry::UpdateSdkStatus()
 
 						// have to convert back to WIndows from Win64
 						FString PlatformName = ConvertToDDPIPlatform(Tokens[0]);
-						DataDrivenPlatforms[PlatformName].SdkStatus = Status;
+						DataDrivenPlatforms[*PlatformName].SdkStatus = Status;
 
 						UE_LOG(LogTemp, Log, TEXT("Turnkey Platform: %s - %d - %s"), *PlatformName, (int)Status, *Tokens[2]);
 					}
@@ -466,7 +507,7 @@ FDataDrivenPlatformInfo& FDataDrivenPlatformInfoRegistry::DeviceIdToInfo(FString
 	}
 
 	// have to convert back to Windows from Win64
-	return FDataDrivenPlatformInfoRegistry::DataDrivenPlatforms[ConvertToDDPIPlatform(PlatformAndDevice[0])];
+	return DataDrivenPlatforms[*ConvertToDDPIPlatform(PlatformAndDevice[0])];
 
 }
 
@@ -541,9 +582,9 @@ void FDataDrivenPlatformInfoRegistry::UpdateDeviceSdkStatus(TArray<FString> Plat
 	TurnkeyProcess->Launch();
 }
 
-void FDataDrivenPlatformInfoRegistry::ClearDeviceStatus(const FString& PlatformName)
+void FDataDrivenPlatformInfoRegistry::ClearDeviceStatus(FName PlatformName)
 {
-	if (PlatformName.Len() == 0)
+	if (PlatformName != NAME_None)
 	{
 		for (auto It : DataDrivenPlatforms)
 		{
@@ -583,7 +624,7 @@ const TArray<FString>& FDataDrivenPlatformInfoRegistry::GetValidPlatformDirector
 		bHasSearchedForPlatforms = true;
 
 		// look for possible platforms
-		const TMap<FString, FDataDrivenPlatformInfo>& Infos = GetAllPlatformInfos();
+		const TMap<FName, FDataDrivenPlatformInfo>& Infos = GetAllPlatformInfos();
 		for (auto Pair : Infos)
 		{
 #if DDPI_HAS_EXTENDED_PLATFORMINFO_DATA
@@ -595,7 +636,7 @@ const TArray<FString>& FDataDrivenPlatformInfoRegistry::GetValidPlatformDirector
 #endif
 
 			// add ourself as valid
-			ValidPlatformDirectories.AddUnique(Pair.Key);
+			ValidPlatformDirectories.AddUnique(Pair.Key.ToString());
 
 			// now add additional directories
 			for (FString& AdditionalDir : Pair.Value.AdditionalRestrictedFolders)
@@ -609,17 +650,16 @@ const TArray<FString>& FDataDrivenPlatformInfoRegistry::GetValidPlatformDirector
 }
 
 
-const FDataDrivenPlatformInfo& FDataDrivenPlatformInfoRegistry::GetPlatformInfo(const FString& PlatformName)
+const FDataDrivenPlatformInfo& FDataDrivenPlatformInfoRegistry::GetPlatformInfo(FName PlatformName)
 {
 	const FDataDrivenPlatformInfo* Info = GetAllPlatformInfos().Find(PlatformName);
 	static FDataDrivenPlatformInfo Empty;
 	return Info ? *Info : Empty;
 }
 
-const FDataDrivenPlatformInfo& FDataDrivenPlatformInfoRegistry::GetPlatformInfo(FName PlatformName)
+const FDataDrivenPlatformInfo& FDataDrivenPlatformInfoRegistry::GetPlatformInfo(const FString& PlatformName)
 {
-	// @todo maybe make the key be FName instead of FString and flip these?
-	return GetPlatformInfo(PlatformName.ToString());
+	return GetPlatformInfo(FName(*PlatformName));
 }
 
 const FDataDrivenPlatformInfo& FDataDrivenPlatformInfoRegistry::GetPlatformInfo(const char* PlatformName)
@@ -628,10 +668,10 @@ const FDataDrivenPlatformInfo& FDataDrivenPlatformInfoRegistry::GetPlatformInfo(
 }
 
 
-const TArray<FString>& FDataDrivenPlatformInfoRegistry::GetConfidentialPlatforms()
+const TArray<FName>& FDataDrivenPlatformInfoRegistry::GetConfidentialPlatforms()
 {
 	static bool bHasSearchedForPlatforms = false;
-	static TArray<FString> FoundPlatforms;
+	static TArray<FName> FoundPlatforms;
 
 	// look on disk for special files
 	if (bHasSearchedForPlatforms == false)
@@ -653,7 +693,7 @@ const TArray<FString>& FDataDrivenPlatformInfoRegistry::GetConfidentialPlatforms
 
 
 #if DDPI_HAS_EXTENDED_PLATFORMINFO_DATA
-bool FDataDrivenPlatformInfoRegistry::HasCompiledSupportForPlatform(const FString& PlatformName, EPlatformNameType PlatformNameType)
+bool FDataDrivenPlatformInfoRegistry::HasCompiledSupportForPlatform(FName PlatformName, EPlatformNameType PlatformNameType)
 {
 	if (PlatformNameType == EPlatformNameType::Ini)
 	{
@@ -663,13 +703,11 @@ bool FDataDrivenPlatformInfoRegistry::HasCompiledSupportForPlatform(const FStrin
 	}
 	else if (PlatformNameType == EPlatformNameType::UBT)
 	{
-		FName PlatformFName(*PlatformName);
-
 		// find all the DataDrivenPlatformInfo objects and find a matching the UBT name
 		for (auto& Pair : FDataDrivenPlatformInfoRegistry::GetAllPlatformInfos())
 		{
 			// if this platform matches the UBT platform name, check it's Ini name
-			if (Pair.Value.UBTPlatformName == PlatformFName)
+			if (Pair.Value.UBTPlatformName == PlatformName)
 			{
 				return HasCompiledSupportForPlatform(Pair.Key, EPlatformNameType::Ini);
 			}
@@ -681,13 +719,11 @@ bool FDataDrivenPlatformInfoRegistry::HasCompiledSupportForPlatform(const FStrin
 	{
 		// was this TP compiled, or a shaderformat (useful for SCW if it ever calls this)
 		return 
-			FModuleManager::Get().ModuleExists(*FString::Printf(TEXT("%sTargetPlatform"), *PlatformName)) || 
-			FModuleManager::Get().ModuleExists(*FString::Printf(TEXT("%sShaderFormat"), *PlatformName));
+			FModuleManager::Get().ModuleExists(*FString::Printf(TEXT("%sTargetPlatform"), *PlatformName.ToString())) || 
+			FModuleManager::Get().ModuleExists(*FString::Printf(TEXT("%sShaderFormat"), *PlatformName.ToString()));
 	}
 
 	return false;
 }
 
-static void ClearDeviceStatus(const FString& PlatformName);
-
 #endif
diff --git a/Engine/Source/Runtime/Core/Private/Misc/Paths.cpp b/Engine/Source/Runtime/Core/Private/Misc/Paths.cpp
--- a/Engine/Source/Runtime/Core/Private/Misc/Paths.cpp
+++ b/Engine/Source/Runtime/Core/Private/Misc/Paths.cpp
@@ -651,9 +651,9 @@ const TArray<FString>& FPaths::GetRestrictedFolderNames()
 		StaticData.RestrictedFolderNames.Add(TEXT("EpicInternal"));
 
 		// Add confidential platforms
-		for (const FString& PlatformStr : FDataDrivenPlatformInfoRegistry::GetConfidentialPlatforms())
+		for (FName PlatformStr : FDataDrivenPlatformInfoRegistry::GetConfidentialPlatforms())
 		{
-			StaticData.RestrictedFolderNames.Add(PlatformStr);
+			StaticData.RestrictedFolderNames.Add(PlatformStr.ToString());
 		}
 
 		StaticData.bRestrictedFolderNamesInitialized = true;
diff --git a/Engine/Source/Runtime/Core/Private/Modules/ModuleManager.cpp b/Engine/Source/Runtime/Core/Private/Modules/ModuleManager.cpp
--- a/Engine/Source/Runtime/Core/Private/Modules/ModuleManager.cpp
+++ b/Engine/Source/Runtime/Core/Private/Modules/ModuleManager.cpp
@@ -1166,7 +1166,10 @@ void FModuleManager::AddExtraBinarySearchPaths()
 	{
 		// Ensure that dependency dlls can be found in restricted sub directories
 		TArray<FString> RestrictedFolderNames = { TEXT("NoRedist"), TEXT("NotForLicensees"), TEXT("CarefullyRedist") };
-		RestrictedFolderNames.Append(FDataDrivenPlatformInfoRegistry::GetConfidentialPlatforms());
+		for (FName PlatformName : FDataDrivenPlatformInfoRegistry::GetConfidentialPlatforms())
+		{
+			RestrictedFolderNames.Add(PlatformName.ToString());
+		}
 
 		FString ModuleDir = FPlatformProcess::GetModulesDirectory();
 		for (const FString& RestrictedFolderName : RestrictedFolderNames)
diff --git a/Engine/Source/Runtime/Core/Public/IOS/IOSPlatformProperties.h b/Engine/Source/Runtime/Core/Public/IOS/IOSPlatformProperties.h
--- a/Engine/Source/Runtime/Core/Public/IOS/IOSPlatformProperties.h
+++ b/Engine/Source/Runtime/Core/Public/IOS/IOSPlatformProperties.h
@@ -106,6 +106,23 @@ struct FIOSPlatformProperties
 	}
 };
 
+struct FTVOSPlatformProperties : public FIOSPlatformProperties
+{
+	// @todo breaking change here!
+	static FORCEINLINE const char* PlatformName()
+	{
+		return "TVOS";
+	}
+
+	static FORCEINLINE const char* IniPlatformName()
+	{
+		return "TVOS";
+	}
+};
+
 #ifdef PROPERTY_HEADER_SHOULD_DEFINE_TYPE
+
+
 typedef FIOSPlatformProperties FPlatformProperties;
+
 #endif
diff --git a/Engine/Source/Runtime/Core/Public/Linux/LinuxPlatformProperties.h b/Engine/Source/Runtime/Core/Public/Linux/LinuxPlatformProperties.h
--- a/Engine/Source/Runtime/Core/Public/Linux/LinuxPlatformProperties.h
+++ b/Engine/Source/Runtime/Core/Public/Linux/LinuxPlatformProperties.h
@@ -64,7 +64,7 @@ struct FLinuxPlatformProperties
 
 		if (HAS_EDITOR_DATA)
 		{
-			return "Linux";
+			return "LinuxEditor";
 		}
 
 		if (IS_CLIENT_ONLY)
@@ -72,7 +72,7 @@ struct FLinuxPlatformProperties
 			return IS_AARCH64 ? "LinuxAArch64Client" : "LinuxClient";
 		}
 
-		return IS_AARCH64 ? "LinuxAArch64" : "LinuxNoEditor";
+		return IS_AARCH64 ? "LinuxAArch64" : "Linux";
 	}
 
 	static FORCEINLINE bool RequiresCookedData( )
diff --git a/Engine/Source/Runtime/Core/Public/Mac/MacPlatformProperties.h b/Engine/Source/Runtime/Core/Public/Mac/MacPlatformProperties.h
--- a/Engine/Source/Runtime/Core/Public/Mac/MacPlatformProperties.h
+++ b/Engine/Source/Runtime/Core/Public/Mac/MacPlatformProperties.h
@@ -58,7 +58,7 @@ struct FMacPlatformProperties
 		
 		if (HAS_EDITOR_DATA)
 		{
-			return "Mac";
+			return "MacEditor";
 		}
 
 		if (IS_CLIENT_ONLY)
@@ -66,7 +66,7 @@ struct FMacPlatformProperties
 			return "MacClient";
 		}
 
-		return "MacNoEditor";
+		return "Mac";
 	}
 
 	static FORCEINLINE bool RequiresCookedData( )
diff --git a/Engine/Source/Runtime/Core/Public/Misc/DataDrivenPlatformInfoRegistry.h b/Engine/Source/Runtime/Core/Public/Misc/DataDrivenPlatformInfoRegistry.h
--- a/Engine/Source/Runtime/Core/Public/Misc/DataDrivenPlatformInfoRegistry.h
+++ b/Engine/Source/Runtime/Core/Public/Misc/DataDrivenPlatformInfoRegistry.h
@@ -97,6 +97,9 @@ struct FPlatformIconPaths
 // Information about a platform loaded from disk
 struct FDataDrivenPlatformInfo
 {
+	// copy of the platform name, same as the Key into GetAllPlatformInfos()
+	FName IniPlatformName;
+
 	// cached list of ini parents
 	TArray<FString> IniParentChain;
 
@@ -229,7 +232,7 @@ struct CORE_API FDataDrivenPlatformInfoRegistry
 	/**
 	* Get the global set of data driven platform information
 	*/
-	static const TMap<FString, FDataDrivenPlatformInfo>& GetAllPlatformInfos();
+	static const TMap<FName, FDataDrivenPlatformInfo>& GetAllPlatformInfos();
 
 	/**
 	 * Gets a set of platform names based on GetAllPlatformInfos, their AdditionalRestrictedFolders, and possibly filtered based on what editor has support compiled for
@@ -245,11 +248,15 @@ struct CORE_API FDataDrivenPlatformInfoRegistry
 	static const FDataDrivenPlatformInfo& GetPlatformInfo(FName PlatformName);
 	static const FDataDrivenPlatformInfo& GetPlatformInfo(const char* PlatformName);
 
+	// get just names or just infos, 
+	static const TArray<FName> GetSortedPlatformNames();
+	static const TArray<const FDataDrivenPlatformInfo*>& GetSortedPlatformInfos();
+
 	/**
 	 * Gets a list of all known confidential platforms (note these are just the platforms you have access to, so, for example PS4 won't be
 	 * returned if you are not a PS4 licensee)
 	 */
-	static const TArray<FString>& GetConfidentialPlatforms();
+	static const TArray<FName>& GetConfidentialPlatforms();
 
 	/**
 	 * Returns the number of discovered ini files that can be loaded with LoadDataDrivenIniFile
@@ -273,10 +280,10 @@ struct CORE_API FDataDrivenPlatformInfoRegistry
 		UBT,
 		// for instance Windows
 		Ini,
-		// for instance WindowsNoEditor
+		// for instance WindowsClient
 		TargetPlatform,
 	};
-	static bool HasCompiledSupportForPlatform(const FString& PlatformName, EPlatformNameType PlatformNameType);
+	static bool HasCompiledSupportForPlatform(FName PlatformName, EPlatformNameType PlatformNameType);
 
 	/**
 	 * Runs UAT to update Sdk status (can be called after user chooses an option that had a bad Sdk status)
@@ -291,15 +298,11 @@ struct CORE_API FDataDrivenPlatformInfoRegistry
 	/**
 	 * Wipes out cached device status for all devices in a platform (or all platforms if PlatformName is empty)
 	 */
-	static void ClearDeviceStatus(const FString& PlatformName);
+	static void ClearDeviceStatus(FName PlatformName);
 
 private:
 	static FDataDrivenPlatformInfo& DeviceIdToInfo(FString DeviceId, FString* OutDeviceName = nullptr);
 
 #endif
-
-
-private:
-	static TMap<FString, FDataDrivenPlatformInfo> DataDrivenPlatforms;
 };
 
diff --git a/Engine/Source/Runtime/Core/Public/Windows/WindowsPlatformProperties.h b/Engine/Source/Runtime/Core/Public/Windows/WindowsPlatformProperties.h
--- a/Engine/Source/Runtime/Core/Public/Windows/WindowsPlatformProperties.h
+++ b/Engine/Source/Runtime/Core/Public/Windows/WindowsPlatformProperties.h
@@ -57,7 +57,7 @@ struct FWindowsPlatformProperties
 		
 		if (HAS_EDITOR_DATA)
 		{
-			return "Windows";
+			return "WindowsEditor";
 		}
 		
 		if (IS_CLIENT_ONLY)
@@ -65,7 +65,7 @@ struct FWindowsPlatformProperties
 			return "WindowsClient";
 		}
 
-		return "WindowsNoEditor";
+		return "Windows";
 	}
 
 	static FORCEINLINE bool RequiresCookedData()
diff --git a/Engine/Source/Runtime/Engine/Classes/Tests/AutomationTestSettings.h b/Engine/Source/Runtime/Engine/Classes/Tests/AutomationTestSettings.h
--- a/Engine/Source/Runtime/Engine/Classes/Tests/AutomationTestSettings.h
+++ b/Engine/Source/Runtime/Engine/Classes/Tests/AutomationTestSettings.h
@@ -327,7 +327,7 @@ struct FLaunchOnTestSettings
 	FFilePath LaunchOnTestmap;
 
 	/** Device to be used **/
-	UPROPERTY(config, EditAnywhere, Category = Automation, meta = (ToolTip = "This is the device to be used for launch on. Example: WindowsNoEditor, Android, IOS, Linux"))
+	UPROPERTY(config, EditAnywhere, Category = Automation, meta = (ToolTip = "This is the device to be used for launch on. Example: WindowsClient, Android, IOS, Linux"))
 	FString DeviceID;
 };
 
diff --git a/Engine/Source/Runtime/Engine/Private/DeviceProfiles/DeviceProfileManager.cpp b/Engine/Source/Runtime/Engine/Private/DeviceProfiles/DeviceProfileManager.cpp
--- a/Engine/Source/Runtime/Engine/Private/DeviceProfiles/DeviceProfileManager.cpp
+++ b/Engine/Source/Runtime/Engine/Private/DeviceProfiles/DeviceProfileManager.cpp
@@ -531,9 +531,9 @@ void UDeviceProfileManager::LoadProfiles()
 	if( !HasAnyFlags( RF_ClassDefaultObject ) )
 	{
 		TMap<FString, FString> DeviceProfileToPlatformConfigMap;
-		TArray<FString> ConfidentialPlatforms = FDataDrivenPlatformInfoRegistry::GetConfidentialPlatforms();
+		TArray<FName> ConfidentialPlatforms = FDataDrivenPlatformInfoRegistry::GetConfidentialPlatforms();
 		
-		checkf(ConfidentialPlatforms.Contains(FString(FPlatformProperties::IniPlatformName())) == false,
+		checkf(ConfidentialPlatforms.Contains(FPlatformProperties::IniPlatformName()) == false,
 			TEXT("UDeviceProfileManager::LoadProfiles is called from a confidential platform (%s). Confidential platforms are not expected to be editor/non-cooked builds."), 
 			ANSI_TO_TCHAR(FPlatformProperties::IniPlatformName()));
 
@@ -541,7 +541,7 @@ void UDeviceProfileManager::LoadProfiles()
 		for (int32 PlatformIndex = 0; PlatformIndex <= ConfidentialPlatforms.Num(); PlatformIndex++)
 		{
 			// which platform's set of ini files should we load from?
-			FString ConfigLoadPlatform = PlatformIndex == 0 ? FString(FPlatformProperties::IniPlatformName()) : ConfidentialPlatforms[PlatformIndex - 1];
+			FString ConfigLoadPlatform = PlatformIndex == 0 ? FString(FPlatformProperties::IniPlatformName()) : ConfidentialPlatforms[PlatformIndex - 1].ToString();
 
 			// load the DP.ini files (from current platform and then by the extra confidential platforms)
 			FConfigFile PlatformConfigFile;
@@ -594,7 +594,7 @@ void UDeviceProfileManager::LoadProfiles()
 				ITargetPlatform* Platform = TargetPlatforms[PlatformIndex];
 
 				// Set TextureLODSettings
-				const UTextureLODSettings* TextureLODSettingsObj = FindProfile(*Platform->GetTargetPlatformInfo().VanillaInfo->PlatformInfoName.ToString(), false);
+				const UTextureLODSettings* TextureLODSettingsObj = FindProfile(*Platform->PlatformName(), false);
 				Platform->RegisterTextureLODSettings(TextureLODSettingsObj);
 			}
 		}
diff --git a/Engine/Source/Runtime/Engine/Private/UnrealEngine.cpp b/Engine/Source/Runtime/Engine/Private/UnrealEngine.cpp
--- a/Engine/Source/Runtime/Engine/Private/UnrealEngine.cpp
+++ b/Engine/Source/Runtime/Engine/Private/UnrealEngine.cpp
@@ -15624,7 +15624,7 @@ static FAutoConsoleCommand SetupThreadAffinityCmd(
 
 static void PakFileTest(const TArray<FString>& Args)
 {
-	FString PakFilename = TEXT("D:\\work\\Dev-Core\\Samples\\Games\\ShooterGame\\Saved\\StagedBuilds\\WindowsNoEditor\\ShooterGame\\Content\\Paks\\test.pak");
+	FString PakFilename = TEXT("D:\\work\\Dev-Core\\Samples\\Games\\ShooterGame\\Saved\\StagedBuilds\\Windows\\ShooterGame\\Content\\Paks\\test.pak");
 	if (Args.Num() < 1)
 	{
 		UE_LOG(LogConsoleResponse, Error, TEXT("Usage: PakFileTest path-to-pak-file"));
diff --git a/Engine/Source/Runtime/Engine/Private/VT/VirtualTextureDataBuilder.cpp b/Engine/Source/Runtime/Engine/Private/VT/VirtualTextureDataBuilder.cpp
--- a/Engine/Source/Runtime/Engine/Private/VT/VirtualTextureDataBuilder.cpp
+++ b/Engine/Source/Runtime/Engine/Private/VT/VirtualTextureDataBuilder.cpp
@@ -790,9 +790,9 @@ static const FName RemovePlatformPrefixFromName(FName const& InName)
 {
 	FString NameString = InName.ToString();
 
-	for (const auto& PlatformInfo : FDataDrivenPlatformInfoRegistry::GetAllPlatformInfos())
+	for (FName PlatformName : FDataDrivenPlatformInfoRegistry::GetSortedPlatformNames())
 	{
-		FString PlatformTextureFormatPrefix = PlatformInfo.Key;
+		FString PlatformTextureFormatPrefix = PlatformName.ToString();
 		PlatformTextureFormatPrefix += TEXT('_');
 		if (NameString.StartsWith(PlatformTextureFormatPrefix, ESearchCase::IgnoreCase))
 		{
diff --git a/Engine/Source/Runtime/Engine/Public/PerPlatformProperties.h b/Engine/Source/Runtime/Engine/Public/PerPlatformProperties.h
--- a/Engine/Source/Runtime/Engine/Public/PerPlatformProperties.h
+++ b/Engine/Source/Runtime/Engine/Public/PerPlatformProperties.h
@@ -30,7 +30,7 @@ struct ENGINE_API TPerPlatformProperty
 	typedef _ValueType ValueType;
 
 #if WITH_EDITOR
-	/** Get the value for the given platform (using standard "ini" name, so Windows, not Win64 or WindowsNoEditor), which can be used to lookup the group */
+	/** Get the value for the given platform (using standard "ini" name, so Windows, not Win64 or WindowsClient), which can be used to lookup the group */
 	_ValueType GetValueForPlatform(FName PlatformName) const
 	{
 		const _StructType* This = StaticCast<const _StructType*>(this);
diff --git a/Engine/Source/Runtime/Projects/Public/Interfaces/IProjectManager.h b/Engine/Source/Runtime/Projects/Public/Interfaces/IProjectManager.h
--- a/Engine/Source/Runtime/Projects/Public/Interfaces/IProjectManager.h
+++ b/Engine/Source/Runtime/Projects/Public/Interfaces/IProjectManager.h
@@ -44,8 +44,9 @@ public:
 
 	/**
 	 * Check to see if the given platform name is supported as a target by the current project
+	 * @todo platformcleanup: Make this use IniPlatformName, not TargetPlatformName!! And other related functions in here
 	 *
-	 * @param	InPlatformName				Name of the platform to target (eg, WindowsNoEditor)
+	 * @param	InPlatformName				Name of the platform to target (eg, WindowsClient)
 	 * @param	bAllowSupportedIfEmptyList	Consider an empty list to mean that all platforms are supported targets?
 	 */
 	bool IsTargetPlatformSupported(const FName& InPlatformName, const bool bAllowSupportedIfEmptyList = true) const
@@ -159,7 +160,7 @@ public:
 	 * Update the list of supported target platforms for the target project based upon the parameters provided
 	 * 
 	 * @param	FilePath			The filepath where the project is stored.
-	 * @param	InPlatformName		Name of the platform to target (eg, WindowsNoEditor)
+	 * @param	InPlatformName		Name of the platform to target (eg, WindowsClient)
 	 * @param	bIsSupported		true if the platform should be supported by this project, false if it should not
 	 */
 	virtual void UpdateSupportedTargetPlatformsForProject(const FString& FilePath, const FName& InPlatformName, const bool bIsSupported) = 0;
@@ -167,7 +168,7 @@ public:
 	/**
 	 * Update the list of supported target platforms for the current project based upon the parameters provided
 	 * 
-	 * @param	InPlatformName		Name of the platform to target (eg, WindowsNoEditor)
+	 * @param	InPlatformName		Name of the platform to target (eg, WindowsClient)
 	 * @param	bIsSupported		true if the platform should be supported by this project, false if it should not
 	 */
 	virtual void UpdateSupportedTargetPlatformsForCurrentProject(const FName& InPlatformName, const bool bIsSupported) = 0;
