- Redid a change from UE4 to fix COTF server crash

diff --git a/Engine/Source/Editor/UnrealEd/Private/CookOnTheFlyServer.cpp b/Engine/Source/Editor/UnrealEd/Private/CookOnTheFlyServer.cpp
--- a/Engine/Source/Editor/UnrealEd/Private/CookOnTheFlyServer.cpp
+++ b/Engine/Source/Editor/UnrealEd/Private/CookOnTheFlyServer.cpp
@@ -3355,7 +3355,7 @@ void UCookOnTheFlyServer::SaveCookedPackage(UE::Cook::FPackageData& PackageData,
 	// Don't resolve, just add to request list as needed
 	TSet<FName> SoftObjectPackages;
 
-	if (!CookByTheBookOptions->bSkipSoftReferences)
+	if (!IsCookByTheBookMode() || !CookByTheBookOptions->bSkipSoftReferences)
 	{
 		GRedirectCollector.ProcessSoftObjectPathPackageList(Package->GetFName(), false, SoftObjectPackages);
 
