Collapsed LevelViewportLayouts into EditorViewportLayouts.

#rnx

diff --git a/Engine/Source/Editor/LevelEditor/Private/LevelViewportLayout.cpp b/Engine/Source/Editor/LevelEditor/Private/LevelViewportLayout.cpp
--- a/Engine/Source/Editor/LevelEditor/Private/LevelViewportLayout.cpp
+++ b/Engine/Source/Editor/LevelEditor/Private/LevelViewportLayout.cpp
@@ -59,10 +59,8 @@ FLevelViewportLayout::FLevelViewportLayout()
 
 FLevelViewportLayout::~FLevelViewportLayout()
 {
-	for (auto& Pair : Viewports)
-	{
-		Pair.Value->OnLayoutDestroyed();
-	}
+	FLevelEditorModule& LevelEditor = FModuleManager::GetModuleChecked<FLevelEditorModule>(LevelEditorModName);
+	LevelEditor.OnTakeHighResScreenShots().RemoveAll(this);
 
 	// Make sure that we're not locking the immersive window after we go away
 	if( bIsImmersive || ( bWasImmersive && bIsTransitioning ) )
@@ -73,40 +71,51 @@ FLevelViewportLayout::~FLevelViewportLayout()
 			OwnerWindow->SetFullWindowOverlayContent( NULL );
 		}
 	}
-
-	FLevelEditorModule& LevelEditor = FModuleManager::GetModuleChecked<FLevelEditorModule>(LevelEditorModName);
-	LevelEditor.OnTakeHighResScreenShots().RemoveAll(this);
 }
 
 
 TSharedRef<SWidget> FLevelViewportLayout::BuildViewportLayout(TSharedPtr<SDockTab> InParentDockTab, TSharedPtr<FEditorViewportTabContent> InParentTab, const FString& LayoutString)
 {
-	MaximizedViewport = NAME_None;
+	TSharedRef<SWidget> ViewportLayoutWidget = FAssetEditorViewportLayout::BuildViewportLayout(InParentDockTab, InParentTab, LayoutString);
 
 	// Important: We use raw bindings here because we are releasing our binding in our destructor (where a weak pointer would be invalid)
 	// It's imperative that our delegate is removed in the destructor for the level editor module to play nicely with reloading.
 	FLevelEditorModule& LevelEditor = FModuleManager::GetModuleChecked<FLevelEditorModule>(LevelEditorModName);
 	LevelEditor.OnTakeHighResScreenShots().AddRaw(this, &FLevelViewportLayout::TakeHighResScreenShot);
 
-	return FAssetEditorViewportLayout::BuildViewportLayout(InParentDockTab, InParentTab, LayoutString);
+	// Prevent maximize if we only have a single viewport
+	bIsMaximizeSupported = (Viewports.Num() > 1);
+
+	return ViewportLayoutWidget;
 }
 
-TSharedRef<IEditorViewportLayoutEntity> FLevelViewportLayout::FactoryViewport(FName InTypeName, const FAssetEditorViewportConstructionArgs& ConstructionArgs) const
+TSharedRef<SWidget> FLevelViewportLayout::FactoryViewport(FName InTypeName, const FAssetEditorViewportConstructionArgs& ConstructionArgs)
 {
+	TSharedPtr<IEditorViewportLayoutEntity> ViewportLayoutEntity;
 	TSharedPtr<FEditorViewportTabContent> PinnedTabContent = ParentTabContent.Pin();
 	if (PinnedTabContent.IsValid())
 	{
+		// Manually use the factory function here based on type, because legacy viewport types don't register with our factory functions
+		// The level editor module will return an appropriate default if the legacy lookup fails too.
 		if (const AssetEditorViewportFactoryFunction* FactoryFunc = PinnedTabContent->FindViewportCreationFactory(InTypeName))
 		{
 			TSharedPtr<SAssetEditorViewport> EditorViewport = (*FactoryFunc)(ConstructionArgs);
-			return MakeShareable(new FLevelViewportLayoutEntity(EditorViewport));
+			ViewportLayoutEntity = MakeShareable(new FLevelViewportLayoutEntity(EditorViewport));
 		}
 	}
-
 	PRAGMA_DISABLE_DEPRECATION_WARNINGS
-	FLevelEditorModule& LevelEditor = FModuleManager::GetModuleChecked<FLevelEditorModule>(LevelEditorModName);
-	return LevelEditor.FactoryViewport(InTypeName, ConstructionArgs);
+	if (!ViewportLayoutEntity.IsValid())
+	{
+		FLevelEditorModule& LevelEditor = FModuleManager::GetModuleChecked<FLevelEditorModule>(LevelEditorModName);
+		ViewportLayoutEntity = LevelEditor.FactoryViewport(InTypeName, ConstructionArgs);
+	}
 	PRAGMA_ENABLE_DEPRECATION_WARNINGS
+	
+	// Both the legacy level editor factory viewport, and the viewport creation functions should fall back to a valid default
+	check(ViewportLayoutEntity.IsValid());
+	
+	Viewports.Add(InTypeName, ViewportLayoutEntity);
+	return ViewportLayoutEntity->AsWidget();
 }
 
 void FLevelViewportLayout::BeginThrottleForAnimatedResize()
@@ -122,7 +131,7 @@ void FLevelViewportLayout::BeginThrottleForAnimatedResize()
 }
 
 
-void  FLevelViewportLayout::EndThrottleForAnimatedResize()
+void FLevelViewportLayout::EndThrottleForAnimatedResize()
 {
 	// Only leave this mode if there is a request
 	if( ViewportResizeThrottleRequest.IsValid() )
@@ -131,11 +140,14 @@ void  FLevelViewportLayout::EndThrottleForAnimatedResize()
 	}
 }
 
-void FLevelViewportLayout::InitCommonLayoutFromString( const FString& SpecificLayoutString, FName DefaultMaximizedViewport )
+void FLevelViewportLayout::InitCommonLayoutFromString( const FString& SpecificLayoutString, const FName PerspectiveViewportKey )
 {
+	FName DefaultMaximizedViewport = PerspectiveViewportKey;
+	MaximizedViewport = NAME_None;
+
 	bool bShouldBeMaximized = bIsMaximizeSupported && ViewportLayoutDefs::bDefaultShouldBeMaximized;
 	bool bShouldBeImmersive = ViewportLayoutDefs::bDefaultShouldBeImmersive;
-	
+
 	if (!SpecificLayoutString.IsEmpty())
 	{
 		const FString& IniSection = FLayoutSaveRestore::GetAdditionalLayoutConfigIni();
@@ -217,6 +229,7 @@ void FLevelViewportLayout::MaximizeViewport( FName ViewportToMaximize, const boo
 	// I.E Maximized viewport is NULL which means this is a new maximize or MaximizeViewport is equal to the passed in one which means this is a restore of the current maximized viewport
 	check( Entity.IsValid() );
 	check( MaximizedViewport.IsNone() || MaximizedViewport == ViewportToMaximize );
+	check(LayoutConfiguration.IsValid());
 
 	// If we're already in immersive mode, toggling maximize just needs to update some state (no visual change)
 	if( bIsImmersive )
@@ -337,7 +350,7 @@ void FLevelViewportLayout::MaximizeViewport( FName ViewportToMaximize, const boo
 				{
 					// Replace our viewport with a dummy widget in it's place during the maximize transition.  We can't
 					// have a single viewport widget in two places at once!
-					ReplaceWidget( MaximizedEntity->AsWidget(), ViewportReplacementWidget.ToSharedRef() );
+					LayoutConfiguration->ReplaceWidget( MaximizedEntity->AsWidget(), ViewportReplacementWidget.ToSharedRef() );
 
 					TSharedRef< SCanvas > Canvas = SNew( SCanvas );
 					Canvas->AddSlot()
@@ -545,6 +558,7 @@ void FLevelViewportLayout::FinishMaximizeTransition()
 	{
 		TSharedPtr<ILevelViewportLayoutEntity> MaximizedViewportEntity = StaticCastSharedPtr<ILevelViewportLayoutEntity>(Viewports.FindRef(MaximizedViewport));
 		check(MaximizedViewportEntity.IsValid());
+		check(LayoutConfiguration.IsValid());
 
 		// The transition animation is complete, allow the engine to tick normally
 		EndThrottleForAnimatedResize();
@@ -614,7 +628,7 @@ void FLevelViewportLayout::FinishMaximizeTransition()
 			ViewportsOverlayWidget.Reset();
 
 			// Restore the viewport widget into the viewport layout splitter
-			ReplaceWidget( ViewportReplacementWidget.ToSharedRef(), MaximizedViewportEntity->AsWidget() );
+			LayoutConfiguration->ReplaceWidget( ViewportReplacementWidget.ToSharedRef(), MaximizedViewportEntity->AsWidget() );
 
 			MaximizedViewport = NAME_None;
 		}
@@ -670,3 +684,33 @@ bool FLevelViewportLayout::IsTickable() const
 {
 	return DeferredMaximizeCommands.Num() > 0 || ( bIsTransitioning && !MaximizeAnimation.IsPlaying() );
 }
+
+void FLevelViewportLayout::LoadConfig(const FString& LayoutString)
+{
+	FAssetEditorViewportLayout::LoadConfig(LayoutString);
+
+	if (!LayoutConfiguration.IsValid())
+	{
+		return;
+	}
+
+	LayoutConfiguration->LoadConfig(LayoutString, [this](const FString& SpecificLayoutString, const FName PerspectiveViewportName)
+	{
+		this->InitCommonLayoutFromString(SpecificLayoutString, PerspectiveViewportName);
+	});
+}
+
+void FLevelViewportLayout::SaveConfig(const FString& LayoutString) const
+{
+	FAssetEditorViewportLayout::SaveConfig(LayoutString);
+
+	if (IsTransitioning() || !LayoutConfiguration.IsValid())
+	{
+		return;
+	}
+
+	LayoutConfiguration->SaveConfig(LayoutString, [this](const FString& SpecificLayoutString)
+	{
+		this->SaveCommonLayoutString(SpecificLayoutString);
+	});
+}
diff --git a/Engine/Source/Editor/LevelEditor/Private/LevelViewportLayout2x2.cpp b/Engine/Source/Editor/LevelEditor/Private/LevelViewportLayout2x2.cpp
deleted file mode 100644
diff --git a/Engine/Source/Editor/LevelEditor/Private/LevelViewportLayout2x2.h b/Engine/Source/Editor/LevelEditor/Private/LevelViewportLayout2x2.h
--- a/Engine/Source/Editor/LevelEditor/Private/LevelViewportLayout2x2.h
+++ b/Engine/Source/Editor/LevelEditor/Private/LevelViewportLayout2x2.h
@@ -2,32 +2,7 @@
 
 #pragma once
 
-#include "CoreMinimal.h"
-#include "Widgets/SWidget.h"
-#include "LevelViewportLayout.h"
-#include "Widgets/Layout/SSplitter.h"
-#include "LevelViewportActions.h"
+#include "EditorViewportLayout2x2.h"
 
-class FLevelViewportLayout2x2 : public FLevelViewportLayout
-{
-public:
-	/**
-	 * Saves viewport layout information between editor sessions
-	 */
-	virtual void SaveLayoutString(const FString& LayoutString) const override;
-
-	virtual const FName& GetLayoutTypeName() const override { return LevelViewportConfigurationNames::FourPanes2x2; }
-protected:
-	/**
-	 * Creates the viewports and splitter for the 2x2 layout                   
-	 */
-	virtual TSharedRef<SWidget> MakeViewportLayout(const FString& LayoutString) override;
-
-	/** Overridden from FLevelViewportLayout */
-	virtual void ReplaceWidget( TSharedRef< SWidget > Source, TSharedRef< SWidget > Replacement ) override;
-
-
-private:
-	/** The splitter widget */
-	TSharedPtr< class SSplitter2x2 > SplitterWidget;
-};
+UE_DEPRECATED(5.0, "Level Pane layouts are now owned by their level layouts. Use FEditorViewportLayout2x2 instead, and add any custom types to FAssetEditorViewportLayout::FactoryPaneConfigurationFromTypeName.")
+typedef FEditorViewportLayout2x2 FLevelViewportLayout2x2;
diff --git a/Engine/Source/Editor/LevelEditor/Private/LevelViewportLayoutFourPanes.cpp b/Engine/Source/Editor/LevelEditor/Private/LevelViewportLayoutFourPanes.cpp
deleted file mode 100644
diff --git a/Engine/Source/Editor/LevelEditor/Private/LevelViewportLayoutFourPanes.h b/Engine/Source/Editor/LevelEditor/Private/LevelViewportLayoutFourPanes.h
--- a/Engine/Source/Editor/LevelEditor/Private/LevelViewportLayoutFourPanes.h
+++ b/Engine/Source/Editor/LevelEditor/Private/LevelViewportLayoutFourPanes.h
@@ -2,92 +2,16 @@
 
 #pragma once
 
-#include "CoreMinimal.h"
-#include "Widgets/SWidget.h"
-#include "LevelViewportLayout.h"
-#include "Widgets/Layout/SSplitter.h"
-#include "LevelViewportActions.h"
+#include "EditorViewportLayoutFourPanes.h"
 
-class FLevelViewportLayoutFourPanes : public FLevelViewportLayout
-{
-public:
-	/**
-	 * Saves viewport layout information between editor sessions
-	 */
-	virtual void SaveLayoutString(const FString& LayoutString) const override;
-protected:
-	/**
-	 * Creates the viewports and splitter for the four-pane layout              
-	 */
-	virtual TSharedRef<SWidget> MakeViewportLayout(const FString& LayoutString) override;
+UE_DEPRECATED(5.0, "Level Pane layouts are now owned by their level layouts. Use FEditorViewportLayoutFourPanesLeft instead, and add any custom types to FAssetEditorViewportLayout::FactoryPaneConfigurationFromTypeName.")
+typedef FEditorViewportLayoutFourPanesLeft FLevelViewportLayoutFourPanesLeft;
 
-	virtual TSharedRef<SWidget> MakeFourPanelWidget(
-		TMap<FName, TSharedPtr< IEditorViewportLayoutEntity >>& ViewportWidgets,
-		TSharedRef<SWidget> Viewport0, TSharedRef<SWidget> Viewport1, TSharedRef<SWidget> Viewport2, TSharedRef<SWidget> Viewport3,
-		float PrimarySplitterPercentage, float SecondarySplitterPercentage0, float SecondarySplitterPercentage1) = 0;
+UE_DEPRECATED(5.0, "Level Pane layouts are now owned by their level layouts. Use FEditorViewportLayoutFourPanesRight instead, and add any custom types to FAssetEditorViewportLayout::FactoryPaneConfigurationFromTypeName.")
+typedef FEditorViewportLayoutFourPanesRight FLevelViewportLayoutFourPanesRight;
 
-	/** Overridden from FLevelViewportLayout */
-	virtual void ReplaceWidget( TSharedRef< SWidget > Source, TSharedRef< SWidget > Replacement ) override;
+UE_DEPRECATED(5.0, "Level Pane layouts are now owned by their level layouts. Use FEditorViewportLayoutFourPanesTop instead, and add any custom types to FAssetEditorViewportLayout::FactoryPaneConfigurationFromTypeName.")
+typedef FEditorViewportLayoutFourPanesTop FLevelViewportLayoutFourPanesTop;
 
-
-protected:
-	/** The splitter widgets */
-	TSharedPtr< class SSplitter > PrimarySplitterWidget;
-	TSharedPtr< class SSplitter > SecondarySplitterWidget;
-};
-
-
-// FLevelViewportLayoutFourPanesLeft /////////////////////////////
-
-class FLevelViewportLayoutFourPanesLeft : public FLevelViewportLayoutFourPanes
-{
-public:
-	virtual const FName& GetLayoutTypeName() const override { return LevelViewportConfigurationNames::FourPanesLeft; }
-
-	virtual TSharedRef<SWidget> MakeFourPanelWidget(
-		TMap<FName, TSharedPtr< IEditorViewportLayoutEntity >>& ViewportWidgets,
-		TSharedRef<SWidget> Viewport0, TSharedRef<SWidget> Viewport1, TSharedRef<SWidget> Viewport2, TSharedRef<SWidget> Viewport3,
-		float PrimarySplitterPercentage, float SecondarySplitterPercentage0, float SecondarySplitterPercentage1) override;
-};
-
-
-// FLevelViewportLayoutFourPanesRight /////////////////////////////
-
-class FLevelViewportLayoutFourPanesRight : public FLevelViewportLayoutFourPanes
-{
-public:
-	virtual const FName& GetLayoutTypeName() const override { return LevelViewportConfigurationNames::FourPanesRight; }
-
-	virtual TSharedRef<SWidget> MakeFourPanelWidget(
-		TMap<FName, TSharedPtr< IEditorViewportLayoutEntity >>& ViewportWidgets,
-		TSharedRef<SWidget> Viewport0, TSharedRef<SWidget> Viewport1, TSharedRef<SWidget> Viewport2, TSharedRef<SWidget> Viewport3,
-		float PrimarySplitterPercentage, float SecondarySplitterPercentage0, float SecondarySplitterPercentage1) override;
-};
-
-
-// FLevelViewportLayoutFourPanesTop /////////////////////////////
-
-class FLevelViewportLayoutFourPanesTop : public FLevelViewportLayoutFourPanes
-{
-public:
-	virtual const FName& GetLayoutTypeName() const override { return LevelViewportConfigurationNames::FourPanesTop; }
-
-	virtual TSharedRef<SWidget> MakeFourPanelWidget(
-		TMap<FName, TSharedPtr< IEditorViewportLayoutEntity >>& ViewportWidgets,
-		TSharedRef<SWidget> Viewport0, TSharedRef<SWidget> Viewport1, TSharedRef<SWidget> Viewport2, TSharedRef<SWidget> Viewport3,
-		float PrimarySplitterPercentage, float SecondarySplitterPercentage0, float SecondarySplitterPercentage1) override;
-};
-
-
-// FLevelViewportLayoutFourPanesBottom /////////////////////////////
-
-class FLevelViewportLayoutFourPanesBottom : public FLevelViewportLayoutFourPanes
-{
-public:
-	virtual const FName& GetLayoutTypeName() const override { return LevelViewportConfigurationNames::FourPanesBottom; }
-
-	virtual TSharedRef<SWidget> MakeFourPanelWidget(
-		TMap<FName, TSharedPtr< IEditorViewportLayoutEntity >>& ViewportWidgets,
-		TSharedRef<SWidget> Viewport0, TSharedRef<SWidget> Viewport1, TSharedRef<SWidget> Viewport2, TSharedRef<SWidget> Viewport3,
-		float PrimarySplitterPercentage, float SecondarySplitterPercentage0, float SecondarySplitterPercentage1) override;
-};
+UE_DEPRECATED(5.0, "Level Pane layouts are now owned by their level layouts. Use FEditorViewportLayoutFourPanesBottom instead, and add any custom types to FAssetEditorViewportLayout::FactoryPaneConfigurationFromTypeName.")
+typedef FEditorViewportLayoutFourPanesBottom FLevelViewportLayoutFourPanesBottom;
diff --git a/Engine/Source/Editor/LevelEditor/Private/LevelViewportLayoutOnePane.cpp b/Engine/Source/Editor/LevelEditor/Private/LevelViewportLayoutOnePane.cpp
deleted file mode 100644
diff --git a/Engine/Source/Editor/LevelEditor/Private/LevelViewportLayoutOnePane.h b/Engine/Source/Editor/LevelEditor/Private/LevelViewportLayoutOnePane.h
--- a/Engine/Source/Editor/LevelEditor/Private/LevelViewportLayoutOnePane.h
+++ b/Engine/Source/Editor/LevelEditor/Private/LevelViewportLayoutOnePane.h
@@ -2,32 +2,7 @@
 
 #pragma once
 
-#include "CoreMinimal.h"
-#include "Widgets/SWidget.h"
-#include "LevelViewportLayout.h"
-#include "LevelViewportActions.h"
+#include "EditorViewportLayoutOnePane.h"
 
-class SHorizontalBox;
-
-class FLevelViewportLayoutOnePane : public FLevelViewportLayout
-{
-public:
-	/**
-	* Saves viewport layout information between editor sessions
-	*/
-	virtual void SaveLayoutString(const FString& LayoutString) const override;
-
-	virtual const FName& GetLayoutTypeName() const override{ return LevelViewportConfigurationNames::OnePane; }
-protected:
-	/**
-	* Creates the viewport for the single pane
-	*/
-	virtual TSharedRef<SWidget> MakeViewportLayout(const FString& LayoutString) override;
-
-	/** Overridden from FLevelViewportLayout */
-	virtual void ReplaceWidget(TSharedRef< SWidget > Source, TSharedRef< SWidget > Replacement) override;
-
-protected:
-	/** The viewport widget parent box */
-	TSharedPtr< SHorizontalBox > ViewportBox;
-};
+UE_DEPRECATED(5.0, "Level Pane layouts are now owned by their level layouts. Use FEditorViewportLayoutOnePane instead, and add any custom types to FAssetEditorViewportLayout::FactoryPaneConfigurationFromTypeName.")
+typedef FEditorViewportLayoutOnePane FLevelViewportLayoutOnePane;
diff --git a/Engine/Source/Editor/LevelEditor/Private/LevelViewportLayoutThreePanes.cpp b/Engine/Source/Editor/LevelEditor/Private/LevelViewportLayoutThreePanes.cpp
deleted file mode 100644
diff --git a/Engine/Source/Editor/LevelEditor/Private/LevelViewportLayoutThreePanes.h b/Engine/Source/Editor/LevelEditor/Private/LevelViewportLayoutThreePanes.h
--- a/Engine/Source/Editor/LevelEditor/Private/LevelViewportLayoutThreePanes.h
+++ b/Engine/Source/Editor/LevelEditor/Private/LevelViewportLayoutThreePanes.h
@@ -2,91 +2,19 @@
 
 #pragma once
 
-#include "CoreMinimal.h"
-#include "Widgets/SWidget.h"
-#include "LevelViewportLayout.h"
-#include "Widgets/Layout/SSplitter.h"
-#include "LevelViewportActions.h"
+#include "EditorViewportLayoutThreePanes.h"
 
-class FLevelViewportLayoutThreePanes : public FLevelViewportLayout
-{
-public:
-	/**
-	 * Saves viewport layout information between editor sessions
-	 */
-	virtual void SaveLayoutString(const FString& LayoutString) const override;
-protected:
-	/**
-	 * Creates the viewports and splitter for the two panes vertical layout                   
-	 */
-	virtual TSharedRef<SWidget> MakeViewportLayout(const FString& LayoutString) override;
+UE_DEPRECATED(5.0, "Level Pane layouts are now owned by their level layouts. Use FEditorViewportLayoutThreePanesLeft instead, and add any custom types to FAssetEditorViewportLayout::FactoryPaneConfigurationFromTypeName.")
+typedef FEditorViewportLayoutThreePanesLeft FLevelViewportLayoutThreePanesLeft;
 
-	virtual TSharedRef<SWidget> MakeThreePanelWidget(
-		TMap< FName, TSharedPtr< IEditorViewportLayoutEntity >>& ViewportWidgets,
-		const TSharedRef<SWidget>& ViewportKey0, const TSharedRef<SWidget>& ViewportKey1, const TSharedRef<SWidget>& ViewportKey2,
-		float PrimarySplitterPercentage, float SecondarySplitterPercentage) = 0;
 
-	/** Overridden from FLevelViewportLayout */
-	virtual void ReplaceWidget( TSharedRef< SWidget > Source, TSharedRef< SWidget > Replacement ) override;
+UE_DEPRECATED(5.0, "Level Pane layouts are now owned by their level layouts. Use FEditorViewportLayoutThreePanesRight instead, and add any custom types to FAssetEditorViewportLayout::FactoryPaneConfigurationFromTypeName.")
+typedef FEditorViewportLayoutThreePanesRight FLevelViewportLayoutThreePanesRight;
 
 
-protected:
-	/** The splitter widgets */
-	TSharedPtr< class SSplitter > PrimarySplitterWidget;
-	TSharedPtr< class SSplitter > SecondarySplitterWidget;
-};
+UE_DEPRECATED(5.0, "Level Pane layouts are now owned by their level layouts. Use FEditorViewportLayoutThreePanesTop instead, and add any custom types to FAssetEditorViewportLayout::FactoryPaneConfigurationFromTypeName.")
+typedef FEditorViewportLayoutThreePanesTop FLevelViewportLayoutThreePanesTop;
 
-// FLevelViewportLayoutThreePanesLeft /////////////////////////////
 
-class FLevelViewportLayoutThreePanesLeft : public FLevelViewportLayoutThreePanes
-{
-public:
-	virtual const FName& GetLayoutTypeName() const override { return LevelViewportConfigurationNames::ThreePanesLeft; }
-
-	virtual TSharedRef<SWidget> MakeThreePanelWidget(
-		TMap< FName, TSharedPtr< IEditorViewportLayoutEntity >>& ViewportWidgets,
-		const TSharedRef<SWidget>& ViewportKey0, const TSharedRef<SWidget>& ViewportKey1, const TSharedRef<SWidget>& ViewportKey2,
-		float PrimarySplitterPercentage, float SecondarySplitterPercentage) override;
-};
-
-
-// FLevelViewportLayoutThreePanesRight /////////////////////////////
-
-class FLevelViewportLayoutThreePanesRight : public FLevelViewportLayoutThreePanes
-{
-public:
-	virtual const FName& GetLayoutTypeName() const override { return LevelViewportConfigurationNames::ThreePanesRight; }
-
-	virtual TSharedRef<SWidget> MakeThreePanelWidget(
-		TMap< FName, TSharedPtr< IEditorViewportLayoutEntity >>& ViewportWidgets,
-		const TSharedRef<SWidget>& ViewportKey0, const TSharedRef<SWidget>& ViewportKey1, const TSharedRef<SWidget>& ViewportKey2,
-		float PrimarySplitterPercentage, float SecondarySplitterPercentage) override;
-};
-
-
-// FLevelViewportLayoutThreePanesTop /////////////////////////////
-
-class FLevelViewportLayoutThreePanesTop : public FLevelViewportLayoutThreePanes
-{
-public:
-	virtual const FName& GetLayoutTypeName() const override { return LevelViewportConfigurationNames::ThreePanesTop; }
-
-	virtual TSharedRef<SWidget> MakeThreePanelWidget(
-		TMap< FName, TSharedPtr< IEditorViewportLayoutEntity >>& ViewportWidgets,
-		const TSharedRef<SWidget>& ViewportKey0, const TSharedRef<SWidget>& ViewportKey1, const TSharedRef<SWidget>& ViewportKey2,
-		float PrimarySplitterPercentage, float SecondarySplitterPercentage) override;
-};
-
-
-// FLevelViewportLayoutThreePanesBottom /////////////////////////////
-
-class FLevelViewportLayoutThreePanesBottom : public FLevelViewportLayoutThreePanes
-{
-public:
-	virtual const FName& GetLayoutTypeName() const override { return LevelViewportConfigurationNames::ThreePanesBottom; }
-
-	virtual TSharedRef<SWidget> MakeThreePanelWidget(
-		TMap< FName, TSharedPtr< IEditorViewportLayoutEntity >>& ViewportWidgets,
-		const TSharedRef<SWidget>& ViewportKey0, const TSharedRef<SWidget>& ViewportKey1, const TSharedRef<SWidget>& ViewportKey2,
-		float PrimarySplitterPercentage, float SecondarySplitterPercentage) override;
-};
+UE_DEPRECATED(5.0, "Level Pane layouts are now owned by their level layouts. Use FEditorViewportLayoutThreePanesBottom instead, and add any custom types to FAssetEditorViewportLayout::FactoryPaneConfigurationFromTypeName.")
+typedef FEditorViewportLayoutThreePanesBottom FLevelViewportLayoutThreePanesBottom;
diff --git a/Engine/Source/Editor/LevelEditor/Private/LevelViewportLayoutTwoPanes.cpp b/Engine/Source/Editor/LevelEditor/Private/LevelViewportLayoutTwoPanes.cpp
deleted file mode 100644
diff --git a/Engine/Source/Editor/LevelEditor/Private/LevelViewportLayoutTwoPanes.h b/Engine/Source/Editor/LevelEditor/Private/LevelViewportLayoutTwoPanes.h
--- a/Engine/Source/Editor/LevelEditor/Private/LevelViewportLayoutTwoPanes.h
+++ b/Engine/Source/Editor/LevelEditor/Private/LevelViewportLayoutTwoPanes.h
@@ -2,49 +2,10 @@
 
 #pragma once
 
-#include "CoreMinimal.h"
-#include "Widgets/SWidget.h"
-#include "LevelViewportLayout.h"
-#include "Widgets/Layout/SSplitter.h"
-#include "LevelViewportActions.h"
+#include "EditorViewportLayoutTwoPanes.h"
 
-template <EOrientation TOrientation>
-class TLevelViewportLayoutTwoPanes : public FLevelViewportLayout
-{
-public:
-	/**
-	 * Saves viewport layout information between editor sessions
-	 */
-	virtual void SaveLayoutString(const FString& LayoutString) const override;
-protected:
-	/**
-	 * Creates the viewports and splitter for the two panes vertical layout                   
-	 */
-	virtual TSharedRef<SWidget> MakeViewportLayout(const FString& LayoutString) override;
+UE_DEPRECATED(5.0, "Level Pane layouts are now owned by their level layouts. Use FEditorViewportLayoutTwoPanesVert instead, and add any custom types to FAssetEditorViewportLayout::FactoryPaneConfigurationFromTypeName.")
+typedef FEditorViewportLayoutTwoPanesVert FLevelViewportLayoutTwoPanesVert;
 
-	/** Overridden from FLevelViewportLayout */
-	virtual void ReplaceWidget( TSharedRef< SWidget > Source, TSharedRef< SWidget > Replacement ) override;
-
-
-private:
-	/** The splitter widget */
-	TSharedPtr< class SSplitter > SplitterWidget;
-};
-
-
-// FLevelViewportLayoutTwoPanesVert /////////////////////////////
-
-class FLevelViewportLayoutTwoPanesVert : public TLevelViewportLayoutTwoPanes<EOrientation::Orient_Vertical>
-{
-public:
-	virtual const FName& GetLayoutTypeName() const override { return LevelViewportConfigurationNames::TwoPanesVert; }
-};
-
-
-// FLevelViewportLayoutTwoPanesHoriz /////////////////////////////
-
-class FLevelViewportLayoutTwoPanesHoriz : public TLevelViewportLayoutTwoPanes<EOrientation::Orient_Horizontal>
-{
-public:
-	virtual const FName& GetLayoutTypeName() const override { return LevelViewportConfigurationNames::TwoPanesHoriz; }
-};
+UE_DEPRECATED(5.0, "Level Pane layouts are now owned by their level layouts. Use FEditorViewportLayoutTwoPanesHoriz instead, and add any custom types to FAssetEditorViewportLayout::FactoryPaneConfigurationFromTypeName.")
+typedef FEditorViewportLayoutTwoPanesHoriz FLevelViewportLayoutTwoPanesHoriz;
diff --git a/Engine/Source/Editor/LevelEditor/Private/LevelViewportTabContent.cpp b/Engine/Source/Editor/LevelEditor/Private/LevelViewportTabContent.cpp
--- a/Engine/Source/Editor/LevelEditor/Private/LevelViewportTabContent.cpp
+++ b/Engine/Source/Editor/LevelEditor/Private/LevelViewportTabContent.cpp
@@ -5,11 +5,6 @@
 #include "Framework/Docking/LayoutService.h"
 #include "Modules/ModuleManager.h"
 #include "LevelEditor.h"
-#include "LevelViewportLayout2x2.h"
-#include "LevelViewportLayoutOnePane.h"
-#include "LevelViewportLayoutTwoPanes.h"
-#include "LevelViewportLayoutThreePanes.h"
-#include "LevelViewportLayoutFourPanes.h"
 #include "Widgets/Docking/SDockTab.h"
 #include "LevelViewportLayout.h"
 #include "LevelEditorViewport.h"
@@ -17,29 +12,10 @@
 
 // FLevelViewportTabContent ///////////////////////////
 
-TSharedPtr< FEditorViewportLayout > FLevelViewportTabContent::ConstructViewportLayoutByTypeName(const FName& TypeName, bool bSwitchingLayouts)
+TSharedPtr< FEditorViewportLayout > FLevelViewportTabContent::FactoryViewportLayout(bool bIsSwitchingLayouts)
 {
-	TSharedPtr< FLevelViewportLayout > ViewportLayout;
-
-	// The items in these ifs should match the names in namespace LevelViewportConfigurationNames
-	if (TypeName == LevelViewportConfigurationNames::FourPanes2x2) ViewportLayout = MakeShareable(new FLevelViewportLayout2x2);
-	else if (TypeName == LevelViewportConfigurationNames::TwoPanesVert) ViewportLayout = MakeShareable(new FLevelViewportLayoutTwoPanesVert);
-	else if (TypeName == LevelViewportConfigurationNames::TwoPanesHoriz) ViewportLayout = MakeShareable(new FLevelViewportLayoutTwoPanesHoriz);
-	else if (TypeName == LevelViewportConfigurationNames::ThreePanesLeft) ViewportLayout = MakeShareable(new FLevelViewportLayoutThreePanesLeft);
-	else if (TypeName == LevelViewportConfigurationNames::ThreePanesRight) ViewportLayout = MakeShareable(new FLevelViewportLayoutThreePanesRight);
-	else if (TypeName == LevelViewportConfigurationNames::ThreePanesTop) ViewportLayout = MakeShareable(new FLevelViewportLayoutThreePanesTop);
-	else if (TypeName == LevelViewportConfigurationNames::ThreePanesBottom) ViewportLayout = MakeShareable(new FLevelViewportLayoutThreePanesBottom);
-	else if (TypeName == LevelViewportConfigurationNames::FourPanesLeft) ViewportLayout = MakeShareable(new FLevelViewportLayoutFourPanesLeft);
-	else if (TypeName == LevelViewportConfigurationNames::FourPanesRight) ViewportLayout = MakeShareable(new FLevelViewportLayoutFourPanesRight);
-	else if (TypeName == LevelViewportConfigurationNames::FourPanesBottom) ViewportLayout = MakeShareable(new FLevelViewportLayoutFourPanesBottom);
-	else if (TypeName == LevelViewportConfigurationNames::FourPanesTop) ViewportLayout = MakeShareable(new FLevelViewportLayoutFourPanesTop);
-	else if (TypeName == LevelViewportConfigurationNames::OnePane) ViewportLayout = MakeShareable(new FLevelViewportLayoutOnePane);
-
-	if (!ensure(ViewportLayout.IsValid()))
-	{
-		ViewportLayout = MakeShareable(new FLevelViewportLayoutOnePane);
-	}
-	ViewportLayout->SetIsReplacement(bSwitchingLayouts);
+	TSharedPtr<FLevelViewportLayout> ViewportLayout = MakeShareable(new FLevelViewportLayout);
+	ViewportLayout->SetIsReplacement(bIsSwitchingLayouts);
 	return ViewportLayout;
 }
 
@@ -61,7 +37,7 @@ void FLevelViewportTabContent::Initialize(AssetEditorViewportFactoryFunction Fun
 	if(InLayoutString.IsEmpty() ||
 		!GConfig->GetString(*IniSection, *(InLayoutString + TEXT(".LayoutType")), LayoutTypeString, GEditorPerProjectIni))
 	{
-		LayoutTypeString = LevelViewportConfigurationNames::FourPanes2x2.ToString();
+		LayoutTypeString = EditorViewportConfigurationNames::FourPanes2x2.ToString();
 	}
 
 	OnViewportTabContentLayoutStartChangeEvent.AddSP(this, &FLevelViewportTabContent::OnLayoutStartChange);
diff --git a/Engine/Source/Editor/LevelEditor/Public/LevelViewportLayout.h b/Engine/Source/Editor/LevelEditor/Public/LevelViewportLayout.h
--- a/Engine/Source/Editor/LevelEditor/Public/LevelViewportLayout.h
+++ b/Engine/Source/Editor/LevelEditor/Public/LevelViewportLayout.h
@@ -69,7 +69,7 @@ public:
 	 */
 	virtual TSharedRef<SWidget> BuildViewportLayout(TSharedPtr<SDockTab> InParentDockTab, TSharedPtr<FEditorViewportTabContent> InParentTab, const FString& LayoutString) override;
 
-	virtual TSharedRef<IEditorViewportLayoutEntity> FactoryViewport(FName InTypeName, const FAssetEditorViewportConstructionArgs& ConstructionArgs) const override;
+	virtual TSharedRef<SWidget> FactoryViewport(FName InTypeName, const FAssetEditorViewportConstructionArgs& ConstructionArgs) override;
 
 	/**
 	 * Makes a request to maximize a specific viewport and hide the others in this layout
@@ -133,6 +133,9 @@ public:
 	/** Returns whether a viewport animation is currently taking place */
 	bool IsTransitioning() const { return bIsTransitioning; }
 
+	virtual void LoadConfig(const FString& LayoutString) override;
+	virtual void SaveConfig(const FString& LayoutString) const override;
+
 protected:
 
 	/**
@@ -145,14 +148,6 @@ protected:
 	 */
 	void MaximizeViewport( FName ViewportToMaximize, const bool bWantMaximize, const bool bWantImmersive, const bool bAllowAnimation );
 
-	/**
-	 * Overridden in derived classes to set up custom layouts  
-	 *
-	 * @param LayoutString		The layout string loaded from a file
-	 * @return The base widget representing the layout.  Usually a splitter
-	 */
-	virtual TSharedRef<SWidget> MakeViewportLayout(const FString& LayoutString) = 0;
-
 	/**
 	 * Delegate called to get the visibility of the non-maximized viewports
 	 * The non-maximized viewports are not visible if there is a maximized viewport on top of them
@@ -175,14 +170,6 @@ protected:
 	 */
 	FVector2D GetMaximizedViewportSizeOnCanvas() const;
 
-	/**
-	 * Inline replaces a viewport content widget within this layout
-	 *
-	 * @param	Source	The widget to replace
-	 * @param	Replacement	The widget to replace the source widget with
-	 */
-	virtual void ReplaceWidget( TSharedRef< SWidget > Source, TSharedRef< SWidget > Replacement ) = 0;
-
 	/** If a viewport animation is in progress, finishes that transition immediately */
 	void FinishMaximizeTransition();
 
@@ -196,23 +183,13 @@ protected:
 	 */
 	void EndThrottleForAnimatedResize();
 
-	/** Generates a layout string for persisting settings for this layout based on the runtime type of layout */
-	FString GetTypeSpecificLayoutString(const FString& LayoutString) const
-	{ 
-		if (LayoutString.IsEmpty())
-		{
-			return LayoutString;
-		}
-		return FString::Printf(TEXT("%s.%s"), *GetLayoutTypeName().ToString(), *LayoutString);
-	}
-
-	/** Called in MakeViewportLayout() functions for derived types of layout to init values common to all types */
-	void InitCommonLayoutFromString( const FString& SpecificLayoutString, FName DefaultMaximizedViewport );
+protected:
+	/** Called in after factory of a new layout functions for derived types of layout to init values common to all types */
+	void InitCommonLayoutFromString(const FString& LayoutString, const FName PerspectiveViewportKey);
 
-	/** Called in SaveLayoutString() functions for derived types of layout to save values common to all types */
-	void SaveCommonLayoutString( const FString& SpecificLayoutString ) const;
+	/** Called in SaveConfig() functions for derived types of layout to save values common to all types */
+	void SaveCommonLayoutString(const FString& SpecificLayoutString) const;
 
-protected:
 	/** True if we've started an animation and are waiting for it to finish */
 	bool bIsTransitioning;
 
diff --git a/Engine/Source/Editor/LevelEditor/Public/LevelViewportTabContent.h b/Engine/Source/Editor/LevelEditor/Public/LevelViewportTabContent.h
--- a/Engine/Source/Editor/LevelEditor/Public/LevelViewportTabContent.h
+++ b/Engine/Source/Editor/LevelEditor/Public/LevelViewportTabContent.h
@@ -22,7 +22,7 @@ public:
 	virtual void Initialize(AssetEditorViewportFactoryFunction Func, TSharedPtr<SDockTab> InParentTab, const FString& InLayoutString) override;
 
 protected:
-	virtual TSharedPtr< FEditorViewportLayout > ConstructViewportLayoutByTypeName(const FName& TypeName, bool bSwitchingLayouts) override;
+	virtual TSharedPtr<FEditorViewportLayout> FactoryViewportLayout(bool bIsSwitchingLayouts) override;
 
 	void OnLayoutStartChange(bool bSwitchingLayouts);
 	void OnLayoutChanged();
diff --git a/Engine/Source/Editor/UnrealEd/Private/AssetEditorViewportLayout.cpp b/Engine/Source/Editor/UnrealEd/Private/AssetEditorViewportLayout.cpp
--- a/Engine/Source/Editor/UnrealEd/Private/AssetEditorViewportLayout.cpp
+++ b/Engine/Source/Editor/UnrealEd/Private/AssetEditorViewportLayout.cpp
@@ -25,6 +25,11 @@
 #include "Templates/SharedPointer.h"
 #include "Framework/MultiBox/MultiBoxBuilder.h"
 #include "EditorViewportTabContent.h"
+#include "EditorViewportLayoutOnePane.h"
+#include "EditorViewportLayoutTwoPanes.h"
+#include "EditorViewportLayoutThreePanes.h"
+#include "EditorViewportLayoutFourPanes.h"
+#include "EditorViewportLayout2x2.h"
 
 
 #define LOCTEXT_NAMESPACE "AssetEditorViewportToolBar"
@@ -89,10 +94,42 @@ TSharedPtr<FViewportTabContent> SAssetEditorViewportsOverlay::GetViewportTab() c
 	return ViewportTab;
 }
 
+void FAssetEditorViewportPaneLayout::LoadConfig(const FString& LayoutString, TFunction<void(const FString&, const FName)> LoadAdditionalLayoutInfoCallback)
+{
+	FString SpecificLayoutString = GetTypeSpecificLayoutString(LayoutString);
+	LoadLayoutString(SpecificLayoutString);
+
+	if (LoadAdditionalLayoutInfoCallback)
+	{
+		LoadAdditionalLayoutInfoCallback(SpecificLayoutString, PerspectiveViewportConfigKey);
+	}
+}
+
+void FAssetEditorViewportPaneLayout::SaveConfig(const FString& LayoutString, TFunction<void(const FString&)> SaveAdditionalLayoutInfoCallback) const
+{
+	FString SpecificLayoutString = GetTypeSpecificLayoutString(LayoutString);
+	SaveLayoutString(SpecificLayoutString);
+
+	if (SaveAdditionalLayoutInfoCallback)
+	{
+		SaveAdditionalLayoutInfoCallback(SpecificLayoutString);
+	}
+}
+
+FString FAssetEditorViewportPaneLayout::GetTypeSpecificLayoutString(const FString& LayoutString) const
+{
+	if (LayoutString.IsEmpty())
+	{
+		return LayoutString;
+	}
+	return FString::Printf(TEXT("%s.%s"), *GetLayoutTypeName().ToString(), *LayoutString);
+}
 
 // FAssetEditorViewportLayout /////////////////////////////
 
-FAssetEditorViewportLayout::FAssetEditorViewportLayout() {}
+FAssetEditorViewportLayout::FAssetEditorViewportLayout()
+{
+}
 
 
 FAssetEditorViewportLayout::~FAssetEditorViewportLayout()
@@ -103,7 +140,7 @@ FAssetEditorViewportLayout::~FAssetEditorViewportLayout()
 	}
 }
 
-TSharedRef<IEditorViewportLayoutEntity> FAssetEditorViewportLayout::FactoryViewport(FName InTypeName, const FAssetEditorViewportConstructionArgs& ConstructionArgs) const
+TSharedRef<SWidget> FAssetEditorViewportLayout::FactoryViewport(FName InTypeName, const FAssetEditorViewportConstructionArgs& ConstructionArgs)
 {
 	TSharedPtr<SAssetEditorViewport> EditorViewport;
 	TSharedPtr<FEditorViewportTabContent> PinnedTabContent = ParentTabContent.Pin();
@@ -116,10 +153,37 @@ TSharedRef<IEditorViewportLayoutEntity> FAssetEditorViewportLayout::FactoryViewp
 		EditorViewport = SNew(SAssetEditorViewport, ConstructionArgs);
 	}
 
-	return MakeShareable(new FEditorViewportLayoutEntity(EditorViewport));
+	TSharedRef<IEditorViewportLayoutEntity> LayoutEntity = MakeShareable(new FEditorViewportLayoutEntity(EditorViewport));
+	Viewports.Add(ConstructionArgs.ConfigKey, LayoutEntity);
+	return LayoutEntity->AsWidget();
 }
 
+void FAssetEditorViewportLayout::FactoryPaneConfigurationFromTypeName(const FName& InLayoutConfigTypeName)
+{
+	//The items in these ifs should match the names in namespace EditorViewportConfigurationNames
+	if (InLayoutConfigTypeName == EditorViewportConfigurationNames::TwoPanesHoriz) LayoutConfiguration = MakeShareable(new FEditorViewportLayoutTwoPanesHoriz);
+	else if (InLayoutConfigTypeName == EditorViewportConfigurationNames::TwoPanesVert) LayoutConfiguration = MakeShareable(new FEditorViewportLayoutTwoPanesVert);
+	else if (InLayoutConfigTypeName == EditorViewportConfigurationNames::FourPanes2x2) LayoutConfiguration = MakeShareable(new FEditorViewportLayout2x2);
+	else if (InLayoutConfigTypeName == EditorViewportConfigurationNames::ThreePanesLeft) LayoutConfiguration = MakeShareable(new FEditorViewportLayoutThreePanesLeft);
+	else if (InLayoutConfigTypeName == EditorViewportConfigurationNames::ThreePanesRight) LayoutConfiguration = MakeShareable(new FEditorViewportLayoutThreePanesRight);
+	else if (InLayoutConfigTypeName == EditorViewportConfigurationNames::ThreePanesTop) LayoutConfiguration = MakeShareable(new FEditorViewportLayoutThreePanesTop);
+	else if (InLayoutConfigTypeName == EditorViewportConfigurationNames::ThreePanesBottom) LayoutConfiguration = MakeShareable(new FEditorViewportLayoutThreePanesBottom);
+	else if (InLayoutConfigTypeName == EditorViewportConfigurationNames::FourPanesLeft) LayoutConfiguration = MakeShareable(new FEditorViewportLayoutFourPanesLeft);
+	else if (InLayoutConfigTypeName == EditorViewportConfigurationNames::FourPanesRight) LayoutConfiguration = MakeShareable(new FEditorViewportLayoutFourPanesRight);
+	else if (InLayoutConfigTypeName == EditorViewportConfigurationNames::FourPanesBottom) LayoutConfiguration = MakeShareable(new FEditorViewportLayoutFourPanesBottom);
+	else if (InLayoutConfigTypeName == EditorViewportConfigurationNames::FourPanesTop) LayoutConfiguration = MakeShareable(new FEditorViewportLayoutFourPanesTop);
+	else if (InLayoutConfigTypeName == EditorViewportConfigurationNames::OnePane) LayoutConfiguration = MakeShareable(new FEditorViewportLayoutOnePane);
+
+	if (!LayoutConfiguration.IsValid())
+	{
+		LayoutConfiguration = MakeShareable(new FEditorViewportLayoutOnePane);
+	}
+}
 
+const FName FAssetEditorViewportLayout::GetActivePaneConfigurationTypeName() const
+{
+	return LayoutConfiguration.IsValid() ? LayoutConfiguration->GetLayoutTypeName() : NAME_None;
+}
 
 TSharedRef<SWidget> FAssetEditorViewportLayout::BuildViewportLayout(TSharedPtr<SDockTab> InParentDockTab, TSharedPtr<FEditorViewportTabContent> InParentTab, const FString& LayoutString)
 {
@@ -143,20 +207,37 @@ TSharedRef<SWidget> FAssetEditorViewportLayout::BuildViewportLayout(TSharedPtr<S
 
 	ViewportsOverlayPtr = ViewportsOverlay;
 
+	// You must have a valid layout configuration before building the layout.
+	if (!ensureMsgf(LayoutConfiguration.IsValid(), TEXT("No valid layout configuration for the viewport layout was found.")))
+	{
+		return ViewportsOverlay;
+	}
+
 	// Don't set the content until the OverlayPtr has been set, because it access this when we want to start with the viewports maximized.
-	ViewportsBorder->SetContent(MakeViewportLayout(LayoutString));
+	TSharedRef<SWidget> ViewportLayoutWidget = LayoutConfiguration->MakeViewportLayout(this->AsShared(), LayoutString);
+	ViewportsBorder->SetContent(ViewportLayoutWidget);
 
 	return ViewportsOverlay;
+}
 
+TSharedRef<SWidget> FAssetEditorViewportLayout::MakeViewportLayout(const FString& LayoutString)
+{
+	return LayoutConfiguration->MakeViewportLayout(this->AsShared(), LayoutString);
 }
 
-FString FAssetEditorViewportLayout::GetTypeSpecificLayoutString(const FString& LayoutString) const
+void FAssetEditorViewportLayout::LoadConfig(const FString& LayoutString)
 {
-	if (LayoutString.IsEmpty())
+}
+
+void FAssetEditorViewportLayout::SaveConfig(const FString& LayoutString) const
+{
+	if (!LayoutString.IsEmpty())
 	{
-		return LayoutString;
+		FString LayoutTypeString = LayoutConfiguration->GetLayoutTypeName().ToString();
+
+		const FString& IniSection = FLayoutSaveRestore::GetAdditionalLayoutConfigIni();
+		GConfig->SetString(*IniSection, *(LayoutString + TEXT(".LayoutType")), *LayoutTypeString, GEditorPerProjectIni);
 	}
-	return FString::Printf(TEXT("%s.%s"), *GetLayoutTypeName().ToString(), *LayoutString);
 }
 
 EVisibility FAssetEditorViewportLayout::OnGetNonMaximizedVisibility() const
diff --git a/Engine/Source/Editor/UnrealEd/Private/EditorViewportLayout2x2.cpp b/Engine/Source/Editor/UnrealEd/Private/EditorViewportLayout2x2.cpp
--- a/Engine/Source/Editor/UnrealEd/Private/EditorViewportLayout2x2.cpp
+++ b/Engine/Source/Editor/UnrealEd/Private/EditorViewportLayout2x2.cpp
@@ -4,9 +4,9 @@
 #include "Framework/Docking/LayoutService.h"
 #include "Editor.h"
 #include "Misc/ConfigCacheIni.h"
-#include "Modules/ModuleManager.h"
 #include "Framework/Application/SlateApplication.h"
-// #include "LevelEditor.h"
+#include "Widgets/SWidget.h"
+#include "Widgets/Layout/SSplitter.h"
 
 namespace ViewportLayout2x2Defs
 {
@@ -16,16 +16,15 @@ namespace ViewportLayout2x2Defs
 
 // FLevelViewportLayout2x2 //////////////////////////////////////////
 
-TSharedRef<SWidget> FEditorViewportLayout2x2::MakeViewportLayout(const FString& LayoutString)
+TSharedRef<SWidget> FEditorViewportLayout2x2::MakeViewportLayout(TSharedPtr<FAssetEditorViewportLayout> InParentLayout, const FString& LayoutString)
 {
- 	FString SpecificLayoutString = GetTypeSpecificLayoutString(LayoutString);
-
 	FString TopLeftKey, BottomLeftKey, TopRightKey, BottomRightKey;
 
 	FString TopLeftType = TEXT("Default"), BottomLeftType = TEXT("Default"), TopRightType = TEXT("Default"), BottomRightType = TEXT("Default");
 
 	TArray<FVector2D> SplitterPercentages;
 	
+	FString SpecificLayoutString = GetTypeSpecificLayoutString(LayoutString);
  	if (!SpecificLayoutString.IsEmpty())
 	{
 		// The Layout String only holds the unique ID of the Additional Layout Configs to use
@@ -53,58 +52,53 @@ TSharedRef<SWidget> FEditorViewportLayout2x2::MakeViewportLayout(const FString&
 		}
 	}
 
-
 	// Set up the viewports
 	FAssetEditorViewportConstructionArgs Args;
- 	Args.ParentLayout = AsShared();
+ 	Args.ParentLayout = InParentLayout;
 
 	// Left viewport
 	Args.bRealtime = false;
 	Args.ConfigKey = *TopLeftKey;
 	Args.ViewportType = LVT_OrthoYZ;
-	TSharedPtr< IEditorViewportLayoutEntity > ViewportTL = FactoryViewport(*TopLeftType, Args);
+	TSharedRef< SWidget > ViewportTL = InParentLayout->FactoryViewport(*TopLeftType, Args);
 
 	// Persp viewport
 	Args.bRealtime = !FPlatformMisc::IsRemoteSession();
 	Args.ConfigKey = *BottomLeftKey;
 	Args.ViewportType = LVT_Perspective;
-	TSharedPtr< IEditorViewportLayoutEntity > ViewportBL = FactoryViewport(*BottomLeftType, Args);
+	TSharedRef< SWidget > ViewportBL = InParentLayout->FactoryViewport(*BottomLeftType, Args);
+	PerspectiveViewportConfigKey = *BottomLeftKey;
 
 	// Front viewport
 	Args.bRealtime = false;
 	Args.ConfigKey = *TopRightKey;
 	Args.ViewportType = LVT_OrthoXZ;
-	TSharedPtr< IEditorViewportLayoutEntity > ViewportTR = FactoryViewport(*TopRightType, Args);
+	TSharedRef< SWidget > ViewportTR = InParentLayout->FactoryViewport(*TopRightType, Args);
 
 	// Top Viewport
 	Args.bRealtime = false;
 	Args.ConfigKey = *BottomRightKey;
 	Args.ViewportType = LVT_OrthoXY;
-	TSharedPtr< IEditorViewportLayoutEntity > ViewportBR = FactoryViewport(*BottomRightType, Args);
-
-	Viewports.Add( *TopLeftKey, ViewportTL );
-	Viewports.Add( *BottomLeftKey, ViewportBL );
-	Viewports.Add( *TopRightKey, ViewportTR );
-	Viewports.Add( *BottomRightKey, ViewportBR );
+	TSharedRef< SWidget > ViewportBR = InParentLayout->FactoryViewport(*BottomRightType, Args);
 
 	// Set up the splitter
 	SplitterWidget = 
 	SNew( SSplitter2x2 )
 	.TopLeft()
 	[
-		ViewportTL->AsWidget()
+		ViewportTL
 	]
 	.BottomLeft()
 	[
-		ViewportBL->AsWidget()
+		ViewportBL
 	]
 	.TopRight()
 	[
-		ViewportTR->AsWidget()
+		ViewportTR
 	]
 	.BottomRight()
 	[
-		ViewportBR->AsWidget()
+		ViewportBR
 	];
 	
 	if (SplitterPercentages.Num() > 0)
@@ -116,3 +110,53 @@ TSharedRef<SWidget> FEditorViewportLayout2x2::MakeViewportLayout(const FString&
 	return SplitterWidget.ToSharedRef();
 }
 
+void FEditorViewportLayout2x2::ReplaceWidget(TSharedRef<SWidget> OriginalWidget, TSharedRef<SWidget> ReplacementWidget)
+{
+	bool bWasFound = false;
+
+	if (SplitterWidget->GetTopLeftContent() == OriginalWidget)
+	{
+		SplitterWidget->SetTopLeftContent(ReplacementWidget);
+		bWasFound = true;
+	}
+
+	else if (SplitterWidget->GetBottomLeftContent() == OriginalWidget)
+	{
+		SplitterWidget->SetBottomLeftContent(ReplacementWidget);
+		bWasFound = true;
+	}
+
+	else if (SplitterWidget->GetTopRightContent() == OriginalWidget)
+	{
+		SplitterWidget->SetTopRightContent(ReplacementWidget);
+		bWasFound = true;
+	}
+
+	else if (SplitterWidget->GetBottomRightContent() == OriginalWidget)
+	{
+		SplitterWidget->SetBottomRightContent(ReplacementWidget);
+		bWasFound = true;
+	}
+
+	// Source widget should have already been a content widget for the splitter
+	check(bWasFound);
+}
+
+const FName& FEditorViewportLayout2x2::GetLayoutTypeName() const
+{
+	return EditorViewportConfigurationNames::FourPanes2x2;
+}
+
+void FEditorViewportLayout2x2::SaveLayoutString(const FString& LayoutString) const
+{
+	FString SpecificLayoutString = GetTypeSpecificLayoutString(LayoutString);
+
+	const FString& IniSection = FLayoutSaveRestore::GetAdditionalLayoutConfigIni();
+
+	TArray<FVector2D> Percentages;
+	SplitterWidget->GetSplitterPercentages(Percentages);
+	for (int32 i = 0; i < Percentages.Num(); ++i)
+	{
+		GConfig->SetString(*IniSection, *(SpecificLayoutString + FString::Printf(TEXT(".Percentages%i"), i)), *Percentages[i].ToString(), GEditorPerProjectIni);
+	}
+}
diff --git a/Engine/Source/Editor/UnrealEd/Private/EditorViewportLayoutFourPanes.cpp b/Engine/Source/Editor/UnrealEd/Private/EditorViewportLayoutFourPanes.cpp
--- a/Engine/Source/Editor/UnrealEd/Private/EditorViewportLayoutFourPanes.cpp
+++ b/Engine/Source/Editor/UnrealEd/Private/EditorViewportLayoutFourPanes.cpp
@@ -18,10 +18,8 @@ namespace ViewportLayoutFourPanesDefs
 
 // FEditorViewportLayoutFourPanes /////////////////////////////
 
-TSharedRef<SWidget> FEditorViewportLayoutFourPanes::MakeViewportLayout(const FString& LayoutString)
+TSharedRef<SWidget> FEditorViewportLayoutFourPanes::MakeViewportLayout(TSharedPtr<FAssetEditorViewportLayout> InParentLayout, const FString& LayoutString)
 {
-	FString SpecificLayoutString = GetTypeSpecificLayoutString(LayoutString);
-
 	FString ViewportKey0, ViewportKey1, ViewportKey2, ViewportKey3;
 	FString ViewportType0 = TEXT("Default"), ViewportType1 = TEXT("Default"), ViewportType2 = TEXT("Default"), ViewportType3 = TEXT("Default");
 
@@ -29,6 +27,7 @@ TSharedRef<SWidget> FEditorViewportLayoutFourPanes::MakeViewportLayout(const FSt
 	float SecondarySplitterPercentage0 = ViewportLayoutFourPanesDefs::DefaultSecondarySplitterPercentage;
 	float SecondarySplitterPercentage1 = ViewportLayoutFourPanesDefs::DefaultSecondarySplitterPercentage;
 
+	FString SpecificLayoutString = GetTypeSpecificLayoutString(LayoutString);
 	if (!SpecificLayoutString.IsEmpty())
 	{
 		// The Layout String only holds the unique ID of the Additional Layout Configs to use
@@ -61,49 +60,96 @@ TSharedRef<SWidget> FEditorViewportLayoutFourPanes::MakeViewportLayout(const FSt
 
 	// Set up the viewports
 	FAssetEditorViewportConstructionArgs Args;
-	Args.ParentLayout = AsShared();
+	Args.ParentLayout = InParentLayout;
 	Args.IsEnabled = FSlateApplication::Get().GetNormalExecutionAttribute();
 
 	Args.bRealtime = !FPlatformMisc::IsRemoteSession();
 	Args.ConfigKey = *ViewportKey0;
 	Args.ViewportType = LVT_Perspective;
-	TSharedRef<IEditorViewportLayoutEntity> Viewport0 = FactoryViewport(*ViewportType0, Args);
+	TSharedRef<SWidget> Viewport0 = InParentLayout->FactoryViewport(*ViewportType0, Args);
 
 	Args.bRealtime = false;
 	Args.ConfigKey = *ViewportKey1;
 	Args.ViewportType = LVT_OrthoXY;
-	TSharedRef<IEditorViewportLayoutEntity> Viewport1 = FactoryViewport(*ViewportType1, Args);
+	TSharedRef<SWidget> Viewport1 = InParentLayout->FactoryViewport(*ViewportType1, Args);
 
 	// Front viewport
 	Args.bRealtime = false;
 	Args.ConfigKey = *ViewportKey2;
 	Args.ViewportType = LVT_OrthoXZ;
-	TSharedRef<IEditorViewportLayoutEntity> Viewport2 = FactoryViewport(*ViewportType2, Args);
+	TSharedRef<SWidget> Viewport2 = InParentLayout->FactoryViewport(*ViewportType2, Args);
 
 	// Top Viewport
 	Args.bRealtime = false;
 	Args.ConfigKey = *ViewportKey3;
 	Args.ViewportType = LVT_OrthoYZ;
-	TSharedRef<IEditorViewportLayoutEntity> Viewport3 = FactoryViewport(*ViewportType2, Args);
-
-	Viewports.Add( *ViewportKey0, Viewport0 );
-	Viewports.Add( *ViewportKey1, Viewport1 );
-	Viewports.Add( *ViewportKey2, Viewport2 );
-	Viewports.Add( *ViewportKey3, Viewport3 );
+	TSharedRef<SWidget> Viewport3 = InParentLayout->FactoryViewport(*ViewportType2, Args);
 
 	TSharedRef<SWidget> LayoutWidget = MakeFourPanelWidget(
-		Viewports,
-		Viewport0->AsWidget(), Viewport1->AsWidget(), Viewport2->AsWidget(), Viewport3->AsWidget(),
+		Viewport0, Viewport1, Viewport2, Viewport3,
 		PrimarySplitterPercentage, SecondarySplitterPercentage0, SecondarySplitterPercentage1);
 
 	return LayoutWidget;
 }
 
+void FEditorViewportLayoutFourPanes::SaveLayoutString(const FString& LayoutString) const
+{
+	FString SpecificLayoutString = GetTypeSpecificLayoutString(LayoutString);
+
+	const FString& IniSection = FLayoutSaveRestore::GetAdditionalLayoutConfigIni();
+
+	check(PrimarySplitterWidget->GetChildren()->Num() == 2);
+	float PrimaryPercentage = PrimarySplitterWidget->SlotAt(0).SizeValue.Get();
+	check(SecondarySplitterWidget->GetChildren()->Num() == 3);
+	float SecondaryPercentage0 = SecondarySplitterWidget->SlotAt(0).SizeValue.Get();
+	float SecondaryPercentage1 = SecondarySplitterWidget->SlotAt(1).SizeValue.Get();
+
+	GConfig->SetString(*IniSection, *(SpecificLayoutString + TEXT(".Percentage0")), *TTypeToString<float>::ToString(PrimaryPercentage), GEditorPerProjectIni);
+	GConfig->SetString(*IniSection, *(SpecificLayoutString + TEXT(".Percentage1")), *TTypeToString<float>::ToString(SecondaryPercentage0), GEditorPerProjectIni);
+	GConfig->SetString(*IniSection, *(SpecificLayoutString + TEXT(".Percentage2")), *TTypeToString<float>::ToString(SecondaryPercentage1), GEditorPerProjectIni);
+}
+
+void FEditorViewportLayoutFourPanes::ReplaceWidget(TSharedRef<SWidget> OriginalWidget, TSharedRef<SWidget> ReplacementWidget)
+{
+	bool bWasFound = false;
+
+	for (int32 SlotIdx = 0; SlotIdx < PrimarySplitterWidget->GetChildren()->Num(); SlotIdx++)
+	{
+		if (PrimarySplitterWidget->GetChildren()->GetChildAt(SlotIdx) == OriginalWidget)
+		{
+			PrimarySplitterWidget->SlotAt(SlotIdx)
+				[
+					ReplacementWidget
+				];
+			bWasFound = true;
+			break;
+		}
+	}
+
+	for (int32 SlotIdx = 0; SlotIdx < SecondarySplitterWidget->GetChildren()->Num(); SlotIdx++)
+	{
+		if (SecondarySplitterWidget->GetChildren()->GetChildAt(SlotIdx) == OriginalWidget)
+		{
+			SecondarySplitterWidget->SlotAt(SlotIdx)
+				[
+					ReplacementWidget
+				];
+			bWasFound = true;
+			break;
+		}
+	}
+
+	// Source widget should have already been a content widget for the splitter
+	check(bWasFound);
+}
 
 // FEditorViewportLayoutFourPanesLeft /////////////////////////////
+const FName& FEditorViewportLayoutFourPanesLeft::GetLayoutTypeName() const
+{
+	return EditorViewportConfigurationNames::FourPanesLeft;
+}
 
 TSharedRef<SWidget> FEditorViewportLayoutFourPanesLeft::MakeFourPanelWidget(
-	TMap<FName, TSharedPtr< IEditorViewportLayoutEntity >>& ViewportWidgets,
 	TSharedRef<SWidget> Viewport0, TSharedRef<SWidget> Viewport1, TSharedRef<SWidget> Viewport2, TSharedRef<SWidget> Viewport3,
 	float PrimarySplitterPercentage, float SecondarySplitterPercentage0, float SecondarySplitterPercentage1)
 {
@@ -142,9 +188,12 @@ TSharedRef<SWidget> FEditorViewportLayoutFourPanesLeft::MakeFourPanelWidget(
 
 
 // FEditorViewportLayoutFourPanesRight /////////////////////////////
+const FName& FEditorViewportLayoutFourPanesRight::GetLayoutTypeName() const
+{
+	return EditorViewportConfigurationNames::FourPanesRight;
+}
 
 TSharedRef<SWidget> FEditorViewportLayoutFourPanesRight::MakeFourPanelWidget(
-	TMap<FName, TSharedPtr< IEditorViewportLayoutEntity >>& ViewportWidgets,
 	TSharedRef<SWidget> Viewport0, TSharedRef<SWidget> Viewport1, TSharedRef<SWidget> Viewport2, TSharedRef<SWidget> Viewport3,
 	float PrimarySplitterPercentage, float SecondarySplitterPercentage0, float SecondarySplitterPercentage1)
 {
@@ -183,9 +232,12 @@ TSharedRef<SWidget> FEditorViewportLayoutFourPanesRight::MakeFourPanelWidget(
 
 
 // FEditorViewportLayoutFourPanesTop /////////////////////////////
+const FName& FEditorViewportLayoutFourPanesTop::GetLayoutTypeName() const
+{
+	return EditorViewportConfigurationNames::FourPanesTop;
+}
 
 TSharedRef<SWidget> FEditorViewportLayoutFourPanesTop::MakeFourPanelWidget(
-	TMap<FName, TSharedPtr< IEditorViewportLayoutEntity >>& ViewportWidgets,
 	TSharedRef<SWidget> Viewport0, TSharedRef<SWidget> Viewport1, TSharedRef<SWidget> Viewport2, TSharedRef<SWidget> Viewport3,
 	float PrimarySplitterPercentage, float SecondarySplitterPercentage0, float SecondarySplitterPercentage1)
 {
@@ -224,9 +276,12 @@ TSharedRef<SWidget> FEditorViewportLayoutFourPanesTop::MakeFourPanelWidget(
 
 
 // FEditorViewportLayoutFourPanesBottom /////////////////////////////
+const FName& FEditorViewportLayoutFourPanesBottom::GetLayoutTypeName() const
+{
+	return EditorViewportConfigurationNames::FourPanesBottom;
+}
 
 TSharedRef<SWidget> FEditorViewportLayoutFourPanesBottom::MakeFourPanelWidget(
-	TMap<FName, TSharedPtr< IEditorViewportLayoutEntity >>& ViewportWidgets,
 	TSharedRef<SWidget> Viewport0, TSharedRef<SWidget> Viewport1, TSharedRef<SWidget> Viewport2, TSharedRef<SWidget> Viewport3,
 	float PrimarySplitterPercentage, float SecondarySplitterPercentage0, float SecondarySplitterPercentage1)
 {
diff --git a/Engine/Source/Editor/UnrealEd/Private/EditorViewportLayoutOnePane.cpp b/Engine/Source/Editor/UnrealEd/Private/EditorViewportLayoutOnePane.cpp
--- a/Engine/Source/Editor/UnrealEd/Private/EditorViewportLayoutOnePane.cpp
+++ b/Engine/Source/Editor/UnrealEd/Private/EditorViewportLayoutOnePane.cpp
@@ -1,21 +1,17 @@
 // Copyright Epic Games, Inc. All Rights Reserved.
 
 #include "EditorViewportLayoutOnePane.h"
-#include "Widgets/DeclarativeSyntaxSupport.h"
-#include "Widgets/SBoxPanel.h"
-#include "Framework/Docking/LayoutService.h"
 #include "ShowFlags.h"
 #include "Editor.h"
+#include "Framework/Docking/LayoutService.h"
+#include "CoreGlobals.h"
 #include "Misc/ConfigCacheIni.h"
-#include "Modules/ModuleManager.h"
-#include "Framework/Application/SlateApplication.h"
-
+#include "Widgets/DeclarativeSyntaxSupport.h"
+#include "Widgets/SBoxPanel.h"
+#include "Widgets/SWidget.h"
 
-TSharedRef<SWidget> FEditorViewportLayoutOnePane::MakeViewportLayout(const FString& LayoutString)
+TSharedRef<SWidget> FEditorViewportLayoutOnePane::MakeViewportLayout(TSharedPtr<FAssetEditorViewportLayout> InParentLayout, const FString& LayoutString)
 {
-	FString SpecificLayoutString = GetTypeSpecificLayoutString(LayoutString);
-
-
 	FEngineShowFlags OrthoShowFlags(ESFIM_Editor);
 	ApplyViewMode(VMI_BrushWireframe, false, OrthoShowFlags);
 
@@ -24,6 +20,7 @@ TSharedRef<SWidget> FEditorViewportLayoutOnePane::MakeViewportLayout(const FStri
 
 	FString ViewportKey, ViewportType;
 
+	FString SpecificLayoutString = GetTypeSpecificLayoutString(LayoutString);
 	if (!SpecificLayoutString.IsEmpty())
 	{
 		const FString& IniSection = FLayoutSaveRestore::GetAdditionalLayoutConfigIni();
@@ -34,19 +31,41 @@ TSharedRef<SWidget> FEditorViewportLayoutOnePane::MakeViewportLayout(const FStri
 
 	// Set up the viewport
 	FAssetEditorViewportConstructionArgs Args;
- 	Args.ParentLayout = AsShared();
+ 	Args.ParentLayout = InParentLayout;
 	Args.bRealtime = !FPlatformMisc::IsRemoteSession();
  	Args.ViewportType = LVT_Perspective;
-	TSharedRef<IEditorViewportLayoutEntity> Viewport = FactoryViewport(*ViewportType, Args);
+	Args.ConfigKey = *ViewportKey;
+	TSharedRef<SWidget> Viewport = InParentLayout->FactoryViewport(*ViewportType, Args);
+	PerspectiveViewportConfigKey = *ViewportKey;
 
 	ViewportBox =
 		SNew(SHorizontalBox)
 		+ SHorizontalBox::Slot()
 		[
-			Viewport->AsWidget()
+			Viewport
 		];
 
-	Viewports.Add( *ViewportKey, Viewport );
-
 	return ViewportBox.ToSharedRef();
 }
+
+const FName& FEditorViewportLayoutOnePane::GetLayoutTypeName() const
+{
+	return EditorViewportConfigurationNames::OnePane;
+}
+
+void FEditorViewportLayoutOnePane::SaveLayoutString(const FString& LayoutString) const
+{
+}
+
+void FEditorViewportLayoutOnePane::ReplaceWidget(TSharedRef<SWidget> OriginalWidget, TSharedRef<SWidget> ReplacementWidget)
+{
+	check(ViewportBox->GetChildren()->Num() == 1)
+		TSharedRef<SWidget> ViewportWidget = ViewportBox->GetChildren()->GetChildAt(0);
+
+	check(ViewportWidget == OriginalWidget);
+	ViewportBox->RemoveSlot(OriginalWidget);
+	ViewportBox->AddSlot()
+		[
+			ReplacementWidget
+		];
+}
diff --git a/Engine/Source/Editor/UnrealEd/Private/EditorViewportLayoutThreePanes.cpp b/Engine/Source/Editor/UnrealEd/Private/EditorViewportLayoutThreePanes.cpp
--- a/Engine/Source/Editor/UnrealEd/Private/EditorViewportLayoutThreePanes.cpp
+++ b/Engine/Source/Editor/UnrealEd/Private/EditorViewportLayoutThreePanes.cpp
@@ -4,9 +4,9 @@
 #include "Framework/Docking/LayoutService.h"
 #include "ShowFlags.h"
 #include "Misc/ConfigCacheIni.h"
-#include "Modules/ModuleManager.h"
 #include "Framework/Application/SlateApplication.h"
-#include "LevelEditor.h"
+#include "Widgets/SWidget.h"
+#include "Widgets/Layout/SSplitter.h"
 
 namespace ViewportLayoutThreePanesDefs
 {
@@ -17,10 +17,8 @@ namespace ViewportLayoutThreePanesDefs
 
 // FEditorViewportLayoutThreePanes /////////////////////////////
 
-TSharedRef<SWidget> FEditorViewportLayoutThreePanes::MakeViewportLayout(const FString& LayoutString)
+TSharedRef<SWidget> FEditorViewportLayoutThreePanes::MakeViewportLayout(TSharedPtr<FAssetEditorViewportLayout> InParentLayout, const FString& LayoutString)
 {
-	FString SpecificLayoutString = GetTypeSpecificLayoutString(LayoutString);
-
 	FEngineShowFlags OrthoShowFlags(ESFIM_Editor);	
 	ApplyViewMode(VMI_BrushWireframe, false, OrthoShowFlags);
 
@@ -32,6 +30,7 @@ TSharedRef<SWidget> FEditorViewportLayoutThreePanes::MakeViewportLayout(const FS
 	float PrimarySplitterPercentage = ViewportLayoutThreePanesDefs::DefaultSplitterPercentage;
 	float SecondarySplitterPercentage = ViewportLayoutThreePanesDefs::DefaultSplitterPercentage;
 
+	FString SpecificLayoutString = GetTypeSpecificLayoutString(LayoutString);
 	if (!SpecificLayoutString.IsEmpty())
 	{
 		// The Layout String only holds the unique ID of the Additional Layout Configs to use
@@ -58,41 +57,87 @@ TSharedRef<SWidget> FEditorViewportLayoutThreePanes::MakeViewportLayout(const FS
 
 	// Set up the viewports
 	FAssetEditorViewportConstructionArgs Args;
-	Args.ParentLayout = AsShared();
+	Args.ParentLayout = InParentLayout;
 	Args.IsEnabled = FSlateApplication::Get().GetNormalExecutionAttribute();
 
 	Args.bRealtime = !FPlatformMisc::IsRemoteSession();
 	Args.ConfigKey = *ViewportKey0;
 	Args.ViewportType = LVT_Perspective;
-	TSharedRef<IEditorViewportLayoutEntity> Viewport0 =FactoryViewport(*ViewportType0, Args);
+	TSharedRef<SWidget> Viewport0 = InParentLayout->FactoryViewport(*ViewportType0, Args);
+	PerspectiveViewportConfigKey = *ViewportKey0;
 
 	Args.bRealtime = false;
 	Args.ConfigKey = *ViewportKey1;
 	Args.ViewportType = LVT_OrthoXY;
-	TSharedRef<IEditorViewportLayoutEntity> Viewport1 =FactoryViewport(*ViewportType1, Args);
+	TSharedRef<SWidget> Viewport1 = InParentLayout->FactoryViewport(*ViewportType1, Args);
 
 	// Front viewport
 	Args.bRealtime = false;
 	Args.ConfigKey = *ViewportKey2;
 	Args.ViewportType = LVT_OrthoXZ;
-	TSharedRef<IEditorViewportLayoutEntity> Viewport2 =FactoryViewport(*ViewportType2, Args);
-
-	Viewports.Add( *ViewportKey0, Viewport0 );
-	Viewports.Add( *ViewportKey1, Viewport1 );
-	Viewports.Add( *ViewportKey2, Viewport2 );
+	TSharedRef<SWidget> Viewport2 = InParentLayout->FactoryViewport(*ViewportType2, Args);
 
 	TSharedRef<SWidget> LayoutWidget = MakeThreePanelWidget(
-		Viewports,
-		Viewport0->AsWidget(), Viewport1->AsWidget(), Viewport2->AsWidget(),
+		Viewport0, Viewport1, Viewport2,
 		PrimarySplitterPercentage, SecondarySplitterPercentage);
 
 	return LayoutWidget;
 }
 
+void FEditorViewportLayoutThreePanes::SaveLayoutString(const FString& SpecificLayoutString) const
+{
+	const FString& IniSection = FLayoutSaveRestore::GetAdditionalLayoutConfigIni();
+
+	check(PrimarySplitterWidget->GetChildren()->Num() == 2);
+	float PrimaryPercentage = PrimarySplitterWidget->SlotAt(0).SizeValue.Get();
+	check(SecondarySplitterWidget->GetChildren()->Num() == 2);
+	float SecondaryPercentage = SecondarySplitterWidget->SlotAt(0).SizeValue.Get();
+
+	GConfig->SetString(*IniSection, *(SpecificLayoutString + TEXT(".Percentage0")), *TTypeToString<float>::ToString(PrimaryPercentage), GEditorPerProjectIni);
+	GConfig->SetString(*IniSection, *(SpecificLayoutString + TEXT(".Percentage1")), *TTypeToString<float>::ToString(SecondaryPercentage), GEditorPerProjectIni);
+}
+
+void FEditorViewportLayoutThreePanes::ReplaceWidget(TSharedRef<SWidget> OriginalWidget, TSharedRef<SWidget> ReplacementWidget)
+{
+	bool bWasFound = false;
+
+	for (int32 SlotIdx = 0; SlotIdx < PrimarySplitterWidget->GetChildren()->Num(); SlotIdx++)
+	{
+		if (PrimarySplitterWidget->GetChildren()->GetChildAt(SlotIdx) == OriginalWidget)
+		{
+			PrimarySplitterWidget->SlotAt(SlotIdx)
+				[
+					ReplacementWidget
+				];
+			bWasFound = true;
+			break;
+		}
+	}
+
+	for (int32 SlotIdx = 0; SlotIdx < SecondarySplitterWidget->GetChildren()->Num(); SlotIdx++)
+	{
+		if (SecondarySplitterWidget->GetChildren()->GetChildAt(SlotIdx) == OriginalWidget)
+		{
+			SecondarySplitterWidget->SlotAt(SlotIdx)
+				[
+					ReplacementWidget
+				];
+			bWasFound = true;
+			break;
+		}
+	}
+
+	// Source widget should have already been a content widget for the splitter
+	check(bWasFound);
+}
+
 // FEditorViewportLayoutThreePanesLeft /////////////////////////////
+const FName& FEditorViewportLayoutThreePanesLeft::GetLayoutTypeName() const
+{
+	return EditorViewportConfigurationNames::ThreePanesLeft;
+}
 
 TSharedRef<SWidget> FEditorViewportLayoutThreePanesLeft::MakeThreePanelWidget(
-	TMap< FName, TSharedPtr< IEditorViewportLayoutEntity >>& ViewportWidgets,
 	const TSharedRef<SWidget>& Viewport0, const TSharedRef<SWidget>& Viewport1, const TSharedRef<SWidget>& Viewport2,
 	float PrimarySplitterPercentage, float SecondarySplitterPercentage)
 {
@@ -126,9 +171,12 @@ TSharedRef<SWidget> FEditorViewportLayoutThreePanesLeft::MakeThreePanelWidget(
 
 
 // FEditorViewportLayoutThreePanesRight /////////////////////////////
+const FName& FEditorViewportLayoutThreePanesRight::GetLayoutTypeName() const
+{
+	return EditorViewportConfigurationNames::ThreePanesRight;
+}
 
 TSharedRef<SWidget> FEditorViewportLayoutThreePanesRight::MakeThreePanelWidget(
-	TMap< FName, TSharedPtr< IEditorViewportLayoutEntity >>& ViewportWidgets,
 	const TSharedRef<SWidget>& Viewport0, const TSharedRef<SWidget>& Viewport1, const TSharedRef<SWidget>& Viewport2,
 	float PrimarySplitterPercentage, float SecondarySplitterPercentage)
 {
@@ -162,9 +210,12 @@ TSharedRef<SWidget> FEditorViewportLayoutThreePanesRight::MakeThreePanelWidget(
 
 
 // FEditorViewportLayoutThreePanesTop /////////////////////////////
+const FName& FEditorViewportLayoutThreePanesTop::GetLayoutTypeName() const
+{
+	return EditorViewportConfigurationNames::ThreePanesTop;
+}
 
 TSharedRef<SWidget> FEditorViewportLayoutThreePanesTop::MakeThreePanelWidget(
-	TMap< FName, TSharedPtr< IEditorViewportLayoutEntity >>& ViewportWidgets,
 	const TSharedRef<SWidget>& Viewport0, const TSharedRef<SWidget>& Viewport1, const TSharedRef<SWidget>& Viewport2,
 	float PrimarySplitterPercentage, float SecondarySplitterPercentage)
 {
@@ -198,9 +249,12 @@ TSharedRef<SWidget> FEditorViewportLayoutThreePanesTop::MakeThreePanelWidget(
 
 
 // FEditorViewportLayoutThreePanesBottom /////////////////////////////
+const FName& FEditorViewportLayoutThreePanesBottom::GetLayoutTypeName() const
+{
+	return EditorViewportConfigurationNames::ThreePanesBottom;
+}
 
 TSharedRef<SWidget> FEditorViewportLayoutThreePanesBottom::MakeThreePanelWidget(
-	TMap< FName, TSharedPtr< IEditorViewportLayoutEntity >>& ViewportWidgets,
 	const TSharedRef<SWidget>& Viewport0, const TSharedRef<SWidget>& Viewport1, const TSharedRef<SWidget>& Viewport2,
 	float PrimarySplitterPercentage, float SecondarySplitterPercentage)
 {
diff --git a/Engine/Source/Editor/UnrealEd/Private/EditorViewportLayoutTwoPanes.cpp b/Engine/Source/Editor/UnrealEd/Private/EditorViewportLayoutTwoPanes.cpp
--- a/Engine/Source/Editor/UnrealEd/Private/EditorViewportLayoutTwoPanes.cpp
+++ b/Engine/Source/Editor/UnrealEd/Private/EditorViewportLayoutTwoPanes.cpp
@@ -16,10 +16,8 @@ namespace ViewportLayoutTwoPanesDefs
 
 
 template <EOrientation TOrientation>
-TSharedRef<SWidget> TEditorViewportLayoutTwoPanes<TOrientation>::MakeViewportLayout(const FString& LayoutString)
+TSharedRef<SWidget> TEditorViewportLayoutTwoPanes<TOrientation>::MakeViewportLayout(TSharedPtr<FAssetEditorViewportLayout> InParentLayout, const FString& SpecificLayoutString)
 {
-	FString SpecificLayoutString = GetTypeSpecificLayoutString(LayoutString);
-
 	FString ViewportKey0, ViewportKey1;
 	FString ViewportType0, ViewportType1;
 	float SplitterPercentage = ViewportLayoutTwoPanesDefs::DefaultSplitterPercentage;
@@ -41,23 +39,22 @@ TSharedRef<SWidget> TEditorViewportLayoutTwoPanes<TOrientation>::MakeViewportLay
 			TTypeFromString<float>::FromString(SplitterPercentage, *PercentageString);
 		}
 	}
+
 	// Set up the viewports
 	FAssetEditorViewportConstructionArgs Args;
-	Args.ParentLayout = AsShared();
+	Args.ParentLayout = InParentLayout;
 	Args.IsEnabled = FSlateApplication::Get().GetNormalExecutionAttribute();
 
 	Args.bRealtime = false;
 	Args.ConfigKey = *ViewportKey0;
 	Args.ViewportType = LVT_OrthoXY;
-	TSharedRef<IEditorViewportLayoutEntity> Viewport0 = FactoryViewport(*ViewportType0, Args);
+	TSharedRef<SWidget> Viewport0 = InParentLayout->FactoryViewport(*ViewportType0, Args);
+	PerspectiveViewportConfigKey = *ViewportKey0;
 
 	Args.bRealtime = !FPlatformMisc::IsRemoteSession();
 	Args.ConfigKey = *ViewportKey1;
 	Args.ViewportType = LVT_Perspective;
-	TSharedRef<IEditorViewportLayoutEntity> Viewport1 = FactoryViewport(*ViewportType1, Args);
-
-	Viewports.Add(*ViewportKey0, Viewport0);
-	Viewports.Add(*ViewportKey1, Viewport1);
+	TSharedRef<SWidget> Viewport1 = InParentLayout->FactoryViewport(*ViewportType1, Args);
 
 	SplitterWidget =
 		SNew(SSplitter)
@@ -65,17 +62,60 @@ TSharedRef<SWidget> TEditorViewportLayoutTwoPanes<TOrientation>::MakeViewportLay
 		+ SSplitter::Slot()
 		.Value(SplitterPercentage)
 		[
-			Viewport0->AsWidget()
+			Viewport0
 		]
 	+ SSplitter::Slot()
 		.Value(1.0f - SplitterPercentage)
 		[
-			Viewport1->AsWidget()
+			Viewport1
 		];
 
 	return SplitterWidget.ToSharedRef();
 }
 
+template <EOrientation TOrientation>
+void TEditorViewportLayoutTwoPanes<TOrientation>::ReplaceWidget(TSharedRef<SWidget> OriginalWidget, TSharedRef<SWidget> ReplacementWidget)
+{
+	bool bWasFound = false;
+
+	for (int32 SlotIdx = 0; SlotIdx < SplitterWidget->GetChildren()->Num(); SlotIdx++)
+	{
+		if (SplitterWidget->GetChildren()->GetChildAt(SlotIdx) == OriginalWidget)
+		{
+			SplitterWidget->SlotAt(SlotIdx)
+				[
+					ReplacementWidget
+				];
+			bWasFound = true;
+			break;
+		}
+	}
+
+	// Source widget should have already been a content widget for the splitter
+	check(bWasFound);
+}
+
+template <EOrientation TOrientation>
+void TEditorViewportLayoutTwoPanes<TOrientation>::SaveLayoutString(const FString& LayoutString) const
+{
+	const FString& IniSection = FLayoutSaveRestore::GetAdditionalLayoutConfigIni();
+
+	check(SplitterWidget->GetChildren()->Num() == 2);
+	float Percentage = SplitterWidget->SlotAt(0).SizeValue.Get();
+
+	FString SpecificLayoutString = GetTypeSpecificLayoutString(LayoutString);
+	GConfig->SetString(*IniSection, *(SpecificLayoutString + TEXT(".Percentage")), *TTypeToString<float>::ToString(Percentage), GEditorPerProjectIni);
+}
+
+const FName& FEditorViewportLayoutTwoPanesVert::GetLayoutTypeName() const
+{
+	return EditorViewportConfigurationNames::TwoPanesVert;
+}
+
+const FName& FEditorViewportLayoutTwoPanesHoriz::GetLayoutTypeName() const
+{
+	return EditorViewportConfigurationNames::TwoPanesHoriz;
+}
 
 /**
 * Function avoids linker errors on the template class functions in this cpp file.
diff --git a/Engine/Source/Editor/UnrealEd/Private/EditorViewportTabContent.cpp b/Engine/Source/Editor/UnrealEd/Private/EditorViewportTabContent.cpp
--- a/Engine/Source/Editor/UnrealEd/Private/EditorViewportTabContent.cpp
+++ b/Engine/Source/Editor/UnrealEd/Private/EditorViewportTabContent.cpp
@@ -2,37 +2,14 @@
 
 #include "EditorViewportTabContent.h"
 #include "SEditorViewport.h"
-#include "EditorViewportLayoutOnePane.h"
-#include "EditorViewportLayoutTwoPanes.h"
-#include "EditorViewportLayoutThreePanes.h"
-#include "EditorViewportLayoutFourPanes.h"
-#include "EditorViewportLayout2x2.h"
 #include "Framework/Docking/LayoutService.h"
 #include "CoreGlobals.h"
 #include "Misc/ConfigCacheIni.h"
 
 TSharedPtr< FEditorViewportLayout > FEditorViewportTabContent::ConstructViewportLayoutByTypeName(const FName& TypeName, bool bSwitchingLayouts)
 {
-	TSharedPtr< FEditorViewportLayout > ViewportLayout;
-
-	//The items in these ifs should match the names in namespace EditorViewportConfigurationNames
-	if (TypeName == EditorViewportConfigurationNames::TwoPanesHoriz) ViewportLayout = MakeShareable(new FEditorViewportLayoutTwoPanesHoriz);
-	else if (TypeName == EditorViewportConfigurationNames::TwoPanesVert) ViewportLayout = MakeShareable(new FEditorViewportLayoutTwoPanesVert);
- 	else if (TypeName == EditorViewportConfigurationNames::FourPanes2x2) ViewportLayout = MakeShareable(new FEditorViewportLayout2x2);
-	else if (TypeName == EditorViewportConfigurationNames::ThreePanesLeft) ViewportLayout = MakeShareable(new FEditorViewportLayoutThreePanesLeft);
-	else if (TypeName == EditorViewportConfigurationNames::ThreePanesRight) ViewportLayout = MakeShareable(new FEditorViewportLayoutThreePanesRight);
-	else if (TypeName == EditorViewportConfigurationNames::ThreePanesTop) ViewportLayout = MakeShareable(new FEditorViewportLayoutThreePanesTop);
-	else if (TypeName == EditorViewportConfigurationNames::ThreePanesBottom) ViewportLayout = MakeShareable(new FEditorViewportLayoutThreePanesBottom);
-	else if (TypeName == EditorViewportConfigurationNames::FourPanesLeft) ViewportLayout = MakeShareable(new FEditorViewportLayoutFourPanesLeft);
-	else if (TypeName == EditorViewportConfigurationNames::FourPanesRight) ViewportLayout = MakeShareable(new FEditorViewportLayoutFourPanesRight);
-	else if (TypeName == EditorViewportConfigurationNames::FourPanesBottom) ViewportLayout = MakeShareable(new FEditorViewportLayoutFourPanesBottom);
-  	else if (TypeName == EditorViewportConfigurationNames::FourPanesTop) ViewportLayout = MakeShareable(new FEditorViewportLayoutFourPanesTop);
-	else /*(TypeName == EditorViewportConfigurationNames::OnePane)*/ ViewportLayout = MakeShareable(new FEditorViewportLayoutOnePane);
-
-	if (!ensure(ViewportLayout.IsValid()))
-	{
-		ViewportLayout = MakeShareable(new FEditorViewportLayoutOnePane);
-	}
+	TSharedPtr<FEditorViewportLayout> ViewportLayout = FactoryViewportLayout(bSwitchingLayouts);
+	ViewportLayout->FactoryPaneConfigurationFromTypeName(TypeName);
 	return ViewportLayout;
 }
 
@@ -84,15 +61,7 @@ void FEditorViewportTabContent::SaveConfig() const
 {
 	if (ActiveViewportLayout.IsValid())
 	{
-		if (!LayoutString.IsEmpty())
-		{
-			FString LayoutTypeString = ActiveViewportLayout->GetLayoutTypeName().ToString();
-
-			const FString& IniSection = FLayoutSaveRestore::GetAdditionalLayoutConfigIni();
-			GConfig->SetString(*IniSection, *(LayoutString + TEXT(".LayoutType")), *LayoutTypeString, GEditorPerProjectIni);
-		}
-
-		ActiveViewportLayout->SaveLayoutString(LayoutString);
+		ActiveViewportLayout->SaveConfig(LayoutString);
 	}
 }
 
@@ -142,7 +111,7 @@ void FEditorViewportTabContent::RefreshViewportConfiguration()
 		return;
 	}
 
-	FName ConfigurationName = ActiveViewportLayout->GetLayoutTypeName();
+	FName ConfigurationName = ActiveViewportLayout->GetActivePaneConfigurationTypeName();
 	for (auto& Pair : ActiveViewportLayout->GetViewports())
 	{
 		if (Pair.Value->AsWidget()->HasFocusedDescendants())
@@ -161,3 +130,8 @@ const AssetEditorViewportFactoryFunction* FEditorViewportTabContent::FindViewpor
 {
 	return ViewportCreationFactories.Find(InTypeName);
 }
+
+TSharedPtr<FEditorViewportLayout> FEditorViewportTabContent::FactoryViewportLayout(bool bIsSwitchingLayouts)
+{
+	return MakeShareable(new FAssetEditorViewportLayout);
+}
diff --git a/Engine/Source/Editor/UnrealEd/Private/ViewportTabContent.cpp b/Engine/Source/Editor/UnrealEd/Private/ViewportTabContent.cpp
--- a/Engine/Source/Editor/UnrealEd/Private/ViewportTabContent.cpp
+++ b/Engine/Source/Editor/UnrealEd/Private/ViewportTabContent.cpp
@@ -28,7 +28,7 @@ bool FViewportTabContent::IsViewportConfigurationSet(const FName& ConfigurationN
 {
 	if (ActiveViewportLayout.IsValid())
 	{
-		return ActiveViewportLayout->GetLayoutTypeName() == ConfigurationName;
+		return ActiveViewportLayout->GetActivePaneConfigurationTypeName() == ConfigurationName;
 	}
 	return false;
 }
diff --git a/Engine/Source/Editor/UnrealEd/Public/AssetEditorViewportLayout.h b/Engine/Source/Editor/UnrealEd/Public/AssetEditorViewportLayout.h
--- a/Engine/Source/Editor/UnrealEd/Public/AssetEditorViewportLayout.h
+++ b/Engine/Source/Editor/UnrealEd/Public/AssetEditorViewportLayout.h
@@ -46,18 +46,18 @@ using AssetEditorViewportFactoryFunction = TFunction<TSharedRef<SAssetEditorView
 
 namespace EditorViewportConfigurationNames
 {
-	static FName TwoPanesHoriz("TwoPanesHoriz");
-	static FName TwoPanesVert("TwoPanesVert");
-	static FName ThreePanesLeft("ThreePanesLeft");
-	static FName ThreePanesRight("ThreePanesRight");
-	static FName ThreePanesTop("ThreePanesTop");
-	static FName ThreePanesBottom("ThreePanesBottom");
-	static FName FourPanesLeft("FourPanesLeft");
-	static FName FourPanesRight("FourPanesRight");
-	static FName FourPanesTop("FourPanesTop");
-	static FName FourPanesBottom("FourPanesBottom");
-	static FName FourPanes2x2("FourPanes2x2");
-	static FName OnePane("OnePane");
+	static const FName TwoPanesHoriz("TwoPanesHoriz");
+	static const FName TwoPanesVert("TwoPanesVert");
+	static const FName ThreePanesLeft("ThreePanesLeft");
+	static const FName ThreePanesRight("ThreePanesRight");
+	static const FName ThreePanesTop("ThreePanesTop");
+	static const FName ThreePanesBottom("ThreePanesBottom");
+	static const FName FourPanesLeft("FourPanesLeft");
+	static const FName FourPanesRight("FourPanesRight");
+	static const FName FourPanesTop("FourPanesTop");
+	static const FName FourPanesBottom("FourPanesBottom");
+	static const FName FourPanes2x2("FourPanes2x2");
+	static const FName OnePane("OnePane");
 }
 
 /**
@@ -113,6 +113,39 @@ private:
 	FVector2D CachedSize;
 };
 
+class UNREALED_API FAssetEditorViewportPaneLayout
+{
+public:
+	virtual ~FAssetEditorViewportPaneLayout() = default;
+
+	/**
+	 * Sets up the layout based on the specific layout configuration implementation.
+	 *
+	 * @param SpecificLayoutString		The layout string loaded from a file.
+	 * @return The base widget representing the layout.  This is commonly a splitter.
+	 */
+	virtual TSharedRef<SWidget> MakeViewportLayout(TSharedPtr<FAssetEditorViewportLayout> InParentLayout, const FString& LayoutString) = 0;
+
+	/**
+	 * Returns the typename of this layout
+	 * Generally one of EditorViewportConfigurationNames
+	 */
+	virtual const FName& GetLayoutTypeName() const = 0;
+
+	void LoadConfig(const FString& LayoutString, TFunction<void(const FString&, const FName)> LoadAdditionalLayoutInfoCallback = nullptr);
+
+	void SaveConfig(const FString& LayoutString, TFunction<void(const FString&)> SaveAdditionalLayoutInfoCallback = nullptr) const;
+
+	virtual void ReplaceWidget(TSharedRef<SWidget> OriginalWidget, TSharedRef<SWidget> ReplacementWidget) = 0;
+
+protected:
+	FString GetTypeSpecificLayoutString(const FString& LayoutString) const;
+	virtual void SaveLayoutString(const FString& LayoutString) const {}
+	virtual void LoadLayoutString(const FString& LayoutString) {}
+
+	FName PerspectiveViewportConfigKey;
+};
+
 /**
  * Base class for viewport layout configurations
  * Handles maximizing and restoring well as visibility of specific viewports.
@@ -131,13 +164,16 @@ public:
 	virtual ~FAssetEditorViewportLayout();
 
 	/** Create an instance of a custom viewport from the specified viewport type name */
-	virtual TSharedRef<IEditorViewportLayoutEntity> FactoryViewport(FName InTypeName, const FAssetEditorViewportConstructionArgs& ConstructionArgs) const;
+	virtual TSharedRef<SWidget> FactoryViewport(FName InTypeName, const FAssetEditorViewportConstructionArgs& ConstructionArgs);
 
 	/** FTickableEditorObject interface */
 	virtual void Tick(float DeltaTime) override {}
 	virtual bool IsTickable() const override { return false; }
  	virtual TStatId GetStatId() const override { return TStatId(); }
 
+	virtual void FactoryPaneConfigurationFromTypeName(const FName& InLayoutConfigTypeName) override;
+	virtual const FName GetActivePaneConfigurationTypeName() const override;
+
 	/**
 	 * Builds a viewport layout and returns the widget containing the layout
 	 * 
@@ -150,16 +186,17 @@ public:
 	/** Returns the parent tab content object */
 	TWeakPtr< FEditorViewportTabContent > GetParentTabContent() const { return ParentTabContent; }
 
-	/** Generates a layout string for persisting settings for this layout based on the runtime type of layout */
-	FString GetTypeSpecificLayoutString(const FString& LayoutString) const;
-
 	/**
-	 * Overridden in derived classes to set up custom layouts  
+	 * Sets up the layout based on the specific layout configuration implementation.
 	 *
-	 * @param LayoutString		The layout string loaded from a file
-	 * @return The base widget representing the layout.  Usually a splitter
+	 * @param LayoutString		The layout string loaded from a file.
+	 * @return The base widget representing the layout.  This is commonly a splitter.
 	 */
-  	virtual TSharedRef<SWidget> MakeViewportLayout(const FString& LayoutString) = 0;
+	UE_DEPRECATED(5.0, "This functionality has moved to the layout configurations. Use BuildViewportLayout.")
+	virtual TSharedRef<SWidget> MakeViewportLayout(const FString& LayoutString) final;
+
+	virtual void LoadConfig(const FString& LayoutString);
+	virtual void SaveConfig(const FString& LayoutString) const;
 
 protected:
 	/**
@@ -179,4 +216,5 @@ protected:
 	/** The parent tab where this layout resides */
 	TWeakPtr< SDockTab > ParentTab;
 
+	TSharedPtr<FAssetEditorViewportPaneLayout> LayoutConfiguration;
 };
diff --git a/Engine/Source/Editor/UnrealEd/Public/EditorViewportLayout.h b/Engine/Source/Editor/UnrealEd/Public/EditorViewportLayout.h
--- a/Engine/Source/Editor/UnrealEd/Public/EditorViewportLayout.h
+++ b/Engine/Source/Editor/UnrealEd/Public/EditorViewportLayout.h
@@ -35,7 +35,7 @@ class FEditorViewportLayout
 {
 public:
 	/** Virtual destruction */
-	virtual ~FEditorViewportLayout() {};
+	virtual ~FEditorViewportLayout() = default;
 
 	virtual bool IsVisible() const { return true; }
 
@@ -44,16 +44,13 @@ public:
 	*/
 	const TMap< FName, TSharedPtr<IEditorViewportLayoutEntity> >& GetViewports() const { return Viewports; }
 
-	/**
-	* Saves viewport layout information between editor sessions
-	*/
-	virtual void SaveLayoutString(const FString& LayoutString) const {}
+	virtual void FactoryPaneConfigurationFromTypeName(const FName& InLayoutConfigTypeName) = 0;
+	virtual const FName GetActivePaneConfigurationTypeName() const = 0;
 
-	virtual const FName& GetLayoutTypeName() const = 0;
+	virtual void SaveConfig(const FString& LayoutSring) const {}
+	virtual void LoadConfig(const FString& LayoutString) {}
 
 protected:
 	/** List of all of the viewports in this layout, keyed on their config key */
 	TMap< FName, TSharedPtr< IEditorViewportLayoutEntity > > Viewports;
-
-
 };
diff --git a/Engine/Source/Editor/UnrealEd/Public/EditorViewportLayout2x2.h b/Engine/Source/Editor/UnrealEd/Public/EditorViewportLayout2x2.h
--- a/Engine/Source/Editor/UnrealEd/Public/EditorViewportLayout2x2.h
+++ b/Engine/Source/Editor/UnrealEd/Public/EditorViewportLayout2x2.h
@@ -2,23 +2,23 @@
 
 #pragma once
 
-#include "CoreMinimal.h"
-#include "Widgets/SWidget.h"
-#include "Widgets/Layout/SSplitter.h"
 #include "AssetEditorViewportLayout.h"
 
-class FEditorViewportLayout2x2 : public FAssetEditorViewportLayout
+class SSplitter;
+class SWidget;
+class FString;
+class FName;
+
+class FEditorViewportLayout2x2 : public FAssetEditorViewportPaneLayout
 {
 public:
+	// IAssetEditorViewportLayoutConfiguration overrides
+	virtual TSharedRef<SWidget> MakeViewportLayout(TSharedPtr<FAssetEditorViewportLayout> InParentLayout, const FString& LayoutString) override;
+	virtual void ReplaceWidget(TSharedRef<SWidget> OriginalWidget, TSharedRef<SWidget> ReplacementWidget) override;
+	virtual const FName& GetLayoutTypeName() const override;
+	virtual void SaveLayoutString(const FString& LayoutString) const override;
 
-	virtual const FName& GetLayoutTypeName() const override { return EditorViewportConfigurationNames::FourPanes2x2; }
 protected:
-	/**
-	 * Creates the viewports and splitter for the 2x2 layout                   
-	 */
-	virtual TSharedRef<SWidget> MakeViewportLayout(const FString& LayoutString) override;
-
-private:
 	/** The splitter widget */
 	TSharedPtr< class SSplitter2x2 > SplitterWidget;
 };
diff --git a/Engine/Source/Editor/UnrealEd/Public/EditorViewportLayoutFourPanes.h b/Engine/Source/Editor/UnrealEd/Public/EditorViewportLayoutFourPanes.h
--- a/Engine/Source/Editor/UnrealEd/Public/EditorViewportLayoutFourPanes.h
+++ b/Engine/Source/Editor/UnrealEd/Public/EditorViewportLayoutFourPanes.h
@@ -8,20 +8,19 @@
 #include "Widgets/Layout/SSplitter.h"
 #include "AssetEditorViewportLayout.h"
 
-class FEditorViewportLayoutFourPanes : public FAssetEditorViewportLayout
+class FEditorViewportLayoutFourPanes : public FAssetEditorViewportPaneLayout
 {
-protected:
-	/**
-	 * Creates the viewports and splitter for the four-pane layout              
-	 */
-	virtual TSharedRef<SWidget> MakeViewportLayout(const FString& LayoutString) override;
+public:
+	// IAssetEditorViewportLayoutConfiguration overrides
+	virtual TSharedRef<SWidget> MakeViewportLayout(TSharedPtr<FAssetEditorViewportLayout> InParentLayout, const FString& LayoutString) override;
+	virtual void SaveLayoutString(const FString& LayoutString) const override;
+	virtual void ReplaceWidget(TSharedRef<SWidget> OriginalWidget, TSharedRef<SWidget> ReplacementWidget) override;
 
+protected:
 	virtual TSharedRef<SWidget> MakeFourPanelWidget(
-		TMap<FName, TSharedPtr< IEditorViewportLayoutEntity >>& ViewportWidgets,
 		TSharedRef<SWidget> Viewport0, TSharedRef<SWidget> Viewport1, TSharedRef<SWidget> Viewport2, TSharedRef<SWidget> Viewport3,
 		float PrimarySplitterPercentage, float SecondarySplitterPercentage0, float SecondarySplitterPercentage1) = 0;
 
-protected:
 	/** The splitter widgets */
 	TSharedPtr< class SSplitter > PrimarySplitterWidget;
 	TSharedPtr< class SSplitter > SecondarySplitterWidget;
@@ -33,10 +32,10 @@ protected:
 class FEditorViewportLayoutFourPanesLeft : public FEditorViewportLayoutFourPanes
 {
 public:
-	virtual const FName& GetLayoutTypeName() const override { return EditorViewportConfigurationNames::FourPanesLeft; }
+	virtual const FName& GetLayoutTypeName() const override;
 
+protected:
 	virtual TSharedRef<SWidget> MakeFourPanelWidget(
-		TMap<FName, TSharedPtr< IEditorViewportLayoutEntity >>& ViewportWidgets,
 		TSharedRef<SWidget> Viewport0, TSharedRef<SWidget> Viewport1, TSharedRef<SWidget> Viewport2, TSharedRef<SWidget> Viewport3,
 		float PrimarySplitterPercentage, float SecondarySplitterPercentage0, float SecondarySplitterPercentage1) override;
 };
@@ -47,10 +46,10 @@ public:
 class FEditorViewportLayoutFourPanesRight : public FEditorViewportLayoutFourPanes
 {
 public:
-	virtual const FName& GetLayoutTypeName() const override { return EditorViewportConfigurationNames::FourPanesRight; }
+	virtual const FName& GetLayoutTypeName() const override;
 
+protected:
 	virtual TSharedRef<SWidget> MakeFourPanelWidget(
-		TMap<FName, TSharedPtr< IEditorViewportLayoutEntity >>& ViewportWidgets,
 		TSharedRef<SWidget> Viewport0, TSharedRef<SWidget> Viewport1, TSharedRef<SWidget> Viewport2, TSharedRef<SWidget> Viewport3,
 		float PrimarySplitterPercentage, float SecondarySplitterPercentage0, float SecondarySplitterPercentage1) override;
 };
@@ -61,10 +60,10 @@ public:
 class FEditorViewportLayoutFourPanesTop : public FEditorViewportLayoutFourPanes
 {
 public:
-	virtual const FName& GetLayoutTypeName() const override { return EditorViewportConfigurationNames::FourPanesTop; }
+	virtual const FName& GetLayoutTypeName() const override;
 
+protected:
 	virtual TSharedRef<SWidget> MakeFourPanelWidget(
-		TMap<FName, TSharedPtr< IEditorViewportLayoutEntity >>& ViewportWidgets,
 		TSharedRef<SWidget> Viewport0, TSharedRef<SWidget> Viewport1, TSharedRef<SWidget> Viewport2, TSharedRef<SWidget> Viewport3,
 		float PrimarySplitterPercentage, float SecondarySplitterPercentage0, float SecondarySplitterPercentage1) override;
 };
@@ -75,10 +74,10 @@ public:
 class FEditorViewportLayoutFourPanesBottom : public FEditorViewportLayoutFourPanes
 {
 public:
-	virtual const FName& GetLayoutTypeName() const override { return EditorViewportConfigurationNames::FourPanesBottom; }
+	virtual const FName& GetLayoutTypeName() const override;
 
+protected:
 	virtual TSharedRef<SWidget> MakeFourPanelWidget(
-		TMap<FName, TSharedPtr< IEditorViewportLayoutEntity >>& ViewportWidgets,
 		TSharedRef<SWidget> Viewport0, TSharedRef<SWidget> Viewport1, TSharedRef<SWidget> Viewport2, TSharedRef<SWidget> Viewport3,
 		float PrimarySplitterPercentage, float SecondarySplitterPercentage0, float SecondarySplitterPercentage1) override;
 };
diff --git a/Engine/Source/Editor/UnrealEd/Public/EditorViewportLayoutOnePane.h b/Engine/Source/Editor/UnrealEd/Public/EditorViewportLayoutOnePane.h
--- a/Engine/Source/Editor/UnrealEd/Public/EditorViewportLayoutOnePane.h
+++ b/Engine/Source/Editor/UnrealEd/Public/EditorViewportLayoutOnePane.h
@@ -2,22 +2,22 @@
 
 #pragma once
 
-#include "CoreMinimal.h"
-#include "Widgets/SWidget.h"
-#include "EditorViewportLayout.h"
 #include "AssetEditorViewportLayout.h"
+#include "Templates/SharedPointer.h"
 
 class SHorizontalBox;
+class SWidget;
+class FString;
+class FName;
 
-class FEditorViewportLayoutOnePane : public FAssetEditorViewportLayout
+class FEditorViewportLayoutOnePane : public FAssetEditorViewportPaneLayout
 {
 public:
-	virtual const FName& GetLayoutTypeName() const override{ return EditorViewportConfigurationNames::OnePane; }
-protected:
-	/**
-	* Creates the viewport for the single pane
-	*/
-	virtual TSharedRef<SWidget> MakeViewportLayout(const FString& LayoutString) override;
+	// IAssetEditorViewportLayoutConfiguration overrides
+	virtual TSharedRef<SWidget> MakeViewportLayout(TSharedPtr<FAssetEditorViewportLayout> InParentLayout, const FString& LayoutString) override;
+	virtual const FName& GetLayoutTypeName() const override;
+	virtual void SaveLayoutString(const FString& LayoutString) const override;
+	virtual void ReplaceWidget(TSharedRef<SWidget> OriginalWidget, TSharedRef<SWidget> ReplacementWidget) override;
 
 protected:
 	/** The viewport widget parent box */
diff --git a/Engine/Source/Editor/UnrealEd/Public/EditorViewportLayoutThreePanes.h b/Engine/Source/Editor/UnrealEd/Public/EditorViewportLayoutThreePanes.h
--- a/Engine/Source/Editor/UnrealEd/Public/EditorViewportLayoutThreePanes.h
+++ b/Engine/Source/Editor/UnrealEd/Public/EditorViewportLayoutThreePanes.h
@@ -2,29 +2,31 @@
 
 #pragma once
 
-#include "CoreMinimal.h"
-#include "Widgets/SWidget.h"
-#include "EditorViewportLayout.h"
-#include "Widgets/Layout/SSplitter.h"
 #include "AssetEditorViewportLayout.h"
 
-class FEditorViewportLayoutThreePanes : public FAssetEditorViewportLayout
+class SSplitter;
+class SWidget;
+class FName;
+
+class FEditorViewportLayoutThreePanes : public FAssetEditorViewportPaneLayout
 {
-protected:
+public:
 	/**
 	 * Creates the viewports and splitter for the two panes vertical layout                   
 	 */
-	virtual TSharedRef<SWidget> MakeViewportLayout(const FString& LayoutString) override;
+	virtual TSharedRef<SWidget> MakeViewportLayout(TSharedPtr<FAssetEditorViewportLayout> InParentLayout, const FString& LayoutString) override;
+	virtual void SaveLayoutString(const FString& LayoutString) const override;
+	virtual void ReplaceWidget(TSharedRef<SWidget> OriginalWidget, TSharedRef<SWidget> ReplacementWidget) override;
 
+protected:
 	virtual TSharedRef<SWidget> MakeThreePanelWidget(
-		TMap< FName, TSharedPtr< IEditorViewportLayoutEntity >>& ViewportWidgets,
 		const TSharedRef<SWidget>& ViewportKey0, const TSharedRef<SWidget>& ViewportKey1, const TSharedRef<SWidget>& ViewportKey2,
 		float PrimarySplitterPercentage, float SecondarySplitterPercentage) = 0;
 
-protected:
 	/** The splitter widgets */
-	TSharedPtr< class SSplitter > PrimarySplitterWidget;
-	TSharedPtr< class SSplitter > SecondarySplitterWidget;
+	TSharedPtr< SSplitter > PrimarySplitterWidget;
+	TSharedPtr< SSplitter > SecondarySplitterWidget;
+	FName PerspectiveViewportConfigKey;
 };
 
 // FEditorViewportLayoutThreePanesLeft /////////////////////////////
@@ -32,10 +34,10 @@ protected:
 class FEditorViewportLayoutThreePanesLeft : public FEditorViewportLayoutThreePanes
 {
 public:
-	virtual const FName& GetLayoutTypeName() const override { return EditorViewportConfigurationNames::ThreePanesLeft; }
+	virtual const FName& GetLayoutTypeName() const override;
 
+protected:
 	virtual TSharedRef<SWidget> MakeThreePanelWidget(
-		TMap< FName, TSharedPtr< IEditorViewportLayoutEntity >>& ViewportWidgets,
 		const TSharedRef<SWidget>& ViewportKey0, const TSharedRef<SWidget>& ViewportKey1, const TSharedRef<SWidget>& ViewportKey2,
 		float PrimarySplitterPercentage, float SecondarySplitterPercentage) override;
 };
@@ -46,10 +48,10 @@ public:
 class FEditorViewportLayoutThreePanesRight : public FEditorViewportLayoutThreePanes
 {
 public:
-	virtual const FName& GetLayoutTypeName() const override { return EditorViewportConfigurationNames::ThreePanesRight; }
+	virtual const FName& GetLayoutTypeName() const override;
 
+protected:
 	virtual TSharedRef<SWidget> MakeThreePanelWidget(
-		TMap< FName, TSharedPtr< IEditorViewportLayoutEntity >>& ViewportWidgets,
 		const TSharedRef<SWidget>& ViewportKey0, const TSharedRef<SWidget>& ViewportKey1, const TSharedRef<SWidget>& ViewportKey2,
 		float PrimarySplitterPercentage, float SecondarySplitterPercentage) override;
 };
@@ -60,10 +62,10 @@ public:
 class FEditorViewportLayoutThreePanesTop : public FEditorViewportLayoutThreePanes
 {
 public:
-	virtual const FName& GetLayoutTypeName() const override { return EditorViewportConfigurationNames::ThreePanesTop; }
+	virtual const FName& GetLayoutTypeName() const override;
 
+protected:
 	virtual TSharedRef<SWidget> MakeThreePanelWidget(
-		TMap< FName, TSharedPtr< IEditorViewportLayoutEntity >>& ViewportWidgets,
 		const TSharedRef<SWidget>& ViewportKey0, const TSharedRef<SWidget>& ViewportKey1, const TSharedRef<SWidget>& ViewportKey2,
 		float PrimarySplitterPercentage, float SecondarySplitterPercentage) override;
 };
@@ -74,10 +76,10 @@ public:
 class FEditorViewportLayoutThreePanesBottom : public FEditorViewportLayoutThreePanes
 {
 public:
-	virtual const FName& GetLayoutTypeName() const override { return EditorViewportConfigurationNames::ThreePanesBottom; }
+	virtual const FName& GetLayoutTypeName() const override;
 
+protected:
 	virtual TSharedRef<SWidget> MakeThreePanelWidget(
-		TMap< FName, TSharedPtr< IEditorViewportLayoutEntity >>& ViewportWidgets,
 		const TSharedRef<SWidget>& ViewportKey0, const TSharedRef<SWidget>& ViewportKey1, const TSharedRef<SWidget>& ViewportKey2,
 		float PrimarySplitterPercentage, float SecondarySplitterPercentage) override;
 };
diff --git a/Engine/Source/Editor/UnrealEd/Public/EditorViewportLayoutTwoPanes.h b/Engine/Source/Editor/UnrealEd/Public/EditorViewportLayoutTwoPanes.h
--- a/Engine/Source/Editor/UnrealEd/Public/EditorViewportLayoutTwoPanes.h
+++ b/Engine/Source/Editor/UnrealEd/Public/EditorViewportLayoutTwoPanes.h
@@ -2,46 +2,38 @@
 
 #pragma once
 
-#include "CoreMinimal.h"
-#include "Widgets/SWidget.h"
 #include "AssetEditorViewportLayout.h"
-#include "Widgets/Layout/SSplitter.h"
 #include "Types/SlateEnums.h"
 
+class SSplitter;
+class SWidget;
+class FString;
+class FName;
 
 template <EOrientation TOrientation>
-class TEditorViewportLayoutTwoPanes : public FAssetEditorViewportLayout
+class TEditorViewportLayoutTwoPanes : public FAssetEditorViewportPaneLayout
 {
 public:
-	/**
-	 * Saves viewport layout information between editor sessions
-	 */
-  	virtual void SaveLayoutString(const FString& LayoutString) const override {}
-
-protected:
-	virtual TSharedRef<SWidget> MakeViewportLayout(const FString& LayoutString) override;
-
+	virtual TSharedRef<SWidget> MakeViewportLayout(TSharedPtr<FAssetEditorViewportLayout> InParentLayout, const FString& LayoutString) override;
+	virtual void ReplaceWidget(TSharedRef<SWidget> OriginalWidget, TSharedRef<SWidget> ReplacementWidget) override;
+	virtual void SaveLayoutString(const FString& LayoutString) const override;
 
 private:
 	/** The splitter widget */
 	TSharedPtr< class SSplitter > SplitterWidget;
 };
 
-
 // FEditorViewportLayoutTwoPanesVert /////////////////////////////
-
-
 class FEditorViewportLayoutTwoPanesVert : public TEditorViewportLayoutTwoPanes<EOrientation::Orient_Vertical>
 {
 public:
-	virtual const FName& GetLayoutTypeName() const override { return EditorViewportConfigurationNames::TwoPanesVert; }
+	virtual const FName& GetLayoutTypeName() const override;
 };
 
 
 // FEditorViewportLayoutTwoPanesHoriz /////////////////////////////
-
 class FEditorViewportLayoutTwoPanesHoriz : public TEditorViewportLayoutTwoPanes<EOrientation::Orient_Horizontal>
 {
 public:
-	virtual const FName& GetLayoutTypeName() const override { return EditorViewportConfigurationNames::TwoPanesHoriz; }
+	virtual const FName& GetLayoutTypeName() const override;
 };
diff --git a/Engine/Source/Editor/UnrealEd/Public/EditorViewportTabContent.h b/Engine/Source/Editor/UnrealEd/Public/EditorViewportTabContent.h
--- a/Engine/Source/Editor/UnrealEd/Public/EditorViewportTabContent.h
+++ b/Engine/Source/Editor/UnrealEd/Public/EditorViewportTabContent.h
@@ -28,7 +28,7 @@ public:
 		return ActiveViewportLayout;
 	}
 	   
-	virtual TSharedPtr<FEditorViewportLayout> ConstructViewportLayoutByTypeName(const FName& TypeName, bool bSwitchingLayouts);
+	TSharedPtr<FEditorViewportLayout> ConstructViewportLayoutByTypeName(const FName& TypeName, bool bSwitchingLayouts);
 
 	virtual void Initialize(AssetEditorViewportFactoryFunction Func, TSharedPtr<SDockTab> InParentTab, const FString& InLayoutString);
 
@@ -62,6 +62,9 @@ public:
 	 */
 	const AssetEditorViewportFactoryFunction* FindViewportCreationFactory(FName InTypeName) const;
 
+protected:
+	virtual TSharedPtr<FEditorViewportLayout> FactoryViewportLayout(bool bIsSwitchingLayouts);
+
 private:
 	using AssetEditorViewportCreationFactories = TMap<FName, AssetEditorViewportFactoryFunction>;
 	AssetEditorViewportCreationFactories ViewportCreationFactories;
