Fix a crash when editor viewport layouts destruct, and their viewport clients use the level editor global mode tools.

#rnx

diff --git a/Engine/Source/Editor/UnrealEd/Private/EditorViewportClient.cpp b/Engine/Source/Editor/UnrealEd/Private/EditorViewportClient.cpp
--- a/Engine/Source/Editor/UnrealEd/Private/EditorViewportClient.cpp
+++ b/Engine/Source/Editor/UnrealEd/Private/EditorViewportClient.cpp
@@ -441,6 +441,15 @@ FEditorViewportClient::FEditorViewportClient(FEditorModeTools* InModeTools, FPre
 
 FEditorViewportClient::~FEditorViewportClient()
 {
+	// There's 3 options for the mode tools pointer in the destructor currently
+	// Either the toolkit owns the mode tools and will handle destructing all viewports and the mode tools properly,
+	// the global level editor mode tools are used, and exist for lifetime of statics in the editor, so we need to cleanup any hooks we made into the globl mode tools
+	// or the viewport client itself created a mode tools and the viewport client needs to fully clean up.
+	if (GEditor && (ModeTools == &GLevelEditorModeTools()))
+	{
+		ModeTools->OnEditorModeIDChanged().RemoveAll(this);
+	}
+
 	if (bOwnsModeTools)
 	{
 		ModeTools->OnEditorModeIDChanged().RemoveAll(this);
