Fixed issue where moving camera to actors didn't calculate bounding box for underlying actors.

diff --git a/Engine/Source/Editor/UnrealEd/Private/EditorServer.cpp b/Engine/Source/Editor/UnrealEd/Private/EditorServer.cpp
--- a/Engine/Source/Editor/UnrealEd/Private/EditorServer.cpp
+++ b/Engine/Source/Editor/UnrealEd/Private/EditorServer.cpp
@@ -4788,54 +4788,59 @@ void UEditorEngine::MoveViewportCamerasToActor(const TArray<AActor*> &Actors, co
 			}
 		}
 	}
-	else 
+	else
 	{
-		for(int32 ActorIdx = 0; ActorIdx < Actors.Num(); ActorIdx++)
+		TSet<AActor*> AlignActors;
+		for (AActor* RootActor : Actors)
 		{
-			AActor* Actor = Actors[ActorIdx];
-
-			if(Actor)
+			if (RootActor)
 			{
-
 				// Don't allow moving the viewport cameras to actors in invisible levels
-				if(!FLevelUtils::IsLevelVisible(Actor->GetLevel()))
+				if (!FLevelUtils::IsLevelVisible(RootActor->GetLevel()))
 				{
-					InvisLevelActors.Add(Actor);
+					InvisLevelActors.Add(RootActor);
 					continue;
 				}
 
-				const bool bActorIsEmitter = (Cast<AEmitter>(Actor) != NULL);
-
-				if(bActorIsEmitter && bCustomCameraAlignEmitter)
-				{
-					const FVector DefaultExtent(CustomCameraAlignEmitterDistance, CustomCameraAlignEmitterDistance, CustomCameraAlignEmitterDistance);
-					const FBox DefaultSizeBox(Actor->GetActorLocation() - DefaultExtent, Actor->GetActorLocation() + DefaultExtent);
-					BoundingBox += DefaultSizeBox;
-				}
-				else if (USceneComponent* RootComponent = Actor->GetRootComponent())
+				AlignActors.Empty(AlignActors.Num());
+				AlignActors.Add(RootActor);
+				RootActor->EditorGetUnderlyingActors(AlignActors);
+				for (AActor* AlignActor : AlignActors)
 				{
-					TArray<USceneComponent*> SceneComponents;
-					RootComponent->GetChildrenComponents(true, SceneComponents);
-					SceneComponents.Add(RootComponent);
+					const bool bActorIsEmitter = (Cast<AEmitter>(AlignActor) != NULL);
 
-					for (USceneComponent* SceneComponent : SceneComponents)
+					if (bActorIsEmitter && bCustomCameraAlignEmitter)
 					{
-						UPrimitiveComponent* PrimitiveComponent = Cast<UPrimitiveComponent>(SceneComponent);
-						if (PrimitiveComponent && PrimitiveComponent->IsRegistered())
+						const FVector DefaultExtent(CustomCameraAlignEmitterDistance, CustomCameraAlignEmitterDistance, CustomCameraAlignEmitterDistance);
+						const FBox DefaultSizeBox(AlignActor->GetActorLocation() - DefaultExtent, AlignActor->GetActorLocation() + DefaultExtent);
+						BoundingBox += DefaultSizeBox;
+					}
+					else if (USceneComponent* RootComponent = AlignActor->GetRootComponent())
+					{
+						TArray<USceneComponent*> SceneComponents;
+						RootComponent->GetChildrenComponents(true, SceneComponents);
+						SceneComponents.Add(RootComponent);
+
+						for (USceneComponent* SceneComponent : SceneComponents)
 						{
-							// Some components can have huge bounds but are not visible.  Ignore these components unless it is the only component on the actor 
-							const bool bIgnore = SceneComponents.Num() > 1 && PrimitiveComponent->IgnoreBoundsForEditorFocus();
+							UPrimitiveComponent* PrimitiveComponent = Cast<UPrimitiveComponent>(SceneComponent);
 
-							if(!bIgnore)
+							if (PrimitiveComponent && PrimitiveComponent->IsRegistered())
 							{
-								FBox LocalBox(ForceInit);
-								if (GLevelEditorModeTools().ComputeBoundingBoxForViewportFocus(Actor, PrimitiveComponent, LocalBox))
-								{
-									BoundingBox += LocalBox;
-								}
-								else
+								// Some components can have huge bounds but are not visible.  Ignore these components unless it is the only component on the actor 
+								const bool bIgnore = SceneComponents.Num() > 1 && PrimitiveComponent->IgnoreBoundsForEditorFocus();
+
+								if (!bIgnore)
 								{
-									BoundingBox += PrimitiveComponent->Bounds.GetBox();
+									FBox LocalBox(ForceInit);
+									if (GLevelEditorModeTools().ComputeBoundingBoxForViewportFocus(AlignActor, PrimitiveComponent, LocalBox))
+									{
+										BoundingBox += LocalBox;
+									}
+									else
+									{
+										BoundingBox += PrimitiveComponent->Bounds.GetBox();
+									}
 								}
 							}
 						}
@@ -4847,7 +4852,7 @@ void UEditorEngine::MoveViewportCamerasToActor(const TArray<AActor*> &Actors, co
 
 	MoveViewportCamerasToBox(BoundingBox, bActiveViewportOnly);
 
-	// Warn the user with a suppressable dialog if they attempted to zoom to actors that are in an invisible level
+	// Warn the user with a supressable dialog if they attempted to zoom to actors that are in an invisible level
 	if ( InvisLevelActors.Num() > 0 )
 	{
 		FString InvisLevelActorString;
