Add FWildcardNodeUtils with some commonly used logic to create/compare wildcard pins in blueprints

#rnx

diff --git a/Engine/Source/Editor/UnrealEd/Private/Kismet2/WildcardNodeUtils.cpp b/Engine/Source/Editor/UnrealEd/Private/Kismet2/WildcardNodeUtils.cpp
new file mode 100644
--- /dev/null
+++ b/Engine/Source/Editor/UnrealEd/Private/Kismet2/WildcardNodeUtils.cpp
@@ -0,0 +1,53 @@
+// Copyright Epic Games, Inc. All Rights Reserved.
+
+#include "Kismet2/WildcardNodeUtils.h"
+#include "EdGraphUtilities.h"
+#include "EdGraphSchema_K2.h"
+
+FEdGraphPinType FWildcardNodeUtils::GetDefaultWildcardPinType()
+{
+	static FEdGraphPinType WildcardPinType(
+		/* InPinCategory */ UEdGraphSchema_K2::PC_Wildcard,
+		/* InPinSubCategory */ NAME_None,
+		/* InPinSubCategoryObject */ nullptr,
+		/* InPinContainerType */ EPinContainerType::None,
+		/* bInIsReference */ false,
+		/* InValueTerminalType */{}
+	);
+
+	return WildcardPinType;
+}
+
+bool FWildcardNodeUtils::IsWildcardPin(const UEdGraphPin* const Pin)
+{
+	return Pin && Pin->PinType.PinCategory == UEdGraphSchema_K2::PC_Wildcard;
+}
+
+UEdGraphPin* FWildcardNodeUtils::CreateWildcardPin(UEdGraphNode* Node, const FName PinName, const EEdGraphPinDirection Direction, const EPinContainerType ContainerType/* = EPinContainerType::None*/)
+{
+	check(Node);
+
+	FEdGraphPinType PinType;
+	PinType.PinCategory = UEdGraphSchema_K2::PC_Wildcard;
+	PinType.PinSubCategory = NAME_None;
+	PinType.PinSubCategoryObject = nullptr;
+	PinType.bIsReference = false;
+	PinType.ContainerType = ContainerType;
+
+	return Node->CreatePin(Direction, PinType, PinName);
+}
+
+bool FWildcardNodeUtils::NodeHasAnyWildcards(const UEdGraphNode* const Node)
+{
+	check(Node);
+
+	for (const UEdGraphPin* const Pin : Node->Pins)
+	{
+		if (FWildcardNodeUtils::IsWildcardPin(Pin))
+		{
+			return true;
+		}
+	}
+
+	return false;
+}
\ No newline at end of file
diff --git a/Engine/Source/Editor/UnrealEd/Public/Kismet2/WildcardNodeUtils.h b/Engine/Source/Editor/UnrealEd/Public/Kismet2/WildcardNodeUtils.h
new file mode 100644
--- /dev/null
+++ b/Engine/Source/Editor/UnrealEd/Public/Kismet2/WildcardNodeUtils.h
@@ -0,0 +1,47 @@
+// Copyright Epic Games, Inc. All Rights Reserved.
+
+#pragma once
+
+#include "CoreMinimal.h"
+#include "EdGraph/EdGraphNode.h"
+
+/**
+* A collection of utilities for creating, using, and changing the value of wildcard pin types
+* throughout the graph editor.
+*/
+class UNREALED_API FWildcardNodeUtils
+{
+public:
+	
+	/** 
+	* Gets the default wildcard pin type. Useful for doing comparisons on other pin types 
+	* and checks during compilation of nodes.
+	*/
+	static FEdGraphPinType GetDefaultWildcardPinType();
+
+	/**
+	* Checks if the given pin is in a wildcard state
+	*
+	* @param	Pin		The pin the consider
+	* @return	True if the given pin is a Wildcard pin
+	*/
+	static bool IsWildcardPin(const UEdGraphPin* const Pin);
+
+	/**
+	* Add a default wildcard pin to the given node
+	* 
+	* @param Node				The node to add this pin to
+	* @param PinName			Name of the given wildcard pin
+	* @param Direction			
+	* @param ContainerType		
+	* @return	The newly created pin or nullptr if failed
+	*/
+	static UEdGraphPin* CreateWildcardPin(UEdGraphNode* Node, const FName PinName, const EEdGraphPinDirection Direction, const EPinContainerType ContainerType = EPinContainerType::None);
+
+	/**
+	* Check this node for any wildcard pins
+	*
+	* @return	True if the given node has any wildcard pins on it
+	*/
+	static bool NodeHasAnyWildcards(const UEdGraphNode* const Node);
+};
\ No newline at end of file
