Added WorldPartitionBuildHLODsCommandlet - first WIP edition
Corrected a few assumptions in UWorldPartition about what needs to be initialized when running from cmdlets

diff --git a/Engine/Source/Editor/UnrealEd/Private/Commandlets/WorldPartitionBuildHLODsCommandlet.cpp b/Engine/Source/Editor/UnrealEd/Private/Commandlets/WorldPartitionBuildHLODsCommandlet.cpp
new file mode 100644
--- /dev/null
+++ b/Engine/Source/Editor/UnrealEd/Private/Commandlets/WorldPartitionBuildHLODsCommandlet.cpp
@@ -0,0 +1,117 @@
+// Copyright Epic Games, Inc. All Rights Reserved.
+
+/*============================================================================================
+ WorldPartitionBuildHLODsCommandlet.cpp: Commandlet to build HLODs for a partionned level
+============================================================================================*/
+
+#include "Commandlets/WorldPartitionBuildHLODsCommandlet.h"
+
+#include "CoreMinimal.h"
+#include "EngineUtils.h"
+#include "Editor.h"
+#include "Logging/LogMacros.h"
+
+#include "WorldPartition/WorldPartitionEditorHash.h"
+#include "WorldPartition/WorldPartition.h"
+#include "WorldPartition/WorldPartitionSubsystem.h"
+#include "WorldPartition/HLOD/HLODActor.h"
+
+DEFINE_LOG_CATEGORY_STATIC(LogWorldPartitionBuildHLODsCommandlet, All, All);
+
+UWorldPartitionBuildHLODsCommandlet::UWorldPartitionBuildHLODsCommandlet(const FObjectInitializer& ObjectInitializer)
+	: Super(ObjectInitializer)
+{
+}
+
+int32 UWorldPartitionBuildHLODsCommandlet::Main(const FString& Params)
+{
+	TArray<FString> Tokens, Switches;
+	ParseCommandLine(*Params, Tokens, Switches);
+
+	// Retrieve map name from the commandline
+	FString MapName;
+	bool bMapProvided = FParse::Value(*Params, TEXT("Map="), MapName);
+	if (!bMapProvided)
+	{
+		UE_LOG(LogWorldPartitionBuildHLODsCommandlet, Error, TEXT("No map specified."));
+		return 1;
+	}
+
+	if (!IsAllowCommandletRendering())
+	{
+		UE_LOG(LogWorldPartitionBuildHLODsCommandlet, Error, TEXT("The option \"-AllowCommandletRendering\" must be provided for the HLOD rebuild process to work"));
+		return 1;
+	}
+
+	// Load the map package
+	UPackage* MapPackage = LoadPackage(NULL, *MapName, LOAD_None);
+	if (!MapPackage)
+	{
+		UE_LOG(LogWorldPartitionBuildHLODsCommandlet, Error, TEXT("Couldn't load package %s."), *MapName);
+		return 1;
+	}
+
+	// Find the world in the given package
+	UWorld* World = UWorld::FindWorldInPackage(MapPackage);
+	if (!World)
+	{
+		UE_LOG(LogWorldPartitionBuildHLODsCommandlet, Error, TEXT("No world in specified package."));
+		return 1;
+	}
+
+	//  Retrieve the world partition.
+	UWorldPartitionSubsystem* WorldPartitionSubsystem = World->GetSubsystem<UWorldPartitionSubsystem>();
+	UWorldPartition* WorldPartition = WorldPartitionSubsystem && WorldPartitionSubsystem->IsEnabled() ? World->GetWorldPartition() : nullptr;
+	if (!WorldPartition)
+	{
+		UE_LOG(LogWorldPartitionBuildHLODsCommandlet, Error, TEXT("Commandlet only works on partitioned maps."));
+		return 1;
+	}
+
+	// Setup the world.
+	World->WorldType = EWorldType::Editor;
+	World->AddToRoot();
+	if (!World->bIsWorldInitialized)
+	{
+		UWorld::InitializationValues IVS;
+		IVS.RequiresHitProxies(false);
+		IVS.ShouldSimulatePhysics(false);
+		IVS.EnableTraceCollision(false);
+		IVS.CreateNavigation(false);
+		IVS.CreateAISystem(false);
+		IVS.AllowAudioPlayback(false);
+		IVS.CreatePhysicsScene(true);
+
+		World->InitWorld(UWorld::InitializationValues(IVS));
+		World->PersistentLevel->UpdateModelComponents();
+		World->UpdateWorldComponents(true, false);
+	}
+
+	FWorldContext& WorldContext = GEditor->GetEditorWorldContext(true);
+	WorldContext.SetCurrentWorld(World);
+	GWorld = World;
+	
+	// For now, load all cells
+	// In the future, we'll want the commandlet to be able to perform partial updates of the map
+	// to allow HLOD rebuild to be distributed on multiple machines.
+	TArray<UWorldPartitionEditorCell*> AllCells;
+	AllCells.Add(WorldPartition->EditorHash->GetAlwaysLoadedCell());
+	WorldPartition->EditorHash->GetAllCells(AllCells);	
+	WorldPartition->LoadCells(AllCells);
+
+	// Clear all existing HLOD actors
+	for (TActorIterator<AWorldPartitionHLOD> ItHLOD(World); ItHLOD; ++ItHLOD)
+	{
+		World->DestroyActor(*ItHLOD);
+	}
+
+	// Rebuild HLOD for the whole world
+	WorldPartition->GenerateHLOD();
+
+	// Cleanup
+	World->RemoveFromRoot();
+	WorldContext.SetCurrentWorld(nullptr);
+	GWorld = nullptr;
+
+	return 0;
+}
diff --git a/Engine/Source/Editor/UnrealEd/Private/Commandlets/WorldPartitionBuildHLODsCommandlet.h b/Engine/Source/Editor/UnrealEd/Private/Commandlets/WorldPartitionBuildHLODsCommandlet.h
new file mode 100644
--- /dev/null
+++ b/Engine/Source/Editor/UnrealEd/Private/Commandlets/WorldPartitionBuildHLODsCommandlet.h
@@ -0,0 +1,16 @@
+// Copyright 1998-2019 Epic Games, Inc. All Rights Reserved.
+
+#pragma once
+
+#include "Commandlets/Commandlet.h"
+#include "WorldPartitionBuildHLODsCommandlet.generated.h"
+
+UCLASS()
+class UWorldPartitionBuildHLODsCommandlet : public UCommandlet
+{
+	GENERATED_UCLASS_BODY()
+
+	//~ Begin UCommandlet Interface
+	virtual int32 Main(const FString& Params) override;
+	//~ End UCommandlet Interface
+};
\ No newline at end of file
diff --git a/Engine/Source/Runtime/Engine/Private/WorldPartition/WorldPartition.cpp b/Engine/Source/Runtime/Engine/Private/WorldPartition/WorldPartition.cpp
--- a/Engine/Source/Runtime/Engine/Private/WorldPartition/WorldPartition.cpp
+++ b/Engine/Source/Runtime/Engine/Private/WorldPartition/WorldPartition.cpp
@@ -204,6 +204,11 @@ void UWorldPartition::Initialize(UWorld* InWorld, const FTransform& InTransform)
 		return;
 	}
 
+	if (IsTemplate())
+	{
+		return;
+	}
+
 	check(InWorld);
 	World = InWorld;
 	InstanceTransform = InTransform;
@@ -216,7 +221,7 @@ void UWorldPartition::Initialize(UWorld* InWorld, const FTransform& InTransform)
 
 #if WITH_EDITOR
 	bool bEditorOnly = !World->IsPlayInEditor();
-	if (bEditorOnly && !IsRunningCommandlet())
+	if (bEditorOnly)
 	{
 		if (!EditorHash)
 		{
@@ -227,7 +232,10 @@ void UWorldPartition::Initialize(UWorld* InWorld, const FTransform& InTransform)
 
 		EditorHash->Initialize();
 
-		LayerSubSystem = GEditor->GetEditorSubsystem<ULayersSubsystem>();
+		if (!IsRunningCommandlet())
+		{
+			LayerSubSystem = GEditor->GetEditorSubsystem<ULayersSubsystem>();
+		}
 
 		if (IsMainWorldPartition())
 		{
@@ -293,25 +301,28 @@ void UWorldPartition::Initialize(UWorld* InWorld, const FTransform& InTransform)
 			}
 		}
 
-		if (bEditorOnly && !IsRunningCommandlet())
+		if (bEditorOnly)
 		{
-			CreateLayers(AllLayersNames);
-
 			for (const auto& Pair : Actors)
 			{
 				HashActorDesc(Pair.Value.Get());
 			}
 
-			// Load the always loaded cell, don't call LoadCells to avoid creating a transaction
-			UpdateLoadingEditorCell(EditorHash->GetAlwaysLoadedCell(), true);
-
-			// When loading a subworld, load all actors
-			if (!IsMainWorldPartition())
+			if (!IsRunningCommandlet())
 			{
-				EditorHash->ForEachCell([this](UWorldPartitionEditorCell* Cell)
+				CreateLayers(AllLayersNames);
+
+				// Load the always loaded cell, don't call LoadCells to avoid creating a transaction
+				UpdateLoadingEditorCell(EditorHash->GetAlwaysLoadedCell(), true);
+
+				// When loading a subworld, load all actors
+				if (!IsMainWorldPartition())
 				{
-					UpdateLoadingEditorCell(Cell, true);
-				});
+					EditorHash->ForEachCell([this](UWorldPartitionEditorCell* Cell)
+					{
+						UpdateLoadingEditorCell(Cell, true);
+					});
+				}
 			}
 		}
 	}
@@ -410,7 +421,7 @@ bool UWorldPartition::IsMainWorldPartition() const
 #if WITH_EDITOR
 void UWorldPartition::RegisterDelegates()
 {
-	if (GEditor && !IsTemplate() && !IsRunningCommandlet())
+	if (GEditor && !IsTemplate())
 	{
 		GEditor->OnActorMoving().AddUObject(this, &UWorldPartition::OnActorMoving);
 		GEditor->OnActorMoved().AddUObject(this, &UWorldPartition::OnActorMoving);
@@ -447,7 +458,7 @@ void UWorldPartition::RegisterDelegates()
 
 void UWorldPartition::UnregisterDelegates()
 {
-	if (GEditor && !IsTemplate() && !IsRunningCommandlet())
+	if (GEditor && !IsTemplate())
 	{
 		if (LayerSubSystem)
 		{
@@ -778,7 +789,7 @@ void UWorldPartition::LoadCells(const TArray<UWorldPartitionEditorCell*>& CellsT
 
 	for (UWorldPartitionEditorCell* Cell : CellsToLoad)
 	{
-		SlowTask.EnterProgressFrame(Cell->Actors.Num() - Cell->LoadedActors.Num());
+		SlowTask.EnterProgressFrame(Cell->Actors.Num() - Cell->LoadedActors.Num(), FText::Format(LOCTEXT("LoadingActors", "Loading actors {0}/{1}"), (int32)SlowTask.CompletedWork, NumActorsToLoad));
 		UpdateLoadingEditorCell(Cell, true);
 	}
 }
@@ -796,6 +807,7 @@ void UWorldPartition::UnloadCells(const TArray<UWorldPartitionEditorCell*>& Cell
 
 	for (UWorldPartitionEditorCell* Cell: CellsToUnload)
 	{
+		SlowTask.EnterProgressFrame(Cell->LoadedActors.Num(), FText::Format(LOCTEXT("UnloadingActors", "Unloading actors {0}/{1}"), (int32)SlowTask.CompletedWork, NumActorsToUnload));
 		SlowTask.EnterProgressFrame(Cell->LoadedActors.Num());
 		UpdateLoadingEditorCell(Cell, false);
 	}
