- Allow for "-run=CookGlobalShaders -platform=<plat> -deploy -reload" commandlet to work to xbox family platforms

diff --git a/Engine/Source/Editor/UnrealEd/Private/Commandlets/CookGlobalShaderCommandlet.cpp b/Engine/Source/Editor/UnrealEd/Private/Commandlets/CookGlobalShaderCommandlet.cpp
--- a/Engine/Source/Editor/UnrealEd/Private/Commandlets/CookGlobalShaderCommandlet.cpp
+++ b/Engine/Source/Editor/UnrealEd/Private/Commandlets/CookGlobalShaderCommandlet.cpp
@@ -162,30 +162,42 @@ int32 UCookGlobalShadersCommandlet::Main(const FString& Params)
 		const TArray<FString> MaterialsToLoad;
 		RecompileShadersForRemote(PlatformName, SP_NumPlatforms, OutputDir, MaterialsToLoad, nullptr, nullptr);
 
+		bool bCopySucceeded = false;
+
+		// Build list of files to copy
+		TArray<TPair<FString, FString>> FilesToCopy;
+		for (FName ShaderFormat : ShaderFormats)
+		{
+			const FString GlobalShaderCacheName = FPaths::Combine(OutputDir, TEXT("Engine"), TEXT("GlobalShaderCache-") + ShaderFormat.ToString() + TEXT(".bin"));
+			const FString OverrideGlobalShaderCacheName = FPaths::Combine(TEXT("Engine"), TEXT("OverrideGlobalShaderCache-") + ShaderFormat.ToString() + TEXT(".bin"));
+			FilesToCopy.Emplace(GlobalShaderCacheName, OverrideGlobalShaderCacheName);
+		}
+
 		// Are we copying the built files somewhere?
 		if (DeviceHelper != nullptr)
 		{
-			// Build list of files to copy
-			TArray<TPair<FString, FString>> FilesToCopy;
-			for (FName ShaderFormat : ShaderFormats)
-			{
-				const FString GlobalShaderCacheName = FPaths::Combine(OutputDir, TEXT("Engine"),  TEXT("GlobalShaderCache-") + ShaderFormat.ToString() + TEXT(".bin"));
-				const FString OverrideGlobalShaderCacheName = FPaths::Combine(TEXT("Engine"), TEXT("OverrideGlobalShaderCache-") + ShaderFormat.ToString() + TEXT(".bin"));
-				FilesToCopy.Emplace(GlobalShaderCacheName, OverrideGlobalShaderCacheName);
-			}
-
 			// Execute Copy
 			UE_LOG(LogCookGlobalShaders, Log, TEXT("Copying Cooked Files..."));
-			if (DeviceHelper->CopyFilesToDevice(TargetDevice.Get(), FilesToCopy) == true )
+			bCopySucceeded =  DeviceHelper->CopyFilesToDevice(TargetDevice.Get(), FilesToCopy);
+		}
+		// if no helper, but we want to deploy, use the TargetPlatform
+		else if (bDeployToDevice && TargetDevice != nullptr)
+		{
+			bCopySucceeded = true;
+
+			for (auto It : FilesToCopy)
 			{
-				// Execute Reload
-				if (bExecuteReload && TargetDevice.IsValid() )
-				{
-					UE_LOG(LogCookGlobalShaders, Log, TEXT("Sending Reload Command..."));
-					TargetDevice->ExecuteConsoleCommand("ReloadGlobalShaders");
-				}
+				bCopySucceeded = bCopySucceeded && TargetPlatform->CopyFileToTarget(TargetDevice->GetId().GetDeviceName(), It.Key, It.Value, {});
 			}
 		}
+
+
+		// Execute Reload
+		if (bCopySucceeded && bExecuteReload && TargetDevice.IsValid())
+		{
+			UE_LOG(LogCookGlobalShaders, Log, TEXT("Sending Reload Command..."));
+			TargetDevice->ExecuteConsoleCommand("ReloadGlobalShaders");
+		}
 	}
 	UE_LOG(LogCookGlobalShaders, Log, TEXT("Complete"));
 
