VR editor mode functionality moved to the level editor viewport client from editor viewport client.

#rnx

diff --git a/Engine/Source/Editor/UnrealEd/Private/EditorViewportClient.cpp b/Engine/Source/Editor/UnrealEd/Private/EditorViewportClient.cpp
--- a/Engine/Source/Editor/UnrealEd/Private/EditorViewportClient.cpp
+++ b/Engine/Source/Editor/UnrealEd/Private/EditorViewportClient.cpp
@@ -49,8 +49,6 @@
 #include "ComponentRecreateRenderStateContext.h"
 #include "EditorBuildUtils.h"
 #include "AudioDevice.h"
-#include "EditorWorldExtension.h"
-#include "ViewportWorldInteraction.h"
 #include "Editor/EditorPerformanceSettings.h"
 #include "ImageWriteQueue.h"
 #include "DebugViewModeHelpers.h"
@@ -1547,24 +1545,6 @@ EMouseCursor::Type FEditorViewportClient::GetCursor(FViewport* InViewport,int32
 		}
 	}
 
-	// Allow the viewport interaction to override any previously set mouse cursor
-	UViewportWorldInteraction* WorldInteraction = Cast<UViewportWorldInteraction>(GEditor->GetEditorWorldExtensionsManager()->GetEditorWorldExtensions(GetWorld())->FindExtension(UViewportWorldInteraction::StaticClass()));
-	if (WorldInteraction != nullptr)
-	{
-		if (WorldInteraction->ShouldForceCursor())
-		{
-			MouseCursor = EMouseCursor::Crosshairs;
-			SetRequiredCursor(false, true);
-			UpdateRequiredCursorVisibility();
-		}
-		else if (WorldInteraction->ShouldSuppressExistingCursor())
-		{
-			MouseCursor = EMouseCursor::None;
-			SetRequiredCursor(false, false);
-			UpdateRequiredCursorVisibility();
-		}
-	}
-
 	CachedMouseX = X;
 	CachedMouseY = Y;
 
@@ -2683,12 +2663,6 @@ bool FEditorViewportClient::InputKey(FViewport* InViewport, int32 ControllerId,
 		return true;
 	}
 
-	UEditorWorldExtensionCollection& EditorWorldExtensionCollection = *GEditor->GetEditorWorldExtensionsManager()->GetEditorWorldExtensions( GetWorld() );
-	if( EditorWorldExtensionCollection.InputKey(this, Viewport, Key, Event))
-	{
-		return true;
-	}
-
 	FInputEventState InputState(InViewport, Key, Event);
 
 	bool bHandled = false;
diff --git a/Engine/Source/Editor/UnrealEd/Private/LevelEditorViewport.cpp b/Engine/Source/Editor/UnrealEd/Private/LevelEditorViewport.cpp
--- a/Engine/Source/Editor/UnrealEd/Private/LevelEditorViewport.cpp
+++ b/Engine/Source/Editor/UnrealEd/Private/LevelEditorViewport.cpp
@@ -76,6 +76,8 @@
 #include "ActorGroupingUtils.h"
 #include "EditorWorldExtension.h"
 #include "VREditorMode.h"
+#include "EditorWorldExtension.h"
+#include "ViewportWorldInteraction.h"
 #include "Subsystems/BrushEditingSubsystem.h"
 #include "Engine/VolumeTexture.h"
 #include "Materials/MaterialExpressionDivide.h"
@@ -2860,6 +2862,12 @@ bool FLevelEditorViewportClient::InputKey(FViewport* InViewport, int32 Controlle
 	bUserIsControllingAtmosphericLight0 = false;	// Disable all atmospheric light controls
 	bUserIsControllingAtmosphericLight1 = false;
 
+	UEditorWorldExtensionCollection& EditorWorldExtensionCollection = *GEditor->GetEditorWorldExtensionsManager()->GetEditorWorldExtensions(GetWorld());
+	if (EditorWorldExtensionCollection.InputKey(this, Viewport, Key, Event))
+	{
+		return true;
+	}
+
 	bool bHandled = FEditorViewportClient::InputKey(InViewport,ControllerId,Key,Event,AmountDepressed,bGamepad);
 
 	// Handle input for the player height preview mode. 
@@ -4397,6 +4405,24 @@ EMouseCursor::Type FLevelEditorViewportClient::GetCursor(FViewport* InViewport,i
 {
 	EMouseCursor::Type CursorType = FEditorViewportClient::GetCursor(InViewport,X,Y);
 
+	// Allow the viewport interaction to override any previously set mouse cursor
+	UViewportWorldInteraction* WorldInteraction = Cast<UViewportWorldInteraction>(GEditor->GetEditorWorldExtensionsManager()->GetEditorWorldExtensions(GetWorld())->FindExtension(UViewportWorldInteraction::StaticClass()));
+	if (WorldInteraction != nullptr)
+	{
+		if (WorldInteraction->ShouldForceCursor())
+		{
+			CursorType = EMouseCursor::Crosshairs;
+			SetRequiredCursor(false, true);
+			UpdateRequiredCursorVisibility();
+		}
+		else if (WorldInteraction->ShouldSuppressExistingCursor())
+		{
+			CursorType = EMouseCursor::None;
+			SetRequiredCursor(false, false);
+			UpdateRequiredCursorVisibility();
+		}
+	}
+
 	// Don't select widget axes by mouse over while they're being controlled by a mouse drag.
 	if( InViewport->IsCursorVisible() && !bWidgetAxisControlledByDrag)
 	{
