Remove fallback for actor and component drag movement in level editor viewport.

#rnx

diff --git a/Engine/Source/Editor/UnrealEd/Classes/Settings/LevelEditorViewportSettings.h b/Engine/Source/Editor/UnrealEd/Classes/Settings/LevelEditorViewportSettings.h
--- a/Engine/Source/Editor/UnrealEd/Classes/Settings/LevelEditorViewportSettings.h
+++ b/Engine/Source/Editor/UnrealEd/Classes/Settings/LevelEditorViewportSettings.h
@@ -426,11 +426,6 @@ public:
 	UPROPERTY(config)
 	TEnumAsByte<ERotationGridMode> CurrentRotGridMode;
 
-	/** Toggles legacy behavior for updating components and actors during drag operations. This could be useful if you're seeing a degradation in performance due to too many PostEditMove calls */
-	UE_DEPRECATED(4.26, "This property is meant to be a temporary toggle for a rollback if too many post edit move calls degrade performance during drag operations.")
-	UPROPERTY(EditAnywhere, config, Category=Controls, AdvancedDisplay, meta = (DisplayName = "Use Legacy Behavior for actor and component updates while dragging"))
-	bool bUseLegacyPostEditBehavior = false;
-
 public:
 
 	/** How to constrain perspective view port FOV */
diff --git a/Engine/Source/Editor/UnrealEd/Private/EditorEngine.cpp b/Engine/Source/Editor/UnrealEd/Private/EditorEngine.cpp
--- a/Engine/Source/Editor/UnrealEd/Private/EditorEngine.cpp
+++ b/Engine/Source/Editor/UnrealEd/Private/EditorEngine.cpp
@@ -2874,9 +2874,6 @@ void UEditorEngine::ApplyDeltaToComponent(USceneComponent* InComponent,
 	// Update the actor before leaving.
 	InComponent->MarkPackageDirty();
 
-	PRAGMA_DISABLE_DEPRECATION_WARNINGS
-	if (!GetDefault<ULevelEditorViewportSettings>()->bUseLegacyPostEditBehavior)
-	PRAGMA_ENABLE_DEPRECATION_WARNINGS
 	{
 		InComponent->PostEditComponentMove(false);
 	}
diff --git a/Engine/Source/Editor/UnrealEd/Private/LevelEditorViewport.cpp b/Engine/Source/Editor/UnrealEd/Private/LevelEditorViewport.cpp
--- a/Engine/Source/Editor/UnrealEd/Private/LevelEditorViewport.cpp
+++ b/Engine/Source/Editor/UnrealEd/Private/LevelEditorViewport.cpp
@@ -3133,9 +3133,6 @@ void FLevelEditorViewportClient::TrackingStopped()
 		FPropertyChangedEvent PropertyChangedEvent(TransformProperty, EPropertyChangeType::ValueSet);
 
 		// Move components and actors
-		PRAGMA_DISABLE_DEPRECATION_WARNINGS
-		if (!LegacyTrackingStoppedForSelectedComponentsAndActors(PropertyChangedEvent))
-		PRAGMA_ENABLE_DEPRECATION_WARNINGS
 		{
 			TArray<USceneComponent*> ComponentsToMove;
 			TArray<AActor*> ActorsToMove;
@@ -3344,217 +3341,6 @@ bool FLevelEditorViewportClient::CanMoveActorInViewport(const AActor* InActor) c
 	return true;
 }
 
-bool FLevelEditorViewportClient::LegacyTrackingStoppedForSelectedComponentsAndActors(FPropertyChangedEvent& PropertyChangedEvent)
-{
-	PRAGMA_DISABLE_DEPRECATION_WARNINGS
-	if (!GetDefault<ULevelEditorViewportSettings>()->bUseLegacyPostEditBehavior)
-	{
-		return false;
-	}
-	PRAGMA_ENABLE_DEPRECATION_WARNINGS
-
-	TArray<AActor*> MovedActors;
-
-	for (FSelectionIterator It(GEditor->GetSelectedActorIterator()); It; ++It)
-	{
-		AActor* Actor = static_cast<AActor*>(*It);
-		checkSlow(Actor->IsA(AActor::StaticClass()));
-
-		// Verify that the actor is in the same world as the viewport before moving it.
-		if (GEditor->PlayWorld)
-		{
-			if (bIsSimulateInEditorViewport)
-			{
-				// If the Actor's outer (level) outer (world) is not the PlayWorld then it cannot be moved in this viewport.
-				if (!(GEditor->PlayWorld == Actor->GetOuter()->GetOuter()))
-				{
-					continue;
-				}
-			}
-			else if (!(GEditor->EditorWorld == Actor->GetOuter()->GetOuter()))
-			{
-				continue;
-			}
-		}
-
-		bool bComponentsMoved = false;
-		if (GEditor->GetSelectedComponentCount() > 0)
-		{
-			USelection* ComponentSelection = GEditor->GetSelectedComponents();
-
-			// Only move the parent-most component(s) that are selected 
-			// Otherwise, if both a parent and child are selected and the delta is applied to both, the child will actually move 2x delta
-			TInlineComponentArray<USceneComponent*> ComponentsToMove;
-			for (FSelectedEditableComponentIterator EditableComponentIt(GEditor->GetSelectedEditableComponentIterator()); EditableComponentIt; ++EditableComponentIt)
-			{
-				USceneComponent* SceneComponent = Cast<USceneComponent>(*EditableComponentIt);
-				if (SceneComponent)
-				{
-					// Check to see if any parent is selected
-					bool bParentAlsoSelected = false;
-					USceneComponent* Parent = SceneComponent->GetAttachParent();
-					while (Parent != nullptr)
-					{
-						if (ComponentSelection->IsSelected(Parent))
-						{
-							bParentAlsoSelected = true;
-							break;
-						}
-
-						Parent = Parent->GetAttachParent();
-					}
-
-					// If no parent of this component is also in the selection set, move it!
-					if (!bParentAlsoSelected)
-					{
-						ComponentsToMove.Add(SceneComponent);
-					}
-				}
-			}
-
-			// Now actually apply the delta to the appropriate component(s)
-			for (USceneComponent* SceneComp : ComponentsToMove)
-			{
-				// Broadcast Post Edit change notification, we can't call PostEditChangeProperty directly on Actor or ActorComponent from here since it wasn't pair with a proper PreEditChange
-				FCoreUObjectDelegates::OnObjectPropertyChanged.Broadcast(SceneComp, PropertyChangedEvent);
-				SceneComp->PostEditComponentMove(true);
-
-				GEditor->BroadcastEndObjectMovement(*SceneComp);
-
-				bComponentsMoved = true;
-			}
-		}
-
-		if (!bComponentsMoved)
-		{
-			// Broadcast Post Edit change notification, we can't call PostEditChangeProperty directly on Actor or ActorComponent from here since it wasn't pair with a proper PreEditChange
-			FCoreUObjectDelegates::OnObjectPropertyChanged.Broadcast(Actor, PropertyChangedEvent);
-			Actor->PostEditMove(true);
-
-			MovedActors.Add(Actor);
-
-			GEditor->BroadcastEndObjectMovement(*Actor);
-		}
-	}
-
-	GEditor->BroadcastActorsMoved(MovedActors);
-
-	return true;
-}
-
-bool FLevelEditorViewportClient::LegacyApplyDeltasForSelectedComponentsAndActors(const FVector& InDrag, const FRotator& InRot, const FVector& ModifiedScale)
-{
-	PRAGMA_DISABLE_DEPRECATION_WARNINGS
-	if (!GetDefault<ULevelEditorViewportSettings>()->bUseLegacyPostEditBehavior)
-	{
-		return false;
-	}
-	PRAGMA_ENABLE_DEPRECATION_WARNINGS
-
-	TArray<AGroupActor*> ActorGroups;
-
-	// Apply the deltas to any selected actors.
-	for ( FSelectionIterator SelectedActorIt( GEditor->GetSelectedActorIterator() ) ; SelectedActorIt ; ++SelectedActorIt )
-	{
-		AActor* Actor = static_cast<AActor*>( *SelectedActorIt );
-		checkSlow( Actor->IsA(AActor::StaticClass()) );
-
-		// Verify that the actor is in the same world as the viewport before moving it.
-		if(GEditor->PlayWorld)
-		{
-			if(bIsSimulateInEditorViewport)
-			{
-				// If the Actor's outer (level) outer (world) is not the PlayWorld then it cannot be moved in this viewport.
-				if( !(GEditor->PlayWorld == Actor->GetWorld()) )
-				{
-					continue;
-				}
-			}
-			else if( !(GEditor->EditorWorld == Actor->GetWorld()) )
-			{
-				continue;
-			}
-		}
-
-		if ( !Actor->bLockLocation )
-		{
-			if (GEditor->GetSelectedComponentCount() > 0)
-			{
-				USelection* ComponentSelection = GEditor->GetSelectedComponents();
-
-				// Only move the parent-most component(s) that are selected 
-				// Otherwise, if both a parent and child are selected and the delta is applied to both, the child will actually move 2x delta
-				TInlineComponentArray<USceneComponent*> ComponentsToMove;
-				for (FSelectedEditableComponentIterator EditableComponentIt(GEditor->GetSelectedEditableComponentIterator()); EditableComponentIt; ++EditableComponentIt)
-				{
-					USceneComponent* SelectedComponent = Cast<USceneComponent>(*EditableComponentIt);
-					if (SelectedComponent)
-					{
-						// Check to see if any parent is selected
-						bool bParentAlsoSelected = false;
-						USceneComponent* Parent = SelectedComponent->GetAttachParent();
-						while (Parent != nullptr)
-						{
-							if (ComponentSelection->IsSelected(Parent))
-							{
-								bParentAlsoSelected = true;
-								break;
-							}
-
-							Parent = Parent->GetAttachParent();
-						}
-
-						// If no parent of this component is also in the selection set, move it!
-						if (!bParentAlsoSelected)
-						{
-							ComponentsToMove.Add(SelectedComponent);
-						}
-					}	
-				}
-
-				// Now actually apply the delta to the appropriate component(s)
-				for (USceneComponent* SceneComp : ComponentsToMove)
-				{
-					ApplyDeltaToComponent(SceneComp, InDrag, InRot, ModifiedScale);
-				}
-			}
-			else
-			{
-				AGroupActor* ParentGroup = AGroupActor::GetRootForActor(Actor, true, true);
-				if (ParentGroup && UActorGroupingUtils::IsGroupingActive())
-				{
-					ActorGroups.AddUnique(ParentGroup);
-				}
-				else
-				{
-					// Finally, verify that no actor in the parent hierarchy is also selected
-					bool bHasParentInSelection = false;
-					AActor* ParentActor = Actor->GetAttachParentActor();
-					while (ParentActor != NULL && !bHasParentInSelection)
-					{
-						if (ParentActor->IsSelected())
-						{
-							bHasParentInSelection = true;
-						}
-						ParentActor = ParentActor->GetAttachParentActor();
-					}
-					if (!bHasParentInSelection)
-					{
-						ApplyDeltaToActor(Actor, InDrag, InRot, ModifiedScale);
-					}
-				}
-			}
-		}
-	}
-	AGroupActor::RemoveSubGroupsFromArray(ActorGroups);
-	for(int32 ActorGroupsIndex=0; ActorGroupsIndex<ActorGroups.Num(); ++ActorGroupsIndex)
-	{
-		ActorGroups[ActorGroupsIndex]->GroupApplyDelta(this, InDrag, InRot, ModifiedScale);
-	}
-
-	return true;
-}
-
 void FLevelEditorViewportClient::NudgeSelectedObjects( const struct FInputEventState& InputState )
 {
 	FViewport* InViewport = InputState.GetViewport();
@@ -4057,9 +3843,6 @@ void FLevelEditorViewportClient::ApplyDeltaToActors(const FVector& InDrag,
 	GEditor->NoteActorMovement();
 		
 	// Move Components and Actors
-	PRAGMA_DISABLE_DEPRECATION_WARNINGS
-	if (!LegacyApplyDeltasForSelectedComponentsAndActors(InDrag, InRot, ModifiedScale))
-	PRAGMA_ENABLE_DEPRECATION_WARNINGS
 	{
 		TArray<USceneComponent*> ComponentsToMove;
 		TArray<AActor*> ActorsToMove;
diff --git a/Engine/Source/Runtime/Engine/Private/Components/SceneComponent.cpp b/Engine/Source/Runtime/Engine/Private/Components/SceneComponent.cpp
--- a/Engine/Source/Runtime/Engine/Private/Components/SceneComponent.cpp
+++ b/Engine/Source/Runtime/Engine/Private/Components/SceneComponent.cpp
@@ -35,10 +35,6 @@
 #include "DeviceProfiles/DeviceProfile.h"
 #include "Net/Core/PushModel/PushModel.h"
 
-#if WITH_EDITOR
-#include "Settings/LevelEditorViewportSettings.h"	// For legacy post edit move behavior
-#endif // WITH_EDITOR
-
 #define LOCTEXT_NAMESPACE "SceneComponent"
 
 namespace SceneComponentStatics
@@ -3344,9 +3340,6 @@ void USceneComponent::PostEditComponentMove(bool bFinished)
 		SnapshotTransactionBuffer(this);
 	}
 
-	PRAGMA_DISABLE_DEPRECATION_WARNINGS
-	if (!GetDefault<ULevelEditorViewportSettings>()->bUseLegacyPostEditBehavior)
-	PRAGMA_ENABLE_DEPRECATION_WARNINGS
 	{
 		// Call on all attached children
 		for (USceneComponent* ChildComponent : GetAttachChildren())
