Updated USelection to be optionally driven by an element list

USelection can now operate on two different backends; FObjectSelectionStore (which is the old, object-array implementation), or FElementSelectionStore (which is the newer element list based implementation).

Breaking changes:
 - USelection::MarkBatchDirty is now USelection::ForceBatchDirty.
 - USelection::Initialize is now private. Use one of USelection::CreateObjectSelection, USelection::CreateActorSelection, or USelection::CreateComponentSelection to create your USelection instance instead.
 - USelection::SelectionChangedEvent and USelection::SelectObjectEvent are now events, and should not have their Broadcast called externally (use USelection::NoteSelectionChanged or USelection::NoteUnknownSelectionChanged if you must manually notify).

Note: The element list implementation is currently disabled (via UE_USE_ELEMENT_LIST_SELECTION) until actor and component element support is submitted.

diff --git a/Engine/Source/Editor/UnrealEd/Private/AssetEditorModeManager.cpp b/Engine/Source/Editor/UnrealEd/Private/AssetEditorModeManager.cpp
--- a/Engine/Source/Editor/UnrealEd/Private/AssetEditorModeManager.cpp
+++ b/Engine/Source/Editor/UnrealEd/Private/AssetEditorModeManager.cpp
@@ -3,39 +3,43 @@
 #include "AssetEditorModeManager.h"
 #include "Engine/Selection.h"
 #include "PreviewScene.h"
+#include "TypedElementRegistry.h"
 
 //////////////////////////////////////////////////////////////////////////
 // FAssetEditorModeManager
 
 FAssetEditorModeManager::FAssetEditorModeManager()
 	: PreviewScene(nullptr)
+	, SelectedElements(UTypedElementRegistry::GetInstance()->CreateElementList())
 {
-	ActorSet = NewObject<USelection>();
-	ActorSet->SetFlags(RF_Transactional);
+	ActorSet = USelection::CreateActorSelection(nullptr, GetTransientPackage(), NAME_None, RF_Transactional);
+	ActorSet->SetElementList(SelectedElements.Get());
 	ActorSet->AddToRoot();
-	ActorSet->Initialize(nullptr);
 
-	ObjectSet = NewObject<USelection>();
-	ObjectSet->SetFlags(RF_Transactional);
+	ObjectSet = USelection::CreateObjectSelection(nullptr, GetTransientPackage(), NAME_None, RF_Transactional);
 	ObjectSet->AddToRoot();
-	ObjectSet->Initialize(nullptr);
 
-	ComponentSet = NewObject<USelection>();
-	ComponentSet->SetFlags(RF_Transactional);
+	ComponentSet = USelection::CreateComponentSelection(nullptr, GetTransientPackage(), NAME_None, RF_Transactional);
+	ComponentSet->SetElementList(SelectedElements.Get());
 	ComponentSet->AddToRoot();
-	ComponentSet->Initialize(nullptr);
 }
 
 FAssetEditorModeManager::~FAssetEditorModeManager()
 {
 	SetPreviewScene(nullptr);
 
+	ActorSet->SetElementList(nullptr);
 	ActorSet->RemoveFromRoot();
 	ActorSet = nullptr;
+
 	ObjectSet->RemoveFromRoot();
 	ObjectSet = nullptr;
+
+	ComponentSet->SetElementList(nullptr);
 	ComponentSet->RemoveFromRoot();
 	ComponentSet = nullptr;
+
+	SelectedElements.Reset();
 }
 
 USelection* FAssetEditorModeManager::GetSelectedActors() const
diff --git a/Engine/Source/Editor/UnrealEd/Private/EditorActor.cpp b/Engine/Source/Editor/UnrealEd/Private/EditorActor.cpp
--- a/Engine/Source/Editor/UnrealEd/Private/EditorActor.cpp
+++ b/Engine/Source/Editor/UnrealEd/Private/EditorActor.cpp
@@ -702,7 +702,7 @@ bool UUnrealEdEngine::edactDeleteSelected( UWorld* InWorld, bool bVerifyDeletion
 					GEditor->SelectComponent(ComponentToSelect, true, false);
 
 					// Make sure the selection changed event fires so the SCS trees can update their selection
-					ComponentSelection->MarkBatchDirty();
+					ComponentSelection->ForceBatchDirty();
 					ComponentSelection->EndBatchSelectOperation(true);
 
 					// Notify the level editor of the new component selection
diff --git a/Engine/Source/Editor/UnrealEd/Private/EditorEngine.cpp b/Engine/Source/Editor/UnrealEd/Private/EditorEngine.cpp
--- a/Engine/Source/Editor/UnrealEd/Private/EditorEngine.cpp
+++ b/Engine/Source/Editor/UnrealEd/Private/EditorEngine.cpp
@@ -284,19 +284,15 @@ static void OnObjectSelected(UObject* Object)
 
 static void PrivateInitSelectedSets()
 {
-	PrivateGetSelectedActors() = NewObject<USelection>(GetTransientPackage(), TEXT("SelectedActors"), RF_Transactional);
+	PrivateGetSelectedActors() = USelection::CreateActorSelection(&GSelectedActorAnnotation, GetTransientPackage(), TEXT("SelectedActors"), RF_Transactional);
 	PrivateGetSelectedActors()->AddToRoot();
-	PrivateGetSelectedActors()->Initialize(&GSelectedActorAnnotation);
-
 	PrivateGetSelectedActors()->SelectObjectEvent.AddStatic(&OnObjectSelected);
 
-	PrivateGetSelectedComponents() = NewObject<USelection>(GetTransientPackage(), TEXT("SelectedComponents"), RF_Transactional);
+	PrivateGetSelectedComponents() = USelection::CreateComponentSelection(&GSelectedComponentAnnotation, GetTransientPackage(), TEXT("SelectedComponents"), RF_Transactional);
 	PrivateGetSelectedComponents()->AddToRoot();
-	PrivateGetSelectedComponents()->Initialize(&GSelectedComponentAnnotation);
 
-	PrivateGetSelectedObjects() = NewObject<USelection>(GetTransientPackage(), TEXT("SelectedObjects"), RF_Transactional);
+	PrivateGetSelectedObjects() = USelection::CreateObjectSelection(&GSelectedObjectAnnotation, GetTransientPackage(), TEXT("SelectedObjects"), RF_Transactional);
 	PrivateGetSelectedObjects()->AddToRoot();
-	PrivateGetSelectedObjects()->Initialize(&GSelectedObjectAnnotation);
 }
 
 static void PrivateDestroySelectedSets()
@@ -304,8 +300,10 @@ static void PrivateDestroySelectedSets()
 #if 0
 	PrivateGetSelectedActors()->RemoveFromRoot();
 	PrivateGetSelectedActors() = NULL;
+
 	PrivateGetSelectedComponents()->RemoveFromRoot();
 	PrivateGetSelectedComponents() = NULL;
+
 	PrivateGetSelectedObjects()->RemoveFromRoot();
 	PrivateGetSelectedObjects() = NULL;
 #endif
diff --git a/Engine/Source/Editor/UnrealEd/Private/EditorSelectUtils.cpp b/Engine/Source/Editor/UnrealEd/Private/EditorSelectUtils.cpp
--- a/Engine/Source/Editor/UnrealEd/Private/EditorSelectUtils.cpp
+++ b/Engine/Source/Editor/UnrealEd/Private/EditorSelectUtils.cpp
@@ -372,7 +372,7 @@ void UUnrealEdEngine::NoteSelectionChange(bool bNotify)
 	if (bNotify)
 	{
 		USelection* Selection = bComponentSelectionChanged ? GetSelectedComponents() : GetSelectedActors();
-		USelection::SelectionChangedEvent.Broadcast(Selection);
+		Selection->NoteSelectionChanged();
 	}
 	
 	if (!bComponentSelectionChanged)
diff --git a/Engine/Source/Editor/UnrealEd/Private/EditorServer.cpp b/Engine/Source/Editor/UnrealEd/Private/EditorServer.cpp
--- a/Engine/Source/Editor/UnrealEd/Private/EditorServer.cpp
+++ b/Engine/Source/Editor/UnrealEd/Private/EditorServer.cpp
@@ -1396,7 +1396,7 @@ void UEditorEngine::PostUndo(bool)
 		OldSelectedComponents.Empty();
 
 		// We want to broadcast the component SelectionChangedEvent even if the selection didn't actually change
-		ComponentSelection->MarkBatchDirty();
+		ComponentSelection->ForceBatchDirty();
 		ComponentSelection->EndBatchSelectOperation();
 	}
 
@@ -4265,7 +4265,7 @@ void UEditorEngine::ExecuteCommandForAllLevelModels( UWorld* InWorld, FSelectCom
 		UModel* Model = (*Iterator)->Model;
 		InSelectCommand.ExecuteIfBound( Model );
 	}
-	USelection::SelectionChangedEvent.Broadcast(NULL);
+	USelection::NoteUnknownSelectionChanged();
 }
 
 void UEditorEngine::ExecuteCommandForAllLevelModels( UWorld* InWorld, FSelectInWorldCommand InSelectCommand, const FText& TransDesription )
@@ -4276,7 +4276,7 @@ void UEditorEngine::ExecuteCommandForAllLevelModels( UWorld* InWorld, FSelectInW
 		UModel* Model = (*Iterator)->Model;
 		InSelectCommand.ExecuteIfBound( InWorld, Model );
 	}
-	USelection::SelectionChangedEvent.Broadcast(NULL);
+	USelection::NoteUnknownSelectionChanged();
 }
 
 void UEditorEngine::FlagModifyAllSelectedSurfacesInLevels( UWorld* InWorld )
@@ -4336,7 +4336,7 @@ bool UEditorEngine::Exec_Poly( UWorld* InWorld, const TCHAR* Str, FOutputDevice&
 			else if (FParse::Command(&Str,TEXT("TEXTURE")))
 			{
 				polySelectMatchingMaterial( InWorld, false );
-				USelection::SelectionChangedEvent.Broadcast(NULL);
+				USelection::NoteUnknownSelectionChanged();
 			}
 			else if (FParse::Command(&Str,TEXT("RESOLUTION")))
 			{
@@ -4348,7 +4348,7 @@ bool UEditorEngine::Exec_Poly( UWorld* InWorld, const TCHAR* Str, FOutputDevice&
 				{
 					polySelectMatchingResolution(InWorld, false);
 				}
-				USelection::SelectionChangedEvent.Broadcast(NULL);
+				USelection::NoteUnknownSelectionChanged();
 			}
 			
 			return true;
diff --git a/Engine/Source/Editor/UnrealEd/Private/Selection.cpp b/Engine/Source/Editor/UnrealEd/Private/Selection.cpp
--- a/Engine/Source/Editor/UnrealEd/Private/Selection.cpp
+++ b/Engine/Source/Editor/UnrealEd/Private/Selection.cpp
@@ -3,15 +3,525 @@
 #include "Selection.h"
 #include "UObject/Package.h"
 #include "UObject/UObjectAnnotation.h"
+#include "GameFramework/Actor.h"
+
+#define UE_USE_ELEMENT_LIST_SELECTION (0)
+
+#if UE_USE_ELEMENT_LIST_SELECTION
+#include "TypedElementList.h"
+#include "TypedElementSelectionInterface.h"
+#endif	// UE_USE_ELEMENT_LIST_SELECTION
 
 DEFINE_LOG_CATEGORY_STATIC(LogSelection, Log, All);
 
+namespace Selection_Private
+{
+
+class FSelectionStoreBase : public ISelectionStore
+{
+public:
+	explicit FSelectionStoreBase(const UClass* InObjectBaseClass);
+
+	//~ ISelectionStore interface
+	virtual void SetSink(ISelectionStoreSink* InSink) override;
+	virtual bool IsValidObjectToSelect(const UObject* InObject) const override;
+
+protected:
+	/** Base type of objects that this selection instance manages. */
+	const UClass* ObjectBaseClass = nullptr;
+
+	/** Sink instance to notify of changes. */
+	ISelectionStoreSink* Sink = nullptr;
+};
+
+FSelectionStoreBase::FSelectionStoreBase(const UClass* InObjectBaseClass)
+	: ObjectBaseClass(InObjectBaseClass)
+{
+}
+
+void FSelectionStoreBase::SetSink(ISelectionStoreSink* InSink)
+{
+	Sink = InSink;
+}
+
+bool FSelectionStoreBase::IsValidObjectToSelect(const UObject* InObject) const
+{
+	return InObject && (!ObjectBaseClass || InObject->IsA(ObjectBaseClass));
+}
+
+
+class FObjectSelectionStore : public FSelectionStoreBase
+{
+public:
+	explicit FObjectSelectionStore(const UClass* InObjectBaseClass);
+
+	//~ ISelectionStore interface
+	virtual void SetElementList(FTypedElementList* InElementList) override;
+	virtual int32 GetNumObjects() const override;
+	virtual UObject* GetObjectAtIndex(const int32 InIndex) const override;
+	virtual bool IsObjectSelected(const UObject* InObject) const override;
+	virtual void SelectObject(UObject* InObject) override;
+	virtual void DeselectObject(UObject* InObject) override;
+	virtual int32 DeselectObjects(TFunctionRef<bool(UObject*)> InPredicate) override;
+	virtual void BeginBatchSelection() override;
+	virtual void EndBatchSelection(const bool InNotify) override;
+	virtual bool IsBatchSelecting() const override;
+	virtual void ForceBatchDirty() override;
+	virtual void NoteSelectionChanged() override;
+
+private:
+	/** Tracks the number of active selection operations.  Allows batched selection operations to only send one notification at the end of the batch */
+	int32 SelectionMutex = 0;
+
+	/** Tracks whether the selection set changed during a batch selection operation */
+	bool bIsBatchDirty = false;
+
+	/** List of selected objects, ordered as they were selected. */
+	TArray<TWeakObjectPtr<UObject>> SelectedObjects;
+};
+
+FObjectSelectionStore::FObjectSelectionStore(const UClass* InObjectBaseClass)
+	: FSelectionStoreBase(InObjectBaseClass)
+{
+}
+
+void FObjectSelectionStore::SetElementList(FTypedElementList* InElementList)
+{
+	checkf(!UE_USE_ELEMENT_LIST_SELECTION, TEXT("This selection store does not support element lists! Only actor and component selections may currently use element lists!"));
+}
+
+int32 FObjectSelectionStore::GetNumObjects() const
+{
+	return SelectedObjects.Num();
+}
+
+UObject* FObjectSelectionStore::GetObjectAtIndex(const int32 InIndex) const
+{
+	return SelectedObjects.IsValidIndex(InIndex)
+		? SelectedObjects[InIndex].Get()
+		: nullptr;
+}
+
+bool FObjectSelectionStore::IsObjectSelected(const UObject* InObject) const
+{
+	return SelectedObjects.Contains(InObject);
+}
+
+void FObjectSelectionStore::SelectObject(UObject* InObject)
+{
+	bIsBatchDirty = true;
+	SelectedObjects.Add(InObject);
+	if (Sink)
+	{
+		Sink->OnObjectSelected(InObject, !IsBatchSelecting());
+	}
+}
+
+void FObjectSelectionStore::DeselectObject(UObject* InObject)
+{
+	bIsBatchDirty = true;
+	SelectedObjects.RemoveSingle(InObject);
+	if (Sink)
+	{
+		Sink->OnObjectDeselected(InObject, !IsBatchSelecting());
+	}
+}
+
+int32 FObjectSelectionStore::DeselectObjects(TFunctionRef<bool(UObject*)> InPredicate)
+{
+	int32 RemovedCount = 0;
+	
+	for (int32 Index = SelectedObjects.Num() - 1; Index >= 0; --Index)
+	{
+		UObject* SelectedObject = SelectedObjects[Index].Get();
+		if (!SelectedObject || InPredicate(SelectedObject))
+		{
+			SelectedObjects.RemoveAt(Index);
+			if (SelectedObject)
+			{
+				Sink->OnObjectDeselected(SelectedObject, /*bNotify*/false);
+				++RemovedCount;
+			}
+		}
+	}
+
+	if (RemovedCount > 0)
+	{
+		bIsBatchDirty = true;
+		if (Sink)
+		{
+			Sink->OnSelectedChanged(/*bSyncState*/false, !IsBatchSelecting());
+		}
+	}
+
+	return RemovedCount;
+}
+
+void FObjectSelectionStore::BeginBatchSelection()
+{
+	if (SelectionMutex++ == 0)
+	{
+		bIsBatchDirty = false;
+	}
+}
+
+void FObjectSelectionStore::EndBatchSelection(const bool InNotify)
+{
+	if (--SelectionMutex == 0)
+	{
+		const bool bSelectionChanged = bIsBatchDirty;
+		bIsBatchDirty = false;
+
+		if (Sink && bSelectionChanged)
+		{
+			Sink->OnSelectedChanged(/*bSyncState*/false, !IsBatchSelecting() && InNotify);
+		}
+	}
+}
+
+bool FObjectSelectionStore::IsBatchSelecting() const
+{
+	return SelectionMutex > 0;
+}
+
+void FObjectSelectionStore::ForceBatchDirty()
+{
+	bIsBatchDirty = true;
+}
+
+void FObjectSelectionStore::NoteSelectionChanged()
+{
+	if (Sink)
+	{
+		Sink->OnSelectedChanged(/*bSyncState*/false, /*bNotify*/true);
+	}
+}
+
+
+#if UE_USE_ELEMENT_LIST_SELECTION
+
+class FElementSelectionStore : public FSelectionStoreBase
+{
+public:
+	explicit FElementSelectionStore(const UClass* InObjectBaseClass);
+	virtual ~FElementSelectionStore();
+
+	//~ ISelectionStore interface
+	virtual void SetElementList(FTypedElementList* InElementList) override;
+	virtual int32 GetNumObjects() const override;
+	virtual UObject* GetObjectAtIndex(const int32 InIndex) const override;
+	virtual bool IsValidObjectToSelect(const UObject* InObject) const override;
+	virtual bool IsObjectSelected(const UObject* InObject) const override;
+	virtual void SelectObject(UObject* InObject) override;
+	virtual void DeselectObject(UObject* InObject) override;
+	virtual int32 DeselectObjects(TFunctionRef<bool(UObject*)> InPredicate) override;
+	virtual void BeginBatchSelection() override;
+	virtual void EndBatchSelection(const bool InNotify) override;
+	virtual bool IsBatchSelecting() const override;
+	virtual void ForceBatchDirty() override;
+	virtual void NoteSelectionChanged() override;
+
+	virtual FTypedElementHandle GetElementHandleForObject(const UObject* InObject, const bool bAllowCreate = true) const = 0;
+
+private:
+	void OnElementListSyncEvent(const FTypedElementList& InElementList, FTypedElementListLegacySync::ESyncType InSyncType, const FTypedElementHandle& InElementHandle, bool bIsWithinBatchOperation);
+	UObject* GetObjectForElementHandle(const FTypedElementHandle& InElementHandle) const;
+
+	/** External list of elements that this selection instance is bridging. This will be null unless a level editor instance exists. */
+	FTypedElementList* ElementList = nullptr;
+};
+
+FElementSelectionStore::FElementSelectionStore(const UClass* InObjectBaseClass)
+	: FSelectionStoreBase(InObjectBaseClass)
+{
+}
+
+FElementSelectionStore::~FElementSelectionStore()
+{
+	if (ElementList)
+	{
+		ElementList->Legacy_GetSync().OnSyncEvent().RemoveAll(this);
+	}
+}
+
+void FElementSelectionStore::SetElementList(FTypedElementList* InElementList)
+{
+	if (ElementList)
+	{
+		ElementList->Legacy_GetSync().OnSyncEvent().RemoveAll(this);
+	}
+
+	ElementList = InElementList;
+
+	if (ElementList)
+	{
+		ElementList->Legacy_GetSync().OnSyncEvent().AddRaw(this, &FElementSelectionStore::OnElementListSyncEvent);
+	}
+}
+
+int32 FElementSelectionStore::GetNumObjects() const
+{
+	return ElementList
+		? ElementList->Num()
+		: 0;
+}
+
+UObject* FElementSelectionStore::GetObjectAtIndex(const int32 InIndex) const
+{
+	if (ElementList && ElementList->IsValidIndex(InIndex))
+	{
+		const FTypedElementHandle& ElementHandle = ElementList->GetElementHandleAt(InIndex);
+		return GetObjectForElementHandle(ElementHandle);
+	}
+
+	return nullptr;
+}
+
+bool FElementSelectionStore::IsValidObjectToSelect(const UObject* InObject) const
+{
+	bool bIsValid = FSelectionStoreBase::IsValidObjectToSelect(InObject);
+
+	if (bIsValid && ElementList)
+	{
+		TTypedElement<UTypedElementSelectionInterface> SelectionElement = ElementList->GetElement<UTypedElementSelectionInterface>(GetElementHandleForObject(InObject));
+		if (SelectionElement)
+		{
+			bIsValid &= SelectionElement.IsValidSelection() && SelectionElement.Legacy_GetSelectionObject() == InObject;
+		}
+		else
+		{
+			// Elements must implement the selection interface in order to be selected!
+			bIsValid = false;
+		}
+	}
+
+	return bIsValid;
+}
+
+bool FElementSelectionStore::IsObjectSelected(const UObject* InObject) const
+{
+	if (ElementList)
+	{
+		if (const FTypedElementHandle ElementHandle = GetElementHandleForObject(InObject, /*bAllowCreate*/false))
+		{
+			return ElementList->Contains(ElementHandle);
+		}
+	}
+
+	return false;
+}
+
+void FElementSelectionStore::SelectObject(UObject* InObject)
+{
+	if (ElementList)
+	{
+		FTypedElementHandle ElementHandle = GetElementHandleForObject(InObject);
+		check(ElementHandle);
+		ElementList->Add(MoveTemp(ElementHandle));
+	}
+}
+
+void FElementSelectionStore::DeselectObject(UObject* InObject)
+{
+	if (ElementList)
+	{
+		if (const FTypedElementHandle ElementHandle = GetElementHandleForObject(InObject, /*bAllowCreate*/false))
+		{
+			ElementList->Remove(ElementHandle);
+		}
+	}
+}
+
+int32 FElementSelectionStore::DeselectObjects(TFunctionRef<bool(UObject*)> InPredicate)
+{
+	if (ElementList)
+	{
+		return ElementList->RemoveAll([this, &InPredicate](const FTypedElementHandle& InElementHandle)
+		{
+			UObject* Object = GetObjectForElementHandle(InElementHandle);
+			return Object && InPredicate(Object);
+		});
+	}
+
+	return 0;
+}
+
+void FElementSelectionStore::BeginBatchSelection()
+{
+	if (ElementList)
+	{
+		ElementList->Legacy_GetSync().BeginBatchOperation();
+	}
+}
+
+void FElementSelectionStore::EndBatchSelection(const bool InNotify)
+{
+	if (ElementList)
+	{
+		ElementList->Legacy_GetSync().EndBatchOperation(InNotify);
+	}
+}
+
+bool FElementSelectionStore::IsBatchSelecting() const
+{
+	return ElementList && ElementList->Legacy_GetSync().IsRunningBatchOperation();
+}
+
+void FElementSelectionStore::ForceBatchDirty()
+{
+	if (ElementList)
+	{
+		ElementList->Legacy_GetSync().ForceBatchOperationDirty();
+	}
+}
+
+void FElementSelectionStore::NoteSelectionChanged()
+{
+	if (ElementList)
+	{
+		// This is a way to force a BatchComplete notification so we won't re-sync the selection list state
+		FTypedElementListLegacySync& LegacySync = ElementList->Legacy_GetSync();
+		LegacySync.BeginBatchOperation();
+		LegacySync.ForceBatchOperationDirty();
+		LegacySync.EndBatchOperation();
+	}
+}
+
+void FElementSelectionStore::OnElementListSyncEvent(const FTypedElementList& InElementList, FTypedElementListLegacySync::ESyncType InSyncType, const FTypedElementHandle& InElementHandle, bool bIsWithinBatchOperation)
+{
+	if (Sink)
+	{
+		switch (InSyncType)
+		{
+		case FTypedElementListLegacySync::ESyncType::Added:
+			if (UObject* Object = GetObjectForElementHandle(InElementHandle))
+			{
+				Sink->OnObjectSelected(Object, !bIsWithinBatchOperation);
+			}
+			break;
+
+		case FTypedElementListLegacySync::ESyncType::Removed:
+			if (UObject* Object = GetObjectForElementHandle(InElementHandle))
+			{
+				Sink->OnObjectDeselected(Object, !bIsWithinBatchOperation);
+			}
+			break;
+
+		case FTypedElementListLegacySync::ESyncType::BatchComplete:
+			Sink->OnSelectedChanged(/*bSyncState*/false, !bIsWithinBatchOperation);
+			break;
+
+		default:
+			Sink->OnSelectedChanged(/*bSyncState*/true, !bIsWithinBatchOperation);
+			break;
+		}
+	}
+}
+
+UObject* FElementSelectionStore::GetObjectForElementHandle(const FTypedElementHandle& InElementHandle) const
+{
+	check(InElementHandle && ElementList);
+	
+	TTypedElement<UTypedElementSelectionInterface> SelectionElement = ElementList->GetElement<UTypedElementSelectionInterface>(InElementHandle);
+	check(SelectionElement);
+
+	UObject* Object = SelectionElement.Legacy_GetSelectionObject();
+	if (Object && Object->IsA(ObjectBaseClass))
+	{
+		return Object;
+	}
+
+	return nullptr;
+}
+
+
+class FActorElementSelectionStore : public FElementSelectionStore
+{
+public:
+	FActorElementSelectionStore();
+
+	//~ FElementSelectionStore interface
+	virtual FTypedElementHandle GetElementHandleForObject(const UObject* InObject, const bool bAllowCreate = true) const override;
+};
+
+FActorElementSelectionStore::FActorElementSelectionStore()
+	: FElementSelectionStore(AActor::StaticClass())
+{
+}
+
+FTypedElementHandle FActorElementSelectionStore::GetElementHandleForObject(const UObject* InObject, const bool bAllowCreate) const
+{
+	check(InObject);
+
+	const AActor* Actor = CastChecked<AActor>(InObject);
+	return Actor->AcquireEditorElementHandle(bAllowCreate);
+}
+
+
+class FComponentElementSelectionStore : public FElementSelectionStore
+{
+public:
+	FComponentElementSelectionStore();
+
+	//~ FElementSelectionStore interface
+	virtual FTypedElementHandle GetElementHandleForObject(const UObject* InObject, const bool bAllowCreate = true) const override;
+};
+
+FComponentElementSelectionStore::FComponentElementSelectionStore()
+	: FElementSelectionStore(UActorComponent::StaticClass())
+{
+}
+
+FTypedElementHandle FComponentElementSelectionStore::GetElementHandleForObject(const UObject* InObject, const bool bAllowCreate) const
+{
+	check(InObject);
+
+	const UActorComponent* Component = CastChecked<UActorComponent>(InObject);
+	return Component->AcquireEditorElementHandle(bAllowCreate);
+}
+
+#endif	// UE_USE_ELEMENT_LIST_SELECTION
+
+} // namespace Selection_Private
+
 USelection::FOnSelectionChanged	USelection::SelectionChangedEvent;
 USelection::FOnSelectionChanged	USelection::SelectObjectEvent;
 FSimpleMulticastDelegate		USelection::SelectNoneEvent;
 
-void USelection::Initialize(FUObjectAnnotationSparseBool* InSelectionAnnotation)
+USelection* USelection::CreateObjectSelection(FUObjectAnnotationSparseBool* InSelectionAnnotation, UObject* InOuter, FName InName, EObjectFlags InFlags)
+{
+	USelection* Selection = NewObject<USelection>(InOuter, InName, InFlags);
+	Selection->Initialize(InSelectionAnnotation, MakeShared<Selection_Private::FObjectSelectionStore>(nullptr));
+	return Selection;
+}
+
+USelection* USelection::CreateActorSelection(FUObjectAnnotationSparseBool* InSelectionAnnotation, UObject* InOuter, FName InName, EObjectFlags InFlags)
+{
+	USelection* Selection = NewObject<USelection>(InOuter, InName, InFlags);
+#if UE_USE_ELEMENT_LIST_SELECTION
+	Selection->Initialize(InSelectionAnnotation, MakeShared<Selection_Private::FActorElementSelectionStore>());
+#else	// UE_USE_ELEMENT_LIST_SELECTION
+	Selection->Initialize(InSelectionAnnotation, MakeShared<Selection_Private::FObjectSelectionStore>(AActor::StaticClass()));
+#endif	// UE_USE_ELEMENT_LIST_SELECTION
+	return Selection;
+}
+
+USelection* USelection::CreateComponentSelection(FUObjectAnnotationSparseBool* InSelectionAnnotation, UObject* InOuter, FName InName, EObjectFlags InFlags)
 {
+	USelection* Selection = NewObject<USelection>(InOuter, InName, InFlags);
+#if UE_USE_ELEMENT_LIST_SELECTION
+	Selection->Initialize(InSelectionAnnotation, MakeShared<Selection_Private::FComponentElementSelectionStore>());
+#else	// UE_USE_ELEMENT_LIST_SELECTION
+	Selection->Initialize(InSelectionAnnotation, MakeShared<Selection_Private::FObjectSelectionStore>(UActorComponent::StaticClass()));
+#endif	// UE_USE_ELEMENT_LIST_SELECTION
+	return Selection;
+}
+
+void USelection::Initialize(FUObjectAnnotationSparseBool* InSelectionAnnotation, TSharedRef<Selection_Private::ISelectionStore>&& InSelectionStore)
+{
+	SelectionStore = MoveTemp(InSelectionStore);
+	SelectionStore->SetSink(this);
+
 	if (InSelectionAnnotation)
 	{
 		SelectionAnnotation = InSelectionAnnotation;
@@ -22,206 +532,209 @@ void USelection::Initialize(FUObjectAnnotationSparseBool* InSelectionAnnotation)
 		SelectionAnnotation = new FUObjectAnnotationSparseBool;
 		bOwnsSelectionAnnotation = true;
 	}
+
+	SyncSelectedState();
 }
 
-void USelection::Select(UObject* InObject)
+void USelection::OnObjectSelected(UObject* InObject, const bool bNotify)
 {
-	check( InObject );
+	check(InObject);
 
-	const bool bSelectionChanged = !SelectionAnnotation->Get(InObject);
 	SelectionAnnotation->Set(InObject);
 
-	if(bSelectionChanged)
-	{
-		// Add to selected lists.
-		SelectedObjects.Add(InObject);
-		MarkBatchDirty();
-	}
-
-	FSelectedClassInfo* SelectedClassInfo = SelectedClasses.Find(InObject->GetClass());
-	if( SelectedClassInfo )
+	if (FSelectedClassInfo* SelectedClassInfo = SelectedClasses.Find(InObject->GetClass()))
 	{
 		++SelectedClassInfo->SelectionCount;
 	}
 	else
 	{
 		// 1 Object with a new class type has been selected
-		FSelectedClassInfo NewInfo(InObject->GetClass(), 1);
-		SelectedClasses.Add(NewInfo);
+		SelectedClasses.Add(FSelectedClassInfo(InObject->GetClass(), 1));
 	}
 
-	if( !IsBatchSelecting() )
+	if (bNotify)
 	{
 		// Call this after the item has been added from the selection set.
-		USelection::SelectObjectEvent.Broadcast( InObject );
+		USelection::SelectObjectEvent.Broadcast(InObject);
 	}
 }
 
-
-void USelection::Deselect(UObject* InObject)
+void USelection::OnObjectDeselected(UObject* InObject, const bool bNotify)
 {
-	check( InObject );
+	check(InObject);
 
-	const bool bSelectionChanged = InObject->IsSelected();
 	SelectionAnnotation->Clear(InObject);
 
-	// Remove from selected list.
-	SelectedObjects.RemoveSingle( InObject );
-
-	if( !IsBatchSelecting() )
-	{
-		// Call this after the item has been removed from the selection set.
-		USelection::SelectObjectEvent.Broadcast( InObject );
-	}
-
 	FSetElementId Id = SelectedClasses.FindId(InObject->GetClass());
-	if(Id.IsValidId())
+	if (Id.IsValidId())
 	{
 		FSelectedClassInfo& ClassInfo = SelectedClasses[Id];
 		// One less object of this class is selected;
 		--ClassInfo.SelectionCount;
 		// If no more objects of the selected class exists, remove it
-		if(ClassInfo.SelectionCount == 0)
+		if (ClassInfo.SelectionCount == 0)
 		{
 			SelectedClasses.Remove(Id);
 		}
 	}
 
-
-	if ( bSelectionChanged )
+	if (bNotify)
 	{
-		MarkBatchDirty();
+		// Call this after the item has been removed from the selection set.
+		USelection::SelectObjectEvent.Broadcast(InObject);
 	}
 }
 
-
-void USelection::Select(UObject* InObject, bool bSelect)
+void USelection::OnSelectedChanged(const bool bSyncState, const bool bNotify)
 {
-	if( bSelect )
+	if (bSyncState)
 	{
-		Select( InObject );
+		SyncSelectedState();
 	}
-	else
+	
+	if (bNotify)
 	{
-		Deselect( InObject );
+		USelection::SelectionChangedEvent.Broadcast(this);
 	}
 }
 
-
-void USelection::ToggleSelect( UObject* InObject )
+void USelection::SyncSelectedState()
 {
-	check( InObject );
+	SelectedClasses.Reset();
+	SelectionAnnotation->ClearAll();
+
+	for (int32 Idx = 0; Idx < Num(); ++Idx)
+	{
+		if (UObject* Object = GetSelectedObject(Idx))
+		{
+			SelectionAnnotation->Set(Object);
 
-	Select( InObject, !InObject->IsSelected() );
+			if (FSelectedClassInfo* SelectedClassInfo = SelectedClasses.Find(Object->GetClass()))
+			{
+				++SelectedClassInfo->SelectionCount;
+			}
+			else
+			{
+				// 1 Object with a new class type has been selected
+				SelectedClasses.Add(FSelectedClassInfo(Object->GetClass(), 1));
+			}
+		}
+	}
 }
 
-void USelection::DeselectAll( UClass* InClass )
+void USelection::Select(UObject* InObject)
 {
-	// Fast path for deselecting all UObjects with any flags
-	if ( InClass == UObject::StaticClass() )
+	check(InObject && SelectionStore->IsValidObjectToSelect(InObject));
+
+	if (!SelectionAnnotation->Get(InObject))
 	{
-		InClass = nullptr;
+		SelectionStore->SelectObject(InObject);
 	}
+}
 
-	bool bSelectionChanged = false;
+void USelection::Deselect(UObject* InObject)
+{
+	check(InObject);
 
-	TSet<FSelectedClassInfo> RemovedClasses;
-	// Remove from the end to minimize memmoves.
-	for ( int32 i = SelectedObjects.Num()-1 ; i >= 0 ; --i )
+	if (SelectionAnnotation->Get(InObject))
 	{
-		UObject* Object = GetSelectedObject(i);
-
-		if ( !Object )
-		{		
-			// Remove NULLs from SelectedObjects array.
-			SelectedObjects.RemoveAt( i );
-		}
-		else if( !InClass || Object->IsA( InClass ) )
-		{
-			// if the object is of type InClass then all objects of that same type will be removed
-			RemovedClasses.Add(FSelectedClassInfo(Object->GetClass()));
-
-			SelectionAnnotation->Clear(Object);
-			SelectedObjects.RemoveAt( i );
-
-			// Call this after the item has been removed from the selection set.
-			USelection::SelectObjectEvent.Broadcast( Object );
-
-			bSelectionChanged = true;
-		}
+		SelectionStore->DeselectObject(InObject);
 	}
+}
 
-	if( InClass == nullptr )
+void USelection::Select(UObject* InObject, bool bSelect)
+{
+	if (bSelect)
 	{
-		SelectedClasses.Empty();
+		Select(InObject);
 	}
 	else
 	{
-		// Remove the passed in class and all child classes that were removed
-		// from the list of currently selected classes
-		RemovedClasses.Add(InClass);
-		SelectedClasses = SelectedClasses.Difference(RemovedClasses);
+		Deselect(InObject);
 	}
+}
 
-	if ( bSelectionChanged )
+void USelection::ToggleSelect(UObject* InObject)
+{
+	Select(InObject, InObject && !InObject->IsSelected());
+}
+
+void USelection::DeselectAll(UClass* InClass)
+{
+	// Fast path for deselecting all UObjects with any flags
+	if (InClass == UObject::StaticClass())
 	{
-		MarkBatchDirty();
-		if ( !IsBatchSelecting() )
-		{
-			USelection::SelectionChangedEvent.Broadcast(this);
-		}
+		InClass = nullptr;
 	}
+
+	SelectionStore->DeselectObjects([InClass](UObject* InObject)
+	{
+		return !InClass || InObject->IsA(InClass);
+	});
 }
 
-void USelection::MarkBatchDirty()
+void USelection::ForceBatchDirty()
 {
-	if ( IsBatchSelecting() )
+	if (IsBatchSelecting())
 	{
-		bIsBatchDirty = true;
+		SelectionStore->ForceBatchDirty();
 	}
 }
 
+void USelection::NoteSelectionChanged()
+{
+	ensureAlwaysMsgf(!SelectionStore->IsBatchSelecting(), TEXT("NoteSelectionChanged does not work correctly if called within a batch operation!"));
+	SelectionStore->NoteSelectionChanged();
+}
 
+void USelection::NoteUnknownSelectionChanged()
+{
+	USelection::SelectionChangedEvent.Broadcast(nullptr);
+}
 
 bool USelection::IsSelected(const UObject* InObject) const
 {
-	if ( InObject )
-	{
-		return SelectedObjects.Contains(InObject);
-	}
-
-	return false;
+	return InObject && SelectionStore->IsObjectSelected(InObject);
 }
 
-
 void USelection::Serialize(FArchive& Ar)
 {
-	Super::Serialize( Ar );
+	Super::Serialize(Ar);
+
+	TArray<TWeakObjectPtr<UObject>> SelectedObjects;
+	if (Ar.IsSaving())
+	{
+		GetSelectedObjects(SelectedObjects); // TODO: The old code used to use bEvenIfPendingKill during the resolve
+	}
 	Ar << SelectedObjects;
 
-	if(Ar.IsLoading())
+	if (Ar.IsLoading())
 	{
+		BeginBatchSelectOperation();
+
 		// The set of selected objects may have changed, so make sure our annotations exactly match the list, otherwise
 		// UObject::IsSelected() could return a result that was different from the list of objects returned by GetSelectedObjects()
 		// This needs to happen in serialize because other code may check the selection state in PostEditUndo and the order of PostEditUndo is indeterminate.
-		SelectionAnnotation->ClearAll();
-
-		for(TWeakObjectPtr<UObject>& ObjectPtr : SelectedObjects)
+		DeselectAll(nullptr);
+		for (const TWeakObjectPtr<UObject>& ObjectPtr : SelectedObjects)
 		{
 			if (UObject* Object = ObjectPtr.Get(true))
 			{
-				SelectionAnnotation->Set(Object);
+				SelectionStore->SelectObject(Object);
 			}
 		}
+		SyncSelectedState();
+
+		EndBatchSelectOperation(/*bNotify*/false);
 	}
 }
 
 bool USelection::Modify(bool bAlwaysMarkDirty/* =true */)
 {
 	// If the selection currently contains any PIE objects we should not be including it in the transaction buffer
-	for (TWeakObjectPtr<UObject>& ObjectPtr : SelectedObjects)
+	for (int32 Idx = 0; Idx < Num(); ++Idx)
 	{
-		UObject* Object = ObjectPtr.Get();
+		UObject* Object = GetSelectedObject(Idx);
 		if (Object && Object->GetOutermost()->HasAnyPackageFlags(PKG_PlayInEditor | PKG_ContainsScript | PKG_CompiledIn))
 		{
 			return false;
@@ -242,3 +755,4 @@ void USelection::BeginDestroy()
 	}
 }
 
+#undef UE_USE_ELEMENT_LIST_SELECTION
diff --git a/Engine/Source/Editor/UnrealEd/Public/AssetEditorModeManager.h b/Engine/Source/Editor/UnrealEd/Public/AssetEditorModeManager.h
--- a/Engine/Source/Editor/UnrealEd/Public/AssetEditorModeManager.h
+++ b/Engine/Source/Editor/UnrealEd/Public/AssetEditorModeManager.h
@@ -4,6 +4,7 @@
 
 #include "CoreMinimal.h"
 #include "EditorModeManager.h"
+#include "TypedElementListPtr.h"
 
 class FPreviewScene;
 
@@ -31,4 +32,5 @@ protected:
 	class USelection* ObjectSet;
 	class USelection* ComponentSet;
 	class FPreviewScene* PreviewScene;
+	FTypedElementListPtr SelectedElements;
 };
diff --git a/Engine/Source/Editor/UnrealEd/Public/Selection.h b/Engine/Source/Editor/UnrealEd/Public/Selection.h
--- a/Engine/Source/Editor/UnrealEd/Public/Selection.h
+++ b/Engine/Source/Editor/UnrealEd/Public/Selection.h
@@ -9,12 +9,121 @@
 #include "Components/ActorComponent.h"
 #include "Selection.generated.h"
 
+class FTypedElementList;
+
+namespace Selection_Private
+{
+
+class ISelectionStoreSink
+{
+public:
+	virtual ~ISelectionStoreSink() = default;
+
+	/**
+	 * Called when the given object is selected within the underlying store.
+	 */
+	virtual void OnObjectSelected(UObject* InObject, const bool bNotify) = 0;
+
+	/**
+	 * Called when the given object is deselected within the underlying store.
+	 */
+	virtual void OnObjectDeselected(UObject* InObject, const bool bNotify) = 0;
+
+	/**
+	 * Called when the the underlying store changes in an unknown way.
+	 */
+	virtual void OnSelectedChanged(const bool bSyncState, const bool bNotify) = 0;
+};
+
+class ISelectionStore
+{
+public:
+	virtual ~ISelectionStore() = default;
+
+	/**
+	 * Set the notification sink instance for this store.
+	 */
+	virtual void SetSink(ISelectionStoreSink* InSink) = 0;
+
+	/**
+	 * Set the element list instance for this store.
+	 * @note Does nothing if element list stores aren't enabled. Asserts for non-element list stores if they are!
+	 */
+	virtual void SetElementList(FTypedElementList* InElementList) = 0;
+
+	/**
+	 * Get the number of objects within the underlying store.
+	 * This is the total number of objects within the store (for use as an upper limit of GetObjectAtIndex),
+	 * however not all of those objects may be valid so it should not be used as a public selection count.
+	 */
+	virtual int32 GetNumObjects() const = 0;
+
+	/**
+	 * Get the object at the internal index of the underlying store.
+	 * This object may be null, both in cases where the underlying store is using weak references, 
+	 * and also in the case that the object does not match the type managed by the underlying store.
+	 */
+	virtual UObject* GetObjectAtIndex(const int32 InIndex) const = 0;
+
+	/**
+	 * Test to see whether the given object is valid to be added to the underlying store.
+	 */
+	virtual bool IsValidObjectToSelect(const UObject* InObject) const = 0;
+
+	/**
+	 * Test to see whether the given object is currently in the underlying store.
+	 */
+	virtual bool IsObjectSelected(const UObject* InObject) const = 0;
+
+	/**
+	 * Add the given object to the underlying store.
+	 */
+	virtual void SelectObject(UObject* InObject) = 0;
+
+	/**
+	 * Remove the given object from the underlying store.
+	 */
+	virtual void DeselectObject(UObject* InObject) = 0;
+
+	/**
+	 * Remove any objects that match the predicate from the underlying store.
+	 */
+	virtual int32 DeselectObjects(TFunctionRef<bool(UObject*)> InPredicate) = 0;
+
+	/**
+	 * Called to begin a batch selection.
+	 */
+	virtual void BeginBatchSelection() = 0;
+
+	/**
+	 * Called to end a batch selection.
+	 */
+	virtual void EndBatchSelection(const bool InNotify) = 0;
+
+	/**
+	 * Are we currently batch selecting?
+	 */
+	virtual bool IsBatchSelecting() const = 0;
+
+	/**
+	 * Forcibly mark this batch as being dirty.
+	 */
+	virtual void ForceBatchDirty() = 0;
+
+	/**
+	 * Manually invoke a selection changed notification.
+	 */
+	virtual void NoteSelectionChanged() = 0;
+};
+
+} // namespace Selection_Private
+
 /**
- * Manages selections of objects.  Used in the editor for selecting
- * objects in the various browser windows.
+ * Manages selections of objects.
+ * Used in the editor for selecting objects in the various browser windows.
  */
 UCLASS(transient)
-class UNREALED_API USelection : public UObject
+class UNREALED_API USelection : public UObject, public Selection_Private::ISelectionStoreSink
 {
 	GENERATED_BODY()
 
@@ -47,15 +156,18 @@ private:
 		}
 	};
 
-	typedef TArray<TWeakObjectPtr<UObject> >	ObjectArray;
-	typedef TSet<FSelectedClassInfo>	ClassArray;
+	typedef TSet<FSelectedClassInfo> ClassArray;
 
 	template<typename SelectionFilter>
 	friend class TSelectionIterator;
 
 public:
+	static USelection* CreateObjectSelection(FUObjectAnnotationSparseBool* InSelectionAnnotation, UObject* InOuter = GetTransientPackage(), FName InName = NAME_None, EObjectFlags InFlags = RF_NoFlags);
+	static USelection* CreateActorSelection(FUObjectAnnotationSparseBool* InSelectionAnnotation, UObject* InOuter = GetTransientPackage(), FName InName = NAME_None, EObjectFlags InFlags = RF_NoFlags);
+	static USelection* CreateComponentSelection(FUObjectAnnotationSparseBool* InSelectionAnnotation, UObject* InOuter = GetTransientPackage(), FName InName = NAME_None, EObjectFlags InFlags = RF_NoFlags);
+
 	/** Params: UObject* NewSelection */
-	DECLARE_MULTICAST_DELEGATE_OneParam(FOnSelectionChanged, UObject*);
+	DECLARE_EVENT_OneParam(USelection, FOnSelectionChanged, UObject*);
 
 	/** Called when selection in editor has changed */
 	static FOnSelectionChanged SelectionChangedEvent;
@@ -64,15 +176,21 @@ public:
 	/** Called to deselect everything */
 	static FSimpleMulticastDelegate SelectNoneEvent;
 
-	/** Initializes the selection set with an annotation used to quickly look up selection state */
-	void Initialize(FUObjectAnnotationSparseBool* InSelectionAnnotation);
-
 	typedef ClassArray::TIterator TClassIterator;
 	typedef ClassArray::TConstIterator TClassConstIterator;
 
 	TClassIterator			ClassItor()				{ return TClassIterator( SelectedClasses ); }
 	TClassConstIterator		ClassConstItor() const	{ return TClassConstIterator( SelectedClasses ); }
 
+	/**
+	 * Set the element list instance for this selection set.
+	 * @note Does nothing if element list stores aren't enabled. Asserts for non-element list stores if they are!
+	 */
+	void SetElementList(FTypedElementList* InElementList)
+	{
+		SelectionStore->SetElementList(InElementList);
+	}
+
 	/**
 	 * Returns the number of objects in the selection set.  This function is used by clients in
 	 * conjunction with op::() to iterate over selected objects.  Note that some of these objects
@@ -83,7 +201,7 @@ public:
 	 */
 	int32 Num() const
 	{
-		return SelectedObjects.Num();
+		return SelectionStore->GetNumObjects();
 	}
 
 	/**
@@ -91,7 +209,7 @@ public:
 	 */
 	UObject* GetSelectedObject(const int32 InIndex)
 	{
-		return (SelectedObjects.IsValidIndex(InIndex) && SelectedObjects[InIndex].IsValid() ? SelectedObjects[InIndex].Get() : nullptr);
+		return SelectionStore->GetObjectAtIndex(InIndex);
 	}
 
 	/**
@@ -99,7 +217,7 @@ public:
 	 */
 	const UObject* GetSelectedObject(const int32 InIndex) const
 	{
-		return (SelectedObjects.IsValidIndex(InIndex) && SelectedObjects[InIndex].IsValid() ? SelectedObjects[InIndex].Get() : nullptr);
+		return SelectionStore->GetObjectAtIndex(InIndex);
 	}
 
 	/**
@@ -107,7 +225,7 @@ public:
 	 */
 	void BeginBatchSelectOperation()
 	{
-		SelectionMutex++;
+		SelectionStore->BeginBatchSelection();
 	}
 
 	/**
@@ -116,17 +234,7 @@ public:
 	 */
 	void EndBatchSelectOperation(bool bNotify = true)
 	{
-		if ( --SelectionMutex == 0 )
-		{
-			const bool bSelectionChanged = bIsBatchDirty;
-			bIsBatchDirty = false;
-
-			if ( bSelectionChanged && bNotify )
-			{
-				// new version - includes which selection set was modified
-				USelection::SelectionChangedEvent.Broadcast(this);
-			}
-		}
+		SelectionStore->EndBatchSelection(bNotify);
 	}
 
 	/**
@@ -134,7 +242,7 @@ public:
 	 */
 	bool IsBatchSelecting() const
 	{
-		return SelectionMutex != 0;
+		return SelectionStore->IsBatchSelecting();
 	}
 
 	/**
@@ -176,7 +284,18 @@ public:
 	/**
 	 * If batch selection is active, sets flag indicating something actually changed.
 	 */
-	void MarkBatchDirty();
+	void ForceBatchDirty();
+
+	/**
+	 * Manually invoke a selection changed notification for this set.
+	 */
+	void NoteSelectionChanged();
+
+	/**
+	 * Manually invoke a selection changed notification for no specific set.
+	 * @note Legacy BSP code only!
+	 */
+	static void NoteUnknownSelectionChanged();
 
 	/**
 	 * Returns the first selected object of the specified class.
@@ -189,7 +308,7 @@ public:
 	UObject* GetTop(UClass* InClass, UClass* RequiredInterface=nullptr, bool bArchetypesOnly=false)
 	{
 		check( InClass );
-		for( int32 i=0; i<SelectedObjects.Num(); ++i )
+		for( int32 i=0; i<Num(); ++i )
 		{
 			UObject* SelectedObject = GetSelectedObject(i);
 			if (SelectedObject)
@@ -230,7 +349,7 @@ public:
 	UObject* GetBottom(UClass* InClass)
 	{
 		check( InClass );
-		for( int32 i = SelectedObjects.Num()-1 ; i > -1 ; --i )
+		for( int32 i = Num()-1 ; i > -1 ; --i )
 		{
 			UObject* SelectedObject = GetSelectedObject(i);
 			if( SelectedObject && SelectedObject->IsA(InClass) )
@@ -280,16 +399,7 @@ public:
 	template< class T >
 	int32 CountSelections( bool bIgnorePendingKill=false )
 	{
-		int32 Count = 0;
-		for( int32 i=0; i<SelectedObjects.Num(); ++i )
-		{
-			UObject* SelectedObject = GetSelectedObject(i);
-			if( SelectedObject && SelectedObject->IsA(T::StaticClass()) && !(bIgnorePendingKill && SelectedObject->IsPendingKill()) )
-			{
-				++Count;
-			}
-		}
-		return Count;
+		return CountSelections(T::StaticClass(), bIgnorePendingKill);
 	}
 
 	/**
@@ -298,7 +408,7 @@ public:
 	int32 CountSelections(UClass *ClassToCount, bool bIgnorePendingKill=false)
 	{
 		int32 Count = 0;
-		for( int32 i=0; i<SelectedObjects.Num(); ++i )
+		for( int32 i=0; i<Num(); ++i )
 		{
 			UObject* SelectedObject = GetSelectedObject(i);
 			if( SelectedObject && SelectedObject->IsA(ClassToCount) && !(bIgnorePendingKill && SelectedObject->IsPendingKill()) )
@@ -332,49 +442,63 @@ public:
 	template< class T > 
 	int32 GetSelectedObjects(TArray<T*> &OutSelectedObjects)
 	{
-		OutSelectedObjects.Empty();
-		for ( int32 Idx = 0 ; Idx < SelectedObjects.Num() ; ++Idx )
+		OutSelectedObjects.Empty(Num());
+		for (int32 Idx = 0; Idx < Num(); ++Idx)
 		{
 			UObject* SelectedObject = GetSelectedObject(Idx);
-			if ( SelectedObject && SelectedObject->IsA(T::StaticClass()) )
+			if (SelectedObject && SelectedObject->IsA(T::StaticClass()))
 			{
-				OutSelectedObjects.Add( (T*)SelectedObject );
+				OutSelectedObjects.Add((T*)SelectedObject);
 			}
 		}
 		return OutSelectedObjects.Num();
 	}
-	int32 GetSelectedObjects( ObjectArray& out_SelectedObjects )
+
+	int32 GetSelectedObjects( TArray<TWeakObjectPtr<UObject>>& OutSelectedObjects )
 	{
-		out_SelectedObjects = SelectedObjects;
-		return SelectedObjects.Num();
+		OutSelectedObjects.Empty(Num());
+		for (int32 Idx = 0; Idx < Num(); ++Idx)
+		{
+			UObject* SelectedObject = GetSelectedObject(Idx);
+			if (SelectedObject)
+			{
+				OutSelectedObjects.Add(SelectedObject);
+			}
+		}
+		return OutSelectedObjects.Num();
 	}
 
 	int32 GetSelectedObjects(UClass *FilterClass, TArray<UObject*> &OutSelectedObjects)
 	{
-		OutSelectedObjects.Empty();
-		for ( int32 Idx = 0 ; Idx < SelectedObjects.Num() ; ++Idx )
+		OutSelectedObjects.Empty(Num());
+		for (int32 Idx = 0; Idx < Num(); ++Idx)
 		{
 			UObject* SelectedObject = GetSelectedObject(Idx);
-			if ( SelectedObject && SelectedObject->IsA(FilterClass) )
+			if (SelectedObject && SelectedObject->IsA(FilterClass))
 			{
-				OutSelectedObjects.Add( SelectedObject );
+				OutSelectedObjects.Add(SelectedObject);
 			}
 		}
 		return OutSelectedObjects.Num();
 	}
 
 protected:
-	/** List of selected objects, ordered as they were selected. */
-	ObjectArray	SelectedObjects;
+	/** Initializes the selection set with an annotation used to quickly look up selection state */
+	void Initialize(FUObjectAnnotationSparseBool* InSelectionAnnotation, TSharedRef<Selection_Private::ISelectionStore>&& InSelectionStore);
 
-	/** Tracks the most recently selected actor classes.  Used for UnrealEd menus. */
-	ClassArray	SelectedClasses;
+	//~ ISelectionStoreSink interface
+	virtual void OnObjectSelected(UObject* InObject, const bool bNotify) override;
+	virtual void OnObjectDeselected(UObject* InObject, const bool bNotify) override;
+	virtual void OnSelectedChanged(const bool bSyncState, const bool bNotify) override;
+	
+	/** Sync the state of the underlying selection store to the annotation and classes data */
+	void SyncSelectedState();
 
-	/** Tracks the number of active selection operations.  Allows batched selection operations to only send one notification at the end of the batch */
-	int32			SelectionMutex = 0;
+	/** Store of selected objects. */
+	TSharedPtr<Selection_Private::ISelectionStore> SelectionStore = nullptr;
 
-	/** Tracks whether the selection set changed during a batch selection operation */
-	bool		bIsBatchDirty = false;
+	/** Tracks the most recently selected actor classes.  Used for UnrealEd menus. */
+	ClassArray SelectedClasses;
 	
 	/** Selection annotation for fast lookup */
 	FUObjectAnnotationSparseBool* SelectionAnnotation = nullptr;
@@ -480,7 +604,7 @@ private:
 
 	bool IsIndexValid() const
 	{
-		return Selection.SelectedObjects.IsValidIndex( Index );
+		return Index >= 0 && Index < Selection.Num();
 	}
 
 	USelection&	Selection;
diff --git a/Engine/Source/Editor/UnrealEd/UnrealEd.Build.cs b/Engine/Source/Editor/UnrealEd/UnrealEd.Build.cs
--- a/Engine/Source/Editor/UnrealEd/UnrealEd.Build.cs
+++ b/Engine/Source/Editor/UnrealEd/UnrealEd.Build.cs
@@ -106,7 +106,9 @@ public class UnrealEd : ModuleRules
                 "MeshBuilder",
                 "MaterialShaderQualitySettings",
                 "EditorSubsystem",
-                "InteractiveToolsFramework",
+				"InteractiveToolsFramework",
+				"TypedElementFramework",
+				"TypedElementInterfaces",
 				"ToolMenusEditor",
 				"StatusBar",
 				"InterchangeCore",
diff --git a/Engine/Source/Runtime/TypedElementInterfaces/Public/TypedElementSelectionInterface.h b/Engine/Source/Runtime/TypedElementInterfaces/Public/TypedElementSelectionInterface.h
new file mode 100644
--- /dev/null
+++ b/Engine/Source/Runtime/TypedElementInterfaces/Public/TypedElementSelectionInterface.h
@@ -0,0 +1,33 @@
+// Copyright Epic Games, Inc. All Rights Reserved.
+
+#pragma once
+
+#include "TypedElementHandle.h"
+#include "TypedElementSelectionInterface.generated.h"
+
+UCLASS()
+class TYPEDELEMENTINTERFACES_API UTypedElementSelectionInterface : public UTypedElementInterface
+{
+	GENERATED_BODY()
+
+public:
+	/**
+	 * Test to see whether the given handle is in a valid state to be selected.
+	 */
+	UFUNCTION(BlueprintCallable, Category="TypedElementInterfaces|Selection")
+	virtual bool IsValidSelection(const FTypedElementHandle& InElementHandle) { return true; }
+
+	/**
+	 * Retreive the object instance that should be selected via this handle.
+	 * @note This exists to allow the legacy USelection to bridge to an FTypedElementList instance. It should not be used in new code!
+	 */
+	UFUNCTION(BlueprintCallable, Category="TypedElementInterfaces|Selection")
+	virtual UObject* Legacy_GetSelectionObject(const FTypedElementHandle& InElementHandle) { return nullptr; }
+};
+
+template <>
+struct TTypedElement<UTypedElementSelectionInterface> : public TTypedElementBase<UTypedElementSelectionInterface>
+{
+	bool IsValidSelection() const { return InterfacePtr->IsValidSelection(*this); }
+	UObject* Legacy_GetSelectionObject() const { return InterfacePtr->Legacy_GetSelectionObject(*this); }
+};
