Move ViewportWorldInteraction code out of the base editor mode manager.

#rnx

diff --git a/Engine/Source/Editor/UnrealEd/Private/EditorModeManager.cpp b/Engine/Source/Editor/UnrealEd/Private/EditorModeManager.cpp
--- a/Engine/Source/Editor/UnrealEd/Private/EditorModeManager.cpp
+++ b/Engine/Source/Editor/UnrealEd/Private/EditorModeManager.cpp
@@ -18,8 +18,6 @@
 #include "Widgets/Notifications/SNotificationList.h"
 #include "Widgets/Input/SButton.h"
 #include "Engine/LevelStreaming.h"
-#include "EditorWorldExtension.h"
-#include "ViewportWorldInteraction.h"
 #include "Editor/EditorEngine.h"
 #include "UnrealEdGlobals.h"
 #include "Editor/UnrealEdEngine.h"
@@ -98,11 +96,10 @@ void FEditorModeTools::LoadConfig(void)
 		GEditorPerProjectIni);
 
 	const bool bGetRawValue = true;
-	int32 Bogus = (int32)GetCoordSystem(bGetRawValue);
-	GConfig->GetInt(TEXT("FEditorModeTools"),TEXT("CoordSystem"),Bogus,
+	int32 CoordSystemAsInt = (int32)GetCoordSystem(bGetRawValue);
+	GConfig->GetInt(TEXT("FEditorModeTools"),TEXT("CoordSystem"), CoordSystemAsInt,
 		GEditorPerProjectIni);
-	SetCoordSystem((ECoordSystem)Bogus);
-
+	SetCoordSystem((ECoordSystem)CoordSystemAsInt);
 
 	LoadWidgetSettings();
 }
@@ -501,22 +498,7 @@ void FEditorModeTools::SetPivotLocation( const FVector& Location, const bool bIn
 
 ECoordSystem FEditorModeTools::GetCoordSystem(bool bGetRawValue)
 {
-	bool bAligningToActors = false;
-	if (GEditor->GetEditorWorldExtensionsManager() != nullptr
-		&& GetWorld() != nullptr)
-	{
-		UEditorWorldExtensionCollection* WorldExtensionCollection = GEditor->GetEditorWorldExtensionsManager()->GetEditorWorldExtensions(GetWorld(), false);
-		if (WorldExtensionCollection != nullptr)
-		{
-			UViewportWorldInteraction* ViewportWorldInteraction = Cast<UViewportWorldInteraction>(WorldExtensionCollection->FindExtension(UViewportWorldInteraction::StaticClass()));
-			if (ViewportWorldInteraction != nullptr && ViewportWorldInteraction->AreAligningToActors() == true)
-			{
-				bAligningToActors = true;
-			}
-		}
-	}
-	if (!bGetRawValue && 
-		((GetWidgetMode() == FWidget::WM_Scale) || bAligningToActors))
+	if (!bGetRawValue && (GetWidgetMode() == FWidget::WM_Scale))
 	{
 		return COORD_Local;
 	}
@@ -528,25 +510,6 @@ ECoordSystem FEditorModeTools::GetCoordSystem(bool bGetRawValue)
 
 void FEditorModeTools::SetCoordSystem(ECoordSystem NewCoordSystem)
 {
-	// If we are trying to enter world space but are aligning to actors, turn off aligning to actors
-	if (GEditor->GetEditorWorldExtensionsManager() != nullptr
-		&& GetWorld() != nullptr
-		&& NewCoordSystem == COORD_World)
-	{
-		UEditorWorldExtensionCollection* WorldExtensionCollection = GEditor->GetEditorWorldExtensionsManager()->GetEditorWorldExtensions(GetWorld(), false);
-		if (WorldExtensionCollection != nullptr)
-		{
-			UViewportWorldInteraction* ViewportWorldInteraction = Cast<UViewportWorldInteraction>(WorldExtensionCollection->FindExtension(UViewportWorldInteraction::StaticClass()));
-			if (ViewportWorldInteraction != nullptr && ViewportWorldInteraction->AreAligningToActors() == true)
-			{
-				if (ViewportWorldInteraction->HasCandidatesSelected())
-				{
-					ViewportWorldInteraction->SetSelectionAsCandidates();
-				}
-				GUnrealEd->Exec(GetWorld(), TEXT("VI.EnableGuides 0"));
-			}
-		}
-	}
 	CoordSystem = NewCoordSystem;
 	BroadcastCoordSystemChanged(NewCoordSystem);
 }
diff --git a/Engine/Source/Editor/ViewportInteraction/Public/ViewportWorldInteraction.h b/Engine/Source/Editor/ViewportInteraction/Public/ViewportWorldInteraction.h
--- a/Engine/Source/Editor/ViewportInteraction/Public/ViewportWorldInteraction.h
+++ b/Engine/Source/Editor/ViewportInteraction/Public/ViewportWorldInteraction.h
@@ -330,7 +330,7 @@ public:
 	void AllowWorldMovement( bool bAllow );
 
 	/** For other systems to check if the Viewport World Interaction system is currently aligning transformables to actors*/
-	bool AreAligningToActors();
+	bool AreAligningToActors() const;
 	/** For other systems to check if the Viewport World Interaction system currently has candidate actors selected */
 	bool HasCandidatesSelected();
 	/** If there are no currently selected candidates, use the currently selected actors as candidates. Otherwise, reset the candidates */
diff --git a/Engine/Source/Editor/ViewportInteraction/ViewportWorldInteraction.cpp b/Engine/Source/Editor/ViewportInteraction/ViewportWorldInteraction.cpp
--- a/Engine/Source/Editor/ViewportInteraction/ViewportWorldInteraction.cpp
+++ b/Engine/Source/Editor/ViewportInteraction/ViewportWorldInteraction.cpp
@@ -2981,14 +2981,28 @@ void UViewportWorldInteraction::CycleTransformGizmoCoordinateSpace()
 
 void UViewportWorldInteraction::SetTransformGizmoCoordinateSpace( const ECoordSystem NewCoordSystem )
 {
+	// If we are trying to enter world space but are aligning to actors, turn off aligning to actors
+	if (NewCoordSystem == COORD_World && AreAligningToActors())
+	{
+		if (HasCandidatesSelected())
+		{
+			SetSelectionAsCandidates();
+		}
+		GUnrealEd->Exec(GetWorld(), TEXT("VI.EnableGuides 0"));
+	}
+
 	GetModeTools().SetCoordSystem( NewCoordSystem );
 }
 
 ECoordSystem UViewportWorldInteraction::GetTransformGizmoCoordinateSpace() const
 {
+	if (AreAligningToActors())
+	{
+		return COORD_Local;
+	}
+
 	const bool bGetRawValue = false;
-	const ECoordSystem CurrentCoordSystem = GetModeTools().GetCoordSystem( bGetRawValue );
-	return CurrentCoordSystem;
+	return GetModeTools().GetCoordSystem(bGetRawValue);
 }
 
 float UViewportWorldInteraction::GetMaxScale()
@@ -3067,7 +3081,7 @@ void UViewportWorldInteraction::DestroyActors()
 	}
 }
 
-bool UViewportWorldInteraction::AreAligningToActors()
+bool UViewportWorldInteraction::AreAligningToActors() const
 {
 	return (VI::ActorSnap->GetInt() == 1) ? true : false;
 }
