Remove deprecated bookmark hooks from editor mode manager.

#rnx

diff --git a/Engine/Source/Editor/UnrealEd/Private/EditorModeManager.cpp b/Engine/Source/Editor/UnrealEd/Private/EditorModeManager.cpp
--- a/Engine/Source/Editor/UnrealEd/Private/EditorModeManager.cpp
+++ b/Engine/Source/Editor/UnrealEd/Private/EditorModeManager.cpp
@@ -10,7 +10,6 @@
 #include "LevelEditorViewport.h"
 #include "EditorModeRegistry.h"
 #include "EditorModes.h"
-#include "Engine/BookMark.h"
 #include "EditorSupportDelegates.h"
 #include "EdMode.h"
 #include "Toolkits/IToolkitHost.h"
@@ -21,7 +20,6 @@
 #include "Editor/EditorEngine.h"
 #include "UnrealEdGlobals.h"
 #include "Editor/UnrealEdEngine.h"
-#include "Bookmarks/IBookmarkTypeTools.h"
 #include "Widgets/Docking/SDockTab.h"
 #include "Widgets/Layout/SWidgetSwitcher.h"
 #include "EditorStyleSet.h"
@@ -1555,62 +1553,6 @@ FWidget::EWidgetMode FEditorModeTools::GetWidgetMode() const
 	return WidgetMode;
 }
 
-bool FEditorModeTools::GetShowFriendlyVariableNames()
-{
-	return GetDefault<UEditorStyleSettings>()->bShowFriendlyNames;
-}
-
-const uint32 FEditorModeTools::GetMaxNumberOfBookmarks(FEditorViewportClient* InViewportClient)
-{
-	return IBookmarkTypeTools::Get().GetMaxNumberOfBookmarks(InViewportClient);
-}
-
-void FEditorModeTools::CompactBookmarks(FEditorViewportClient* InViewportClient)
-{
-	IBookmarkTypeTools::Get().CompactBookmarks(InViewportClient);
-}
-
-/**
- * Sets a bookmark in the levelinfo file, allocating it if necessary.
- */
-void FEditorModeTools::SetBookmark( uint32 InIndex, FEditorViewportClient* InViewportClient )
-{
-	IBookmarkTypeTools::Get().CreateOrSetBookmark(InIndex, InViewportClient);
-}
-
-/**
- * Checks to see if a bookmark exists at a given index
- */
-
-bool FEditorModeTools::CheckBookmark( uint32 InIndex, FEditorViewportClient* InViewportClient )
-{
-	return IBookmarkTypeTools::Get().CheckBookmark(InIndex, InViewportClient);
-}
-
-/**
- * Retrieves a bookmark from the list.
- */
-void FEditorModeTools::JumpToBookmark(uint32 InIndex, TSharedPtr<FBookmarkBaseJumpToSettings> InSettings, FEditorViewportClient* InViewportClient)
-{
-	IBookmarkTypeTools::Get().JumpToBookmark(InIndex, InSettings, InViewportClient);
-}
-
-/**
- * Clears a bookmark
- */
-void FEditorModeTools::ClearBookmark( uint32 InIndex, FEditorViewportClient* InViewportClient )
-{
-	IBookmarkTypeTools::Get().ClearBookmark(InIndex, InViewportClient);
-}
-
-/**
-* Clears all book marks
-*/
-void FEditorModeTools::ClearAllBookmarks( FEditorViewportClient* InViewportClient )
-{
-	IBookmarkTypeTools::Get().ClearAllBookmarks(InViewportClient);
-}
-
 void FEditorModeTools::AddReferencedObjects( FReferenceCollector& Collector )
 {
 	Collector.AddReferencedObjects(ActiveScriptableModes);
diff --git a/Engine/Source/Editor/UnrealEd/Public/EditorModeManager.h b/Engine/Source/Editor/UnrealEd/Public/EditorModeManager.h
--- a/Engine/Source/Editor/UnrealEd/Public/EditorModeManager.h
+++ b/Engine/Source/Editor/UnrealEd/Public/EditorModeManager.h
@@ -323,84 +323,6 @@ public:
 	 */
 	FWidget::EWidgetMode GetWidgetMode() const;
 
-	/**
-	 * Gets the current state of script editor usage of show friendly names
-	 * @deprecated Use GetDefault<UEditorStyleSettings>()->bShowFriendlyNames instead
-	 *
-	 * @ return - If true, replace variable names with sanitized ones
-	 */
-	UE_DEPRECATED(4.26, "Use GetDefault<UEditorStyleSettings>()->bShowFriendlyNames")
-	static bool GetShowFriendlyVariableNames();
-
-	/**
-	 * Gets the maximum number of bookmarks.
-	 * @deprecated Use IBookmarkTypeTools::Get().GetMaxNumberOfBookmarks instead
-	 *
-	 * @param InViewportClient	Level editor viewport client used to reference the world which owns the bookmarks.
-	 */
-	UE_DEPRECATED(4.26, "Use IBookmarkTypeTools::Get().GetMaxNumberOfBookmarks")
-	static const uint32 GetMaxNumberOfBookmarks(FEditorViewportClient* InViewportClient);
-
-	/**
-	 * Compacts the available bookmarks into mapped spaces.
-	 * Does nothing if all valid bookmarks are already mapped.
-	 * @deprecated Use IBookmarkTypeTools::Get().CompactBookmarks instead
-	 *
-	 * @param InViewportClient	Level editor viewport client used to reference the world which owns the bookmarks.
-	 */
-	UE_DEPRECATED(4.26, "Use IBookmarkTypeTools::Get().CompactBookmarks")
-	static void CompactBookmarks(FEditorViewportClient* InViewportClient);
-
-	/**
-	 * Sets a bookmark in the levelinfo file, allocating it if necessary.
-	 * @deprecated Use IBookmarkTypeTools::Get().CreateOrSetBookmark instead
-	 * 
-	 * @param InIndex			Index of the bookmark to set.
-	 * @param InViewportClient	Level editor viewport client used to reference the world which owns the bookmark.
-	 */
-	UE_DEPRECATED(4.26, "Use IBookmarkTypeTools::Get().CreateOrSetBookmark")
-	static void SetBookmark( uint32 InIndex, FEditorViewportClient* InViewportClient );
-
-	/**
-	 * Checks to see if a bookmark exists at a given index
-	 * @deprecated Use IBookmarkTypeTools::Get().CheckBookmark instead
-	 * 
-	 * @param InIndex			Index of the bookmark to set.
-	 * @param InViewportClient	Level editor viewport client used to reference the world which owns the bookmark.
-	 */
-	UE_DEPRECATED(4.26, "Use IBookmarkTypeTools::Get().CheckBookmark")
-	static bool CheckBookmark( uint32 InIndex, FEditorViewportClient* InViewportClient );
-
-	/**
-	 * Activates a bookmark from the list.
-	 * @deprecated Use IBookmarkTypeTools::Get().JumpToBookmark instead
-	 * 
-	 * @param InIndex			Index of the bookmark to set.
-	 * @param InSettings		Settings to used when jumpting to the bookmark.
-	 * @param InViewportClient	Level editor viewport client used to reference the world which owns the bookmark.
-	 */
-	UE_DEPRECATED(4.26, "Use IBookmarkTypeTools::Get().JumpToBookmark")
-	void JumpToBookmark( uint32 InIndex, TSharedPtr<struct FBookmarkBaseJumpToSettings> InSettings, FEditorViewportClient* InViewportClient );
-
-	/**
-	 * Clears a bookmark from the list.
-	 * @deprecated Use IBookmarkTypeTools::Get().ClearBookmark instead
-	 * 
-	 * @param InIndex			Index of the bookmark to clear.
-	 * @param InViewportClient	Level editor viewport client used to reference the world which owns the bookmark.
-	 */
-	UE_DEPRECATED(4.26, "Use IBookmarkTypeTools::Get().ClearBookmark")
-	static void ClearBookmark( uint32 InIndex, FEditorViewportClient* InViewportClient );
-
-	/**
-	 * Clears all book marks
-	 * @deprecated Use IBookmarkTypeTools::Get().ClearAllBookmarks instead
-	 * 
-	 * @param InViewportClient	Level editor viewport client used to reference the world which owns the bookmarks.
-	 */
-	UE_DEPRECATED(4.26, "Use IBookmarkTypeTools::Get().ClearAllBookmarks")
-	static void ClearAllBookmarks( FEditorViewportClient* InViewportClient );
-
 	// FGCObject interface
 	virtual void AddReferencedObjects( FReferenceCollector& Collector ) override;
 	// End of FGCObject interface
