Move some FEditorViewportClient functions to matching implementation files.

#rnx

diff --git a/Engine/Source/Editor/UnrealEd/Private/EditorViewportClient.cpp b/Engine/Source/Editor/UnrealEd/Private/EditorViewportClient.cpp
--- a/Engine/Source/Editor/UnrealEd/Private/EditorViewportClient.cpp
+++ b/Engine/Source/Editor/UnrealEd/Private/EditorViewportClient.cpp
@@ -29,7 +29,6 @@
 #include "Editor/Matinee/Public/MatineeConstants.h"
 #include "HighResScreenshot.h"
 #include "EditorDragTools.h"
-#include "Editor/MeshPaintMode/Public/MeshPaintEdMode.h"
 #include "EngineAnalytics.h"
 #include "AnalyticsEventAttribute.h"
 #include "Interfaces/IAnalyticsProvider.h"
@@ -158,6 +157,67 @@ public:
 	TMap <FKey, EInputEvent> KeyEventValues;
 };
 
+
+///////////////////////////////////////////////////////////////////////////////////////////////////
+//
+//	FViewportCursorLocation
+//	Contains information about a mouse cursor position within a viewport, transformed into the correct
+//	coordinate system for the viewport.
+//
+///////////////////////////////////////////////////////////////////////////////////////////////////
+FViewportCursorLocation::FViewportCursorLocation(const FSceneView* View, FEditorViewportClient* InViewportClient, int32 X, int32 Y)
+	: Origin(ForceInit), Direction(ForceInit), CursorPos(X, Y)
+{
+
+	FVector4 ScreenPos = View->PixelToScreen(X, Y, 0);
+
+	const FMatrix InvViewMatrix = View->ViewMatrices.GetInvViewMatrix();
+	const FMatrix InvProjMatrix = View->ViewMatrices.GetInvProjectionMatrix();
+
+	const float ScreenX = ScreenPos.X;
+	const float ScreenY = ScreenPos.Y;
+
+	ViewportClient = InViewportClient;
+
+	if (ViewportClient->IsPerspective())
+	{
+		Origin = View->ViewMatrices.GetViewOrigin();
+		Direction = InvViewMatrix.TransformVector(FVector(InvProjMatrix.TransformFVector4(FVector4(ScreenX * GNearClippingPlane, ScreenY * GNearClippingPlane, 0.0f, GNearClippingPlane)))).GetSafeNormal();
+	}
+	else
+	{
+		Origin = InvViewMatrix.TransformFVector4(InvProjMatrix.TransformFVector4(FVector4(ScreenX, ScreenY, 0.5f, 1.0f)));
+		Direction = InvViewMatrix.TransformVector(FVector(0, 0, 1)).GetSafeNormal();
+	}
+}
+
+FViewportCursorLocation::~FViewportCursorLocation()
+{
+}
+
+ELevelViewportType FViewportCursorLocation::GetViewportType() const
+{
+	return ViewportClient->GetViewportType();
+}
+
+///////////////////////////////////////////////////////////////////////////////////////////////////
+//
+//	FViewportClick::FViewportClick - Calculates useful information about a click for the below ClickXXX functions to use.
+//
+///////////////////////////////////////////////////////////////////////////////////////////////////
+FViewportClick::FViewportClick(const FSceneView* View, FEditorViewportClient* ViewportClient, FKey InKey, EInputEvent InEvent, int32 X, int32 Y)
+	: FViewportCursorLocation(View, ViewportClient, X, Y)
+	, Key(InKey), Event(InEvent)
+{
+	ControlDown = ViewportClient->IsCtrlPressed();
+	ShiftDown = ViewportClient->IsShiftPressed();
+	AltDown = ViewportClient->IsAltPressed();
+}
+
+FViewportClick::~FViewportClick()
+{
+}
+
 FViewportCameraTransform::FViewportCameraTransform()
 	: TransitionCurve( new FCurveSequence( 0.0f, FocusConstants::TransitionTime, ECurveEaseFunction::CubicOut ) )
 	, ViewLocation( FVector::ZeroVector )
diff --git a/Engine/Source/Editor/UnrealEd/Private/LevelEditorViewport.cpp b/Engine/Source/Editor/UnrealEd/Private/LevelEditorViewport.cpp
--- a/Engine/Source/Editor/UnrealEd/Private/LevelEditorViewport.cpp
+++ b/Engine/Source/Editor/UnrealEd/Private/LevelEditorViewport.cpp
@@ -102,67 +102,6 @@ TSet<FViewportHoverTarget> FLevelEditorViewportClient::HoveredObjects;
 
 IMPLEMENT_HIT_PROXY( HLevelSocketProxy, HHitProxy );
 
-
-///////////////////////////////////////////////////////////////////////////////////////////////////
-//
-//	FViewportCursorLocation
-//	Contains information about a mouse cursor position within a viewport, transformed into the correct
-//	coordinate system for the viewport.
-//
-///////////////////////////////////////////////////////////////////////////////////////////////////
-FViewportCursorLocation::FViewportCursorLocation( const FSceneView* View, FEditorViewportClient* InViewportClient, int32 X, int32 Y )
-	:	Origin(ForceInit), Direction(ForceInit), CursorPos(X, Y)
-{
-
-	FVector4 ScreenPos = View->PixelToScreen(X, Y, 0);
-
-	const FMatrix InvViewMatrix = View->ViewMatrices.GetInvViewMatrix();
-	const FMatrix InvProjMatrix = View->ViewMatrices.GetInvProjectionMatrix();
-
-	const float ScreenX = ScreenPos.X;
-	const float ScreenY = ScreenPos.Y;
-
-	ViewportClient = InViewportClient;
-
-	if ( ViewportClient->IsPerspective() )
-	{
-		Origin = View->ViewMatrices.GetViewOrigin();
-		Direction = InvViewMatrix.TransformVector(FVector(InvProjMatrix.TransformFVector4(FVector4(ScreenX * GNearClippingPlane,ScreenY * GNearClippingPlane,0.0f,GNearClippingPlane)))).GetSafeNormal();
-	}
-	else
-	{
-		Origin = InvViewMatrix.TransformFVector4(InvProjMatrix.TransformFVector4(FVector4(ScreenX,ScreenY,0.5f,1.0f)));
-		Direction = InvViewMatrix.TransformVector(FVector(0,0,1)).GetSafeNormal();
-	}
-}
-
-FViewportCursorLocation::~FViewportCursorLocation()
-{
-}
-
-ELevelViewportType FViewportCursorLocation::GetViewportType() const
-{
-	return ViewportClient->GetViewportType();
-}
-
-///////////////////////////////////////////////////////////////////////////////////////////////////
-//
-//	FViewportClick::FViewportClick - Calculates useful information about a click for the below ClickXXX functions to use.
-//
-///////////////////////////////////////////////////////////////////////////////////////////////////
-FViewportClick::FViewportClick(const FSceneView* View,FEditorViewportClient* ViewportClient,FKey InKey,EInputEvent InEvent,int32 X,int32 Y)
-	:	FViewportCursorLocation(View, ViewportClient, X, Y)
-	,	Key(InKey), Event(InEvent)
-{
-	ControlDown = ViewportClient->IsCtrlPressed();
-	ShiftDown = ViewportClient->IsShiftPressed();
-	AltDown = ViewportClient->IsAltPressed();
-}
-
-FViewportClick::~FViewportClick()
-{
-}
-
 /** Helper function to compute a new location that is snapped to the origin plane given the users cursor location and camera angle */
 static FVector4 AttemptToSnapLocationToOriginPlane( const FViewportCursorLocation& Cursor, FVector4 Location )
 {
