Fix copyright message
#rnx

diff --git a/Engine/Source/Editor/UnrealEd/Private/Commandlets/WorldPartitionBuildHLODsCommandlet.h b/Engine/Source/Editor/UnrealEd/Private/Commandlets/WorldPartitionBuildHLODsCommandlet.h
--- a/Engine/Source/Editor/UnrealEd/Private/Commandlets/WorldPartitionBuildHLODsCommandlet.h
+++ b/Engine/Source/Editor/UnrealEd/Private/Commandlets/WorldPartitionBuildHLODsCommandlet.h
@@ -1,4 +1,4 @@
-// Copyright 1998-2019 Epic Games, Inc. All Rights Reserved.
+// Copyright Epic Games, Inc. All Rights Reserved.
 
 #pragma once
 
