Use explicit template instantiations instead of a potentially optimized out function creating dummy variables.

#rnx

diff --git a/Engine/Source/Editor/UnrealEd/Private/EditorViewportLayoutTwoPanes.cpp b/Engine/Source/Editor/UnrealEd/Private/EditorViewportLayoutTwoPanes.cpp
--- a/Engine/Source/Editor/UnrealEd/Private/EditorViewportLayoutTwoPanes.cpp
+++ b/Engine/Source/Editor/UnrealEd/Private/EditorViewportLayoutTwoPanes.cpp
@@ -14,7 +14,6 @@ namespace ViewportLayoutTwoPanesDefs
 	static const float DefaultSplitterPercentage = 0.5f;
 }
 
-
 template <EOrientation TOrientation>
 TSharedRef<SWidget> TEditorViewportLayoutTwoPanes<TOrientation>::MakeViewportLayout(TSharedPtr<FAssetEditorViewportLayout> InParentLayout, const FString& SpecificLayoutString)
 {
@@ -116,14 +115,5 @@ const FName& FEditorViewportLayoutTwoPanesHoriz::GetLayoutTypeName() const
 	return EditorViewportConfigurationNames::TwoPanesHoriz;
 }
 
-/**
-* Function avoids linker errors on the template class functions in this cpp file.
-* It avoids the need to put the contents of this file into the header.
-* It doesn't get called.
-*/
-void EditorViewportLayoutTwoPanes_LinkErrorFixFunc()
-{
-	check(0);
-	FEditorViewportLayoutTwoPanesVert DummyVert;
-	FEditorViewportLayoutTwoPanesHoriz DummyHoriz;
-}
+template class TEditorViewportLayoutTwoPanes<EOrientation::Orient_Vertical>;
+template class TEditorViewportLayoutTwoPanes<EOrientation::Orient_Horizontal>;
