Fixed static analysis warning

#rnx

diff --git a/Engine/Source/Editor/UnrealEd/Private/Selection.cpp b/Engine/Source/Editor/UnrealEd/Private/Selection.cpp
--- a/Engine/Source/Editor/UnrealEd/Private/Selection.cpp
+++ b/Engine/Source/Editor/UnrealEd/Private/Selection.cpp
@@ -139,7 +139,10 @@ int32 FObjectSelectionStore::DeselectObjects(TFunctionRef<bool(UObject*)> InPred
 			SelectedObjects.RemoveAt(Index);
 			if (SelectedObject)
 			{
-				Sink->OnObjectDeselected(SelectedObject, /*bNotify*/false);
+				if (Sink)
+				{
+					Sink->OnObjectDeselected(SelectedObject, /*bNotify*/false);
+				}
 				++RemovedCount;
 			}
 		}
