World Partition convert commandlet: delete external actors even if the converted map doesn't exist.

diff --git a/Engine/Source/Editor/UnrealEd/Private/Commandlets/WorldPartitionConvertCommandlet.cpp b/Engine/Source/Editor/UnrealEd/Private/Commandlets/WorldPartitionConvertCommandlet.cpp
--- a/Engine/Source/Editor/UnrealEd/Private/Commandlets/WorldPartitionConvertCommandlet.cpp
+++ b/Engine/Source/Editor/UnrealEd/Private/Commandlets/WorldPartitionConvertCommandlet.cpp
@@ -734,7 +734,6 @@ void UWorldPartitionConvertCommandlet::SetupHLODLayerAssets()
 	}
 }
 
-
 int32 UWorldPartitionConvertCommandlet::Main(const FString& Params)
 {
 	UE_SCOPED_TIMER(TEXT("Conversion"), LogWorldPartitionConvertCommandlet, Display);
@@ -808,25 +807,22 @@ int32 UWorldPartitionConvertCommandlet::Main(const FString& Params)
 	{
 		UE_SCOPED_TIMER(TEXT("Deleting existing conversion results"), LogWorldPartitionConvertCommandlet, Display);
 
-		FString OldLevelName = Tokens[0] + ConversionSuffix;
-		if (FPackageName::SearchForPackageOnDisk(OldLevelName, &OldLevelName))
+		// Update AssetRegistry
 		{
-			// Update AssetRegistry
-			{
-				UE_SCOPED_TIMER(TEXT("Waiting for asset registry"), LogWorldPartitionConvertCommandlet, Display);
-				IAssetRegistry& AssetRegistry = FModuleManager::LoadModuleChecked<FAssetRegistryModule>(TEXT("AssetRegistry")).Get();
-				AssetRegistry.SearchAllAssets(true);
-			}
+			UE_SCOPED_TIMER(TEXT("Waiting for asset registry"), LogWorldPartitionConvertCommandlet, Display);
+			IAssetRegistry& AssetRegistry = FModuleManager::LoadModuleChecked<FAssetRegistryModule>(TEXT("AssetRegistry")).Get();
+			AssetRegistry.SearchAllAssets(true);
+		}
 
-			TArray<FAssetData> Assets;
-			FActorRegistry::GetLevelActors(*OldLevelName, Assets);
+		TArray<FAssetData> Assets;
+		FString OldLevelName = Tokens[0] + ConversionSuffix;
+		FActorRegistry::GetLevelActors(*OldLevelName, Assets);
 
-			for (const FAssetData& Asset : Assets)
+		for (const FAssetData& Asset : Assets)
+		{
+			if (!DeleteFile(SourceControlHelpers::PackageFilename(Asset.PackageName.ToString())))
 			{
-				if (!DeleteFile(SourceControlHelpers::PackageFilename(Asset.PackageName.ToString())))
-				{
-					return 1;
-				}
+				return 1;
 			}
 		}
 	}
