Fix ensure when notifying selection change during a batch

diff --git a/Engine/Source/Editor/UnrealEd/Private/Selection.cpp b/Engine/Source/Editor/UnrealEd/Private/Selection.cpp
--- a/Engine/Source/Editor/UnrealEd/Private/Selection.cpp
+++ b/Engine/Source/Editor/UnrealEd/Private/Selection.cpp
@@ -67,7 +67,6 @@ public:
 	virtual void EndBatchSelection(const bool InNotify) override;
 	virtual bool IsBatchSelecting() const override;
 	virtual void ForceBatchDirty() override;
-	virtual void NoteSelectionChanged() override;
 
 private:
 	/** Tracks the number of active selection operations.  Allows batched selection operations to only send one notification at the end of the batch */
@@ -192,14 +191,6 @@ void FObjectSelectionStore::ForceBatchDirty()
 	bIsBatchDirty = true;
 }
 
-void FObjectSelectionStore::NoteSelectionChanged()
-{
-	if (Sink)
-	{
-		Sink->OnSelectedChanged(/*bSyncState*/false, /*bNotify*/true);
-	}
-}
-
 
 #if UE_USE_ELEMENT_LIST_SELECTION
 
@@ -222,7 +213,6 @@ public:
 	virtual void EndBatchSelection(const bool InNotify) override;
 	virtual bool IsBatchSelecting() const override;
 	virtual void ForceBatchDirty() override;
-	virtual void NoteSelectionChanged() override;
 
 	virtual FTypedElementHandle GetElementHandleForObject(const UObject* InObject, const bool bAllowCreate = true) const = 0;
 
@@ -378,18 +368,6 @@ void FElementSelectionStore::ForceBatchDirty()
 	}
 }
 
-void FElementSelectionStore::NoteSelectionChanged()
-{
-	if (ElementList)
-	{
-		// This is a way to force a BatchComplete notification so we won't re-sync the selection list state
-		FTypedElementListLegacySync& LegacySync = ElementList->Legacy_GetSync();
-		LegacySync.BeginBatchOperation();
-		LegacySync.ForceBatchOperationDirty();
-		LegacySync.EndBatchOperation();
-	}
-}
-
 void FElementSelectionStore::OnElementListSyncEvent(const FTypedElementList& InElementList, FTypedElementListLegacySync::ESyncType InSyncType, const FTypedElementHandle& InElementHandle, bool bIsWithinBatchOperation)
 {
 	if (Sink)
@@ -686,8 +664,7 @@ void USelection::ForceBatchDirty()
 
 void USelection::NoteSelectionChanged()
 {
-	ensureAlwaysMsgf(!SelectionStore->IsBatchSelecting(), TEXT("NoteSelectionChanged does not work correctly if called within a batch operation!"));
-	SelectionStore->NoteSelectionChanged();
+	USelection::SelectionChangedEvent.Broadcast(this);
 }
 
 void USelection::NoteUnknownSelectionChanged()
diff --git a/Engine/Source/Editor/UnrealEd/Public/Selection.h b/Engine/Source/Editor/UnrealEd/Public/Selection.h
--- a/Engine/Source/Editor/UnrealEd/Public/Selection.h
+++ b/Engine/Source/Editor/UnrealEd/Public/Selection.h
@@ -109,11 +109,6 @@ public:
 	 * Forcibly mark this batch as being dirty.
 	 */
 	virtual void ForceBatchDirty() = 0;
-
-	/**
-	 * Manually invoke a selection changed notification.
-	 */
-	virtual void NoteSelectionChanged() = 0;
 };
 
 } // namespace Selection_Private
