Fix PIE netmode drop down appearing hard to read. Converted to a sub-menu for consistency

diff --git a/Engine/Source/Editor/UnrealEd/Private/Kismet2/DebuggerCommands.cpp b/Engine/Source/Editor/UnrealEd/Private/Kismet2/DebuggerCommands.cpp
--- a/Engine/Source/Editor/UnrealEd/Private/Kismet2/DebuggerCommands.cpp
+++ b/Engine/Source/Editor/UnrealEd/Private/Kismet2/DebuggerCommands.cpp
@@ -203,7 +203,7 @@ public:
 	static void SetNumberOfClients(int32 NumClients, ETextCommit::Type CommitInfo);
 
 	static int32 GetNetPlayMode();
-	static void SetNetPlayMode(int32 Value, ESelectInfo::Type CommitInfo);
+	static void SetNetPlayMode(int32 Value);
 
 protected:
 
@@ -960,17 +960,24 @@ TSharedRef< SWidget > FPlayWorldCommands::GeneratePlayMenuContent(TSharedRef<FUI
 			}
 			// Net Mode
 			{
-				const UEnum* PlayNetModeEnum = FindObject<UEnum>(ANY_PACKAGE, TEXT("EPlayNetMode"));
-			
-				TSharedRef<SWidget> NetMode = SNew(SEnumComboBox, PlayNetModeEnum)
-					.CurrentValue(TAttribute<int32>::Create(TAttribute<int32>::FGetter::CreateStatic(&FInternalPlayWorldCommandCallbacks::GetNetPlayMode)))
-					.ButtonStyle(FEditorStyle::Get(), "FlatButton.Light")
-					.ContentPadding(FMargin(2, 0))
-					.Font(FEditorStyle::GetFontStyle("Sequencer.AnimationOutliner.RegularFont"))
-					.OnEnumSelectionChanged(SEnumComboBox::FOnEnumSelectionChanged::CreateStatic(&FInternalPlayWorldCommandCallbacks::SetNetPlayMode))
-					.ToolTipText(LOCTEXT("NetworkModeToolTip", "Which network mode should the clients launch in? A server will automatically be started if needed."));
-
-				Section.AddEntry(FToolMenuEntry::InitWidget("NetMode", NetMode, LOCTEXT("NetworkModeMenuWidget", "Net Mode")));
+				Section.AddSubMenu(
+					"NetMode",
+					LOCTEXT("NetworkModeMenu", "Net Mode"),
+					LOCTEXT("NetworkModeToolTip", "Which network mode should the clients launch in? A server will automatically be started if needed."),
+					FNewMenuDelegate::CreateLambda([](FMenuBuilder& InMenuBuilder)
+						{
+							const UEnum* PlayNetModeEnum = FindObject<UEnum>(ANY_PACKAGE, TEXT("EPlayNetMode"));
+
+							for (int32 i = 0; i < PlayNetModeEnum->NumEnums() - 1; i++)
+							{
+								if (PlayNetModeEnum->HasMetaData(TEXT("Hidden"), i) == false)
+								{
+									FUIAction Action(FExecuteAction::CreateStatic(&FInternalPlayWorldCommandCallbacks::SetNetPlayMode, i), FCanExecuteAction(), FIsActionChecked::CreateLambda([](int32 Index) {return FInternalPlayWorldCommandCallbacks::GetNetPlayMode() == Index; }, i));
+									InMenuBuilder.AddMenuEntry(PlayNetModeEnum->GetDisplayNameTextByIndex(i), PlayNetModeEnum->GetToolTipTextByIndex(i), FSlateIcon(), Action, NAME_None, EUserInterfaceActionType::RadioButton);
+								}
+							}
+						})
+						,false);
 			}
 		}
 
@@ -3023,7 +3030,7 @@ int32 FInternalPlayWorldCommandCallbacks::GetNetPlayMode()
 	return (int32)NetMode;
 }
 
-void FInternalPlayWorldCommandCallbacks::SetNetPlayMode(int32 Value, ESelectInfo::Type CommitInfo)
+void FInternalPlayWorldCommandCallbacks::SetNetPlayMode(int32 Value)
 {
 	ULevelEditorPlaySettings* PlayInSettings = GetMutableDefault<ULevelEditorPlaySettings>();
 	PlayInSettings->SetPlayNetMode((EPlayNetMode)Value);
diff --git a/Engine/Source/Runtime/Slate/Private/Framework/MultiBox/SWidgetBlock.cpp b/Engine/Source/Runtime/Slate/Private/Framework/MultiBox/SWidgetBlock.cpp
--- a/Engine/Source/Runtime/Slate/Private/Framework/MultiBox/SWidgetBlock.cpp
+++ b/Engine/Source/Runtime/Slate/Private/Framework/MultiBox/SWidgetBlock.cpp
@@ -111,6 +111,7 @@ void SWidgetBlock::BuildMultiBlockWidget(const ISlateStyle* StyleSet, const FNam
 				SNew( STextBlock )
 				.TextStyle(LabelStyle)
 				.Text(WidgetBlock->Label)
+				.ColorAndOpacity(FAppStyle::Get().GetSlateColor("Colors.ForegroundHover"))
 			]
 		]
 		+SHorizontalBox::Slot()
