WorldPartition:
- Added generateini commandline option to only generate the conversion congif file, useful when initially converting a map.
- Added default grid placement for actors that are in the persistent level, by default they are not always loaded, unless there was streaming levels present.
- Hardcoded conversion suffix to avoid converting the same map several times with a different suffix and leave old conversion results behind.
- Always delete the umap destination file before conversion, right after deleting associated external actors.

diff --git a/Engine/Source/Editor/UnrealEd/Classes/Commandlets/WorldPartitionConvertCommandlet.h b/Engine/Source/Editor/UnrealEd/Classes/Commandlets/WorldPartitionConvertCommandlet.h
--- a/Engine/Source/Editor/UnrealEd/Classes/Commandlets/WorldPartitionConvertCommandlet.h
+++ b/Engine/Source/Editor/UnrealEd/Classes/Commandlets/WorldPartitionConvertCommandlet.h
@@ -36,9 +36,12 @@ public:
 	virtual int32 Main(const FString& Params) override;
 	//~ End UCommandlet Interface
 
+private:
+	void GatherAndPrepareSubLevelsToConvert(ULevel* Level, TArray<ULevel*>& SubLevels);
+
 protected:
-	virtual void GetSubLevelsToConvert(ULevel* Level, TArray<ULevel*>& SubLevels);
 	virtual bool GetAdditionalLevelsToConvert(ULevel* Level, TArray<ULevel*>& SubLevels);
+	virtual bool PrepareStreamingLevelForConversion(ULevelStreaming* StreamingLevel);	
 	virtual bool ShouldDeleteActor(AActor* Actor, bool bMainLevel) const;
 	virtual void PerformAdditionalWorldCleanup(UWorld* World) const;
 	virtual void OutputConversionReport() const;
@@ -81,6 +84,7 @@ protected:
 
 	bool bNoSourceControl;
 	bool bDeleteSourceLevels;
+	bool bGenerateIni;
 	bool bReportOnly;
 	bool bVerbose;
 	bool bConversionSuffix;
diff --git a/Engine/Source/Editor/UnrealEd/Private/Commandlets/WorldPartitionConvertCommandlet.cpp b/Engine/Source/Editor/UnrealEd/Private/Commandlets/WorldPartitionConvertCommandlet.cpp
--- a/Engine/Source/Editor/UnrealEd/Private/Commandlets/WorldPartitionConvertCommandlet.cpp
+++ b/Engine/Source/Editor/UnrealEd/Private/Commandlets/WorldPartitionConvertCommandlet.cpp
@@ -142,8 +142,8 @@ public:
 UWorldPartitionConvertCommandlet::UWorldPartitionConvertCommandlet(const FObjectInitializer& ObjectInitializer)
 : Super(ObjectInitializer)
 , bConversionSuffix(false)
-{
-}
+, ConversionSuffix(TEXT("_WP"))
+{}
 
 ULevel* UWorldPartitionConvertCommandlet::LoadLevel(const FString& LevelToLoad)
 {
@@ -212,10 +212,10 @@ UWorldPartition* UWorldPartitionConvertCommandlet::CreateWorldPartition(AWorldSe
 	return WorldPartition;
 }
 
-void UWorldPartitionConvertCommandlet::GetSubLevelsToConvert(ULevel* Level, TArray<ULevel*>& SubLevels)
+void UWorldPartitionConvertCommandlet::GatherAndPrepareSubLevelsToConvert(ULevel* Level, TArray<ULevel*>& SubLevels)
 {
-	UWorld* World = Level->GetTypedOuter<UWorld>();
-	
+	UWorld* World = Level->GetTypedOuter<UWorld>();	
+
 	// Set all streaming levels to be loaded/visible for next Flush
 	for (ULevelStreaming* StreamingLevel : World->GetStreamingLevels())
 	{
@@ -227,36 +227,47 @@ void UWorldPartitionConvertCommandlet::GetSubLevelsToConvert(ULevel* Level, TArr
 	
 	for(ULevelStreaming* StreamingLevel: World->GetStreamingLevels())
 	{
-		if (StreamingLevel->bDisableDistanceStreaming)
+		if (PrepareStreamingLevelForConversion(StreamingLevel))
 		{
-			UE_LOG(LogWorldPartitionConvertCommandlet, Warning, TEXT("Converting non distance-based streaming level %s"), *StreamingLevel->GetWorldAssetPackageName());
+			ULevel* SubLevel = StreamingLevel->GetLoadedLevel();
+			check(SubLevel);
+
+			SubLevels.Add(SubLevel);
+
+			// Recursively obtain sub levels to convert
+			GatherAndPrepareSubLevelsToConvert(SubLevel, SubLevels);
 		}
+	}
+}
+
+bool UWorldPartitionConvertCommandlet::PrepareStreamingLevelForConversion(ULevelStreaming* StreamingLevel)
+{
+	if (StreamingLevel->bDisableDistanceStreaming)
+	{
+		UE_LOG(LogWorldPartitionConvertCommandlet, Warning, TEXT("Converting non distance-based streaming level %s"), *StreamingLevel->GetWorldAssetPackageName());
+	}
+
+	if (StreamingLevel->ShouldBeAlwaysLoaded())
+	{
+		UE_LOG(LogWorldPartitionConvertCommandlet, Log, TEXT("Converting always loaded streaming level %s"), *StreamingLevel->GetWorldAssetPackageName());
 
 		ULevel* SubLevel = StreamingLevel->GetLoadedLevel();
 		check(SubLevel);
-	
-		if (StreamingLevel->ShouldBeAlwaysLoaded())
-		{
-			UE_LOG(LogWorldPartitionConvertCommandlet, Log, TEXT("Converting always loaded streaming level %s"), *StreamingLevel->GetWorldAssetPackageName());
 
-			for (AActor* Actor: SubLevel->Actors)
+		for (AActor* Actor: SubLevel->Actors)
+		{
+			if (Actor)
 			{
-				if (Actor)
-				{
-					Actor->GridPlacement = EActorGridPlacement::AlwaysLoaded;
-				}
+				Actor->GridPlacement = EActorGridPlacement::AlwaysLoaded;
 			}
 		}
-		else
-		{
-			UE_LOG(LogWorldPartitionConvertCommandlet, Log, TEXT("Converting dynamic streaming level %s"), *StreamingLevel->GetWorldAssetPackageName());
-		}
-
-		SubLevels.Add(SubLevel);
-
-		// Recursively obtain sub levels to convert
-		GetSubLevelsToConvert(SubLevel, SubLevels);
 	}
+	else
+	{
+		UE_LOG(LogWorldPartitionConvertCommandlet, Log, TEXT("Converting dynamic streaming level %s"), *StreamingLevel->GetWorldAssetPackageName());
+	}
+
+	return true;
 }
 
 bool UWorldPartitionConvertCommandlet::GetAdditionalLevelsToConvert(ULevel* Level, TArray<ULevel*>& SubLevels)
@@ -290,6 +301,8 @@ bool UWorldPartitionConvertCommandlet::ShouldDeleteActor(AActor* Actor, bool bMa
 
 bool UWorldPartitionConvertCommandlet::DeleteFile(const FString& Filename)
 {
+	UE_LOG(LogWorldPartitionConvertCommandlet, Verbose, TEXT("Deleting %s"), *Filename);
+
 	if (!UseSourceControl())
 	{
 		if (!IPlatformFile::GetPlatformPhysical().SetReadOnly(*Filename, false) ||
@@ -756,9 +769,10 @@ int32 UWorldPartitionConvertCommandlet::Main(const FString& Params)
 
 	bNoSourceControl = Switches.Contains(TEXT("NoSourceControl"));
 	bDeleteSourceLevels = Switches.Contains(TEXT("DeleteSourceLevels"));
-	bReportOnly = Switches.Contains(TEXT("ReportOnly"));
+	bGenerateIni = Switches.Contains(TEXT("GenerateIni"));
+	bReportOnly = bGenerateIni || Switches.Contains(TEXT("ReportOnly"));
 	bVerbose = Switches.Contains(TEXT("Verbose"));
-	
+
 	ReadAdditionalTokensAndSwitches(Tokens, Switches);
 
 	if (bVerbose)
@@ -766,7 +780,7 @@ int32 UWorldPartitionConvertCommandlet::Main(const FString& Params)
 		LogWorldPartitionConvertCommandlet.SetVerbosity(ELogVerbosity::Verbose);
 	}
 
-	bConversionSuffix = FParse::Value(*Params, TEXT("ConversionSuffix="), ConversionSuffix);
+	bConversionSuffix = Switches.Contains(TEXT("ConversionSuffix"));
 
 	FScopedSourceControl SourceControl;
 	SourceControlProvider = bNoSourceControl ? nullptr : &ISourceControlModule::Get().GetProvider();
@@ -802,8 +816,8 @@ int32 UWorldPartitionConvertCommandlet::Main(const FString& Params)
 
 	SetupHLODLayerAssets();
 
-	// Delete existing result from running the commandlet
-	if (bConversionSuffix && !bReportOnly)
+	// Delete existing result from running the commandlet, even if not using the suffix mode to cleanup previous conversion
+	if (!bReportOnly)
 	{
 		UE_SCOPED_TIMER(TEXT("Deleting existing conversion results"), LogWorldPartitionConvertCommandlet, Display);
 
@@ -825,6 +839,11 @@ int32 UWorldPartitionConvertCommandlet::Main(const FString& Params)
 				return 1;
 			}
 		}
+
+		if (FPackageName::SearchForPackageOnDisk(OldLevelName, &OldLevelName))
+		{
+			DeleteFile(SourceControlHelpers::PackageFilename(OldLevelName));
+		}
 	}
 
 	// Load persistent level
@@ -847,7 +866,7 @@ int32 UWorldPartitionConvertCommandlet::Main(const FString& Params)
 
 	OnWorldLoaded(MainWorld);
 
-	auto DeleteUnrequiredActors = [this](ULevel* Level, bool bMainLevel, EActorGridPlacement DefaultGridPlacement)
+	auto PrepareLevelActors = [this](ULevel* Level, bool bMainLevel, EActorGridPlacement DefaultGridPlacement)
 	{
 		for (AActor* Actor: Level->Actors)
 		{
@@ -866,8 +885,18 @@ int32 UWorldPartitionConvertCommandlet::Main(const FString& Params)
 		}
 	};
 
+	// Gather sublevels
+	TArray<ULevel*> SubLevelsToConvert;
+	GatherAndPrepareSubLevelsToConvert(MainLevel, SubLevelsToConvert);
+
+	if (!GetAdditionalLevelsToConvert(MainLevel, SubLevelsToConvert))
+	{
+		return 1;
+	}
+
+	// Prepare levels for conversion
 	DetachDependantLevelPackages(MainLevel);
-	DeleteUnrequiredActors(MainLevel, true, EActorGridPlacement::AlwaysLoaded);
+	PrepareLevelActors(MainLevel, true, SubLevelsToConvert.Num() ? EActorGridPlacement::AlwaysLoaded : EActorGridPlacement::Bounds);
 
 	if (bConversionSuffix)
 	{
@@ -885,19 +914,10 @@ int32 UWorldPartitionConvertCommandlet::Main(const FString& Params)
 		RemapSoftObjectPaths.Add(OldPackagePath, FSoftObjectPath(MainPackage).ToString());
 	}
 
-	// Gather sublevels
-	TArray<ULevel*> SubLevelsToConvert;
-	GetSubLevelsToConvert(MainLevel, SubLevelsToConvert);
-
-	if (!GetAdditionalLevelsToConvert(MainLevel, SubLevelsToConvert))
-	{
-		return 1;
-	}
-
 	for (ULevel* SubLevel : SubLevelsToConvert)
 	{
 		DetachDependantLevelPackages(SubLevel);
-		DeleteUnrequiredActors(SubLevel, false, EActorGridPlacement::Location);
+		PrepareLevelActors(SubLevel, false, EActorGridPlacement::Bounds);
 	}
 
 	TMap<UObject*, UObject*> PrivateRefsMap;
@@ -1034,6 +1054,15 @@ int32 UWorldPartitionConvertCommandlet::Main(const FString& Params)
 	ActorList.Append(ChildActorList);
 	ChildActorList.Empty();
 
+	// Setup the world partition object
+	UWorldPartition* WorldPartition = CreateWorldPartition(MainWorldSettings);
+	if (!WorldPartition)
+	{
+		return 1;
+	}
+
+	WorldPartition->AddToRoot();
+
 	if (!bReportOnly)
 	{
 		for(AActor* Actor: ActorList)
@@ -1061,13 +1090,6 @@ int32 UWorldPartitionConvertCommandlet::Main(const FString& Params)
 
 		MainLevel->bUseExternalActors = true;
 
-		// Setup the world partition object
-		UWorldPartition* WorldPartition = CreateWorldPartition(MainWorldSettings);
-		if (!WorldPartition)
-		{
-			return 1;
-		}
-
 		MainWorldSettings->SetWorldPartition(WorldPartition);
 		MainLevel->bIsPartitioned = true;
 
@@ -1138,6 +1160,11 @@ int32 UWorldPartitionConvertCommandlet::Main(const FString& Params)
 			GEditor->CleanupPhysicsSceneThatWasInitializedForSave(MainWorld, bForceInitializeWorld);
 		}
 
+		UE_LOG(LogWorldPartitionConvertCommandlet, Log, TEXT("######## CONVERSION COMPLETED SUCCESSFULLY ########"));
+	}
+
+	if (bGenerateIni || !bReportOnly)
+	{
 		if (!FPlatformFileManager::Get().GetPlatformFile().FileExists(*LevelConfigFilename))
 		{
 			SaveConfig(CPF_Config, *LevelConfigFilename);
@@ -1149,8 +1176,6 @@ int32 UWorldPartitionConvertCommandlet::Main(const FString& Params)
 				Pair.Value->SaveConfig(CPF_Config, *LevelConfigFilename);
 			}
 		}
-
-		UE_LOG(LogWorldPartitionConvertCommandlet, Log, TEXT("######## CONVERSION COMPLETED SUCCESSFULLY ########"));
 	}
 
 	UPackage::WaitForAsyncFileWrites();
