Fix class picker having squashed documentation links and some very minor styling cleanup

diff --git a/Engine/Source/Editor/UnrealEd/Private/SClassPickerDialog.cpp b/Engine/Source/Editor/UnrealEd/Private/SClassPickerDialog.cpp
--- a/Engine/Source/Editor/UnrealEd/Private/SClassPickerDialog.cpp
+++ b/Engine/Source/Editor/UnrealEd/Private/SClassPickerDialog.cpp
@@ -83,21 +83,30 @@ void SClassPickerDialog::Construct(const FArguments& InArgs)
 	ChildSlot
 	[
 		SNew(SBorder)
-		.Visibility(EVisibility::Visible)
-		.BorderImage(FEditorStyle::GetBrush("Menu.Background"))
+		.BorderImage(FEditorStyle::GetBrush("Brushes.Background"))
 		[
 			SNew(SBox)
-			.Visibility(EVisibility::Visible)
-			.WidthOverride(520.0f)
+			.WidthOverride(610.0f)
 			[
 				SNew(SVerticalBox)
 				+SVerticalBox::Slot()
 				.AutoHeight()
 				[
 					SNew(SExpandableArea)
+					.BorderImage(FStyleDefaults::GetNoBrush())
+					.BodyBorderImage(FAppStyle::Get().GetBrush("Brushes.Recessed"))
+					.HeaderPadding(FMargin(5.0f, 3.0f))
+					.AllowAnimatedTransition(false)
 					.InitiallyCollapsed(!bExpandDefaultClassPicker)
-					.AreaTitle(NSLOCTEXT("SClassPickerDialog", "CommonClassesAreaTitle", "Common Classes"))
 					.OnAreaExpansionChanged(this, &SClassPickerDialog::OnDefaultAreaExpansionChanged)
+					.HeaderContent()
+					[
+						SNew(STextBlock)
+						.Text(NSLOCTEXT("SClassPickerDialog", "CommonClassesAreaTitle", "Common"))
+						.TextStyle(FAppStyle::Get(), "ButtonText")
+						.Font(FAppStyle::Get().GetFontStyle("NormalFontBold"))
+						.TransformPolicy(ETextTransformPolicy::ToUpper)
+					]
 					.BodyContent()
 					[
 						SAssignNew(DefaultClassViewer, SListView < TSharedPtr<FClassPickerDefaults> >)
@@ -113,17 +122,27 @@ void SClassPickerDialog::Construct(const FArguments& InArgs)
 				.Padding(0.0f, 10.0f, 0.0f, 0.0f)
 				[
 					SNew(SExpandableArea)
+					.BorderImage(FStyleDefaults::GetNoBrush())
+					.BodyBorderImage(FAppStyle::Get().GetBrush("Brushes.Recessed"))
+					.HeaderPadding(FMargin(5.0f, 3.0f))
+					.AllowAnimatedTransition(false)
 					.MaxHeight(320.f)
 					.InitiallyCollapsed(!bExpandCustomClassPicker)
-					.AreaTitle(NSLOCTEXT("SClassPickerDialog", "AllClassesAreaTitle", "All Classes"))
 					.OnAreaExpansionChanged(this, &SClassPickerDialog::OnCustomAreaExpansionChanged)
+					.HeaderContent()
+					[
+						SNew(STextBlock)
+						.Text(NSLOCTEXT("SClassPickerDialog", "AllClassesAreaTitle", "All Classes"))
+						.TextStyle(FAppStyle::Get(), "ButtonText")
+						.Font(FAppStyle::Get().GetFontStyle("NormalFontBold"))
+						.TransformPolicy(ETextTransformPolicy::ToUpper)
+					]
 					.BodyContent()
 					[
 						ClassViewer.ToSharedRef()
 					]
 				]
 				+SVerticalBox::Slot()
-				.AutoHeight()
 				.HAlign(HAlign_Right)
 				.VAlign(VAlign_Bottom)
 				.Padding(8)
@@ -136,20 +155,15 @@ void SClassPickerDialog::Construct(const FArguments& InArgs)
 						.Text(NSLOCTEXT("SClassPickerDialog", "ClassPickerSelectButton", "Select"))
 						.HAlign(HAlign_Center)
 						.Visibility( this, &SClassPickerDialog::GetSelectButtonVisibility )
-						.ContentPadding(FEditorStyle::GetMargin("StandardDialog.ContentPadding"))
 						.OnClicked(this, &SClassPickerDialog::OnClassPickerConfirmed)
-						.ButtonStyle(FEditorStyle::Get(), "FlatButton.Success")
-						.TextStyle(FEditorStyle::Get(), "FlatButton.DefaultTextStyle")
+						.ButtonStyle(FEditorStyle::Get(), "PrimaryButton")
 					]
 					+SUniformGridPanel::Slot(1,0)
 					[
 						SNew(SButton)
 						.Text(NSLOCTEXT("SClassPickerDialog", "ClassPickerCancelButton", "Cancel"))
 						.HAlign(HAlign_Center)
-						.ContentPadding(FEditorStyle::GetMargin("StandardDialog.ContentPadding"))
 						.OnClicked(this, &SClassPickerDialog::OnClassPickerCanceled)
-						.ButtonStyle(FEditorStyle::Get(), "FlatButton.Default")
-						.TextStyle(FEditorStyle::Get(), "FlatButton.DefaultTextStyle")
 					]
 				]
 			]
@@ -222,7 +236,7 @@ TSharedRef<ITableRow> SClassPickerDialog::GenerateListRow(TSharedPtr<FClassPicke
 		[
 			SNew(SHorizontalBox)
 			+SHorizontalBox::Slot()
-			.FillWidth(0.45f)
+			.FillWidth(0.5f)
 			[
 				SNew(SButton)
 				.OnClicked(this, &SClassPickerDialog::OnDefaultClassPicked, ItemClass)
@@ -231,9 +245,8 @@ TSharedRef<ITableRow> SClassPickerDialog::GenerateListRow(TSharedPtr<FClassPicke
 				[
 					SNew(SHorizontalBox)
 					+SHorizontalBox::Slot()
-					.HAlign(HAlign_Center)
+					.AutoWidth()
 					.VAlign(VAlign_Center)
-					.FillWidth(0.12f)
 					[
 						SNew(SImage)
 						.Image(ItemBrush)
@@ -241,7 +254,6 @@ TSharedRef<ITableRow> SClassPickerDialog::GenerateListRow(TSharedPtr<FClassPicke
 					+SHorizontalBox::Slot()
 					.VAlign(VAlign_Center)
 					.Padding(4.0f, 0.0f)
-					.FillWidth(0.8f)
 					[
 						SNew(STextBlock)
 						.Text(Obj->GetName())
@@ -257,14 +269,10 @@ TSharedRef<ITableRow> SClassPickerDialog::GenerateListRow(TSharedPtr<FClassPicke
 				.AutoWrapText(true)
 			]
 			+SHorizontalBox::Slot()
+			.VAlign(VAlign_Center)
 			.AutoWidth()
 			[
-				SNew(SBox)
-				.WidthOverride(32)
-				.HeightOverride(32)
-				[
-					FEditorClassUtils::GetDocumentationLinkWidget(ItemClass)
-				]
+				FEditorClassUtils::GetDocumentationLinkWidget(ItemClass)
 			]
 		]
 	];
