Fix data table row importer not showing selected row type.

diff --git a/Engine/Source/Editor/UnrealEd/Private/DataTableEditorUtils.cpp b/Engine/Source/Editor/UnrealEd/Private/DataTableEditorUtils.cpp
--- a/Engine/Source/Editor/UnrealEd/Private/DataTableEditorUtils.cpp
+++ b/Engine/Source/Editor/UnrealEd/Private/DataTableEditorUtils.cpp
@@ -38,9 +38,9 @@ public:
 		: _Content()
 		, _ComboBoxStyle(&FCoreStyle::Get().GetWidgetStyle< FComboBoxStyle >("ComboBox"))
 		, _ButtonStyle(nullptr)
-		, _ItemStyle(&FCoreStyle::Get().GetWidgetStyle< FTableRowStyle >("TableView.Row"))
-		, _ContentPadding(FMargin(4.0, 2.0))
-		, _ForegroundColor(FCoreStyle::Get().GetSlateColor("InvertedForeground"))
+		, _ItemStyle(&FAppStyle::Get().GetWidgetStyle<FTableRowStyle>("ComboBox.Row"))
+		, _ContentPadding(_ComboBoxStyle->ContentPadding)
+		, _ForegroundColor(FSlateColor::UseStyle())
 		, _OnStructSelected()
 		, _InitiallySelectedItem(nullptr)
 		, _Method()
@@ -138,7 +138,8 @@ private:
 private:
 	/** Delegate that is invoked when the selected item in the combo box changes */
 	FDataTableEditorUtils::FOnDataTableStructSelected OnStructSelected;
-
+	/** The padding around each menu row */
+	FMargin MenuRowPadding;
 	/** The item currently selected in the combo box */
 	TSharedPtr<FString> SelectedItem;
 	/** The search field used for the combox box's contents */
@@ -160,6 +161,8 @@ void SDataTableStructComboBox::Construct(const FArguments& InArgs)
 
 	ItemStyle = InArgs._ItemStyle;
 
+	MenuRowPadding = InArgs._ComboBoxStyle->MenuRowPadding;
+
 	// Work out which values we should use based on whether we were given an override, or should use the style's version
 	const FComboButtonStyle& OurComboButtonStyle = InArgs._ComboBoxStyle->ComboButtonStyle;
 	const FButtonStyle* const OurButtonStyle = InArgs._ButtonStyle ? InArgs._ButtonStyle : &OurComboButtonStyle.ButtonStyle;
@@ -219,6 +222,7 @@ void SDataTableStructComboBox::Construct(const FArguments& InArgs)
 		.IsFocusable(true)
 		);
 
+
 	// Better to select search field so you can type right away
 	SetMenuContentWidgetToFocus(SearchField);
 
@@ -336,6 +340,7 @@ TSharedRef<ITableRow> SDataTableStructComboBox::GenerateMenuItemRow(TSharedPtr<F
 	return SNew(SComboRow<TSharedPtr<FString>>, OwnerTable)
 		.Style(ItemStyle)
 		.Visibility(WidgetVisibility)
+		.Padding(MenuRowPadding)
 		[
 			SNew(STextBlock)
 			.Text(FText::FromString(*InItem))
@@ -439,8 +444,7 @@ const FString FDataTableEditorUtils::VariableTypesTooltipDocLink = TEXT("Shared/
 TSharedRef<SWidget> FDataTableEditorUtils::MakeRowStructureComboBox(FOnDataTableStructSelected OnSelected)
 {
 	TSharedRef<SDataTableStructComboBox> ComboBox = SNew(SDataTableStructComboBox)
-		.OnStructSelected(OnSelected)
-		.ContentPadding(3);
+		.OnStructSelected(OnSelected);
 
 	return ComboBox;
 }
diff --git a/Engine/Source/Editor/UnrealEd/Private/SCSVImportOptions.cpp b/Engine/Source/Editor/UnrealEd/Private/SCSVImportOptions.cpp
--- a/Engine/Source/Editor/UnrealEd/Private/SCSVImportOptions.cpp
+++ b/Engine/Source/Editor/UnrealEd/Private/SCSVImportOptions.cpp
@@ -56,16 +56,15 @@ void SCSVImportOptions::Construct(const FArguments& InArgs)
 	this->ChildSlot
 	[
 		SNew(SBorder)
-		.BorderImage(FEditorStyle::GetBrush(TEXT("Menu.Background")))
+		.BorderImage(FAppStyle::Get().GetBrush(TEXT("Brushes.Background")))
 		.Padding(10)
 		[
 			SNew(SVerticalBox)
 			+SVerticalBox::Slot()
 			.AutoHeight()
 			[
-				SNew(SBorder)
+				SNew(SBox)
 				.Padding(FMargin(3))
-				.BorderImage(FEditorStyle::GetBrush("ToolPanel.GroupBorder"))
 				.Visibility( InArgs._FullPath.IsEmpty() ? EVisibility::Collapsed : EVisibility::Visible )
 				[
 					SNew(SHorizontalBox)
@@ -73,8 +72,7 @@ void SCSVImportOptions::Construct(const FArguments& InArgs)
 					.AutoWidth()
 					[
 						SNew(STextBlock)
-						.Font(FEditorStyle::GetFontStyle("CurveEd.LabelFont"))
-						.Text(LOCTEXT("Import_CurrentFileTitle", "Current File: "))
+						.Text(LOCTEXT("Import_CurrentFileTitle", "Current File "))
 					]
 					+SHorizontalBox::Slot()
 					.Padding(5, 0, 0, 0)
@@ -82,7 +80,6 @@ void SCSVImportOptions::Construct(const FArguments& InArgs)
 					.VAlign(VAlign_Center)
 					[
 						SNew(STextBlock)
-						.Font(FEditorStyle::GetFontStyle("CurveEd.InfoFont"))
 						.Text(InArgs._FullPath)
 					]
 				]
@@ -94,7 +91,7 @@ void SCSVImportOptions::Construct(const FArguments& InArgs)
 			.Padding(2)
 			[
 				SNew(STextBlock)
-				.Text( LOCTEXT("ChooseAssetType", "Import As:") )
+				.Text( LOCTEXT("ChooseAssetType", "Import As") )
 			]
 			+SVerticalBox::Slot()
 			.AutoHeight()
@@ -114,7 +111,7 @@ void SCSVImportOptions::Construct(const FArguments& InArgs)
 			.Padding(2)
 			[
 				SNew(STextBlock)
-				.Text( LOCTEXT("ChooseRowType", "Choose DataTable Row Type:") )
+				.Text( LOCTEXT("ChooseRowType", "Choose DataTable Row Type") )
 				.Visibility( this, &SCSVImportOptions::GetTableRowOptionVis )
 			]
 			+SVerticalBox::Slot()
@@ -128,7 +125,7 @@ void SCSVImportOptions::Construct(const FArguments& InArgs)
 			.Padding(2)
 			[
 				SNew(STextBlock)
-				.Text( LOCTEXT("ChooseCurveType", "Choose Curve Interpolation Type:") )
+				.Text( LOCTEXT("ChooseCurveType", "Choose Curve Interpolation Type") )
 				.Visibility( this, &SCSVImportOptions::GetCurveTypeVis )
 			]
 			+SVerticalBox::Slot()
@@ -161,11 +158,12 @@ void SCSVImportOptions::Construct(const FArguments& InArgs)
 			[
 				SNew(SHorizontalBox)
 				+SHorizontalBox::Slot()
-				.AutoWidth()
+				.HAlign(HAlign_Right)
 				.Padding(2)
 				[
 					SNew(SButton)
 					.Text(LOCTEXT("Import", "Apply"))
+					.ButtonStyle(FAppStyle::Get(), "PrimaryButton")
 					.OnClicked( this, &SCSVImportOptions::OnImport )
 					.IsEnabled( this, &SCSVImportOptions::CanImport )
 				]
diff --git a/Engine/Source/Runtime/SlateCore/Private/Styling/StarshipCoreStyle.cpp b/Engine/Source/Runtime/SlateCore/Private/Styling/StarshipCoreStyle.cpp
--- a/Engine/Source/Runtime/SlateCore/Private/Styling/StarshipCoreStyle.cpp
+++ b/Engine/Source/Runtime/SlateCore/Private/Styling/StarshipCoreStyle.cpp
@@ -1152,9 +1152,9 @@ void FStarshipCoreStyle::SetupComboButtonStyles(TSharedRef<FStyle>& Style)
 		.SetActiveHoveredBrush(FSlateColorBrush(FStyleColors::PrimaryHover))
 		.SetActiveHighlightedBrush(FSlateColorBrush(FStyleColors::PrimaryHover))
 
-		.SetInactiveBrush(FSlateNoResource())
-		.SetInactiveHoveredBrush(FSlateNoResource())
-		.SetInactiveHighlightedBrush(FSlateNoResource())
+		.SetInactiveBrush(FSlateColorBrush(FStyleColors::Primary))
+		.SetInactiveHoveredBrush(FSlateColorBrush(FStyleColors::PrimaryHover))
+		.SetInactiveHighlightedBrush(FSlateColorBrush(FStyleColors::PrimaryHover))
 
 
 		.SetTextColor(FStyleColors::White)
