WorldPartition:
- Basic support for converting world using World Composition to World Partition in WorldPartitionConvertCommandlet
- On top of the default runtime grid and/or config file, each WorldComposition Layer will create an extra runtime grid based on the world composition layer name and streaming distance

diff --git a/Engine/Source/Editor/UnrealEd/Classes/Commandlets/WorldPartitionConvertCommandlet.h b/Engine/Source/Editor/UnrealEd/Classes/Commandlets/WorldPartitionConvertCommandlet.h
--- a/Engine/Source/Editor/UnrealEd/Classes/Commandlets/WorldPartitionConvertCommandlet.h
+++ b/Engine/Source/Editor/UnrealEd/Classes/Commandlets/WorldPartitionConvertCommandlet.h
@@ -14,6 +14,9 @@
 
 UNREALED_API DECLARE_LOG_CATEGORY_EXTERN(LogWorldPartitionConvertCommandlet, Log, All);
 
+class UWorldPartition;
+class UWorldComposition;
+
 USTRUCT()
 struct FHLODLayerActorMapping
 {
@@ -37,18 +40,18 @@ public:
 	//~ End UCommandlet Interface
 
 private:
-	void GatherAndPrepareSubLevelsToConvert(ULevel* Level, TArray<ULevel*>& SubLevels);
+	void GatherAndPrepareSubLevelsToConvert(const UWorldPartition* WorldPartition, ULevel* Level, TArray<ULevel*>& SubLevels);
 
 protected:
 	virtual bool GetAdditionalLevelsToConvert(ULevel* Level, TArray<ULevel*>& SubLevels);
-	virtual bool PrepareStreamingLevelForConversion(ULevelStreaming* StreamingLevel);	
+	virtual bool PrepareStreamingLevelForConversion(const UWorldPartition* WorldPartition, ULevelStreaming* StreamingLevel);
 	virtual bool ShouldDeleteActor(AActor* Actor, bool bMainLevel) const;
 	virtual void PerformAdditionalWorldCleanup(UWorld* World) const;
 	virtual void OutputConversionReport() const;
-	virtual void OnWorldLoaded(UWorld* World) {}
+	virtual void OnWorldLoaded(UWorld* World);
 	virtual void ReadAdditionalTokensAndSwitches(const TArray<FString>& Tokens, const TArray<FString>& Switches) {}
 
-	class UWorldPartition* CreateWorldPartition(class AWorldSettings* MainWorldSettings) const;
+	UWorldPartition* CreateWorldPartition(class AWorldSettings* MainWorldSettings, UWorldComposition* WorldComposition) const;
 	ULevel* LoadLevel(const FString& LevelToLoad);
 
 	bool UseSourceControl() const { return SourceControlProvider != nullptr; }
diff --git a/Engine/Source/Editor/UnrealEd/Private/Commandlets/WorldPartitionConvertCommandlet.cpp b/Engine/Source/Editor/UnrealEd/Private/Commandlets/WorldPartitionConvertCommandlet.cpp
--- a/Engine/Source/Editor/UnrealEd/Private/Commandlets/WorldPartitionConvertCommandlet.cpp
+++ b/Engine/Source/Editor/UnrealEd/Private/Commandlets/WorldPartitionConvertCommandlet.cpp
@@ -40,6 +40,7 @@
 #include "ProfilingDebugging/ScopedTimers.h"
 #include "AssetRegistryModule.h"
 #include "FoliageHelper.h"
+#include "Engine/WorldComposition.h"
 
 DEFINE_LOG_CATEGORY(LogWorldPartitionConvertCommandlet);
 
@@ -188,7 +189,7 @@ ULevel* UWorldPartitionConvertCommandlet::LoadLevel(const FString& LevelToLoad)
 	return nullptr;
 }
 
-UWorldPartition* UWorldPartitionConvertCommandlet::CreateWorldPartition(AWorldSettings* MainWorldSettings) const
+UWorldPartition* UWorldPartitionConvertCommandlet::CreateWorldPartition(AWorldSettings* MainWorldSettings, UWorldComposition* WorldComposition) const
 {
 	UWorldPartition* WorldPartition = NewObject<UWorldPartition>(MainWorldSettings);
 
@@ -209,10 +210,12 @@ UWorldPartition* UWorldPartitionConvertCommandlet::CreateWorldPartition(AWorldSe
 		WorldPartition->DefaultHLODLayer = nullptr;
 	}
 
+	WorldPartition->RuntimeHash->ImportFromWorldComposition(WorldComposition);
+
 	return WorldPartition;
 }
 
-void UWorldPartitionConvertCommandlet::GatherAndPrepareSubLevelsToConvert(ULevel* Level, TArray<ULevel*>& SubLevels)
+void UWorldPartitionConvertCommandlet::GatherAndPrepareSubLevelsToConvert(const UWorldPartition* WorldPartition, ULevel* Level, TArray<ULevel*>& SubLevels)
 {
 	UWorld* World = Level->GetTypedOuter<UWorld>();	
 
@@ -227,7 +230,7 @@ void UWorldPartitionConvertCommandlet::GatherAndPrepareSubLevelsToConvert(ULevel
 	
 	for(ULevelStreaming* StreamingLevel: World->GetStreamingLevels())
 	{
-		if (PrepareStreamingLevelForConversion(StreamingLevel))
+		if (PrepareStreamingLevelForConversion(WorldPartition, StreamingLevel))
 		{
 			ULevel* SubLevel = StreamingLevel->GetLoadedLevel();
 			check(SubLevel);
@@ -235,24 +238,19 @@ void UWorldPartitionConvertCommandlet::GatherAndPrepareSubLevelsToConvert(ULevel
 			SubLevels.Add(SubLevel);
 
 			// Recursively obtain sub levels to convert
-			GatherAndPrepareSubLevelsToConvert(SubLevel, SubLevels);
+			GatherAndPrepareSubLevelsToConvert(WorldPartition, SubLevel, SubLevels);
 		}
 	}
 }
 
-bool UWorldPartitionConvertCommandlet::PrepareStreamingLevelForConversion(ULevelStreaming* StreamingLevel)
+bool UWorldPartitionConvertCommandlet::PrepareStreamingLevelForConversion(const UWorldPartition* WorldPartition, ULevelStreaming* StreamingLevel)
 {
-	if (StreamingLevel->bDisableDistanceStreaming)
-	{
-		UE_LOG(LogWorldPartitionConvertCommandlet, Warning, TEXT("Converting non distance-based streaming level %s"), *StreamingLevel->GetWorldAssetPackageName());
-	}
+	ULevel* SubLevel = StreamingLevel->GetLoadedLevel();
+	check(SubLevel);
 
-	if (StreamingLevel->ShouldBeAlwaysLoaded())
+	if (StreamingLevel->ShouldBeAlwaysLoaded() || StreamingLevel->bDisableDistanceStreaming)
 	{
-		UE_LOG(LogWorldPartitionConvertCommandlet, Log, TEXT("Converting always loaded streaming level %s"), *StreamingLevel->GetWorldAssetPackageName());
-
-		ULevel* SubLevel = StreamingLevel->GetLoadedLevel();
-		check(SubLevel);
+		UE_LOG(LogWorldPartitionConvertCommandlet, Log, TEXT("Converting %s streaming level %s"), StreamingLevel->bDisableDistanceStreaming ? TEXT("non distance-based") : TEXT("always loaded"), *StreamingLevel->GetWorldAssetPackageName());
 
 		for (AActor* Actor: SubLevel->Actors)
 		{
@@ -265,6 +263,17 @@ bool UWorldPartitionConvertCommandlet::PrepareStreamingLevelForConversion(ULevel
 	else
 	{
 		UE_LOG(LogWorldPartitionConvertCommandlet, Log, TEXT("Converting dynamic streaming level %s"), *StreamingLevel->GetWorldAssetPackageName());
+
+		const UWorldPartitionRuntimeHash* RuntimeHash = WorldPartition->RuntimeHash;
+		check(RuntimeHash);
+
+		for (AActor* Actor : SubLevel->Actors)
+		{
+			if (Actor)
+			{
+				Actor->RuntimeGrid = RuntimeHash->GetActorRuntimeGrid(Actor);
+			}
+		}
 	}
 
 	return true;
@@ -747,6 +756,27 @@ void UWorldPartitionConvertCommandlet::SetupHLODLayerAssets()
 	}
 }
 
+void UWorldPartitionConvertCommandlet::OnWorldLoaded(UWorld* World)
+{
+	if (UWorldComposition* WorldComposition = World->WorldComposition)
+	{
+		// Add tiles streaming levels to world
+		World->SetStreamingLevels(WorldComposition->TilesStreaming);
+
+		// Make sure to force bDisableDistanceStreaming on streaming levels of World Composition non distance dependent tiles (for the rest of the process to handle streaming level as always loaded)
+		UWorldComposition::FTilesList& Tiles = WorldComposition->GetTilesList();
+		for (int32 TileIdx = 0; TileIdx < Tiles.Num(); TileIdx++)
+		{
+			FWorldCompositionTile& Tile = Tiles[TileIdx];
+			ULevelStreaming* StreamingLevel = WorldComposition->TilesStreaming[TileIdx];
+			if (StreamingLevel && !WorldComposition->IsDistanceDependentLevel(Tile.PackageName))
+			{
+				StreamingLevel->bDisableDistanceStreaming = true;
+			}
+		}
+	}
+}
+
 int32 UWorldPartitionConvertCommandlet::Main(const FString& Params)
 {
 	UE_SCOPED_TIMER(TEXT("Conversion"), LogWorldPartitionConvertCommandlet, Display);
@@ -866,6 +896,13 @@ int32 UWorldPartitionConvertCommandlet::Main(const FString& Params)
 
 	OnWorldLoaded(MainWorld);
 
+	// Setup the world partition object
+	UWorldPartition* WorldPartition = CreateWorldPartition(MainWorldSettings, MainWorld->WorldComposition);
+	if (!WorldPartition)
+	{
+		return 1;
+	}
+	
 	auto PrepareLevelActors = [this](ULevel* Level, bool bMainLevel, EActorGridPlacement DefaultGridPlacement)
 	{
 		for (AActor* Actor: Level->Actors)
@@ -887,7 +924,7 @@ int32 UWorldPartitionConvertCommandlet::Main(const FString& Params)
 
 	// Gather sublevels
 	TArray<ULevel*> SubLevelsToConvert;
-	GatherAndPrepareSubLevelsToConvert(MainLevel, SubLevelsToConvert);
+	GatherAndPrepareSubLevelsToConvert(WorldPartition, MainLevel, SubLevelsToConvert);
 
 	if (!GetAdditionalLevelsToConvert(MainLevel, SubLevelsToConvert))
 	{
@@ -1054,13 +1091,6 @@ int32 UWorldPartitionConvertCommandlet::Main(const FString& Params)
 	ActorList.Append(ChildActorList);
 	ChildActorList.Empty();
 
-	// Setup the world partition object
-	UWorldPartition* WorldPartition = CreateWorldPartition(MainWorldSettings);
-	if (!WorldPartition)
-	{
-		return 1;
-	}
-
 	WorldPartition->AddToRoot();
 
 	if (!bReportOnly)
@@ -1089,8 +1119,8 @@ int32 UWorldPartitionConvertCommandlet::Main(const FString& Params)
 		CollectGarbage(RF_Standalone);
 
 		MainLevel->bUseExternalActors = true;
-
 		MainWorldSettings->SetWorldPartition(WorldPartition);
+		MainWorld->WorldComposition = nullptr;
 		MainLevel->bIsPartitioned = true;
 
 		if (bDeleteSourceLevels)
diff --git a/Engine/Source/Runtime/Engine/Classes/Engine/WorldComposition.h b/Engine/Source/Runtime/Engine/Classes/Engine/WorldComposition.h
--- a/Engine/Source/Runtime/Engine/Classes/Engine/WorldComposition.h
+++ b/Engine/Source/Runtime/Engine/Classes/Engine/WorldComposition.h
@@ -88,6 +88,8 @@ class ENGINE_API UWorldComposition : public UObject
 	 *  @returns Whether streaming levels state was updated by this call
 	 */
 	bool UpdateEditorStreamingState(const FVector& InLocation);
+
+	TArray<FWorldTileLayer> GetDistanceDependentLayers() const;
 #endif// WITH_EDITOR
 
 	/**
diff --git a/Engine/Source/Runtime/Engine/Private/WorldComposition.cpp b/Engine/Source/Runtime/Engine/Private/WorldComposition.cpp
--- a/Engine/Source/Runtime/Engine/Private/WorldComposition.cpp
+++ b/Engine/Source/Runtime/Engine/Private/WorldComposition.cpp
@@ -19,6 +19,7 @@
 #include "Engine/LocalPlayer.h"
 #include "Engine/LevelStreamingDynamic.h"
 #include "Engine/AssetManager.h"
+#include "Engine/Level.h"
 
 DEFINE_LOG_CATEGORY_STATIC(LogWorldComposition, Log, All);
 
@@ -147,7 +148,12 @@ struct FWorldTilesGatherer
 		{
 			FString FullPath = FilenameOrDirectory;
 
-			if (FPaths::GetExtension(FullPath, true) == FPackageName::GetMapPackageExtension())
+			if (FPaths::GetExtension(FullPath, true) == FPackageName::GetMapPackageExtension() 
+#if WITH_EDITOR
+				// Make sure we don't gather partitioned levels in World Composition
+				&& !ULevel::GetIsLevelPartitionedFromPackage(FName(FPackageName::FilenameToLongPackageName(FullPath)))
+#endif
+				)
 			{
 				MapFilesToConsider.Emplace(MoveTemp(FullPath));
 			}
@@ -747,6 +753,19 @@ void UWorldComposition::UpdateStreamingState()
 }
 
 #if WITH_EDITOR
+TArray<FWorldTileLayer> UWorldComposition::GetDistanceDependentLayers() const
+{
+	TArray<FWorldTileLayer> Layers;
+	for (const FWorldCompositionTile& Tile : Tiles)
+	{
+		if (IsDistanceDependentLevel(Tile.PackageName))
+		{
+			Layers.AddUnique(Tile.Info.Layer);
+		}
+	}
+	return Layers;
+}
+
 bool UWorldComposition::UpdateEditorStreamingState(const FVector& InLocation)
 {
 	UWorld* OwningWorld = GetWorld();
diff --git a/Engine/Source/Runtime/Engine/Private/WorldPartition/WorldPartitionRuntimeSpatialHash.cpp b/Engine/Source/Runtime/Engine/Private/WorldPartition/WorldPartitionRuntimeSpatialHash.cpp
--- a/Engine/Source/Runtime/Engine/Private/WorldPartition/WorldPartitionRuntimeSpatialHash.cpp
+++ b/Engine/Source/Runtime/Engine/Private/WorldPartition/WorldPartitionRuntimeSpatialHash.cpp
@@ -30,6 +30,9 @@
 #include "AssetRegistryModule.h"
 #include "AssetData.h"
 
+#include "Engine/WorldComposition.h"
+#include "LevelUtils.h"
+
 extern UNREALED_API class UEditorEngine* GEditor;
 #endif
 
@@ -453,16 +456,63 @@ UWorldPartitionRuntimeSpatialHash::UWorldPartitionRuntimeSpatialHash(const FObje
 {}
 
 #if WITH_EDITOR
+FName UWorldPartitionRuntimeSpatialHash::GetActorRuntimeGrid(const AActor* Actor) const
+{
+	if (ULevel* Level = Actor ? Actor->GetLevel() : nullptr)
+	{
+		if (const FName* ActorRuntimeGrid = WorldCompositionStreamingLevelToRuntimeGrid.Find(FLevelUtils::FindStreamingLevel(Level)))
+		{
+			return *ActorRuntimeGrid;
+		}
+	}
+	return Super::GetActorRuntimeGrid(Actor);
+}
+
 void UWorldPartitionRuntimeSpatialHash::SetDefaultValues()
 {
 	FSpatialHashRuntimeGrid& MainGrid = Grids.AddDefaulted_GetRef();
-
 	MainGrid.GridName = TEXT("MainGrid");
 	MainGrid.CellSize = 3200;
 	MainGrid.LoadingRange = 25600;
 	MainGrid.DebugColor = FLinearColor::Gray;
 }
 
+void UWorldPartitionRuntimeSpatialHash::ImportFromWorldComposition(UWorldComposition* WorldComposition)
+{
+	check(IsRunningCommandlet());
+
+	if (WorldComposition)
+	{
+		const TArray<FWorldTileLayer> WorldCompositionTileLayers = WorldComposition->GetDistanceDependentLayers();
+		for (const FWorldTileLayer& Layer : WorldCompositionTileLayers)
+		{
+			FName GridName = FName(Layer.Name);
+			FSpatialHashRuntimeGrid* Grid = Algo::FindByPredicate(Grids, [GridName](const FSpatialHashRuntimeGrid& Grid) { return Grid.GridName == GridName; });
+			if (!Grid)
+			{
+				Grid = &Grids.AddDefaulted_GetRef();
+				Grid->GridName = GridName;
+				Grid->CellSize = 3200;
+				Grid->DebugColor = FLinearColor::MakeRandomColor();
+			}
+			// World Composition Layer Streaming Distance always wins over existing value (config file)
+			Grid->LoadingRange = Layer.StreamingDistance;
+		}
+
+		const UWorldComposition::FTilesList& Tiles = WorldComposition->GetTilesList();
+		for (int32 TileIdx = 0; TileIdx < Tiles.Num(); TileIdx++)
+		{
+			const FWorldCompositionTile& Tile = Tiles[TileIdx];
+			ULevelStreaming* StreamingLevel = WorldComposition->TilesStreaming[TileIdx];
+			if (StreamingLevel && WorldComposition->IsDistanceDependentLevel(Tile.PackageName))
+			{
+				// Map WorldComposition tiles streaming level to Runtime Grid
+				WorldCompositionStreamingLevelToRuntimeGrid.Add(StreamingLevel, FName(Tile.Info.Layer.Name));
+			}
+		}
+	}
+}
+
 bool UWorldPartitionRuntimeSpatialHash::GenerateStreaming(EWorldPartitionStreamingMode Mode, UWorldPartitionStreamingPolicy* StreamingPolicy)
 {
 	TRACE_CPUPROFILER_EVENT_SCOPE(UWorldPartitionRuntimeSpatialHash::GenerateStreaming);
diff --git a/Engine/Source/Runtime/Engine/Public/WorldPartition/WorldPartitionRuntimeHash.h b/Engine/Source/Runtime/Engine/Public/WorldPartition/WorldPartitionRuntimeHash.h
--- a/Engine/Source/Runtime/Engine/Public/WorldPartition/WorldPartitionRuntimeHash.h
+++ b/Engine/Source/Runtime/Engine/Public/WorldPartition/WorldPartitionRuntimeHash.h
@@ -17,9 +17,11 @@ class ENGINE_API UWorldPartitionRuntimeHash : public UObject
 
 #if WITH_EDITOR
 	virtual void SetDefaultValues() {}
+	virtual void ImportFromWorldComposition(class UWorldComposition* WorldComposition) {}
 	virtual bool GenerateStreaming(EWorldPartitionStreamingMode Mode, class UWorldPartitionStreamingPolicy* Policy) { return false; }
 	virtual void FlushStreaming() {}
 	virtual bool GenerateHLOD() { return false; }
+	virtual FName GetActorRuntimeGrid(const AActor* Actor) const { return NAME_None; }
 #endif
 
 	// Streaming interface
diff --git a/Engine/Source/Runtime/Engine/Public/WorldPartition/WorldPartitionRuntimeSpatialHash.h b/Engine/Source/Runtime/Engine/Public/WorldPartition/WorldPartitionRuntimeSpatialHash.h
--- a/Engine/Source/Runtime/Engine/Public/WorldPartition/WorldPartitionRuntimeSpatialHash.h
+++ b/Engine/Source/Runtime/Engine/Public/WorldPartition/WorldPartitionRuntimeSpatialHash.h
@@ -139,10 +139,12 @@ class ENGINE_API UWorldPartitionRuntimeSpatialHash : public UWorldPartitionRunti
 
 public:
 #if WITH_EDITOR
-	virtual void SetDefaultValues();
+	virtual void SetDefaultValues() override;
+	virtual void ImportFromWorldComposition(class UWorldComposition* WorldComposition) override;
 	virtual bool GenerateStreaming(EWorldPartitionStreamingMode Mode, class UWorldPartitionStreamingPolicy* StreamingPolicy) override;
 	virtual void FlushStreaming() override;
 	virtual bool GenerateHLOD() override;
+	virtual FName GetActorRuntimeGrid(const AActor* Actor) const override;
 #endif
 
 	// streaming interface
@@ -150,6 +152,10 @@ public:
 	virtual int32 GetStreamingCells(const TArray<FWorldPartitionStreamingSource>& Sources, TSet<const UWorldPartitionRuntimeCell*>& Cells) const override;
 	virtual void SortStreamingCellsByDistance(const TSet<const UWorldPartitionRuntimeCell*>& InCells, const TArray<FWorldPartitionStreamingSource>& InSources, TArray<const UWorldPartitionRuntimeCell*>& OutSortedCells) override;
 
+protected:
+	// Used to convert from a world using World Composition
+	TMap<class ULevelStreaming*, FName> WorldCompositionStreamingLevelToRuntimeGrid;
+
 private:
 
 #if WITH_EDITOR
