Add TGA translator to Interchange import

diff --git a/Engine/Plugins/Runtime/Interchange/Source/Import/Private/InterchangeImportPlugin.cpp b/Engine/Plugins/Runtime/Interchange/Source/Import/Private/InterchangeImportPlugin.cpp
--- a/Engine/Plugins/Runtime/Interchange/Source/Import/Private/InterchangeImportPlugin.cpp
+++ b/Engine/Plugins/Runtime/Interchange/Source/Import/Private/InterchangeImportPlugin.cpp
@@ -12,7 +12,7 @@
 #include "Texture/InterchangePNGTranslator.h"
 #include "Texture/InterchangePSDTranslator.h"
 #include "Texture/InterchangeTextureFactory.h"
-
+#include "Texture/InterchangeTGATranslator.h"
 
 DEFINE_LOG_CATEGORY(LogInterchangeImportPlugin);
 
@@ -40,6 +40,7 @@ void FInterchangeImportPlugin::StartupModule()
 	InterchangeManager.RegisterTranslator(UInterchangePCXTranslator::StaticClass());
 	InterchangeManager.RegisterTranslator(UInterchangePNGTranslator::StaticClass());
 	InterchangeManager.RegisterTranslator(UInterchangePSDTranslator::StaticClass());
+	InterchangeManager.RegisterTranslator(UInterchangeTGATranslator::StaticClass());
 
 	//Register the factories
 	InterchangeManager.RegisterFactory(UInterchangeTextureFactory::StaticClass());
diff --git a/Engine/Plugins/Runtime/Interchange/Source/Import/Private/Texture/InterchangeTGATranslator.cpp b/Engine/Plugins/Runtime/Interchange/Source/Import/Private/Texture/InterchangeTGATranslator.cpp
new file mode 100644
--- /dev/null
+++ b/Engine/Plugins/Runtime/Interchange/Source/Import/Private/Texture/InterchangeTGATranslator.cpp
@@ -0,0 +1,168 @@
+// Copyright Epic Games, Inc. All Rights Reserved. 
+#include "Texture/InterchangeTGATranslator.h"
+
+#include "IImageWrapper.h"
+#include "IImageWrapperModule.h"
+#include "LogInterchangeImportPlugin.h"
+#include "Misc/ConfigCacheIni.h"
+#include "Misc/FileHelper.h"
+#include "Misc/Paths.h"
+#include "Modules/ModuleManager.h"
+#include "Nodes/BaseNodeContainer.h"
+#include "TextureNode.h"
+#include "TgaImageSupport.h"
+
+
+namespace TGATranslatorImpl
+{
+	bool DecompressTGA( const FTGAFileHeader* TGA, Interchange::FImportImage& OutImage )
+	{
+		if (TGA->ColorMapType == 1 && TGA->ImageTypeCode == 1 && TGA->BitsPerPixel == 8)
+		{
+			// Notes: The Scaleform GFx exporter (dll) strips all font glyphs into a single 8-bit texture.
+			// The targa format uses this for a palette index; GFx uses a palette of (i,i,i,i) so the index
+			// is also the alpha value.
+			//
+			// We store the image as PF_G8, where it will be used as alpha in the Glyph shader.
+			OutImage.Init2DWithOneMip(
+				TGA->Width,
+				TGA->Height,
+				TSF_G8);
+			OutImage.CompressionSettings = TC_Grayscale;
+		}
+		else if(TGA->ColorMapType == 0 && TGA->ImageTypeCode == 3 && TGA->BitsPerPixel == 8)
+		{
+			// standard grayscale images
+			OutImage.Init2DWithOneMip(
+				TGA->Width,
+				TGA->Height,
+				TSF_G8);
+			OutImage.CompressionSettings = TC_Grayscale;
+		}
+		else
+		{
+			if(TGA->ImageTypeCode == 10) // 10 = RLE compressed 
+			{
+				if( TGA->BitsPerPixel != 32 &&
+					TGA->BitsPerPixel != 24 &&
+					TGA->BitsPerPixel != 16 )
+				{
+					UE_LOG( LogInterchangeImportPlugin, Error, TEXT("TGA uses an unsupported rle-compressed bit-depth: %u"), TGA->BitsPerPixel );
+					return false;
+				}
+			}
+			else
+			{
+				if( TGA->BitsPerPixel != 32 &&
+					TGA->BitsPerPixel != 16 &&
+					TGA->BitsPerPixel != 24)
+				{
+					UE_LOG( LogInterchangeImportPlugin, Error, TEXT("TGA uses an unsupported bit-depth: %u"), TGA->BitsPerPixel );
+					return false;
+				}
+			}
+
+			OutImage.Init2DWithOneMip(
+				TGA->Width,
+				TGA->Height,
+				TSF_BGRA8);
+		}
+
+		return true;
+	}
+}
+
+bool UInterchangeTGATranslator::CanImportSourceData(const UInterchangeSourceData* InSourceData) const
+{
+	FString Extension = FPaths::GetExtension(InSourceData->GetFilename());
+	FString TGAExtension = (TEXT("tga;Texture"));
+	return TGAExtension.StartsWith(Extension);
+}
+
+bool UInterchangeTGATranslator::Translate(const UInterchangeSourceData* SourceData, Interchange::FBaseNodeContainer& BaseNodeContainer) const
+{
+	FString Filename = SourceData->GetFilename();
+	if ( !FPaths::FileExists( Filename ) )
+	{
+		return false;
+	}
+
+	FName DisplayLabel = *FPaths::GetBaseFilename( Filename );
+	Interchange::FNodeUniqueID NodeUID( *Filename );
+	//TGA is creating a UTexture2D
+	TUniquePtr< Interchange::FTextureNode > TextureNode = MakeUnique< Interchange::FTextureNode >( NodeUID, DisplayLabel, UTexture2D::StaticClass() );
+	TextureNode->SetPayLoadKey( Filename );
+
+	BaseNodeContainer.AddNode( MoveTemp( TextureNode ) );
+
+	return true;
+}
+
+const TOptional<Interchange::FImportImage> UInterchangeTGATranslator::GetPayloadData(const UInterchangeSourceData* SourceData, const FString& PayLoadKey) const
+{
+	if (!SourceData)
+	{
+		UE_LOG(LogInterchangeImportPlugin, Error, TEXT("Failed to import TGA, bad source data."));
+		return {};
+	}
+
+	TArray64<uint8> SourceDataBuffer;
+	FString Filename = SourceData->GetFilename();
+	
+	//Make sure the key fit the filename, The key should always be valid
+	if (!Filename.Equals(PayLoadKey))
+	{
+		UE_LOG(LogInterchangeImportPlugin, Error, TEXT("Failed to import TGA, wrong payload key. [%s]"), *Filename);
+		return {};
+	}
+
+	if (!FPaths::FileExists(Filename))
+	{
+		UE_LOG(LogInterchangeImportPlugin, Error, TEXT("Failed to import TGA, cannot open file. [%s]"), *Filename);
+		return {};
+	}
+
+	if (!FFileHelper::LoadFileToArray(SourceDataBuffer, *Filename))
+	{
+		UE_LOG(LogInterchangeImportPlugin, Error, TEXT("Failed to import TGA, cannot load file content into an array. [%s]"), *Filename);
+		return {};
+	}
+
+	const uint8* Buffer = SourceDataBuffer.GetData();
+	const uint8* BufferEnd = Buffer + SourceDataBuffer.Num();
+
+	bool bAllowNonPowerOfTwo = false;
+	GConfig->GetBool(TEXT("TextureImporter"), TEXT("AllowNonPowerOfTwoTextures"), bAllowNonPowerOfTwo, GEditorIni);
+
+	const int32 Length = BufferEnd - Buffer;
+
+	IImageWrapperModule& ImageWrapperModule = FModuleManager::LoadModuleChecked<IImageWrapperModule>(FName("ImageWrapper"));
+
+	TSharedPtr<IImageWrapper> TgaImageWrapper = ImageWrapperModule.CreateImageWrapper(EImageFormat::TGA);
+	if ( TgaImageWrapper.IsValid() && TgaImageWrapper->SetCompressed( Buffer, Length ) )
+	{
+		// Check the resolution of the imported texture to ensure validity
+		if ( !Interchange::FImportImageHelper::IsImportResolutionValid( TgaImageWrapper->GetWidth(), TgaImageWrapper->GetHeight(), bAllowNonPowerOfTwo ) )
+		{
+			return {};
+		}
+
+		const FTGAFileHeader* TGA = (FTGAFileHeader *)Buffer;
+
+		Interchange::FImportImage PayloadData;
+
+		TGATranslatorImpl::DecompressTGA( TGA, PayloadData );
+
+		const bool bResult = TgaImageWrapper->GetRaw( TgaImageWrapper->GetFormat(), TgaImageWrapper->GetBitDepth(), PayloadData.RawData );
+
+		if ( bResult && PayloadData.CompressionSettings == TC_Grayscale && TGA->ImageTypeCode == 3 )
+		{
+			// default grayscales to linear as they wont get compression otherwise and are commonly used as masks
+			PayloadData.SRGB = false;
+		}
+
+		return PayloadData;
+	}
+
+	return {};
+}
\ No newline at end of file
diff --git a/Engine/Plugins/Runtime/Interchange/Source/Import/Public/Texture/InterchangeTGATranslator.h b/Engine/Plugins/Runtime/Interchange/Source/Import/Public/Texture/InterchangeTGATranslator.h
new file mode 100644
--- /dev/null
+++ b/Engine/Plugins/Runtime/Interchange/Source/Import/Public/Texture/InterchangeTGATranslator.h
@@ -0,0 +1,45 @@
+// Copyright Epic Games, Inc. All Rights Reserved.
+
+#pragma once
+
+#include "CoreMinimal.h"
+#include "UObject/Object.h"
+#include "UObject/ObjectMacros.h"
+#include "InterchangeTranslatorBase.h"
+#include "Texture/InterchangeTexturePayloadData.h"
+#include "Texture/InterchangeTextureTranslator.h"
+
+#include "InterchangeTGATranslator.generated.h"
+
+UCLASS(BlueprintType)
+class INTERCHANGEIMPORTPLUGIN_API UInterchangeTGATranslator : public UInterchangeTextureTranslator
+{
+	GENERATED_BODY()
+public:
+	
+	/**
+	 * return true if the translator can translate the specified file.
+	 */
+	virtual bool CanImportSourceData(const UInterchangeSourceData* InSourceData) const override;
+
+	/**
+	 * Translate a source data into a node hold by the specified nodes container.
+	 *
+	 * @param InSourceData - The source data containing the JPG data
+	 * @param BaseNodeContainer - The unreal objects descriptions container where to put the translated source data.
+	 * @return true if the translator can translate the source data, false otherwise.
+	 */
+	virtual bool Translate(const UInterchangeSourceData* SourceData, Interchange::FBaseNodeContainer& BaseNodeContainer) const override;
+
+	/**
+	 * Once the translation is done, the import process need a way to retrieve payload data.
+	 * This payload will be use by the factories to create the asset.
+	 *
+	 * @param SourceData - The source data containing the data to translate
+	 * @param PayloadKey - The key to retrieve the a particular payload contain into the specified source data.
+	 * @return a PayloadData containing the import image data. The TOptional will not be set if there is an error.
+	 */
+	virtual const TOptional<Interchange::FImportImage> GetPayloadData(const UInterchangeSourceData* SourceData, const FString& PayLoadKey) const override;
+};
+
+
diff --git a/Engine/Source/Editor/UnrealEd/Private/Factories/EditorFactories.cpp b/Engine/Source/Editor/UnrealEd/Private/Factories/EditorFactories.cpp
--- a/Engine/Source/Editor/UnrealEd/Private/Factories/EditorFactories.cpp
+++ b/Engine/Source/Editor/UnrealEd/Private/Factories/EditorFactories.cpp
@@ -199,6 +199,7 @@
 #include "Factories/TIFFLoader.h"
 #include "IImageWrapper.h"
 #include "IImageWrapperModule.h"
+#include "TgaImageSupport.h"
 
 #include "FbxImporter.h"
 #include "Misc/FbxErrors.h"
@@ -2647,374 +2648,6 @@ static void psd_GetPSDHeader( const uint8* Buffer, FPSDFileHeader& Info )
 		((int32)Buffer[25] <<  0);
 }
 
-
-void DecompressTGA_RLE_32bpp( const FTGAFileHeader* TGA, uint32* TextureData )
-{
-	uint8*	IdData		= (uint8*)TGA + sizeof(FTGAFileHeader); 
-	uint8*	ColorMap	= IdData + TGA->IdFieldLength;
-	uint8*	ImageData	= (uint8*) (ColorMap + (TGA->ColorMapEntrySize + 4) / 8 * TGA->ColorMapLength);				
-	uint32	Pixel		= 0;
-	int32     RLERun		= 0;
-	int32     RAWRun		= 0;
-
-	for(int32 Y = TGA->Height-1; Y >=0; Y--) // Y-flipped.
-	{					
-		for(int32 X = 0;X < TGA->Width;X++)
-		{						
-			if( RLERun > 0 )
-			{
-				RLERun--;  // reuse current Pixel data.
-			}
-			else if( RAWRun == 0 ) // new raw pixel or RLE-run.
-			{
-				uint8 RLEChunk = *(ImageData++);							
-				if( RLEChunk & 0x80 )
-				{
-					RLERun = ( RLEChunk & 0x7F ) + 1;
-					RAWRun = 1;
-				}
-				else
-				{
-					RAWRun = ( RLEChunk & 0x7F ) + 1;
-				}
-			}							
-			// Retrieve new pixel data - raw run or single pixel for RLE stretch.
-			if( RAWRun > 0 )
-			{
-				Pixel = *(uint32*)ImageData; // RGBA 32-bit dword.
-				ImageData += 4;
-				RAWRun--;
-				RLERun--;
-			}
-			// Store.
-			*( (TextureData + Y*TGA->Width)+X ) = Pixel;
-		}
-	}
-}
-
-void DecompressTGA_RLE_24bpp( const FTGAFileHeader* TGA, uint32* TextureData )
-{
-	uint8*	IdData = (uint8*)TGA + sizeof(FTGAFileHeader); 
-	uint8*	ColorMap = IdData + TGA->IdFieldLength;
-	uint8*	ImageData = (uint8*) (ColorMap + (TGA->ColorMapEntrySize + 4) / 8 * TGA->ColorMapLength);
-	uint8    Pixel[4] = {};
-	int32     RLERun = 0;
-	int32     RAWRun = 0;
-
-	for(int32 Y = TGA->Height-1; Y >=0; Y--) // Y-flipped.
-	{					
-		for(int32 X = 0;X < TGA->Width;X++)
-		{						
-			if( RLERun > 0 )
-				RLERun--;  // reuse current Pixel data.
-			else if( RAWRun == 0 ) // new raw pixel or RLE-run.
-			{
-				uint8 RLEChunk = *(ImageData++);
-				if( RLEChunk & 0x80 )
-				{
-					RLERun = ( RLEChunk & 0x7F ) + 1;
-					RAWRun = 1;
-				}
-				else
-				{
-					RAWRun = ( RLEChunk & 0x7F ) + 1;
-				}
-			}							
-			// Retrieve new pixel data - raw run or single pixel for RLE stretch.
-			if( RAWRun > 0 )
-			{
-				Pixel[0] = *(ImageData++);
-				Pixel[1] = *(ImageData++);
-				Pixel[2] = *(ImageData++);
-				Pixel[3] = 255;
-				RAWRun--;
-				RLERun--;
-			}
-			// Store.
-			*( (TextureData + Y*TGA->Width)+X ) = *(uint32*)&Pixel;
-		}
-	}
-}
-
-void DecompressTGA_RLE_16bpp( const FTGAFileHeader* TGA, uint32* TextureData )
-{
-	uint8*	IdData = (uint8*)TGA + sizeof(FTGAFileHeader);
-	uint8*	ColorMap = IdData + TGA->IdFieldLength;				
-	uint16*	ImageData = (uint16*) (ColorMap + (TGA->ColorMapEntrySize + 4) / 8 * TGA->ColorMapLength);
-	uint16  FilePixel = 0;
-	uint32	TexturePixel = 0;
-	int32     RLERun = 0;
-	int32     RAWRun = 0;
-
-	for(int32 Y = TGA->Height-1; Y >=0; Y--) // Y-flipped.
-	{					
-		for( int32 X=0;X<TGA->Width;X++ )
-		{						
-			if( RLERun > 0 )
-				RLERun--;  // reuse current Pixel data.
-			else if( RAWRun == 0 ) // new raw pixel or RLE-run.
-			{
-				uint8 RLEChunk =  *((uint8*)ImageData);
-				ImageData = (uint16*)(((uint8*)ImageData)+1);
-				if( RLEChunk & 0x80 )
-				{
-					RLERun = ( RLEChunk & 0x7F ) + 1;
-					RAWRun = 1;
-				}
-				else
-				{
-					RAWRun = ( RLEChunk & 0x7F ) + 1;
-				}
-			}							
-			// Retrieve new pixel data - raw run or single pixel for RLE stretch.
-			if( RAWRun > 0 )
-			{ 
-				FilePixel = *(ImageData++);
-				RAWRun--;
-				RLERun--;
-			}
-			// Convert file format A1R5G5B5 into pixel format B8G8R8B8
-			TexturePixel = (FilePixel & 0x001F) << 3;
-			TexturePixel |= (FilePixel & 0x03E0) << 6;
-			TexturePixel |= (FilePixel & 0x7C00) << 9;
-			TexturePixel |= (FilePixel & 0x8000) << 16;
-			// Store.
-			*( (TextureData + Y*TGA->Width)+X ) = TexturePixel;
-		}
-	}
-}
-
-void DecompressTGA_32bpp( const FTGAFileHeader* TGA, uint32* TextureData )
-{
-
-	uint8*	IdData = (uint8*)TGA + sizeof(FTGAFileHeader);
-	uint8*	ColorMap = IdData + TGA->IdFieldLength;
-	uint32*	ImageData = (uint32*) (ColorMap + (TGA->ColorMapEntrySize + 4) / 8 * TGA->ColorMapLength);
-
-	for(int32 Y = 0;Y < TGA->Height;Y++)
-	{
-		FMemory::Memcpy(TextureData + Y * TGA->Width,ImageData + (TGA->Height - Y - 1) * TGA->Width,TGA->Width * 4);
-	}
-}
-
-void DecompressTGA_16bpp( const FTGAFileHeader* TGA, uint32* TextureData )
-{
-	uint8*	IdData = (uint8*)TGA + sizeof(FTGAFileHeader);
-	uint8*	ColorMap = IdData + TGA->IdFieldLength;
-	uint16*	ImageData = (uint16*) (ColorMap + (TGA->ColorMapEntrySize + 4) / 8 * TGA->ColorMapLength);
-	uint16    FilePixel = 0;
-	uint32	TexturePixel = 0;
-
-	for (int32 Y = TGA->Height - 1; Y >= 0; Y--)
-	{					
-		for (int32 X = 0; X<TGA->Width; X++)
-		{
-			FilePixel = *ImageData++;
-			// Convert file format A1R5G5B5 into pixel format B8G8R8A8
-			TexturePixel = (FilePixel & 0x001F) << 3;
-			TexturePixel |= (FilePixel & 0x03E0) << 6;
-			TexturePixel |= (FilePixel & 0x7C00) << 9;
-			TexturePixel |= (FilePixel & 0x8000) << 16;
-			// Store.
-			*((TextureData + Y*TGA->Width) + X) = TexturePixel;						
-		}
-	}
-}
-
-void DecompressTGA_24bpp( const FTGAFileHeader* TGA, uint32* TextureData )
-{
-	uint8*	IdData = (uint8*)TGA + sizeof(FTGAFileHeader);
-	uint8*	ColorMap = IdData + TGA->IdFieldLength;
-	uint8*	ImageData = (uint8*) (ColorMap + (TGA->ColorMapEntrySize + 4) / 8 * TGA->ColorMapLength);
-	uint8    Pixel[4];
-
-	for(int32 Y = 0;Y < TGA->Height;Y++)
-	{
-		for(int32 X = 0;X < TGA->Width;X++)
-		{
-			Pixel[0] = *(( ImageData+( TGA->Height-Y-1 )*TGA->Width*3 )+X*3+0);
-			Pixel[1] = *(( ImageData+( TGA->Height-Y-1 )*TGA->Width*3 )+X*3+1);
-			Pixel[2] = *(( ImageData+( TGA->Height-Y-1 )*TGA->Width*3 )+X*3+2);
-			Pixel[3] = 255;
-			*((TextureData+Y*TGA->Width)+X) = *(uint32*)&Pixel;
-		}
-	}
-}
-
-void DecompressTGA_8bpp( const FTGAFileHeader* TGA, uint8* TextureData )
-{
-	const uint8*  const IdData = (uint8*)TGA + sizeof(FTGAFileHeader);
-	const uint8*  const ColorMap = IdData + TGA->IdFieldLength;
-	const uint8*  const ImageData = (uint8*) (ColorMap + (TGA->ColorMapEntrySize + 4) / 8 * TGA->ColorMapLength);
-
-	int32 RevY = 0;
-	for (int32 Y = TGA->Height-1; Y >= 0; --Y)
-	{
-		const uint8* ImageCol = ImageData + (Y * TGA->Width); 
-		uint8* TextureCol = TextureData + (RevY++ * TGA->Width);
-		FMemory::Memcpy(TextureCol, ImageCol, TGA->Width);
-	}
-}
-
-bool DecompressTGA_helper(
-	const FTGAFileHeader* TGA,
-	uint32*& TextureData,
-	const int32 TextureDataSize,
-	FFeedbackContext* Warn )
-{
-	if(TGA->ImageTypeCode == 10) // 10 = RLE compressed 
-	{
-		// RLE compression: CHUNKS: 1 -byte header, high bit 0 = raw, 1 = compressed
-		// bits 0-6 are a 7-bit count; count+1 = number of raw pixels following, or rle pixels to be expanded. 
-		if(TGA->BitsPerPixel == 32)
-		{
-			DecompressTGA_RLE_32bpp(TGA, TextureData);
-		}
-		else if( TGA->BitsPerPixel == 24 )
-		{	
-			DecompressTGA_RLE_24bpp(TGA, TextureData);
-		}
-		else if( TGA->BitsPerPixel == 16 )
-		{
-			DecompressTGA_RLE_16bpp(TGA, TextureData);
-		}
-		else
-		{
-			Warn->Logf(ELogVerbosity::Error, TEXT("TGA uses an unsupported rle-compressed bit-depth: %u"),TGA->BitsPerPixel);
-			return false;
-		}
-	}
-	else if(TGA->ImageTypeCode == 2) // 2 = Uncompressed RGB
-	{
-		if(TGA->BitsPerPixel == 32)
-		{
-			DecompressTGA_32bpp(TGA, TextureData);
-		}
-		else if(TGA->BitsPerPixel == 16)
-		{
-			DecompressTGA_16bpp(TGA, TextureData);
-		}            
-		else if(TGA->BitsPerPixel == 24)
-		{
-			DecompressTGA_24bpp(TGA, TextureData);
-		}
-		else
-		{
-			Warn->Logf(ELogVerbosity::Error, TEXT("TGA uses an unsupported bit-depth: %u"),TGA->BitsPerPixel);
-			return false;
-		}
-	}
-	// Support for alpha stored as pseudo-color 8-bit TGA
-	else if(TGA->ColorMapType == 1 && TGA->ImageTypeCode == 1 && TGA->BitsPerPixel == 8)
-	{
-		DecompressTGA_8bpp(TGA, (uint8*)TextureData);
-	}
-	// standard grayscale
-	else if(TGA->ColorMapType == 0 && TGA->ImageTypeCode == 3 && TGA->BitsPerPixel == 8)
-	{
-		DecompressTGA_8bpp(TGA, (uint8*)TextureData);
-	}
-	else
-	{
-		Warn->Logf(ELogVerbosity::Error, TEXT("TGA is an unsupported type: %u"),TGA->ImageTypeCode);
-		return false;
-	}
-
-	// Flip the image data if the flip bits are set in the TGA header.
-	bool FlipX = (TGA->ImageDescriptor & 0x10) ? 1 : 0;
-	bool FlipY = (TGA->ImageDescriptor & 0x20) ? 1 : 0;
-	if(FlipY || FlipX)
-	{
-		TArray<uint8> FlippedData;
-		FlippedData.AddUninitialized(TextureDataSize);
-
-		int32 NumBlocksX = TGA->Width;
-		int32 NumBlocksY = TGA->Height;
-		int32 BlockBytes = TGA->BitsPerPixel == 8 ? 1 : 4;
-
-		uint8* MipData = (uint8*)TextureData;
-
-		for(int32 Y = 0;Y < NumBlocksY;Y++)
-		{
-			for(int32 X  = 0;X < NumBlocksX;X++)
-			{
-				int32 DestX = FlipX ? (NumBlocksX - X - 1) : X;
-				int32 DestY = FlipY ? (NumBlocksY - Y - 1) : Y;
-				FMemory::Memcpy(
-					&FlippedData[(DestX + DestY * NumBlocksX) * BlockBytes],
-					&MipData[(X + Y * NumBlocksX) * BlockBytes],
-					BlockBytes
-					);
-			}
-		}
-		FMemory::Memcpy(MipData,FlippedData.GetData(),FlippedData.Num());
-	}
-
-	return true;
-}
-
-bool DecompressTGA(
-	const FTGAFileHeader*	TGA,
-	FImportImage&			OutImage,
-	FFeedbackContext*		Warn)
-{
-	if (TGA->ColorMapType == 1 && TGA->ImageTypeCode == 1 && TGA->BitsPerPixel == 8)
-	{
-		// Notes: The Scaleform GFx exporter (dll) strips all font glyphs into a single 8-bit texture.
-		// The targa format uses this for a palette index; GFx uses a palette of (i,i,i,i) so the index
-		// is also the alpha value.
-		//
-		// We store the image as PF_G8, where it will be used as alpha in the Glyph shader.
-		OutImage.Init2DWithOneMip(
-			TGA->Width,
-			TGA->Height,
-			TSF_G8);
-		OutImage.CompressionSettings = TC_Grayscale;
-	}
-	else if(TGA->ColorMapType == 0 && TGA->ImageTypeCode == 3 && TGA->BitsPerPixel == 8)
-	{
-		// standard grayscale images
-		OutImage.Init2DWithOneMip(
-			TGA->Width,
-			TGA->Height,
-			TSF_G8);
-		OutImage.CompressionSettings = TC_Grayscale;
-	}
-	else
-	{
-		if(TGA->ImageTypeCode == 10) // 10 = RLE compressed 
-		{
-			if( TGA->BitsPerPixel != 32 &&
-				TGA->BitsPerPixel != 24 &&
-			    TGA->BitsPerPixel != 16 )
-			{
-				Warn->Logf(ELogVerbosity::Error, TEXT("TGA uses an unsupported rle-compressed bit-depth: %u"),TGA->BitsPerPixel);
-				return false;
-			}
-		}
-		else
-		{
-			if( TGA->BitsPerPixel != 32 &&
-				TGA->BitsPerPixel != 16 &&
-				TGA->BitsPerPixel != 24)
-			{
-				Warn->Logf(ELogVerbosity::Error, TEXT("TGA uses an unsupported bit-depth: %u"),TGA->BitsPerPixel);
-				return false;
-			}
-		}
-
-		OutImage.Init2DWithOneMip(
-			TGA->Width,
-			TGA->Height,
-			TSF_BGRA8);
-	}
-
-	int32 TextureDataSize = OutImage.RawData.Num();
-	uint32* TextureData = (uint32*)OutImage.RawData.GetData();
-
-	return DecompressTGA_helper(TGA, TextureData, TextureDataSize, Warn);
-}
-
 bool UTextureFactory::bSuppressImportOverwriteDialog = false;
 bool UTextureFactory::bForceOverwriteExistingSettings = false;
 
@@ -3649,22 +3282,70 @@ bool UTextureFactory::ImportImage(const uint8* Buffer, uint32 Length, FFeedbackC
 	//
 	// TGA
 	//
-	// Support for alpha stored as pseudo-color 8-bit TGA
-	const FTGAFileHeader*    TGA = (FTGAFileHeader *)Buffer;
-	if (Length >= sizeof(FTGAFileHeader) &&
-		((TGA->ColorMapType == 0 && TGA->ImageTypeCode == 2) ||
-			// ImageTypeCode 3 is greyscale
-		(TGA->ColorMapType == 0 && TGA->ImageTypeCode == 3) ||
-			(TGA->ColorMapType == 0 && TGA->ImageTypeCode == 10) ||
-			(TGA->ColorMapType == 1 && TGA->ImageTypeCode == 1 && TGA->BitsPerPixel == 8)))
+	TSharedPtr<IImageWrapper> TgaImageWrapper = ImageWrapperModule.CreateImageWrapper(EImageFormat::TGA);
+	if (TgaImageWrapper.IsValid() && TgaImageWrapper->SetCompressed(Buffer, Length))
 	{
 		// Check the resolution of the imported texture to ensure validity
-		if (!IsImportResolutionValid(TGA->Width, TGA->Height, bAllowNonPowerOfTwo, Warn))
+		if (!IsImportResolutionValid(TgaImageWrapper->GetWidth(), TgaImageWrapper->GetHeight(), bAllowNonPowerOfTwo, Warn))
 		{
 			return false;
 		}
 
-		const bool bResult = DecompressTGA(TGA, OutImage, Warn);
+		const FTGAFileHeader* TGA = (FTGAFileHeader*)Buffer;
+
+		if (TGA->ColorMapType == 1 && TGA->ImageTypeCode == 1 && TGA->BitsPerPixel == 8)
+		{
+			// Notes: The Scaleform GFx exporter (dll) strips all font glyphs into a single 8-bit texture.
+			// The targa format uses this for a palette index; GFx uses a palette of (i,i,i,i) so the index
+			// is also the alpha value.
+			//
+			// We store the image as PF_G8, where it will be used as alpha in the Glyph shader.
+			OutImage.Init2DWithOneMip(
+				TGA->Width,
+				TGA->Height,
+				TSF_G8);
+			OutImage.CompressionSettings = TC_Grayscale;
+		}
+		else if(TGA->ColorMapType == 0 && TGA->ImageTypeCode == 3 && TGA->BitsPerPixel == 8)
+		{
+			// standard grayscale images
+			OutImage.Init2DWithOneMip(
+				TGA->Width,
+				TGA->Height,
+				TSF_G8);
+			OutImage.CompressionSettings = TC_Grayscale;
+		}
+		else
+		{
+			if(TGA->ImageTypeCode == 10) // 10 = RLE compressed 
+			{
+				if( TGA->BitsPerPixel != 32 &&
+					TGA->BitsPerPixel != 24 &&
+					TGA->BitsPerPixel != 16 )
+				{
+					Warn->Logf(ELogVerbosity::Error, TEXT("TGA uses an unsupported rle-compressed bit-depth: %u"),TGA->BitsPerPixel);
+					return false;
+				}
+			}
+			else
+			{
+				if( TGA->BitsPerPixel != 32 &&
+					TGA->BitsPerPixel != 16 &&
+					TGA->BitsPerPixel != 24)
+				{
+					Warn->Logf(ELogVerbosity::Error, TEXT("TGA uses an unsupported bit-depth: %u"),TGA->BitsPerPixel);
+					return false;
+				}
+			}
+
+			OutImage.Init2DWithOneMip(
+				TGA->Width,
+				TGA->Height,
+				TSF_BGRA8);
+		}
+
+		const bool bResult = TgaImageWrapper->GetRaw( TgaImageWrapper->GetFormat(), TgaImageWrapper->GetBitDepth(), OutImage.RawData );
+
 		if (bResult && OutImage.CompressionSettings == TC_Grayscale && TGA->ImageTypeCode == 3)
 		{
 			// default grayscales to linear as they wont get compression otherwise and are commonly used as masks
@@ -3673,6 +3354,7 @@ bool UTextureFactory::ImportImage(const uint8* Buffer, uint32 Length, FFeedbackC
 
 		return bResult;
 	}
+
 	//
 	// PSD File
 	//
diff --git a/Engine/Source/Editor/UnrealEd/Public/Factories.h b/Engine/Source/Editor/UnrealEd/Public/Factories.h
--- a/Engine/Source/Editor/UnrealEd/Public/Factories.h
+++ b/Engine/Source/Editor/UnrealEd/Public/Factories.h
@@ -113,41 +113,3 @@ public:
 
     virtual USkeletalMesh* CreateSkeletalMesh(const TCHAR* SrcFilename, EMorphImportError* Error ) = 0;
 };
-
-#pragma pack(push,1)
-
-struct FTGAFileHeader
-{
-	uint8 IdFieldLength;
-	uint8 ColorMapType;
-	uint8 ImageTypeCode;		// 2 for uncompressed RGB format
-	uint16 ColorMapOrigin;
-	uint16 ColorMapLength;
-	uint8 ColorMapEntrySize;
-	uint16 XOrigin;
-	uint16 YOrigin;
-	uint16 Width;
-	uint16 Height;
-	uint8 BitsPerPixel;
-	uint8 ImageDescriptor;
-	friend FArchive& operator<<( FArchive& Ar, FTGAFileHeader& H )
-	{
-		Ar << H.IdFieldLength << H.ColorMapType << H.ImageTypeCode;
-		Ar << H.ColorMapOrigin << H.ColorMapLength << H.ColorMapEntrySize;
-		Ar << H.XOrigin << H.YOrigin << H.Width << H.Height << H.BitsPerPixel;
-		Ar << H.ImageDescriptor;
-		return Ar;
-	}
-};
-
-#pragma pack(pop)
-
-/**
- * This helper allows to decompress TGA data in a pre-allocated memory block.
- * The pixel format is necessarily PF_A8R8G8B8.
- */
-UNREALED_API bool DecompressTGA_helper(
-	const FTGAFileHeader* TGA,
-	uint32*& TextureData,
-	const int32 TextureDataSize,
-	FFeedbackContext* Warn );
diff --git a/Engine/Source/Runtime/ImageWrapper/Private/Formats/TgaImageWrapper.cpp b/Engine/Source/Runtime/ImageWrapper/Private/Formats/TgaImageWrapper.cpp
new file mode 100644
--- /dev/null
+++ b/Engine/Source/Runtime/ImageWrapper/Private/Formats/TgaImageWrapper.cpp
@@ -0,0 +1,57 @@
+// Copyright Epic Games, Inc. All Rights Reserved.
+
+#include "TgaImageWrapper.h"
+#include "ImageWrapperPrivate.h"
+
+#include "TgaImageSupport.h"
+
+
+void FTgaImageWrapper::Compress(int32 Quality)
+{
+	checkf(false, TEXT("TGA compression not supported"));
+}
+
+bool FTgaImageWrapper::SetCompressed(const void* InCompressedData, int64 InCompressedSize)
+{
+	bool bResult = FImageWrapperBase::SetCompressed(InCompressedData, InCompressedSize);
+
+	return bResult && LoadTGAHeader();
+}
+
+void FTgaImageWrapper::Uncompress(const ERGBFormat InFormat, const int32 InBitDepth)
+{
+	const int32 BytesPerPixel = ( InFormat == ERGBFormat::Gray ? 1 : 4 );
+	RawData.AddUninitialized( (int64)Width * Height * BytesPerPixel );
+
+	const int32 TextureDataSize = RawData.Num();
+	uint32* TextureData = reinterpret_cast< uint32* >( RawData.GetData() );
+
+	FTGAFileHeader* TgaHeader = reinterpret_cast< FTGAFileHeader* >( CompressedData.GetData() );
+	DecompressTGA_helper( TgaHeader, TextureData, TextureDataSize );
+}
+
+bool FTgaImageWrapper::LoadTGAHeader()
+{
+	check( CompressedData.Num() );
+
+	const FTGAFileHeader* TgaHeader = (FTGAFileHeader*)CompressedData.GetData();
+
+	if ( CompressedData.Num() >= sizeof( FTGAFileHeader ) &&
+		( (TgaHeader->ColorMapType == 0 && TgaHeader->ImageTypeCode == 2) ||
+		  ( TgaHeader->ColorMapType == 0 && TgaHeader->ImageTypeCode == 3 ) || // ImageTypeCode 3 is greyscale
+		  ( TgaHeader->ColorMapType == 0 && TgaHeader->ImageTypeCode == 10 ) ||
+		  ( TgaHeader->ColorMapType == 1 && TgaHeader->ImageTypeCode == 1 && TgaHeader->BitsPerPixel == 8 ) ) )
+	{
+		Width = TgaHeader->Width;
+		Height = TgaHeader->Height;
+		ColorMapType = TgaHeader->ColorMapType;
+		ImageTypeCode = TgaHeader->ImageTypeCode;
+		BitDepth = TgaHeader->BitsPerPixel;
+
+		return true;
+	}
+	else
+	{
+		return false;
+	}
+}
diff --git a/Engine/Source/Runtime/ImageWrapper/Private/Formats/TgaImageWrapper.h b/Engine/Source/Runtime/ImageWrapper/Private/Formats/TgaImageWrapper.h
new file mode 100644
--- /dev/null
+++ b/Engine/Source/Runtime/ImageWrapper/Private/Formats/TgaImageWrapper.h
@@ -0,0 +1,33 @@
+// Copyright Epic Games, Inc. All Rights Reserved.
+
+#pragma once
+
+#include "CoreTypes.h"
+#include "ImageWrapperBase.h"
+
+
+/**
+ * TGA implementation of the helper class
+ */
+class FTgaImageWrapper : public FImageWrapperBase
+{
+
+public:
+	/** 
+	 * Load the header information, returns true if successful.
+	 */
+	bool LoadTGAHeader();
+
+public:
+
+	//~ FImageWrapper interface
+
+	virtual void Compress(int32 Quality) override;
+	virtual bool SetCompressed(const void* InCompressedData, int64 InCompressedSize) override;
+	virtual void Uncompress(const ERGBFormat InFormat, int32 InBitDepth) override;
+
+private:
+	/** The color type as defined in the header. */
+	int32 ColorMapType;
+	uint8 ImageTypeCode;
+};
diff --git a/Engine/Source/Runtime/ImageWrapper/Private/ImageWrapperModule.cpp b/Engine/Source/Runtime/ImageWrapper/Private/ImageWrapperModule.cpp
--- a/Engine/Source/Runtime/ImageWrapper/Private/ImageWrapperModule.cpp
+++ b/Engine/Source/Runtime/ImageWrapper/Private/ImageWrapperModule.cpp
@@ -11,6 +11,7 @@
 #include "Formats/IcoImageWrapper.h"
 #include "Formats/JpegImageWrapper.h"
 #include "Formats/PngImageWrapper.h"
+#include "Formats/TgaImageWrapper.h"
 #include "IImageWrapperModule.h"
 
 
@@ -97,6 +98,10 @@ public:
 			ImageWrapper = new FIcnsImageWrapper();
 			break;
 
+		case EImageFormat::TGA:
+			ImageWrapper = new FTgaImageWrapper();
+			break;
+
 		default:
 			break;
 		}
diff --git a/Engine/Source/Runtime/ImageWrapper/Private/TgaImageSupport.cpp b/Engine/Source/Runtime/ImageWrapper/Private/TgaImageSupport.cpp
new file mode 100644
--- /dev/null
+++ b/Engine/Source/Runtime/ImageWrapper/Private/TgaImageSupport.cpp
@@ -0,0 +1,312 @@
+// Copyright Epic Games, Inc. All Rights Reserved.
+
+#include "TgaImageSupport.h"
+
+namespace TgaImageSupportImpl
+{
+	void DecompressTGA_RLE_32bpp( const FTGAFileHeader* TgaHeader, uint32* TextureData )
+	{
+		uint8*  IdData      = (uint8*)TgaHeader + sizeof(FTGAFileHeader); 
+		uint8*  ColorMap    = IdData + TgaHeader->IdFieldLength;
+		uint8*  ImageData   = (uint8*) (ColorMap + (TgaHeader->ColorMapEntrySize + 4) / 8 * TgaHeader->ColorMapLength);				
+		uint32  Pixel       = 0;
+		int32   RLERun      = 0;
+		int32   RAWRun      = 0;
+
+		for ( int32 Y = TgaHeader->Height - 1 ; Y >=0; Y-- ) // Y-flipped.
+		{					
+			for ( int32 X = 0; X < TgaHeader->Width; X++ )
+			{						
+				if( RLERun > 0 )
+				{
+					RLERun--;  // reuse current Pixel data.
+				}
+				else if( RAWRun == 0 ) // new raw pixel or RLE-run.
+				{
+					uint8 RLEChunk = *(ImageData++);
+					if( RLEChunk & 0x80 )
+					{
+						RLERun = ( RLEChunk & 0x7F ) + 1;
+						RAWRun = 1;
+					}
+					else
+					{
+						RAWRun = ( RLEChunk & 0x7F ) + 1;
+					}
+				}							
+				// Retrieve new pixel data - raw run or single pixel for RLE stretch.
+				if( RAWRun > 0 )
+				{
+					Pixel = *(uint32*)ImageData; // RGBA 32-bit dword.
+					ImageData += 4;
+					RAWRun--;
+					RLERun--;
+				}
+				// Store.
+				*( (TextureData + Y*TgaHeader->Width)+X ) = Pixel;
+			}
+		}
+	}
+
+	void DecompressTGA_RLE_24bpp( const FTGAFileHeader* TgaHeader, uint32* TextureData )
+	{
+		uint8*  IdData = (uint8*)TgaHeader + sizeof(FTGAFileHeader); 
+		uint8*  ColorMap = IdData + TgaHeader->IdFieldLength;
+		uint8*  ImageData = (uint8*) (ColorMap + (TgaHeader->ColorMapEntrySize + 4) / 8 * TgaHeader->ColorMapLength);
+		uint8   Pixel[4] = {};
+		int32   RLERun = 0;
+		int32   RAWRun = 0;
+
+		for(int32 Y = TgaHeader->Height-1; Y >=0; Y--) // Y-flipped.
+		{					
+			for(int32 X = 0;X < TgaHeader->Width;X++)
+			{						
+				if( RLERun > 0 )
+					RLERun--;  // reuse current Pixel data.
+				else if( RAWRun == 0 ) // new raw pixel or RLE-run.
+				{
+					uint8 RLEChunk = *(ImageData++);
+					if( RLEChunk & 0x80 )
+					{
+						RLERun = ( RLEChunk & 0x7F ) + 1;
+						RAWRun = 1;
+					}
+					else
+					{
+						RAWRun = ( RLEChunk & 0x7F ) + 1;
+					}
+				}							
+				// Retrieve new pixel data - raw run or single pixel for RLE stretch.
+				if( RAWRun > 0 )
+				{
+					Pixel[0] = *(ImageData++);
+					Pixel[1] = *(ImageData++);
+					Pixel[2] = *(ImageData++);
+					Pixel[3] = 255;
+					RAWRun--;
+					RLERun--;
+				}
+				// Store.
+				*( (TextureData + Y*TgaHeader->Width)+X ) = *(uint32*)&Pixel;
+			}
+		}
+	}
+
+	void DecompressTGA_RLE_16bpp( const FTGAFileHeader* TgaHeader, uint32* TextureData )
+	{
+		uint8*  IdData = (uint8*)TgaHeader + sizeof(FTGAFileHeader);
+		uint8*  ColorMap = IdData + TgaHeader->IdFieldLength;				
+		uint16* ImageData = (uint16*) (ColorMap + (TgaHeader->ColorMapEntrySize + 4) / 8 * TgaHeader->ColorMapLength);
+		uint16  FilePixel = 0;
+		uint32  TexturePixel = 0;
+		int32   RLERun = 0;
+		int32   RAWRun = 0;
+
+		for ( int32 Y = TgaHeader->Height-1; Y >=0; Y-- ) // Y-flipped.
+		{					
+			for ( int32 X=0;X<TgaHeader->Width;X++ )
+			{						
+				if ( RLERun > 0 )
+				{
+					RLERun--;  // reuse current Pixel data.
+				}
+				else if ( RAWRun == 0 ) // new raw pixel or RLE-run.
+				{
+					uint8 RLEChunk =  *((uint8*)ImageData);
+					ImageData = (uint16*)(((uint8*)ImageData)+1);
+					if( RLEChunk & 0x80 )
+					{
+						RLERun = ( RLEChunk & 0x7F ) + 1;
+						RAWRun = 1;
+					}
+					else
+					{
+						RAWRun = ( RLEChunk & 0x7F ) + 1;
+					}
+				}
+
+				// Retrieve new pixel data - raw run or single pixel for RLE stretch.
+				if ( RAWRun > 0 )
+				{ 
+					FilePixel = *(ImageData++);
+					RAWRun--;
+					RLERun--;
+				}
+
+				// Convert file format A1R5G5B5 into pixel format B8G8R8B8
+				TexturePixel = (FilePixel & 0x001F) << 3;
+				TexturePixel |= (FilePixel & 0x03E0) << 6;
+				TexturePixel |= (FilePixel & 0x7C00) << 9;
+				TexturePixel |= (FilePixel & 0x8000) << 16;
+
+				// Store.
+				*( (TextureData + Y*TgaHeader->Width)+X ) = TexturePixel;
+			}
+		}
+	}
+
+	void DecompressTGA_32bpp( const FTGAFileHeader* TgaHeader, uint32* TextureData )
+	{
+		uint8*	IdData = (uint8*)TgaHeader + sizeof(FTGAFileHeader);
+		uint8*	ColorMap = IdData + TgaHeader->IdFieldLength;
+		uint32*	ImageData = (uint32*) (ColorMap + (TgaHeader->ColorMapEntrySize + 4) / 8 * TgaHeader->ColorMapLength);
+
+		for ( int32 Y = 0; Y < TgaHeader->Height; Y++ )
+		{
+			FMemory::Memcpy(TextureData + Y * TgaHeader->Width,ImageData + (TgaHeader->Height - Y - 1) * TgaHeader->Width,TgaHeader->Width * 4);
+		}
+	}
+
+	void DecompressTGA_16bpp( const FTGAFileHeader* TgaHeader, uint32* TextureData )
+	{
+		uint8*	IdData = (uint8*)TgaHeader + sizeof(FTGAFileHeader);
+		uint8*	ColorMap = IdData + TgaHeader->IdFieldLength;
+		uint16*	ImageData = (uint16*) (ColorMap + (TgaHeader->ColorMapEntrySize + 4) / 8 * TgaHeader->ColorMapLength);
+		uint16    FilePixel = 0;
+		uint32	TexturePixel = 0;
+
+		for ( int32 Y = TgaHeader->Height - 1; Y >= 0; Y-- )
+		{					
+			for ( int32 X = 0; X < TgaHeader->Width; X++ )
+			{
+				FilePixel = *ImageData++;
+				// Convert file format A1R5G5B5 into pixel format B8G8R8A8
+				TexturePixel = (FilePixel & 0x001F) << 3;
+				TexturePixel |= (FilePixel & 0x03E0) << 6;
+				TexturePixel |= (FilePixel & 0x7C00) << 9;
+				TexturePixel |= (FilePixel & 0x8000) << 16;
+				// Store.
+				*((TextureData + Y*TgaHeader->Width) + X) = TexturePixel;						
+			}
+		}
+	}
+
+	void DecompressTGA_24bpp( const FTGAFileHeader* TgaHeader, uint32* TextureData )
+	{
+		uint8*	IdData = (uint8*)TgaHeader + sizeof(FTGAFileHeader);
+		uint8*	ColorMap = IdData + TgaHeader->IdFieldLength;
+		uint8*	ImageData = (uint8*) (ColorMap + (TgaHeader->ColorMapEntrySize + 4) / 8 * TgaHeader->ColorMapLength);
+		uint8    Pixel[4];
+
+		for( int32 Y = 0; Y < TgaHeader->Height; Y++ )
+		{
+			for( int32 X = 0; X < TgaHeader->Width; X++ )
+			{
+				Pixel[0] = *(( ImageData+( TgaHeader->Height-Y-1 )*TgaHeader->Width*3 )+X*3+0);
+				Pixel[1] = *(( ImageData+( TgaHeader->Height-Y-1 )*TgaHeader->Width*3 )+X*3+1);
+				Pixel[2] = *(( ImageData+( TgaHeader->Height-Y-1 )*TgaHeader->Width*3 )+X*3+2);
+				Pixel[3] = 255;
+				*((TextureData+Y*TgaHeader->Width)+X) = *(uint32*)&Pixel;
+			}
+		}
+	}
+
+	void DecompressTGA_8bpp( const FTGAFileHeader* TgaHeader, uint8* TextureData )
+	{
+		const uint8*  const IdData = (uint8*)TgaHeader + sizeof(FTGAFileHeader);
+		const uint8*  const ColorMap = IdData + TgaHeader->IdFieldLength;
+		const uint8*  const ImageData = (uint8*) (ColorMap + (TgaHeader->ColorMapEntrySize + 4) / 8 * TgaHeader->ColorMapLength);
+
+		int32 RevY = 0;
+		for (int32 Y = TgaHeader->Height-1; Y >= 0; --Y)
+		{
+			const uint8* ImageCol = ImageData + (Y * TgaHeader->Width); 
+			uint8* TextureCol = TextureData + (RevY++ * TgaHeader->Width);
+			FMemory::Memcpy(TextureCol, ImageCol, TgaHeader->Width);
+		}
+	}
+}
+
+bool DecompressTGA_helper( const FTGAFileHeader* TgaHeader, uint32*& TextureData, const int32 TextureDataSize )
+{
+	if ( TgaHeader->ImageTypeCode == 10 ) // 10 = RLE compressed 
+	{
+		// RLE compression: CHUNKS: 1 -byte header, high bit 0 = raw, 1 = compressed
+		// bits 0-6 are a 7-bit count; count+1 = number of raw pixels following, or rle pixels to be expanded. 
+		if(TgaHeader->BitsPerPixel == 32)
+		{
+			TgaImageSupportImpl::DecompressTGA_RLE_32bpp(TgaHeader, TextureData);
+		}
+		else if( TgaHeader->BitsPerPixel == 24 )
+		{	
+			TgaImageSupportImpl::DecompressTGA_RLE_24bpp(TgaHeader, TextureData);
+		}
+		else if( TgaHeader->BitsPerPixel == 16 )
+		{
+			TgaImageSupportImpl::DecompressTGA_RLE_16bpp(TgaHeader, TextureData);
+		}
+		else
+		{
+			UE_LOG( LogImageWrapper, Error, TEXT("TgaHeader uses an unsupported rle-compressed bit-depth: %u"), TgaHeader->BitsPerPixel );
+			return false;
+		}
+	}
+	else if(TgaHeader->ImageTypeCode == 2) // 2 = Uncompressed RGB
+	{
+		if(TgaHeader->BitsPerPixel == 32)
+		{
+			TgaImageSupportImpl::DecompressTGA_32bpp(TgaHeader, TextureData);
+		}
+		else if(TgaHeader->BitsPerPixel == 16)
+		{
+			TgaImageSupportImpl::DecompressTGA_16bpp(TgaHeader, TextureData);
+		}
+		else if(TgaHeader->BitsPerPixel == 24)
+		{
+			TgaImageSupportImpl::DecompressTGA_24bpp(TgaHeader, TextureData);
+		}
+		else
+		{
+			UE_LOG( LogImageWrapper, Error, TEXT("TgaHeader uses an unsupported bit-depth: %u"), TgaHeader->BitsPerPixel );
+			return false;
+		}
+	}
+	// Support for alpha stored as pseudo-color 8-bit TgaHeader
+	else if(TgaHeader->ColorMapType == 1 && TgaHeader->ImageTypeCode == 1 && TgaHeader->BitsPerPixel == 8)
+	{
+		TgaImageSupportImpl::DecompressTGA_8bpp(TgaHeader, (uint8*)TextureData);
+	}
+	// standard grayscale
+	else if(TgaHeader->ColorMapType == 0 && TgaHeader->ImageTypeCode == 3 && TgaHeader->BitsPerPixel == 8)
+	{
+		TgaImageSupportImpl::DecompressTGA_8bpp(TgaHeader, (uint8*)TextureData);
+	}
+	else
+	{
+		UE_LOG( LogImageWrapper, Error, TEXT("TgaHeader is an unsupported type: %u"), TgaHeader->ImageTypeCode );
+		return false;
+	}
+
+	// Flip the image data if the flip bits are set in the TgaHeader header.
+	const bool bFlipX = (TgaHeader->ImageDescriptor & 0x10) ? 1 : 0;
+	const bool bFlipY = (TgaHeader->ImageDescriptor & 0x20) ? 1 : 0;
+	if ( bFlipX || bFlipY )
+	{
+		TArray<uint8> FlippedData;
+		FlippedData.AddUninitialized(TextureDataSize);
+
+		int32 NumBlocksX = TgaHeader->Width;
+		int32 NumBlocksY = TgaHeader->Height;
+		int32 BlockBytes = TgaHeader->BitsPerPixel == 8 ? 1 : 4;
+
+		uint8* MipData = (uint8*)TextureData;
+
+		for( int32 Y = 0; Y < NumBlocksY;Y++ )
+		{
+			for( int32 X  = 0; X < NumBlocksX; X++ )
+			{
+				int32 DestX = bFlipX ? (NumBlocksX - X - 1) : X;
+				int32 DestY = bFlipY ? (NumBlocksY - Y - 1) : Y;
+				FMemory::Memcpy(
+					&FlippedData[(DestX + DestY * NumBlocksX) * BlockBytes],
+					&MipData[(X + Y * NumBlocksX) * BlockBytes],
+					BlockBytes
+				);
+			}
+		}
+
+		FMemory::Memcpy( MipData, FlippedData.GetData(), FlippedData.Num() );
+	}
+
+	return true;
+}
\ No newline at end of file
diff --git a/Engine/Source/Runtime/ImageWrapper/Public/IImageWrapper.h b/Engine/Source/Runtime/ImageWrapper/Public/IImageWrapper.h
--- a/Engine/Source/Runtime/ImageWrapper/Public/IImageWrapper.h
+++ b/Engine/Source/Runtime/ImageWrapper/Public/IImageWrapper.h
@@ -34,7 +34,10 @@ enum class EImageFormat : int8
 	EXR,
 
 	/** Mac icon. */
-	ICNS
+	ICNS,
+	
+	/** Truevision TGA / TARGA */
+	TGA
 };
 
 
diff --git a/Engine/Source/Runtime/ImageWrapper/Public/TgaImageSupport.h b/Engine/Source/Runtime/ImageWrapper/Public/TgaImageSupport.h
new file mode 100644
--- /dev/null
+++ b/Engine/Source/Runtime/ImageWrapper/Public/TgaImageSupport.h
@@ -0,0 +1,39 @@
+// Copyright Epic Games, Inc. All Rights Reserved.
+
+#pragma once
+
+#include "CoreMinimal.h"
+#include "Serialization/Archive.h"
+
+#pragma pack(push,1)
+struct FTGAFileHeader
+{
+	uint8 IdFieldLength;
+	uint8 ColorMapType;
+	uint8 ImageTypeCode;		// 2 for uncompressed RGB format
+	uint16 ColorMapOrigin;
+	uint16 ColorMapLength;
+	uint8 ColorMapEntrySize;
+	uint16 XOrigin;
+	uint16 YOrigin;
+	uint16 Width;
+	uint16 Height;
+	uint8 BitsPerPixel;
+	uint8 ImageDescriptor;
+
+	friend FArchive& operator<<( FArchive& Ar, FTGAFileHeader& H )
+	{
+		Ar << H.IdFieldLength << H.ColorMapType << H.ImageTypeCode;
+		Ar << H.ColorMapOrigin << H.ColorMapLength << H.ColorMapEntrySize;
+		Ar << H.XOrigin << H.YOrigin << H.Width << H.Height << H.BitsPerPixel;
+		Ar << H.ImageDescriptor;
+		return Ar;
+	}
+};
+#pragma pack(pop)
+
+/**
+ * This helper allows to decompress TGA data in a pre-allocated memory block.
+ * The pixel format is necessarily PF_A8R8G8B8.
+ */
+IMAGEWRAPPER_API bool DecompressTGA_helper(	const FTGAFileHeader* TGA, uint32*& TextureData, const int32 TextureDataSize );
