Add IsLinkedToWildcard util function

#rnx

diff --git a/Engine/Source/Editor/UnrealEd/Private/Kismet2/WildcardNodeUtils.cpp b/Engine/Source/Editor/UnrealEd/Private/Kismet2/WildcardNodeUtils.cpp
--- a/Engine/Source/Editor/UnrealEd/Private/Kismet2/WildcardNodeUtils.cpp
+++ b/Engine/Source/Editor/UnrealEd/Private/Kismet2/WildcardNodeUtils.cpp
@@ -23,6 +23,21 @@ bool FWildcardNodeUtils::IsWildcardPin(const UEdGraphPin* const Pin)
 	return Pin && Pin->PinType.PinCategory == UEdGraphSchema_K2::PC_Wildcard;
 }
 
+bool FWildcardNodeUtils::IsLinkedToWildcard(const UEdGraphPin* const Pin)
+{
+	if (Pin)
+	{
+		for (const UEdGraphPin* const LinkedPin : Pin->LinkedTo)
+		{
+			if (FWildcardNodeUtils::IsWildcardPin(LinkedPin))
+			{
+				return true;
+			}
+		}
+	}
+	return false;
+}
+
 UEdGraphPin* FWildcardNodeUtils::CreateWildcardPin(UEdGraphNode* Node, const FName PinName, const EEdGraphPinDirection Direction, const EPinContainerType ContainerType/* = EPinContainerType::None*/)
 {
 	check(Node);
diff --git a/Engine/Source/Editor/UnrealEd/Public/Kismet2/WildcardNodeUtils.h b/Engine/Source/Editor/UnrealEd/Public/Kismet2/WildcardNodeUtils.h
--- a/Engine/Source/Editor/UnrealEd/Public/Kismet2/WildcardNodeUtils.h
+++ b/Engine/Source/Editor/UnrealEd/Public/Kismet2/WildcardNodeUtils.h
@@ -27,6 +27,13 @@ public:
 	*/
 	static bool IsWildcardPin(const UEdGraphPin* const Pin);
 
+	/**
+	* Checks if the given pin is linked to any wildcard pins
+	* 
+	* @return	True if the given pin is linked to any wildcard pins
+	*/
+	static bool IsLinkedToWildcard(const UEdGraphPin* const Pin);
+
 	/**
 	* Add a default wildcard pin to the given node
 	* 
