Add missing header guard

diff --git a/Engine/Source/Editor/UnrealEd/Public/EditorScriptingHelpers.h b/Engine/Source/Editor/UnrealEd/Public/EditorScriptingHelpers.h
--- a/Engine/Source/Editor/UnrealEd/Public/EditorScriptingHelpers.h
+++ b/Engine/Source/Editor/UnrealEd/Public/EditorScriptingHelpers.h
@@ -1,5 +1,7 @@
 // Copyright Epic Games, Inc. All Rights Reserved.
 
+#pragma once 
+
 #include "CoreMinimal.h"
 
 namespace EditorScriptingHelpers
