- Switched over to new Turnkey-enabled LaunchOn menu only, removed the old one (can get it back with -oldlaunchon)

diff --git a/Engine/Source/Editor/UnrealEd/Private/Kismet2/DebuggerCommands.cpp b/Engine/Source/Editor/UnrealEd/Private/Kismet2/DebuggerCommands.cpp
--- a/Engine/Source/Editor/UnrealEd/Private/Kismet2/DebuggerCommands.cpp
+++ b/Engine/Source/Editor/UnrealEd/Private/Kismet2/DebuggerCommands.cpp
@@ -663,28 +663,30 @@ void FPlayWorldCommands::BuildToolbar(FToolMenuSection& InSection, bool bInclude
 				FUIAction LaunchMenuAction;
 				LaunchMenuAction.IsActionVisibleDelegate = FIsActionButtonVisible::CreateStatic(&FInternalPlayWorldCommandCallbacks::CanShowNonPlayWorldOnlyActions);
 
-				InDynamicSection.AddEntry(FToolMenuEntry::InitComboButton(
-					"LaunchCombo",
-					LaunchMenuAction,
-					FOnGetContent::CreateStatic(&GenerateLaunchMenuContent, GlobalPlayWorldActions.ToSharedRef()),
-					LOCTEXT("LaunchCombo_Label", "Launch Options"),
-					LOCTEXT("PODComboToolTip", "Options for launching on a device"),
-					FSlateIcon(FEditorStyle::GetStyleSetName(), "LevelEditor.RepeatLastLaunch"),
-					true
-				));
-
-				FUIAction TurnkeyLaunchMenuAction;
-				TurnkeyLaunchMenuAction.IsActionVisibleDelegate = FIsActionButtonVisible::CreateStatic(&FInternalPlayWorldCommandCallbacks::CanShowNonPlayWorldOnlyActions);
-
-				InDynamicSection.AddEntry(FToolMenuEntry::InitComboButton(
-					"TurnkeyCombo",
-					TurnkeyLaunchMenuAction,
-					FOnGetContent::CreateStatic(&GenerateTurnkeyMenuContent, GlobalPlayWorldActions.ToSharedRef()),
-					LOCTEXT("LaunchCombo_Label", "Launch Options"),
-					LOCTEXT("PlayWithTurnkeyComboToolTip", "Options for launching on a device and manage platforms"),
-					FSlateIcon(FEditorStyle::GetStyleSetName(), "LevelEditor.RepeatLastLaunch"),
-					true
-				));
+				if (FParse::Param(FCommandLine::Get(), TEXT("OldLaunchOn")))
+				{
+					InDynamicSection.AddEntry(FToolMenuEntry::InitComboButton(
+						"LaunchCombo",
+						LaunchMenuAction,
+						FOnGetContent::CreateStatic(&GenerateLaunchMenuContent, GlobalPlayWorldActions.ToSharedRef()),
+						LOCTEXT("LaunchCombo_Label", "Launch Options"),
+						LOCTEXT("PODComboToolTip", "Options for launching on a device"),
+						FSlateIcon(FEditorStyle::GetStyleSetName(), "LevelEditor.RepeatLastLaunch"),
+						true
+					));
+				}
+				else
+				{
+					InDynamicSection.AddEntry(FToolMenuEntry::InitComboButton(
+						"TurnkeyCombo",
+						LaunchMenuAction,
+						FOnGetContent::CreateStatic(&GenerateTurnkeyMenuContent, GlobalPlayWorldActions.ToSharedRef()),
+						LOCTEXT("LaunchCombo_Label", "Launch Options"),
+						LOCTEXT("PlayWithTurnkeyComboToolTip", "Options for launching on a device and manage platforms"),
+						FSlateIcon(FEditorStyle::GetStyleSetName(), "LevelEditor.RepeatLastLaunch"),
+						true
+					));
+				}
 			}
 		}));
 	}
@@ -1328,6 +1330,19 @@ TSharedRef< SWidget > FPlayWorldCommands::GenerateTurnkeyMenuContent(TSharedRef<
 
 	TArray<FString> DeviceIdsToQuery;
 
+	MenuBuilder.BeginSection("LevelEditorLaunchHint", FText::FromString(TEXT("Turnkey Info")));
+	{
+		MenuBuilder.AddWidget(
+			SNew(STextBlock)
+			.ColorAndOpacity(FSlateColor::UseSubduedForeground())
+			// this is temp and doesn't need to be loc'd
+			.WrapTextAt(300)
+			.Text(FText::FromString(TEXT("The Turnkey SDK management system is now enabled for this menu. If you have any issues, run with \"-OldLaunchOn\"  on the editor commandline to use the old Launch On menu"))),
+			FText::GetEmpty()
+		);
+	}
+	MenuBuilder.EndSection();
+
 	MenuBuilder.BeginSection("LevelEditorLaunchDevices", LOCTEXT("TurnkeySection_LaunchButtonDevices", "Quick Launch"));
 	{
 		for (const auto Pair : FDataDrivenPlatformInfoRegistry::GetAllPlatformInfos())
@@ -1421,7 +1436,7 @@ TSharedRef< SWidget > FPlayWorldCommands::GenerateTurnkeyMenuContent(TSharedRef<
 			SNew(STextBlock)
 			.ColorAndOpacity(FSlateColor::UseSubduedForeground())
 			.Text(LOCTEXT("ZoomToFitHorizontal", "Launching a game on a different device will change your default 'Launch' device in the toolbar"))
-			.WrapTextAt(250),
+			.WrapTextAt(300),
 			FText::GetEmpty()
 		);
 	}
@@ -2562,10 +2577,23 @@ void FInternalPlayWorldCommandCallbacks::RepeatLastLaunch_Clicked()
 	switch (PlaySettings->LastExecutedLaunchModeType)
 	{
 	case LaunchMode_OnDevice:
-		if (IsReadyToLaunchOnDevice(PlaySettings->LastExecutedLaunchDevice))
+		if (FParse::Param(FCommandLine::Get(), TEXT("OldLaunchOn")))
+		{
+			if (IsReadyToLaunchOnDevice(PlaySettings->LastExecutedLaunchDevice))
+			{
+				LaunchOnDevice(PlaySettings->LastExecutedLaunchDevice, PlaySettings->LastExecutedLaunchName, false);
+			}
+		}
+		else
 		{
-			LaunchOnDevice(PlaySettings->LastExecutedLaunchDevice, PlaySettings->LastExecutedLaunchName, false);
+			// @todo turnkey - remove this function completely and all the platform support, and put it into Turnkey
+//			if (IsReadyToLaunchOnDevice(PlaySettings->LastExecutedLaunchDevice))
+			{
+				LaunchOnDevice(PlaySettings->LastExecutedLaunchDevice, PlaySettings->LastExecutedLaunchName, true);
+			}
 		}
+
+
 		break;
 
 	default:
@@ -2786,36 +2814,8 @@ static void PrepareLaunchOn(FString DeviceId, FString DeviceName)
 
 void FInternalPlayWorldCommandCallbacks::HandleLaunchOnDeviceActionExecute(FString DeviceId, FString DeviceName, bool bUseTurnkey)
 {
-	if (bUseTurnkey)
-	{
-		int32 Index = 0;
-		DeviceId.FindChar(TEXT('@'), Index);
-		FString TargetPlatformName = DeviceId.Left(Index);
-		const PlatformInfo::FTargetPlatformInfo* const PlatformInfo = PlatformInfo::FindPlatformInfo(FName(*TargetPlatformName));
-		FString IniPlatformName = PlatformInfo->IniPlatformName.ToString();
-		FString UBTPlatformName = PlatformInfo->DataDrivenPlatformInfo->UBTPlatformString;
-
-		FString CommandLine = FString::Printf(TEXT("Turnkey -command=VerifySdk -UpdateIfNeeded -platform=%s -EditorIO -noturnkeyvariables -device=%s"), *UBTPlatformName, *DeviceName);
-		FText TaskName = LOCTEXT("VerifyingSDK", "Verifying SDK and Device");
-
-		IUATHelperModule::Get().CreateUatTask(CommandLine, FText::FromString(IniPlatformName), TaskName, TaskName, FEditorStyle::GetBrush(TEXT("MainFrame.PackageProject")),
-			[DeviceId, DeviceName](FString, double)
-		{
-			AsyncTask(ENamedThreads::GameThread, [DeviceId, DeviceName]()
-			{
-				PrepareLaunchOn(DeviceId, DeviceName);
-				LaunchOnDevice(DeviceId, DeviceName, false);
-			});
-		});
-	}
-	else
-	{
-		if (IsReadyToLaunchOnDevice(DeviceId))
-		{
-			PrepareLaunchOn(DeviceId, DeviceName);
-			LaunchOnDevice(DeviceId, DeviceName, false);
-		}
-	}
+	PrepareLaunchOn(DeviceId, DeviceName);
+	LaunchOnDevice(DeviceId, DeviceName, bUseTurnkey);
 }
 
 
@@ -3135,9 +3135,12 @@ void FInternalPlayWorldCommandCallbacks::LaunchOnDevice(const FString& DeviceId,
 	if (FTargetDeviceId::Parse(DeviceId, TargetDeviceId))
 	{
 		const PlatformInfo::FTargetPlatformInfo* const PlatformInfo = PlatformInfo::FindPlatformInfo(*TargetDeviceId.GetPlatformName());
+		FString UBTPlatformName = PlatformInfo->DataDrivenPlatformInfo->UBTPlatformString;
+		FString IniPlatformName = PlatformInfo->IniPlatformName.ToString();
+
 		check(PlatformInfo);
 
-		if (FInstalledPlatformInfo::Get().IsPlatformMissingRequiredFile(PlatformInfo->DataDrivenPlatformInfo->UBTPlatformString))
+		if (FInstalledPlatformInfo::Get().IsPlatformMissingRequiredFile(UBTPlatformName))
 		{
 			if (!FInstalledPlatformInfo::OpenInstallerOptions())
 			{
@@ -3153,13 +3156,50 @@ void FInternalPlayWorldCommandCallbacks::LaunchOnDevice(const FString& DeviceId,
 			FRequestPlaySessionParams::FLauncherDeviceInfo DeviceInfo;
 			DeviceInfo.DeviceId = DeviceId;
 			DeviceInfo.DeviceName = DeviceName;
-			DeviceInfo.bUpdateDeviceFlash = bUseTurnkey;
+			// @todo turnkey: we set this to false because we will kick off a Turnkey run before cooking, etc, to get an early warning. however, if it's too difficult
+			// to get an error back from CreateUatTask, then we should set this to bUseTurnkey and remove the block below, and let the code in FLauncherWorker::CreateAndExecuteTasks handle it
+			DeviceInfo.bUpdateDeviceFlash = false;
 
 			FRequestPlaySessionParams SessionParams;
 			SessionParams.SessionDestination = EPlaySessionDestinationType::Launcher;
 			SessionParams.LauncherTargetDevice = DeviceInfo;
 
-			GUnrealEd->RequestPlaySession(SessionParams);
+			// if we want to check device flash before we start cooking, kick it off now. we could delay this 
+			if (bUseTurnkey)
+			{
+				FString CommandLine = FString::Printf(TEXT("Turnkey -command=VerifySdk -UpdateIfNeeded -platform=%s -EditorIO -noturnkeyvariables -device=%s"), *UBTPlatformName, *DeviceName);
+				FText TaskName = LOCTEXT("VerifyingSDK", "Verifying SDK and Device");
+
+				IUATHelperModule::Get().CreateUatTask(CommandLine, FText::FromString(IniPlatformName), TaskName, TaskName, FEditorStyle::GetBrush(TEXT("MainFrame.PackageProject")),
+					[SessionParams](FString Result, double)
+					{
+						// unfortunate string comparison for success
+						bool bWasSuccessful = Result == TEXT("Completed");
+						AsyncTask(ENamedThreads::GameThread, [SessionParams, bWasSuccessful]()
+						{
+ 							if (bWasSuccessful)
+ 							{
+ 								GUnrealEd->RequestPlaySession(SessionParams);
+ 							}
+ 							else
+							{
+								TSharedRef<SWindow> Win = OpenMsgDlgInt_NonModal(EAppMsgType::YesNo, LOCTEXT("SDKCheckFailed", "SDK Verification failed. Would you like to attempt the Launch On anyway?"), LOCTEXT("SDKCheckFailedTitle", "SDK Verification"),
+									FOnMsgDlgResult::CreateLambda([SessionParams](const TSharedRef<SWindow>&, EAppReturnType::Type Choice)
+									{
+										if (Choice == EAppReturnType::Yes)
+										{
+											GUnrealEd->RequestPlaySession(SessionParams);
+										}
+									}));
+								Win->ShowWindow();
+							}
+						});
+					});
+			}
+			else
+			{
+				GUnrealEd->RequestPlaySession(SessionParams);
+			}
 		}
 	}
 }
