The default editor mode is now a UEdMode.

#rnx

diff --git a/Engine/Source/Editor/UnrealEd/Private/EdMode.cpp b/Engine/Source/Editor/UnrealEd/Private/EdMode.cpp
--- a/Engine/Source/Editor/UnrealEd/Private/EdMode.cpp
+++ b/Engine/Source/Editor/UnrealEd/Private/EdMode.cpp
@@ -1115,16 +1115,3 @@ bool FEdMode::IsSnapRotationEnabled()
 {
 	return GetDefault<ULevelEditorViewportSettings>()->RotGridEnabled;
 }
-
-/*------------------------------------------------------------------------------
-	Default.
-------------------------------------------------------------------------------*/
-
-FEdModeDefault::FEdModeDefault()
-{
-	bDrawGrid = true;
-	bDrawPivot = false;
-	bDrawBaseInfo = false;
-	bDrawWorldBox = false;
-	bDrawKillZ = false;
-}
diff --git a/Engine/Source/Editor/UnrealEd/Private/EditorModeRegistry.cpp b/Engine/Source/Editor/UnrealEd/Private/EditorModeRegistry.cpp
--- a/Engine/Source/Editor/UnrealEd/Private/EditorModeRegistry.cpp
+++ b/Engine/Source/Editor/UnrealEd/Private/EditorModeRegistry.cpp
@@ -50,20 +50,6 @@ void FEditorModeRegistry::Initialize()
 		AssetEditorSubsystem->OnEditorModeRegistered().Broadcast(ModeEntry.Key);
 	}
 
-	if(!GetDefault<UEditorStyleSettings>()->bEnableLegacyEditorModeUI)
-	{
-		// Add default editor modes
-		RegisterMode<FEdModeDefault>(
-			FBuiltinEditorModes::EM_Default,
-			NSLOCTEXT("DefaultMode", "DisplayName", "Select"),
-			FSlateIcon(FEditorStyle::GetStyleSetName(), "LevelEditor.SelectMode", "LevelEditor.SelectMode.Small"),
-			true, 0);
-	}
-	else
-	{
-		RegisterMode<FEdModeDefault>(FBuiltinEditorModes::EM_Default);
-	}
-
 	RegisterMode<FEdModeInterpEdit>(FBuiltinEditorModes::EM_InterpEdit);
 
 	FModuleManager::LoadModuleChecked<IPlacementModeModule>(TEXT("PlacementMode"));
diff --git a/Engine/Source/Editor/UnrealEd/Private/Tools/UEdMode.cpp b/Engine/Source/Editor/UnrealEd/Private/Tools/UEdMode.cpp
--- a/Engine/Source/Editor/UnrealEd/Private/Tools/UEdMode.cpp
+++ b/Engine/Source/Editor/UnrealEd/Private/Tools/UEdMode.cpp
@@ -217,7 +217,7 @@ void UEdMode::Enter()
 	bPendingDeletion = false;
 
 	// initialize the adapter that attaches the ToolsContext to this FEdMode
-	ToolsContext = NewObject<UEdModeInteractiveToolsContext>(GetTransientPackage(), ToolsContextClass.Get(), TEXT("ToolsContext"), RF_Transient);
+	ToolsContext = NewObject<UEdModeInteractiveToolsContext>(this, ToolsContextClass.Get(), NAME_None, RF_Transient);
 	ToolsContext->InitializeContextWithEditorModeManager(GetModeManager());
 
 
@@ -394,7 +394,7 @@ void UEdMode::DrawBrackets(FEditorViewportClient* ViewportClient, FViewport* Vie
 
 bool UEdMode::UsesToolkits() const
 {
-	return true;
+	return Toolkit.IsValid();
 }
 
 UWorld* UEdMode::GetWorld() const
@@ -443,12 +443,10 @@ void UEdMode::CreateToolkit()
 	{
 		Toolkit = MakeShareable(new FModeToolkit);
 		Toolkit->Init(Owner->GetToolkitHost());
-
-
 	}
 
 	UClass* LoadedSettingsObject = SettingsClass.LoadSynchronous();
-	SettingsObject = NewObject<UObject>(this, LoadedSettingsObject);
+	SettingsObject = LoadedSettingsObject ? NewObject<UObject>(this, LoadedSettingsObject) : nullptr;
 	if (SettingsObject)
 	{
 		Toolkit->SetModeSettingsObject(SettingsObject);
@@ -459,3 +457,23 @@ bool UEdMode::IsSnapRotationEnabled()
 {
 	return GetDefault<ULevelEditorViewportSettings>()->RotGridEnabled;
 }
+
+UEdModeDefault::UEdModeDefault()
+{
+	bDrawGrid = true;
+	bDrawPivot = false;
+	bDrawBaseInfo = false;
+	bDrawWorldBox = false;
+	bDrawKillZ = false;
+
+	Info = FEditorModeInfo(FBuiltinEditorModes::EM_Default,
+		NSLOCTEXT("DefaultMode", "DisplayName", "Select"),
+		FSlateIcon(FEditorStyle::GetStyleSetName(), "LevelEditor.SelectMode", "LevelEditor.SelectMode.Small"),
+		true, 0);
+}
+
+void UEdModeDefault::CreateToolkit()
+{
+	// We don't need a toolkit for the default mode, since there isn't any UI or details to show.
+	// It's just a placeholder mode, so that the current implementation of the mode manager will always have at least one active mode.
+}
diff --git a/Engine/Source/Editor/UnrealEd/Public/EdMode.h b/Engine/Source/Editor/UnrealEd/Public/EdMode.h
--- a/Engine/Source/Editor/UnrealEd/Public/EdMode.h
+++ b/Engine/Source/Editor/UnrealEd/Public/EdMode.h
@@ -398,21 +398,3 @@ protected:
 	/** Indicates  */
 	bool bEditedPropertyIsTransform;
 };
-
-/*------------------------------------------------------------------------------
-	Default.
-------------------------------------------------------------------------------*/
-
-/**
- * The default editing mode. 
- */
-class FEdModeDefault : public FEdMode
-{
-public:
-	FEdModeDefault();
-
-	// FEdMode interface
-	virtual bool UsesPropertyWidgets() const override { return true; }
-	// End of FEdMode interface
-};
-
diff --git a/Engine/Source/Editor/UnrealEd/Public/Tools/UEdMode.h b/Engine/Source/Editor/UnrealEd/Public/Tools/UEdMode.h
--- a/Engine/Source/Editor/UnrealEd/Public/Tools/UEdMode.h
+++ b/Engine/Source/Editor/UnrealEd/Public/Tools/UEdMode.h
@@ -304,3 +304,20 @@ protected:
 
 };
 
+/*------------------------------------------------------------------------------
+	Default.
+------------------------------------------------------------------------------*/
+
+/**
+ * The default editing mode.
+ */
+UCLASS()
+class UEdModeDefault : public UEdMode
+{
+	GENERATED_BODY()
+
+public:
+	UEdModeDefault();
+
+	virtual void CreateToolkit() override;
+};
