Fix for nounity error

diff --git a/Engine/Source/Editor/UnrealEd/Public/EditorFolderUtils.h b/Engine/Source/Editor/UnrealEd/Public/EditorFolderUtils.h
--- a/Engine/Source/Editor/UnrealEd/Public/EditorFolderUtils.h
+++ b/Engine/Source/Editor/UnrealEd/Public/EditorFolderUtils.h
@@ -3,6 +3,7 @@
 #pragma once
 
 #include "CoreMinimal.h"
+#include "Misc/Paths.h"
 
 class UNREALED_API FEditorFolderUtils
 {
