Fix for nounity error

diff --git a/Engine/Source/Editor/UnrealEd/Private/Kismet2/DebuggerCommands.cpp b/Engine/Source/Editor/UnrealEd/Private/Kismet2/DebuggerCommands.cpp
--- a/Engine/Source/Editor/UnrealEd/Private/Kismet2/DebuggerCommands.cpp
+++ b/Engine/Source/Editor/UnrealEd/Private/Kismet2/DebuggerCommands.cpp
@@ -72,7 +72,7 @@
 #include "ToolMenus.h"
 #include "SBlueprintEditorToolbar.h"
 #include "SEnumCombobox.h"
-
+#include "Dialogs/Dialogs.h"
 
 #include "Misc/DataDrivenPlatformInfoRegistry.h"
 #include "IUATHelperModule.h"
