WorldPartitionConvertCommandlet
- Moved package / source control utility methods to an helper class
- Reused in the BuildHLOD commandlet

diff --git a/Engine/Source/Editor/UnrealEd/Classes/Commandlets/CommandletPackageHelper.h b/Engine/Source/Editor/UnrealEd/Classes/Commandlets/CommandletPackageHelper.h
new file mode 100644
--- /dev/null
+++ b/Engine/Source/Editor/UnrealEd/Classes/Commandlets/CommandletPackageHelper.h
@@ -0,0 +1,27 @@
+// Copyright Epic Games, Inc. All Rights Reserved.
+
+#pragma once
+
+#include "CoreMinimal.h"
+#include "UObject/ObjectMacros.h"
+#include "SourceControlHelpers.h"
+
+class UNREALED_API FCommandletPackageHelper
+{
+public:
+	FCommandletPackageHelper();
+
+	void SetSourceControlEnabled(bool bWithSourceControl);
+	bool UseSourceControl() const { return SourceControlProvider != nullptr; }
+	ISourceControlProvider& GetSourceControlProvider() const { check(UseSourceControl()); return *SourceControlProvider; }
+
+	bool Delete(const FString& PackageName) const;
+	bool Delete(UPackage* Package) const;
+	bool AddToSourceControl(UPackage* Package) const;
+	bool Checkout(UPackage* Package) const;
+	bool Save(UPackage* Package) const;
+
+private:
+	FScopedSourceControl SourceControl;
+	mutable ISourceControlProvider* SourceControlProvider;
+};
diff --git a/Engine/Source/Editor/UnrealEd/Classes/Commandlets/WorldPartitionConvertCommandlet.h b/Engine/Source/Editor/UnrealEd/Classes/Commandlets/WorldPartitionConvertCommandlet.h
--- a/Engine/Source/Editor/UnrealEd/Classes/Commandlets/WorldPartitionConvertCommandlet.h
+++ b/Engine/Source/Editor/UnrealEd/Classes/Commandlets/WorldPartitionConvertCommandlet.h
@@ -5,8 +5,10 @@
 #include "CoreMinimal.h"
 #include "UObject/ObjectMacros.h"
 #include "UObject/UObjectAnnotation.h"
+
 #include "Commandlets/Commandlet.h"
-#include "ISourceControlProvider.h"
+#include "Commandlets/CommandletPackageHelper.h"
+
 #include "WorldPartition/WorldPartitionEditorHash.h"
 #include "WorldPartition/WorldPartitionRuntimeHash.h"
 #include "WorldPartition/HLOD/HLODLayer.h"
@@ -54,16 +56,9 @@ protected:
 	UWorldPartition* CreateWorldPartition(class AWorldSettings* MainWorldSettings, UWorldComposition* WorldComposition) const;
 	ULevel* LoadLevel(const FString& LevelToLoad);
 
-	bool UseSourceControl() const { return SourceControlProvider != nullptr; }
-	ISourceControlProvider& GetSourceControlProvider() { check(UseSourceControl()); return *SourceControlProvider; }
-
 	void ChangeObjectOuter(UObject* Object, UObject* NewOuter);
 	void FixupSoftObjectPaths(UPackage* OuterPackage);
 
-	bool DeleteFile(const FString& Filename);
-	bool AddPackageToSourceControl(UPackage* Package);
-	bool CheckoutPackage(UPackage* Package);
-	bool SavePackage(UPackage* Package);
 	bool DetachDependantLevelPackages(ULevel* Level);
 	bool RenameWorldPackageWithSuffix(UWorld* World);
 
@@ -78,14 +73,13 @@ protected:
 	TSet<FString> ActorsInGroupActors;
 	TSet<FString> ActorsReferencesToActors;
 
-	ISourceControlProvider* SourceControlProvider;
 	TMap<FString, FString> RemapSoftObjectPaths;
 
 	FString LevelConfigFilename;
 	TArray<UPackage*> PackagesToSave;
 	TArray<UPackage*> PackagesToDelete;
+	FCommandletPackageHelper PackageHelper;
 
-	bool bNoSourceControl;
 	bool bDeleteSourceLevels;
 	bool bGenerateIni;
 	bool bReportOnly;
diff --git a/Engine/Source/Editor/UnrealEd/Private/Commandlets/CommandletPackageHelper.cpp b/Engine/Source/Editor/UnrealEd/Private/Commandlets/CommandletPackageHelper.cpp
new file mode 100644
--- /dev/null
+++ b/Engine/Source/Editor/UnrealEd/Private/Commandlets/CommandletPackageHelper.cpp
@@ -0,0 +1,187 @@
+// Copyright Epic Games, Inc. All Rights Reserved.
+
+/*=============================================================================================================
+ CommandletPackageHelper.cpp: Utility class that provides tools to handle packages & source control operations.
+=============================================================================================================*/
+
+#include "Commandlets/CommandletPackageHelper.h"
+#include "Logging/LogMacros.h"
+#include "UObject/Package.h"
+#include "HAL/FileManager.h"
+#include "HAL/PlatformFileManager.h"
+#include "ISourceControlOperation.h"
+#include "ISourceControlModule.h"
+#include "ISourceControlState.h"
+#include "ISourceControlProvider.h"
+#include "SourceControlOperations.h"
+
+DEFINE_LOG_CATEGORY_STATIC(LogCommandletPackageHelper, Log, All);
+
+FCommandletPackageHelper::FCommandletPackageHelper()
+	: SourceControlProvider(nullptr)
+{
+}
+
+void FCommandletPackageHelper::SetSourceControlEnabled(bool bWithSourceControl)
+{
+	SourceControlProvider = bWithSourceControl ? &ISourceControlModule::Get().GetProvider() : nullptr;
+}
+
+bool FCommandletPackageHelper::Delete(const FString& PackageName) const
+{
+	FString Filename = SourceControlHelpers::PackageFilename(PackageName);
+
+	UE_LOG(LogCommandletPackageHelper, Verbose, TEXT("Deleting %s"), *Filename);
+	
+	if (!UseSourceControl())
+	{
+		if (!IPlatformFile::GetPlatformPhysical().SetReadOnly(*Filename, false) ||
+			!IPlatformFile::GetPlatformPhysical().DeleteFile(*Filename))
+		{
+			UE_LOG(LogCommandletPackageHelper, Error, TEXT("Error deleting %s"), *Filename);
+			return false;
+		}
+	}
+	else
+	{
+		FSourceControlStatePtr SourceControlState = GetSourceControlProvider().GetState(Filename, EStateCacheUsage::ForceUpdate);
+
+		if (SourceControlState.IsValid() && SourceControlState->IsSourceControlled())
+		{
+			FString OtherCheckedOutUser;
+			if (SourceControlState->IsCheckedOutOther(&OtherCheckedOutUser))
+			{
+				UE_LOG(LogCommandletPackageHelper, Error, TEXT("Overwriting package %s already checked out by %s, will not submit"), *Filename, *OtherCheckedOutUser);
+				return false;
+			}
+			else if (!SourceControlState->IsCurrent())
+			{
+				UE_LOG(LogCommandletPackageHelper, Error, TEXT("Overwriting package %s (not at head revision), will not submit"), *Filename);
+				return false;
+			}
+			else if (SourceControlState->IsAdded())
+			{
+				if (GetSourceControlProvider().Execute(ISourceControlOperation::Create<FRevert>(), Filename) != ECommandResult::Succeeded)
+				{
+					UE_LOG(LogCommandletPackageHelper, Error, TEXT("Error reverting package %s from source control"), *Filename);
+					return false;
+				}
+			}
+			else
+			{
+				UE_LOG(LogCommandletPackageHelper, Log, TEXT("Deleting package %s from source control"), *Filename);
+
+				if (SourceControlState->IsCheckedOut())
+				{
+					if (GetSourceControlProvider().Execute(ISourceControlOperation::Create<FRevert>(), Filename) != ECommandResult::Succeeded)
+					{
+						UE_LOG(LogCommandletPackageHelper, Error, TEXT("Error reverting package %s from source control"), *Filename);
+						return false;
+					}
+				}
+
+				if (GetSourceControlProvider().Execute(ISourceControlOperation::Create<FDelete>(), Filename) != ECommandResult::Succeeded)
+				{
+					UE_LOG(LogCommandletPackageHelper, Error, TEXT("Error deleting package %s from source control"), *Filename);
+					return false;
+				}
+			}
+		}
+		else
+		{
+			if (!IFileManager::Get().Delete(*Filename, false, true))
+			{
+				UE_LOG(LogCommandletPackageHelper, Error, TEXT("Error deleting package %s locally"), *Filename);
+				return false;
+			}
+		}
+	}
+
+	return true;
+}
+
+bool FCommandletPackageHelper::Delete(UPackage* Package) const
+{
+	return Delete(Package->GetName());
+}
+
+bool FCommandletPackageHelper::AddToSourceControl(UPackage* Package) const
+{
+	if (UseSourceControl())
+	{
+		FString PackageFilename = SourceControlHelpers::PackageFilename(Package);
+		FSourceControlStatePtr SourceControlState = GetSourceControlProvider().GetState(PackageFilename, EStateCacheUsage::ForceUpdate);
+
+		if (SourceControlState.IsValid() && !SourceControlState->IsSourceControlled())
+		{
+			UE_LOG(LogCommandletPackageHelper, Log, TEXT("Adding package %s to source control"), *PackageFilename);
+			if (GetSourceControlProvider().Execute(ISourceControlOperation::Create<FMarkForAdd>(), Package) != ECommandResult::Succeeded)
+			{
+				UE_LOG(LogCommandletPackageHelper, Error, TEXT("Error adding %s to source control."), *PackageFilename);
+				return false;
+			}
+		}
+	}
+
+	return true;
+}
+
+bool FCommandletPackageHelper::Save(UPackage* Package) const
+{
+	FString PackageFileName = SourceControlHelpers::PackageFilename(Package);
+	if (!UPackage::SavePackage(Package, nullptr, RF_Standalone, *PackageFileName, GError, nullptr, false, true, SAVE_Async))
+	{
+		UE_LOG(LogCommandletPackageHelper, Error, TEXT("Error saving %s"), *PackageFileName);
+		return false;
+	}
+
+	return true;
+}
+
+bool FCommandletPackageHelper::Checkout(UPackage* Package) const
+{
+	if (UseSourceControl())
+	{
+		FString PackageFilename = SourceControlHelpers::PackageFilename(Package);
+		FSourceControlStatePtr SourceControlState = GetSourceControlProvider().GetState(PackageFilename, EStateCacheUsage::ForceUpdate);
+
+		if (SourceControlState.IsValid())
+		{
+			FString OtherCheckedOutUser;
+			if (SourceControlState->IsCheckedOutOther(&OtherCheckedOutUser))
+			{
+				UE_LOG(LogCommandletPackageHelper, Error, TEXT("Overwriting package %s already checked out by %s, will not submit"), *PackageFilename, *OtherCheckedOutUser);
+				return false;
+			}
+			else if (!SourceControlState->IsCurrent())
+			{
+				UE_LOG(LogCommandletPackageHelper, Error, TEXT("Overwriting package %s (not at head revision), will not submit"), *PackageFilename);
+				return false;
+			}
+			else if (SourceControlState->IsCheckedOut() || SourceControlState->IsAdded())
+			{
+				UE_LOG(LogCommandletPackageHelper, Log, TEXT("Skipping package %s (already checked out)"), *PackageFilename);
+				return true;
+			}
+			else if (SourceControlState->IsSourceControlled())
+			{
+				UE_LOG(LogCommandletPackageHelper, Log, TEXT("Checking out package %s from source control"), *PackageFilename);
+				return GetSourceControlProvider().Execute(ISourceControlOperation::Create<FCheckOut>(), Package) == ECommandResult::Succeeded;
+			}
+		}
+	}
+	else
+	{
+		FString PackageFilename = SourceControlHelpers::PackageFilename(Package);
+		if (IPlatformFile::GetPlatformPhysical().FileExists(*PackageFilename))
+		{
+			if (!IPlatformFile::GetPlatformPhysical().SetReadOnly(*PackageFilename, false))
+			{
+				UE_LOG(LogCommandletPackageHelper, Error, TEXT("Error setting %s writable"), *PackageFilename);
+				return false;
+			}
+		}
+	}
+
+	return true;
+}
diff --git a/Engine/Source/Editor/UnrealEd/Private/Commandlets/WorldPartitionBuildHLODsCommandlet.cpp b/Engine/Source/Editor/UnrealEd/Private/Commandlets/WorldPartitionBuildHLODsCommandlet.cpp
--- a/Engine/Source/Editor/UnrealEd/Private/Commandlets/WorldPartitionBuildHLODsCommandlet.cpp
+++ b/Engine/Source/Editor/UnrealEd/Private/Commandlets/WorldPartitionBuildHLODsCommandlet.cpp
@@ -27,6 +27,9 @@ int32 UWorldPartitionBuildHLODsCommandlet::Main(const FString& Params)
 	TArray<FString> Tokens, Switches;
 	ParseCommandLine(*Params, Tokens, Switches);
 
+	bool bNoSourceControl = Switches.Contains(TEXT("NoSourceControl"));
+	PackageHelper.SetSourceControlEnabled(!bNoSourceControl);
+
 	// Retrieve map name from the commandline
 	FString MapName;
 	bool bMapProvided = FParse::Value(*Params, TEXT("Map="), MapName);
@@ -96,14 +99,75 @@ int32 UWorldPartitionBuildHLODsCommandlet::Main(const FString& Params)
 	const FBox LoadBox(FVector(-WORLD_MAX, -WORLD_MAX, -WORLD_MAX), FVector(WORLD_MAX, WORLD_MAX, WORLD_MAX));
 	WorldPartition->LoadEditorCells(LoadBox);
 
+	TSet<UPackage*> PackagesToDelete;
+	TSet<UPackage*> PackagesToSave;
+
 	// Clear all existing HLOD actors
 	for (TActorIterator<AWorldPartitionHLOD> ItHLOD(World); ItHLOD; ++ItHLOD)
 	{
+		PackagesToDelete.Add(ItHLOD->GetPackage());
 		World->DestroyActor(*ItHLOD);
 	}
 
 	// Rebuild HLOD for the whole world
 	WorldPartition->GenerateHLOD();
+	
+	// Required to clear any deleted actors from the level
+	CollectGarbage(RF_Standalone);
+
+	// Gather all packages
+	for (TActorIterator<AWorldPartitionHLOD> ItHLOD(World); ItHLOD; ++ItHLOD)
+	{
+		UPackage* Package = ItHLOD->GetPackage();
+		if (Package && Package->IsDirty())
+		{
+			PackagesToSave.Add(Package);
+		}
+	}
+
+	// Delete packages
+	UE_LOG(LogWorldPartitionBuildHLODsCommandlet, Log, TEXT("Deleting %d packages."), PackagesToDelete.Num());
+	for (UPackage* Package : PackagesToDelete)
+	{
+		if (!PackageHelper.Delete(Package->GetName()))
+		{
+			UE_LOG(LogWorldPartitionBuildHLODsCommandlet, Error, TEXT("Error deleting package %s."), *Package->GetName());
+			return 1;
+		}
+	}
+
+	// Checkout packages
+	UE_LOG(LogWorldPartitionBuildHLODsCommandlet, Log, TEXT("Checking out %d actor packages."), PackagesToSave.Num());
+	for (UPackage* Package : PackagesToSave)
+	{
+		if (!PackageHelper.Checkout(Package))
+		{
+			UE_LOG(LogWorldPartitionBuildHLODsCommandlet, Error, TEXT("Error checking out package %s."), *Package->GetName());
+			return 1;
+		}
+	}
+
+	// Save packages
+	UE_LOG(LogWorldPartitionBuildHLODsCommandlet, Log, TEXT("Saving %d packages."), PackagesToSave.Num());
+	for (UPackage* Package : PackagesToSave)
+	{
+		if (!PackageHelper.Save(Package))
+		{
+			UE_LOG(LogWorldPartitionBuildHLODsCommandlet, Error, TEXT("Error saving package %s."), *Package->GetName());
+			return 1;
+		}
+	}
+
+	// Add new packages to source control
+	UE_LOG(LogWorldPartitionBuildHLODsCommandlet, Log, TEXT("Adding packages to source control."));
+	for (UPackage* Package : PackagesToSave)
+	{
+		if (!PackageHelper.AddToSourceControl(Package))
+		{
+			UE_LOG(LogWorldPartitionBuildHLODsCommandlet, Error, TEXT("Error adding package %s to source control."), *Package->GetName());
+			return 1;
+		}
+	}
 
 	// Cleanup
 	World->RemoveFromRoot();
diff --git a/Engine/Source/Editor/UnrealEd/Private/Commandlets/WorldPartitionBuildHLODsCommandlet.h b/Engine/Source/Editor/UnrealEd/Private/Commandlets/WorldPartitionBuildHLODsCommandlet.h
--- a/Engine/Source/Editor/UnrealEd/Private/Commandlets/WorldPartitionBuildHLODsCommandlet.h
+++ b/Engine/Source/Editor/UnrealEd/Private/Commandlets/WorldPartitionBuildHLODsCommandlet.h
@@ -3,6 +3,7 @@
 #pragma once
 
 #include "Commandlets/Commandlet.h"
+#include "Commandlets/CommandletPackageHelper.h"
 #include "WorldPartitionBuildHLODsCommandlet.generated.h"
 
 UCLASS()
@@ -13,4 +14,7 @@ class UWorldPartitionBuildHLODsCommandlet : public UCommandlet
 	//~ Begin UCommandlet Interface
 	virtual int32 Main(const FString& Params) override;
 	//~ End UCommandlet Interface
+
+private:
+	FCommandletPackageHelper PackageHelper;
 };
\ No newline at end of file
diff --git a/Engine/Source/Editor/UnrealEd/Private/Commandlets/WorldPartitionConvertCommandlet.cpp b/Engine/Source/Editor/UnrealEd/Private/Commandlets/WorldPartitionConvertCommandlet.cpp
--- a/Engine/Source/Editor/UnrealEd/Private/Commandlets/WorldPartitionConvertCommandlet.cpp
+++ b/Engine/Source/Editor/UnrealEd/Private/Commandlets/WorldPartitionConvertCommandlet.cpp
@@ -6,7 +6,6 @@
 
 #include "Commandlets/WorldPartitionConvertCommandlet.h"
 #include "Algo/Transform.h"
-#include "HAL/FileManager.h"
 #include "HAL/PlatformFileManager.h"
 #include "Misc/Paths.h"
 #include "Misc/ConfigCacheIni.h"
@@ -24,10 +23,6 @@
 #include "GameFramework/WorldSettings.h"
 #include "UObject/UObjectHash.h"
 #include "PackageHelperFunctions.h"
-#include "ISourceControlOperation.h"
-#include "SourceControlOperations.h"
-#include "SourceControlHelpers.h"
-#include "ISourceControlModule.h"
 #include "UObject/MetaData.h"
 #include "Editor.h"
 #include "HierarchicalLOD.h"
@@ -308,162 +303,10 @@ bool UWorldPartitionConvertCommandlet::ShouldDeleteActor(AActor* Actor, bool bMa
 	return false;
 }
 
-bool UWorldPartitionConvertCommandlet::DeleteFile(const FString& Filename)
-{
-	UE_LOG(LogWorldPartitionConvertCommandlet, Verbose, TEXT("Deleting %s"), *Filename);
-
-	if (!UseSourceControl())
-	{
-		if (!IPlatformFile::GetPlatformPhysical().SetReadOnly(*Filename, false) ||
-			!IPlatformFile::GetPlatformPhysical().DeleteFile(*Filename))
-		{
-			UE_LOG(LogWorldPartitionConvertCommandlet, Error, TEXT("Error deleting %s"), *Filename);
-			return false;
-		}
-	}
-	else 
-	{
-		FSourceControlStatePtr SourceControlState = GetSourceControlProvider().GetState(Filename, EStateCacheUsage::ForceUpdate);
-
-		if (SourceControlState.IsValid() && SourceControlState->IsSourceControlled())
-		{
-			FString OtherCheckedOutUser;
-			if (SourceControlState->IsCheckedOutOther(&OtherCheckedOutUser))
-			{
-				UE_LOG(LogWorldPartitionConvertCommandlet, Error, TEXT("Overwriting package %s already checked out by %s, will not submit"), *Filename, *OtherCheckedOutUser);
-				return false;
-			}
-			else if (!SourceControlState->IsCurrent())
-			{
-				UE_LOG(LogWorldPartitionConvertCommandlet, Error, TEXT("Overwriting package %s (not at head revision), will not submit"), *Filename);
-				return false;
-			}
-			else if (SourceControlState->IsAdded())
-			{
-				if (GetSourceControlProvider().Execute(ISourceControlOperation::Create<FRevert>(), Filename) != ECommandResult::Succeeded)
-				{
-					UE_LOG(LogWorldPartitionConvertCommandlet, Error, TEXT("Error reverting package %s from source control"), *Filename);
-					return false;
-				}
-			}
-			else
-			{
-				UE_LOG(LogWorldPartitionConvertCommandlet, Log, TEXT("Deleting package %s from source control"), *Filename);
-
-				if (SourceControlState->IsCheckedOut())
-				{
-					if (GetSourceControlProvider().Execute(ISourceControlOperation::Create<FRevert>(), Filename) != ECommandResult::Succeeded)
-					{
-						UE_LOG(LogWorldPartitionConvertCommandlet, Error, TEXT("Error reverting package %s from source control"), *Filename);
-						return false;
-					}
-				}
-
-				if (GetSourceControlProvider().Execute(ISourceControlOperation::Create<FDelete>(), Filename) != ECommandResult::Succeeded)
-				{
-					UE_LOG(LogWorldPartitionConvertCommandlet, Error, TEXT("Error deleting package %s from source control"), *Filename);
-					return false;
-				}
-			}
-		}
-		else
-		{
-			if (!IFileManager::Get().Delete(*Filename, false, true))
-			{
-				UE_LOG(LogWorldPartitionConvertCommandlet, Error, TEXT("Error deleting package %s locally"), *Filename);
-				return false;
-			}
-		}
-	}
-
-	return true;
-}
-
 void UWorldPartitionConvertCommandlet::PerformAdditionalWorldCleanup(UWorld* World) const
 {
 }
 
-bool UWorldPartitionConvertCommandlet::AddPackageToSourceControl(UPackage* Package)
-{
-	if (UseSourceControl())
-	{
-		FString PackageFilename = SourceControlHelpers::PackageFilename(Package);
-		FSourceControlStatePtr SourceControlState = GetSourceControlProvider().GetState(PackageFilename, EStateCacheUsage::ForceUpdate);
-
-		if (SourceControlState.IsValid() && !SourceControlState->IsSourceControlled())
-		{
-			UE_LOG(LogWorldPartitionConvertCommandlet, Log, TEXT("Adding package %s to source control"), *PackageFilename);
-			if (GetSourceControlProvider().Execute(ISourceControlOperation::Create<FMarkForAdd>(), Package) != ECommandResult::Succeeded)
-			{
-				UE_LOG(LogWorldPartitionConvertCommandlet, Error, TEXT("Error adding %s to source control."), *PackageFilename);
-				return false;
-			}
-		}
-	}
-
-	return true;
-}
-
-bool UWorldPartitionConvertCommandlet::SavePackage(UPackage* Package)
-{
-	FString PackageFileName = SourceControlHelpers::PackageFilename(Package);
-	if (!UPackage::SavePackage(Package, nullptr, RF_Standalone, *PackageFileName, GError, nullptr, false, true, SAVE_Async))
-	{
-		UE_LOG(LogWorldPartitionConvertCommandlet, Error, TEXT("Error saving %s"), *PackageFileName);
-		return false;
-	}
-
-	return true;
-}
-
-bool UWorldPartitionConvertCommandlet::CheckoutPackage(UPackage* Package)
-{
-	if (UseSourceControl())
-	{
-		FString PackageFilename = SourceControlHelpers::PackageFilename(Package);
-		FSourceControlStatePtr SourceControlState = GetSourceControlProvider().GetState(PackageFilename, EStateCacheUsage::ForceUpdate);
-
-		if (SourceControlState.IsValid())
-		{
-			FString OtherCheckedOutUser;
-			if (SourceControlState->IsCheckedOutOther(&OtherCheckedOutUser))
-			{
-				UE_LOG(LogWorldPartitionConvertCommandlet, Error, TEXT("Overwriting package %s already checked out by %s, will not submit"), *PackageFilename, *OtherCheckedOutUser);
-				return false;
-			}
-			else if (!SourceControlState->IsCurrent())
-			{
-				UE_LOG(LogWorldPartitionConvertCommandlet, Error, TEXT("Overwriting package %s (not at head revision), will not submit"), *PackageFilename);
-				return false;
-			}
-			else if (SourceControlState->IsCheckedOut() || SourceControlState->IsAdded())
-			{
-				UE_LOG(LogWorldPartitionConvertCommandlet, Log, TEXT("Skipping package %s (already checked out)"), *PackageFilename);
-				return true;
-			}
-			else if (SourceControlState->IsSourceControlled())
-			{
-				UE_LOG(LogWorldPartitionConvertCommandlet, Log, TEXT("Checking out package %s from source control"), *PackageFilename);
-				return GetSourceControlProvider().Execute(ISourceControlOperation::Create<FCheckOut>(), Package) == ECommandResult::Succeeded;
-			}
-		}
-	}
-	else
-	{
-		FString PackageFilename = SourceControlHelpers::PackageFilename(Package);
-		if (IPlatformFile::GetPlatformPhysical().FileExists(*PackageFilename))
-		{
-			if (!IPlatformFile::GetPlatformPhysical().SetReadOnly(*PackageFilename, false))
-			{
-				UE_LOG(LogWorldPartitionConvertCommandlet, Error, TEXT("Error setting %s writable"), *PackageFilename);
-				return false;
-			}
-		}
-	}
-
-	return true;
-}
-
 void UWorldPartitionConvertCommandlet::OutputConversionReport() const
 {
 	UE_LOG(LogWorldPartitionConvertCommandlet, Display, TEXT("WorldPartitionConvertCommandlet report:"));
@@ -797,7 +640,6 @@ int32 UWorldPartitionConvertCommandlet::Main(const FString& Params)
 		return 1;
 	}
 
-	bNoSourceControl = Switches.Contains(TEXT("NoSourceControl"));
 	bDeleteSourceLevels = Switches.Contains(TEXT("DeleteSourceLevels"));
 	bGenerateIni = Switches.Contains(TEXT("GenerateIni"));
 	bReportOnly = bGenerateIni || Switches.Contains(TEXT("ReportOnly"));
@@ -812,8 +654,8 @@ int32 UWorldPartitionConvertCommandlet::Main(const FString& Params)
 
 	bConversionSuffix = Switches.Contains(TEXT("ConversionSuffix"));
 
-	FScopedSourceControl SourceControl;
-	SourceControlProvider = bNoSourceControl ? nullptr : &ISourceControlModule::Get().GetProvider();
+	bool bNoSourceControl = Switches.Contains(TEXT("NoSourceControl"));
+	PackageHelper.SetSourceControlEnabled(!bNoSourceControl);
 
 	// Load configuration file
 	FString LevelLongPackageName;
@@ -864,7 +706,7 @@ int32 UWorldPartitionConvertCommandlet::Main(const FString& Params)
 
 		for (const FAssetData& Asset : Assets)
 		{
-			if (!DeleteFile(SourceControlHelpers::PackageFilename(Asset.PackageName.ToString())))
+			if (!PackageHelper.Delete(Asset.PackageName.ToString()))
 			{
 				return 1;
 			}
@@ -872,7 +714,7 @@ int32 UWorldPartitionConvertCommandlet::Main(const FString& Params)
 
 		if (FPackageName::SearchForPackageOnDisk(OldLevelName, &OldLevelName))
 		{
-			DeleteFile(SourceControlHelpers::PackageFilename(OldLevelName));
+			PackageHelper.Delete(OldLevelName);
 		}
 	}
 
@@ -1127,7 +969,7 @@ int32 UWorldPartitionConvertCommandlet::Main(const FString& Params)
 		{
 			for (UPackage* Package : PackagesToDelete)
 			{
-				if (!DeleteFile(SourceControlHelpers::PackageFilename(Package)))
+				if (!PackageHelper.Delete(Package))
 				{
 					return 1;
 				}
@@ -1137,7 +979,7 @@ int32 UWorldPartitionConvertCommandlet::Main(const FString& Params)
 		UE_LOG(LogWorldPartitionConvertCommandlet, Log, TEXT("Checking out %d actor packages."), PackagesToSave.Num());
 		for(UPackage* Package: PackagesToSave)
 		{
-			if (!CheckoutPackage(Package))
+			if (!PackageHelper.Checkout(Package))
 			{
 				return 1;
 			}
@@ -1156,7 +998,7 @@ int32 UWorldPartitionConvertCommandlet::Main(const FString& Params)
 		UE_LOG(LogWorldPartitionConvertCommandlet, Log, TEXT("Saving %d packages."), PackagesToSave.Num());
 		for (UPackage* PackageToSave : PackagesToSave)
 		{
-			if (!SavePackage(PackageToSave))
+			if (!PackageHelper.Save(PackageToSave))
 			{
 				return 1;
 			}
@@ -1165,7 +1007,7 @@ int32 UWorldPartitionConvertCommandlet::Main(const FString& Params)
 		// Add new packages to source control
 		for(UPackage* PackageToSave: PackagesToSave)
 		{
-			if(!AddPackageToSourceControl(PackageToSave))
+			if(!PackageHelper.AddToSourceControl(PackageToSave))
 			{
 				return 1;
 			}
@@ -1174,13 +1016,13 @@ int32 UWorldPartitionConvertCommandlet::Main(const FString& Params)
 		// Checkout level
 		UE_LOG(LogWorldPartitionConvertCommandlet, Log, TEXT("Saving %s."), *MainPackage->GetName());
 
-		if (!CheckoutPackage(MainPackage))
+		if (!PackageHelper.Checkout(MainPackage))
 		{
 			return 1;
 		}
 
 		// Save level
-		if (!SavePackage(MainPackage))
+		if (!PackageHelper.Save(MainPackage))
 		{
 			return 1;
 		}
