Fix ensure in ForceDeleteObjects when deleting StaticMeshes that are still referenced in async tasks

diff --git a/Engine/Source/Editor/UnrealEd/Private/ObjectTools.cpp b/Engine/Source/Editor/UnrealEd/Private/ObjectTools.cpp
--- a/Engine/Source/Editor/UnrealEd/Private/ObjectTools.cpp
+++ b/Engine/Source/Editor/UnrealEd/Private/ObjectTools.cpp
@@ -94,6 +94,8 @@
 #include "Subsystems/AssetEditorSubsystem.h"
 #include "TextureCompiler.h"
 #include "UObject/ReferencerFinder.h"
+#include "DistanceFieldAtlas.h"
+#include "MeshCardRepresentation.h"
 
 DEFINE_LOG_CATEGORY_STATIC(LogObjectTools, Log, All);
 
@@ -2486,6 +2488,14 @@ namespace ObjectTools
 			{
 				World->CleanupWorld();
 			}
+
+			// Make sure the object is not still referenced by async tasks
+			UStaticMesh* StaticMesh = Cast<UStaticMesh>(ObjectToDelete);
+			if (StaticMesh != nullptr)
+			{
+				GDistanceFieldAsyncQueue->BlockUntilBuildComplete(StaticMesh, true);
+				GCardRepresentationAsyncQueue->BlockUntilBuildComplete(StaticMesh, true);
+			}
 		}
 
 		if ( bPerformReferenceCheck )
