WorldPartitionConvertCommandlet:
- Add the possibility to override the default grid placement for specific levels via the config file.
- Don't allow converting levels that contains unstable GUIDs.

diff --git a/Engine/Source/Editor/UnrealEd/Classes/Commandlets/WorldPartitionConvertCommandlet.h b/Engine/Source/Editor/UnrealEd/Classes/Commandlets/WorldPartitionConvertCommandlet.h
--- a/Engine/Source/Editor/UnrealEd/Classes/Commandlets/WorldPartitionConvertCommandlet.h
+++ b/Engine/Source/Editor/UnrealEd/Classes/Commandlets/WorldPartitionConvertCommandlet.h
@@ -43,6 +43,7 @@ public:
 
 private:
 	void GatherAndPrepareSubLevelsToConvert(const UWorldPartition* WorldPartition, ULevel* Level, TArray<ULevel*>& SubLevels);
+	EActorGridPlacement GetLevelGridPlacement(ULevel* Level, EActorGridPlacement DefaultGridPlacement);
 
 protected:
 	virtual bool GetAdditionalLevelsToConvert(ULevel* Level, TArray<ULevel*>& SubLevels);
@@ -93,6 +94,9 @@ protected:
 	UPROPERTY(Config)
 	TSubclassOf<UWorldPartitionRuntimeHash> RuntimeHashClass;
 
+	UPROPERTY(Config)
+	TMap<FName, EActorGridPlacement> LevelsGridPlacement;
+
 	UPROPERTY(Config)
 	FString HLODLayerAssetsPath;
 
diff --git a/Engine/Source/Editor/UnrealEd/Private/Commandlets/WorldPartitionConvertCommandlet.cpp b/Engine/Source/Editor/UnrealEd/Private/Commandlets/WorldPartitionConvertCommandlet.cpp
--- a/Engine/Source/Editor/UnrealEd/Private/Commandlets/WorldPartitionConvertCommandlet.cpp
+++ b/Engine/Source/Editor/UnrealEd/Private/Commandlets/WorldPartitionConvertCommandlet.cpp
@@ -238,6 +238,16 @@ void UWorldPartitionConvertCommandlet::GatherAndPrepareSubLevelsToConvert(const
 	}
 }
 
+EActorGridPlacement UWorldPartitionConvertCommandlet::GetLevelGridPlacement(ULevel* Level, EActorGridPlacement DefaultGridPlacement)
+{
+	FString WorldPath = Level->GetPackage()->GetName();
+	if (EActorGridPlacement* CustomLevelGridPlacement = LevelsGridPlacement.Find(*WorldPath))
+	{
+		return *CustomLevelGridPlacement;
+	}
+	return DefaultGridPlacement;
+}
+
 bool UWorldPartitionConvertCommandlet::PrepareStreamingLevelForConversion(const UWorldPartition* WorldPartition, ULevelStreaming* StreamingLevel)
 {
 	ULevel* SubLevel = StreamingLevel->GetLoadedLevel();
@@ -245,13 +255,17 @@ bool UWorldPartitionConvertCommandlet::PrepareStreamingLevelForConversion(const
 
 	if (StreamingLevel->ShouldBeAlwaysLoaded() || StreamingLevel->bDisableDistanceStreaming)
 	{
-		UE_LOG(LogWorldPartitionConvertCommandlet, Log, TEXT("Converting %s streaming level %s"), StreamingLevel->bDisableDistanceStreaming ? TEXT("non distance-based") : TEXT("always loaded"), *StreamingLevel->GetWorldAssetPackageName());
-
-		for (AActor* Actor: SubLevel->Actors)
+		FString WorldPath = SubLevel->GetPackage()->GetName();
+		if (!LevelsGridPlacement.Contains(*WorldPath))
 		{
-			if (Actor)
+			UE_LOG(LogWorldPartitionConvertCommandlet, Log, TEXT("Converting %s streaming level %s"), StreamingLevel->bDisableDistanceStreaming ? TEXT("non distance-based") : TEXT("always loaded"), *StreamingLevel->GetWorldAssetPackageName());
+
+			for (AActor* Actor: SubLevel->Actors)
 			{
-				Actor->GridPlacement = EActorGridPlacement::AlwaysLoaded;
+				if (Actor)
+				{
+					Actor->GridPlacement = EActorGridPlacement::AlwaysLoaded;
+				}
 			}
 		}
 	}
@@ -764,7 +778,7 @@ int32 UWorldPartitionConvertCommandlet::Main(const FString& Params)
 		}
 	};
 
-	// Gather sublevels
+	// Gather and load sublevels
 	TArray<ULevel*> SubLevelsToConvert;
 	GatherAndPrepareSubLevelsToConvert(WorldPartition, MainLevel, SubLevelsToConvert);
 
@@ -773,9 +787,26 @@ int32 UWorldPartitionConvertCommandlet::Main(const FString& Params)
 		return 1;
 	}
 
+	// Validate levels for conversion
+	bool bNeedsResaveSubLevels = false;
+
+	for (ULevel* Level: SubLevelsToConvert)
+	{
+		if (!Level->bContainsStableActorGUIDs)
+		{
+			bNeedsResaveSubLevels |= true;
+			UE_LOG(LogWorldPartitionConvertCommandlet, Error, TEXT("Unable to convert level '%s' with non-stable actor GUIDs. Resave the level before converting."), *Level->GetPackage()->GetName());
+		}
+	}
+
+	if (bNeedsResaveSubLevels)
+	{
+		return 1;
+	}
+
 	// Prepare levels for conversion
 	DetachDependantLevelPackages(MainLevel);
-	PrepareLevelActors(MainLevel, true, SubLevelsToConvert.Num() ? EActorGridPlacement::AlwaysLoaded : EActorGridPlacement::Bounds);
+	PrepareLevelActors(MainLevel, true, GetLevelGridPlacement(MainLevel, SubLevelsToConvert.Num() ? EActorGridPlacement::AlwaysLoaded : EActorGridPlacement::Bounds));
 
 	if (bConversionSuffix)
 	{
@@ -796,7 +827,7 @@ int32 UWorldPartitionConvertCommandlet::Main(const FString& Params)
 	for (ULevel* SubLevel : SubLevelsToConvert)
 	{
 		DetachDependantLevelPackages(SubLevel);
-		PrepareLevelActors(SubLevel, false, EActorGridPlacement::Bounds);
+		PrepareLevelActors(SubLevel, false, GetLevelGridPlacement(SubLevel, EActorGridPlacement::Bounds));
 	}
 
 	TMap<UObject*, UObject*> PrivateRefsMap;
