World Partition - Added  "-SkipStableGUIDValidation" option to conversion commandlet

diff --git a/Engine/Source/Editor/UnrealEd/Private/Commandlets/WorldPartitionConvertCommandlet.cpp b/Engine/Source/Editor/UnrealEd/Private/Commandlets/WorldPartitionConvertCommandlet.cpp
--- a/Engine/Source/Editor/UnrealEd/Private/Commandlets/WorldPartitionConvertCommandlet.cpp
+++ b/Engine/Source/Editor/UnrealEd/Private/Commandlets/WorldPartitionConvertCommandlet.cpp
@@ -788,20 +788,24 @@ int32 UWorldPartitionConvertCommandlet::Main(const FString& Params)
 	}
 
 	// Validate levels for conversion
-	bool bNeedsResaveSubLevels = false;
-
-	for (ULevel* Level: SubLevelsToConvert)
+	bool bSkipStableGUIDValidation = Switches.Contains(TEXT("SkipStableGUIDValidation"));
+	if (!bSkipStableGUIDValidation)
 	{
-		if (!Level->bContainsStableActorGUIDs)
+		bool bNeedsResaveSubLevels = false;
+
+		for (ULevel* Level: SubLevelsToConvert)
 		{
-			bNeedsResaveSubLevels |= true;
-			UE_LOG(LogWorldPartitionConvertCommandlet, Error, TEXT("Unable to convert level '%s' with non-stable actor GUIDs. Resave the level before converting."), *Level->GetPackage()->GetName());
+			if (!Level->bContainsStableActorGUIDs)
+			{
+				bNeedsResaveSubLevels |= true;
+				UE_LOG(LogWorldPartitionConvertCommandlet, Error, TEXT("Unable to convert level '%s' with non-stable actor GUIDs. Resave the level before converting."), *Level->GetPackage()->GetName());
+			}
 		}
-	}
 
-	if (bNeedsResaveSubLevels)
-	{
-		return 1;
+		if (bNeedsResaveSubLevels)
+		{
+			return 1;
+		}
 	}
 
 	// Prepare levels for conversion
