- FIxed infinite recursion when re-scanning TargetPlatforms after installing Turnkey SDK

diff --git a/Engine/Source/Editor/UnrealEd/Private/CookOnTheFlyServer.cpp b/Engine/Source/Editor/UnrealEd/Private/CookOnTheFlyServer.cpp
--- a/Engine/Source/Editor/UnrealEd/Private/CookOnTheFlyServer.cpp
+++ b/Engine/Source/Editor/UnrealEd/Private/CookOnTheFlyServer.cpp
@@ -4193,13 +4193,11 @@ void UCookOnTheFlyServer::ProcessAccessedIniSettings(const FConfigFile* Config,
 		return;
 	}
 	// try figure out if this config file is for a specific platform 
-	ITargetPlatformManagerModule& TPM = GetTargetPlatformManagerRef();
-	const TArray<ITargetPlatform*>& Platforms = TPM.GetTargetPlatforms();
 	FString PlatformName;
 	bool bFoundPlatformName = false;
-	for (const ITargetPlatform* Platform : Platforms )
+	for (auto It : FDataDrivenPlatformInfoRegistry::GetAllPlatformInfos())
 	{
-		FString CurrentPlatformName = Platform->IniPlatformName();
+		FString CurrentPlatformName = It.Key.ToString();
 		for ( const auto& SourceIni : Config->SourceIniHierarchy )
 		{
 			if ( SourceIni.Value.Filename.Contains(CurrentPlatformName) )
