Fixing various crashes related to Tool Manager being nullptr when a valid context should exist.

#rnx

diff --git a/Engine/Source/Editor/UnrealEd/Private/EditorModeManager.cpp b/Engine/Source/Editor/UnrealEd/Private/EditorModeManager.cpp
--- a/Engine/Source/Editor/UnrealEd/Private/EditorModeManager.cpp
+++ b/Engine/Source/Editor/UnrealEd/Private/EditorModeManager.cpp
@@ -63,7 +63,7 @@ FEditorModeTools::FEditorModeTools()
 {
 	DefaultModeIDs.Add( FBuiltinEditorModes::EM_Default );
 
-	InteractiveToolsContext = NewObject<UEdModeInteractiveToolsContext>(GetTransientPackage(), UEdModeInteractiveToolsContext::StaticClass(), TEXT("ToolsContext"), RF_Transient);
+	InteractiveToolsContext = NewObject<UEdModeInteractiveToolsContext>(GetTransientPackage(), UEdModeInteractiveToolsContext::StaticClass(), NAME_None, RF_Transient);
 	InteractiveToolsContext->InitializeContextWithEditorModeManager(this);
 
 	// Load the last used settings
