Cleaned up some common BSP deselection logic

#rnx

diff --git a/Engine/Source/Editor/UnrealEd/Private/EditorSelectUtils.cpp b/Engine/Source/Editor/UnrealEd/Private/EditorSelectUtils.cpp
--- a/Engine/Source/Editor/UnrealEd/Private/EditorSelectUtils.cpp
+++ b/Engine/Source/Editor/UnrealEd/Private/EditorSelectUtils.cpp
@@ -726,7 +726,7 @@ void UUnrealEdEngine::SelectBSPSurf(UModel* InModel, int32 iSurf, bool bSelected
 /**
  * Deselects all BSP surfaces in the specified level.
  *
- * @param	Level		The level for which to deselect all levels.
+ * @param	Level		The level for which to deselect all surfaces.
  * @return				The number of surfaces that were deselected
  */
 static uint32 DeselectAllSurfacesForLevel(ULevel* Level)
@@ -749,24 +749,35 @@ static uint32 DeselectAllSurfacesForLevel(ULevel* Level)
 	return NumSurfacesDeselected;
 }
 
-
-void UUnrealEdEngine::DeselectAllSurfaces()
+/**
+ * Deselects all BSP surfaces in the specified world.
+ *
+ * @param	World		The world for which to deselect all surfaces.
+ * @return				The number of surfaces that were deselected
+ */
+static uint32 DeselectAllSurfacesForWorld(UWorld* World)
 {
-	UWorld* World = GWorld;
+	uint32 NumSurfacesDeselected = 0;
 	if (World)
 	{
-		DeselectAllSurfacesForLevel(World->PersistentLevel);
-		for (ULevelStreaming* StreamingLevel: World->GetStreamingLevels())
+		NumSurfacesDeselected += DeselectAllSurfacesForLevel(World->PersistentLevel);
+		for (ULevelStreaming* StreamingLevel : World->GetStreamingLevels())
 		{
 			if (StreamingLevel)
 			{
 				if (ULevel* Level = StreamingLevel->GetLoadedLevel())
 				{
-					DeselectAllSurfacesForLevel(Level);
+					NumSurfacesDeselected += DeselectAllSurfacesForLevel(Level);
 				}
 			}
 		}
 	}
+	return NumSurfacesDeselected;
+}
+
+void UUnrealEdEngine::DeselectAllSurfaces()
+{
+	DeselectAllSurfacesForWorld(GWorld);
 }
 
 void UUnrealEdEngine::SelectNone(bool bNoteSelectionChange, bool bDeselectBSPSurfs, bool WarnAboutManyActors)
@@ -827,21 +838,9 @@ void UUnrealEdEngine::SelectNone(bool bNoteSelectionChange, bool bDeselectBSPSur
 	}
 
 	uint32 NumDeselectSurfaces = 0;
-	UWorld* World = GWorld;
-	if( bDeselectBSPSurfs && World )
+	if( bDeselectBSPSurfs )
 	{
-		// Unselect all surfaces in all levels.
-		NumDeselectSurfaces += DeselectAllSurfacesForLevel( World->PersistentLevel );
-		for (ULevelStreaming* StreamingLevel : World->GetStreamingLevels())
-		{
-			if( StreamingLevel )
-			{
-				if (ULevel* Level = StreamingLevel->GetLoadedLevel())
-				{
-					NumDeselectSurfaces += DeselectAllSurfacesForLevel( Level );
-				}
-			}
-		}
+		NumDeselectSurfaces = DeselectAllSurfacesForWorld(GWorld);
 	}
 
 	SelectedActors->EndBatchSelectOperation(bNoteSelectionChange);
