Edigrate various OFPA fixes from UE4.

diff --git a/Engine/Source/Editor/UnrealEd/Private/Commandlets/ConvertLevelsToExternalActorsCommandlet.cpp b/Engine/Source/Editor/UnrealEd/Private/Commandlets/ConvertLevelsToExternalActorsCommandlet.cpp
--- a/Engine/Source/Editor/UnrealEd/Private/Commandlets/ConvertLevelsToExternalActorsCommandlet.cpp
+++ b/Engine/Source/Editor/UnrealEd/Private/Commandlets/ConvertLevelsToExternalActorsCommandlet.cpp
@@ -71,13 +71,13 @@ bool UConvertLevelsToExternalActorsCommandlet::CheckExternalActors(const FString
 		TMap<FName, FName> ActorFiles;
 
 		IAssetRegistry& AssetRegistry = FModuleManager::LoadModuleChecked<FAssetRegistryModule>(TEXT("AssetRegistry")).Get();
-		AssetRegistry.OnAssetAdded().AddLambda([&ActorFiles](const FAssetData& AssetData)
+		FDelegateHandle AddedCheckHandle = AssetRegistry.OnAssetAdded().AddLambda([&ActorFiles](const FAssetData& AssetData)
 		{
 			check(!ActorFiles.Contains(AssetData.ObjectPath));
 			ActorFiles.Add(AssetData.ObjectPath, AssetData.PackageName);
 		});
 
-		AssetRegistry.OnAssetUpdated().AddLambda([&ActorFiles, &DuplicatedActorFiles](const FAssetData& AssetData)
+		FDelegateHandle UpdatedCheckHandle = AssetRegistry.OnAssetUpdated().AddLambda([&ActorFiles, &DuplicatedActorFiles](const FAssetData& AssetData)
 		{
 			FName ExistingPackageName;
 			if (ActorFiles.RemoveAndCopyValue(AssetData.ObjectPath, ExistingPackageName))
@@ -89,6 +89,9 @@ bool UConvertLevelsToExternalActorsCommandlet::CheckExternalActors(const FString
 		});
 
 		AssetRegistry.ScanPathsSynchronous({LevelExternalPathActors});
+
+		AssetRegistry.OnAssetAdded().Remove(AddedCheckHandle);
+		AssetRegistry.OnAssetUpdated().Remove(UpdatedCheckHandle);
 	}
 
 	if (DuplicatedActorFiles.Num())
@@ -170,8 +173,12 @@ bool UConvertLevelsToExternalActorsCommandlet::AddPackageToSourceControl(UPackag
 
 bool UConvertLevelsToExternalActorsCommandlet::SavePackage(UPackage* Package)
 {
+	// Use GEditor save as it does some UWorld specific shenanigans such as handle level offsets
 	FString PackageFileName = SourceControlHelpers::PackageFilename(Package);
-	if (!UPackage::SavePackage(Package, nullptr, RF_Standalone, *PackageFileName, GError, nullptr, false, true, SAVE_None))
+	FSavePackageResultStruct SaveResult = GEditor->Save(Package, nullptr, RF_Standalone, *PackageFileName,
+		GError, nullptr, false, true, SAVE_None);
+
+	if (SaveResult.Result != ESavePackageResult::Success)
 	{
 		UE_LOG(LogConvertLevelsToExternalActorsCommandlet, Error, TEXT("Error saving %s"), *PackageFileName);
 		return false;
