External Actors Transactions
- Don't dirty packages when applying transactions on objects, as their packages should already be dirty by previous changes when undoing or set when redoing.
- Don't dirty actor packages when changing their levels visibility; the level package still gets dirtied when changing its visibility, as this setting is persistent.

diff --git a/Engine/Source/Editor/UnrealEd/Private/EditorLevelUtils.cpp b/Engine/Source/Editor/UnrealEd/Private/EditorLevelUtils.cpp
--- a/Engine/Source/Editor/UnrealEd/Private/EditorLevelUtils.cpp
+++ b/Engine/Source/Editor/UnrealEd/Private/EditorLevelUtils.cpp
@@ -1025,7 +1025,7 @@ void SetLevelVisibilityNoGlobalUpdateInternal(ULevel* Level, const bool bShouldB
 			{
 				if (ModifyMode == ELevelVisibilityDirtyMode::ModifyOnChange)
 				{
-					CurActor->Modify();
+					CurActor->Modify(false);
 				}
 				
 				CurActor->bHiddenEdLevel = !bShouldBeVisible;
@@ -1158,7 +1158,7 @@ void SetLevelVisibilityNoGlobalUpdateInternal(ULevel* Level, const bool bShouldB
 					{
 						if (ModifyMode == ELevelVisibilityDirtyMode::ModifyOnChange)
 						{
-							bModified = Actor->Modify();
+							bModified = Actor->Modify(false);
 						}
 						
 						Actor->bHiddenEdLayer = false;
@@ -1173,7 +1173,7 @@ void SetLevelVisibilityNoGlobalUpdateInternal(ULevel* Level, const bool bShouldB
 				{
 					if (!bModified && ModifyMode == ELevelVisibilityDirtyMode::ModifyOnChange)
 					{
-						bModified = Actor->Modify();
+						bModified = Actor->Modify(false);
 					}
 					Actor->bHiddenEdLevel = !bShouldBeVisible;
 
diff --git a/Engine/Source/Runtime/CoreUObject/Private/UObject/Obj.cpp b/Engine/Source/Runtime/CoreUObject/Private/UObject/Obj.cpp
--- a/Engine/Source/Runtime/CoreUObject/Private/UObject/Obj.cpp
+++ b/Engine/Source/Runtime/CoreUObject/Private/UObject/Obj.cpp
@@ -332,7 +332,7 @@ void UObject::PostLoad()
 #if WITH_EDITOR
 void UObject::PreEditChange(FProperty* PropertyAboutToChange)
 {
-	Modify();
+	Modify(!GIsTransacting);
 }
 
 
