Cloth Painter - Do not rebuilt the DDC while previewing. [Merging CL 14268331]

diff --git a/Engine/Source/Editor/ClothPainter/Private/ClothPainter.cpp b/Engine/Source/Editor/ClothPainter/Private/ClothPainter.cpp
--- a/Engine/Source/Editor/ClothPainter/Private/ClothPainter.cpp
+++ b/Engine/Source/Editor/ClothPainter/Private/ClothPainter.cpp
@@ -271,8 +271,8 @@ void FClothPainter::Tick(FEditorViewportClient* ViewportClient, float DeltaTime)
 					Asset->ApplyParameterMasks();
 				}
 			}
-		
-			SkeletalMeshComponent->RebuildClothingSectionsFixedVerts();
+			static const bool bInvalidateDerivedDataCache = false;  // No need to rebuild the DDC while previewing
+			SkeletalMeshComponent->RebuildClothingSectionsFixedVerts(bInvalidateDerivedDataCache);
 		}
 
 		FComponentReregisterContext ReregisterContext(SkeletalMeshComponent);
diff --git a/Engine/Source/Editor/UnrealEd/Classes/Animation/DebugSkelMeshComponent.h b/Engine/Source/Editor/UnrealEd/Classes/Animation/DebugSkelMeshComponent.h
--- a/Engine/Source/Editor/UnrealEd/Classes/Animation/DebugSkelMeshComponent.h
+++ b/Engine/Source/Editor/UnrealEd/Classes/Animation/DebugSkelMeshComponent.h
@@ -442,7 +442,8 @@ public:
 
 	// Rebuilds the fixed parameter on the mesh to mesh data, to be used if the editor has
 	// changed a vert to be fixed or unfixed otherwise the simulation will not work
-	void RebuildClothingSectionsFixedVerts();
+	// bInvalidateDerivedDataCache can only be false during previewing as otherwise the changes won't be correctly saved
+	void RebuildClothingSectionsFixedVerts(bool bInvalidateDerivedDataCache = true);
 
 	TArray<FVector> SkinnedSelectedClothingPositions;
 	TArray<FVector> SkinnedSelectedClothingNormals;
diff --git a/Engine/Source/Editor/UnrealEd/Private/Animation/DebugSkelMeshComponent.cpp b/Engine/Source/Editor/UnrealEd/Private/Animation/DebugSkelMeshComponent.cpp
--- a/Engine/Source/Editor/UnrealEd/Private/Animation/DebugSkelMeshComponent.cpp
+++ b/Engine/Source/Editor/UnrealEd/Private/Animation/DebugSkelMeshComponent.cpp
@@ -723,7 +723,7 @@ void UDebugSkelMeshComponent::ResetMeshSectionVisibility()
 	}
 }
 
-void UDebugSkelMeshComponent::RebuildClothingSectionsFixedVerts()
+void UDebugSkelMeshComponent::RebuildClothingSectionsFixedVerts(bool bInvalidateDerivedDataCache)
 {
 	FSkeletalMeshModel* Resource = SkeletalMesh->GetImportedModel();
 	FScopedSkeletalMeshPostEditChange ScopedSkeletalMeshPostEditChange(SkeletalMesh);
@@ -752,8 +752,11 @@ void UDebugSkelMeshComponent::RebuildClothingSectionsFixedVerts()
 							VertData.SourceMeshVertIndices[1],
 							VertData.SourceMeshVertIndices[2]) ? 0xFFFF : 0;
 					}
-					//We must dirty the DDC key
-					SkeletalMesh->InvalidateDeriveDataCacheGUID();
+					if (bInvalidateDerivedDataCache)
+					{
+						// We must always dirty the DDC key unless previewing
+						SkeletalMesh->InvalidateDeriveDataCacheGUID();
+					}
 				}
 			}
 		}
