Edigrate cl 14223025:

Fix moving an actor to a world outliner folder not properly creating a transaction.

Empty or no-op transaction should not be canceled anymore as they are handled automatically by the transaction system. Also undoing partial transaction is not supported anymore since the current API is dangerous and unreliable hence why the Cancel were removed in this cl.

diff --git a/Engine/Source/Editor/UnrealEd/Private/EditorActorFolders.cpp b/Engine/Source/Editor/UnrealEd/Private/EditorActorFolders.cpp
--- a/Engine/Source/Editor/UnrealEd/Private/EditorActorFolders.cpp
+++ b/Engine/Source/Editor/UnrealEd/Private/EditorActorFolders.cpp
@@ -176,10 +176,6 @@ void FActorFolders::OnActorFolderChanged(const AActor* InActor, FName OldPath)
 	{
 		OnFolderCreate.Broadcast(*World, NewPath);
 	}
-	else
-	{
-		Transaction.Cancel();
-	}
 }
 
 void FActorFolders::RebuildFolderListForWorld(UWorld& InWorld)
@@ -418,10 +414,6 @@ void FActorFolders::CreateFolder(UWorld& InWorld, FName Path)
 	{
 		OnFolderCreate.Broadcast(InWorld, Path);
 	}
-	else
-	{
-		Transaction.Cancel();
-	}
 }
 
 void FActorFolders::DeleteFolder(UWorld& InWorld, FName FolderToDelete)
