UMG - Adding cleanup code for an unknown bug.  Some how widgets are being pasted from one location to inside a Named slot, or possible through another means, like Replace With Child?  Unclear, but the end result is a widget who is still keeping their old parent widget alive through the Slot ptr - even though that widget is no longer in the tree.  This new cleanup code on open goes through those named slots, and also checks the root widget to see if such a situation exists and if so, to null those ptrs out.

diff --git a/Engine/Source/Editor/UnrealEd/Private/BaseWidgetBlueprint.cpp b/Engine/Source/Editor/UnrealEd/Private/BaseWidgetBlueprint.cpp
--- a/Engine/Source/Editor/UnrealEd/Private/BaseWidgetBlueprint.cpp
+++ b/Engine/Source/Editor/UnrealEd/Private/BaseWidgetBlueprint.cpp
@@ -18,19 +18,56 @@ void UBaseWidgetBlueprint::PostLoad()
 
 	// Update any older widget trees that don't have the same flags.
 	WidgetTree->SetFlags(RF_Transactional | RF_ArchetypeObject);
+
+	// Clean up named slot content that somehow still refers to a parent widget.  Nothing
+	// hosted inside of a named slot should have a Slot property set to anything.
+	// This shouldn't be the case, but has happened in the past due to various paste bugs
+	// so clean it up.
+	ForEachSourceWidget([&](UWidget* InWidget) {
+		if (INamedSlotInterface* NamedSlotHost = Cast<INamedSlotInterface>(InWidget))
+		{
+			TArray<FName> SlotNames;
+			NamedSlotHost->GetSlotNames(SlotNames);
+
+			for (FName SlotName : SlotNames)
+			{
+				if (UWidget* SlotContent = NamedSlotHost->GetContentForSlot(SlotName))
+				{
+					if (SlotContent->Slot)
+					{
+						SlotContent->Slot = nullptr;
+					}
+				}
+			}
+		}
+	});
+
+	// Clean up any cases where the root widget some how still has a valid reference to a non-existent
+	// slot parent.  This shouldn't be the case, but has happened in the past due to various paste bugs
+	// so clean it up.
+	if (WidgetTree->RootWidget && WidgetTree->RootWidget->Slot)
+	{
+		WidgetTree->RootWidget->Slot = nullptr;
+	}
 }
 
 TArray<UWidget*> UBaseWidgetBlueprint::GetAllSourceWidgets()
 {
 	TArray<UWidget*> Ret;
-	ForEachSourceWidgetImpl([&Ret](UWidget* Inner) { Ret.Push(Inner); });
+	ForEachSourceWidgetImpl([&Ret](UWidget* Inner) {
+		Ret.Push(Inner);
+	});
+
 	return Ret;
 }
 
 TArray<const UWidget*> UBaseWidgetBlueprint::GetAllSourceWidgets() const
 {
 	TArray<const UWidget*> Ret;
-	ForEachSourceWidgetImpl([&Ret](UWidget* Inner) { Ret.Push(Inner); });
+	ForEachSourceWidgetImpl([&Ret](UWidget* Inner) {
+		Ret.Push(Inner);
+	});
+
 	return Ret;
 }
 
@@ -57,6 +94,7 @@ void UBaseWidgetBlueprint::ForEachSourceWidgetImpl(TFunctionRef<void(UWidget*)>
 		WidgetTree,
 		[Fn, WidgetTreeForCapture](UObject* Inner)
 	{
+		FString Path = Inner->GetPathName();
 		if (UWidget* AsWidget = Cast<UWidget>(Inner))
 		{
 			// Widgets owned by another UWidgetBlueprint aren't really 'source' widgets, E.g. widgets
