MeshModelingTools: Allow tools to show details panel again

#rnx

diff --git a/Engine/Plugins/Experimental/ModelingToolsEditorMode/Source/ModelingToolsEditorMode/Private/ModelingToolsEditorModeToolkit.cpp b/Engine/Plugins/Experimental/ModelingToolsEditorMode/Source/ModelingToolsEditorMode/Private/ModelingToolsEditorModeToolkit.cpp
--- a/Engine/Plugins/Experimental/ModelingToolsEditorMode/Source/ModelingToolsEditorMode/Private/ModelingToolsEditorModeToolkit.cpp
+++ b/Engine/Plugins/Experimental/ModelingToolsEditorMode/Source/ModelingToolsEditorMode/Private/ModelingToolsEditorModeToolkit.cpp
@@ -656,13 +656,16 @@ void FModelingToolsEditorModeToolkit::OnToolStarted(UInteractiveToolManager* Man
 void FModelingToolsEditorModeToolkit::OnToolEnded(UInteractiveToolManager* Manager, UInteractiveTool* Tool)
 {
 	DetailsView->SetObject(nullptr);
-	UInteractiveTool* CurTool = GetToolsEditorMode()->GetToolManager()->GetActiveTool(EToolSide::Left);
 	ActiveToolName = FText::GetEmpty();
 	ModeHeaderArea->SetVisibility(EVisibility::Visible);
 	ModeHeaderArea->SetText(LOCTEXT("SelectToolLabel", "Select a Tool from the Toolbar"));
 	ClearNotification();
 	ClearWarning();
-	CurTool->OnPropertySetsModified.RemoveAll(this);
+	UInteractiveTool* CurTool = GetToolsEditorMode()->GetToolManager()->GetActiveTool(EToolSide::Left);
+	if ( CurTool )
+	{
+		CurTool->OnPropertySetsModified.RemoveAll(this);
+	}
 }
 
 void FModelingToolsEditorModeToolkit::UpdateShowWarnings()
diff --git a/Engine/Source/Editor/UnrealEd/Private/Toolkits/BaseToolkit.cpp b/Engine/Source/Editor/UnrealEd/Private/Toolkits/BaseToolkit.cpp
--- a/Engine/Source/Editor/UnrealEd/Private/Toolkits/BaseToolkit.cpp
+++ b/Engine/Source/Editor/UnrealEd/Private/Toolkits/BaseToolkit.cpp
@@ -1,6 +1,7 @@
 // Copyright Epic Games, Inc. All Rights Reserved.
 
 #include "Toolkits/BaseToolkit.h"
+#include "EdModeInteractiveToolsContext.h"
 #include "Toolkits/ToolkitManager.h"
 #include "PropertyEditorModule.h"
 #include "Modules/ModuleManager.h"
@@ -11,6 +12,7 @@
 #include "Tools/UEdMode.h"
 #include "Framework/MultiBox/MultiBoxBuilder.h"
 #include "EdMode.h"
+#include "EditorModeManager.h"
 
 #define LOCTEXT_NAMESPACE "BaseToolkit"
 
@@ -130,16 +132,14 @@ void FModeToolkit::Init(const TSharedPtr< class IToolkitHost >& InitToolkitHost)
 void FModeToolkit::Init(const TSharedPtr< class IToolkitHost >& InitToolkitHost, TWeakObjectPtr<UEdMode> InOwningMode)
 {
 	check( InitToolkitHost.IsValid() );
-	
+
 	ToolkitMode = EToolkitMode::WorldCentric;
 	ToolkitHost = InitToolkitHost;
 	OwningEditorMode = InOwningMode;
 
-	if (OwningEditorMode.IsValid())
-	{
-		OwningEditorMode->GetToolManager()->OnToolStarted.AddSP(this, &FModeToolkit::OnToolStarted);
-		OwningEditorMode->GetToolManager()->OnToolEnded.AddSP(this, &FModeToolkit::OnToolEnded);
-	}
+	UInteractiveToolManager* ToolManager = GetEditorModeManager().GetInteractiveToolsContext()->ToolManager;
+	ToolManager->OnToolStarted.AddSP(this, &FModeToolkit::OnToolStarted);
+	ToolManager->OnToolEnded.AddSP(this, &FModeToolkit::OnToolEnded);
 
 	FPropertyEditorModule& PropertyEditorModule = FModuleManager::GetModuleChecked<FPropertyEditorModule>("PropertyEditor");
 
@@ -180,11 +180,9 @@ void FModeToolkit::Init(const TSharedPtr< class IToolkitHost >& InitToolkitHost,
 
 FModeToolkit::~FModeToolkit()
 {
-	if (OwningEditorMode.IsValid())
-	{
-		OwningEditorMode->GetToolManager()->OnToolStarted.RemoveAll(this);
-		OwningEditorMode->GetToolManager()->OnToolEnded.RemoveAll(this);
-	}
+	UInteractiveToolManager* ToolManager = GetEditorModeManager().GetInteractiveToolsContext()->ToolManager;
+	ToolManager->OnToolStarted.RemoveAll(this);
+	ToolManager->OnToolEnded.RemoveAll(this);
 
 	OwningEditorMode.Reset();
 }
