Make sure that FCommandletPackageHelper::Delete(Packages) won't fail if the provided array is empty
* Modified UPackageTools::UnloadPackages() to return true if there are no packages to process

diff --git a/Engine/Source/Editor/UnrealEd/Private/Commandlets/CommandletPackageHelper.cpp b/Engine/Source/Editor/UnrealEd/Private/Commandlets/CommandletPackageHelper.cpp
--- a/Engine/Source/Editor/UnrealEd/Private/Commandlets/CommandletPackageHelper.cpp
+++ b/Engine/Source/Editor/UnrealEd/Private/Commandlets/CommandletPackageHelper.cpp
@@ -129,6 +129,11 @@ bool FCommandletPackageHelper::Delete(const TArray<FAssetData>& Assets) const
 
 bool FCommandletPackageHelper::Delete(const TArray<UPackage*>& Packages) const
 {
+	if (Packages.IsEmpty())
+	{
+		return true;
+	}
+
 	// Store all packages names as we won't be able to retrieve them from the UPackages once they are unloaded
 	TArray<FString> PackagesNames;
 	PackagesNames.Reserve(Packages.Num());
diff --git a/Engine/Source/Editor/UnrealEd/Private/Commandlets/WorldPartitionBuildHLODsCommandlet.cpp b/Engine/Source/Editor/UnrealEd/Private/Commandlets/WorldPartitionBuildHLODsCommandlet.cpp
--- a/Engine/Source/Editor/UnrealEd/Private/Commandlets/WorldPartitionBuildHLODsCommandlet.cpp
+++ b/Engine/Source/Editor/UnrealEd/Private/Commandlets/WorldPartitionBuildHLODsCommandlet.cpp
@@ -132,43 +132,50 @@ int32 UWorldPartitionBuildHLODsCommandlet::Main(const FString& Params)
 	}
 
 	// Delete packages
-	UE_LOG(LogWorldPartitionBuildHLODsCommandlet, Log, TEXT("Deleting %d packages."), PackagesToDelete.Num());
-	if (!PackageHelper.Delete(PackagesToDelete))
+	if (!PackagesToDelete.IsEmpty())
 	{
-		UE_LOG(LogWorldPartitionBuildHLODsCommandlet, Error, TEXT("Error deleting packages."));
-		return 1;
-	}
-
-	// Checkout packages
-	UE_LOG(LogWorldPartitionBuildHLODsCommandlet, Log, TEXT("Checking out %d actor packages."), PackagesToSave.Num());
-	for (UPackage* Package : PackagesToSave)
-	{
-		if (!PackageHelper.Checkout(Package))
+		UE_LOG(LogWorldPartitionBuildHLODsCommandlet, Log, TEXT("Deleting %d packages."), PackagesToDelete.Num());
+		if (!PackageHelper.Delete(PackagesToDelete))
 		{
-			UE_LOG(LogWorldPartitionBuildHLODsCommandlet, Error, TEXT("Error checking out package %s."), *Package->GetName());
+			UE_LOG(LogWorldPartitionBuildHLODsCommandlet, Error, TEXT("Error deleting packages."));
 			return 1;
 		}
 	}
 
 	// Save packages
-	UE_LOG(LogWorldPartitionBuildHLODsCommandlet, Log, TEXT("Saving %d packages."), PackagesToSave.Num());
-	for (UPackage* Package : PackagesToSave)
+	if (!PackagesToSave.IsEmpty())
 	{
-		if (!PackageHelper.Save(Package))
+		// Checkout packages
+		UE_LOG(LogWorldPartitionBuildHLODsCommandlet, Log, TEXT("Checking out %d actor packages."), PackagesToSave.Num());
+		for (UPackage* Package : PackagesToSave)
 		{
-			UE_LOG(LogWorldPartitionBuildHLODsCommandlet, Error, TEXT("Error saving package %s."), *Package->GetName());
-			return 1;
+			if (!PackageHelper.Checkout(Package))
+			{
+				UE_LOG(LogWorldPartitionBuildHLODsCommandlet, Error, TEXT("Error checking out package %s."), *Package->GetName());
+				return 1;
+			}
 		}
-	}
 
-	// Add new packages to source control
-	UE_LOG(LogWorldPartitionBuildHLODsCommandlet, Log, TEXT("Adding packages to source control."));
-	for (UPackage* Package : PackagesToSave)
-	{
-		if (!PackageHelper.AddToSourceControl(Package))
+		// Save packages
+		UE_LOG(LogWorldPartitionBuildHLODsCommandlet, Log, TEXT("Saving %d packages."), PackagesToSave.Num());
+		for (UPackage* Package : PackagesToSave)
 		{
-			UE_LOG(LogWorldPartitionBuildHLODsCommandlet, Error, TEXT("Error adding package %s to source control."), *Package->GetName());
-			return 1;
+			if (!PackageHelper.Save(Package))
+			{
+				UE_LOG(LogWorldPartitionBuildHLODsCommandlet, Error, TEXT("Error saving package %s."), *Package->GetName());
+				return 1;
+			}
+		}
+
+		// Add new packages to source control
+		UE_LOG(LogWorldPartitionBuildHLODsCommandlet, Log, TEXT("Adding packages to source control."));
+		for (UPackage* Package : PackagesToSave)
+		{
+			if (!PackageHelper.AddToSourceControl(Package))
+			{
+				UE_LOG(LogWorldPartitionBuildHLODsCommandlet, Error, TEXT("Error adding package %s to source control."), *Package->GetName());
+				return 1;
+			}
 		}
 	}
 
diff --git a/Engine/Source/Editor/UnrealEd/Private/PackageTools.cpp b/Engine/Source/Editor/UnrealEd/Private/PackageTools.cpp
--- a/Engine/Source/Editor/UnrealEd/Private/PackageTools.cpp
+++ b/Engine/Source/Editor/UnrealEd/Private/PackageTools.cpp
@@ -272,6 +272,12 @@ UPackageTools::UPackageTools(const FObjectInitializer& ObjectInitializer)
 
 	bool UPackageTools::UnloadPackages( const TArray<UPackage*>& TopLevelPackages, FText& OutErrorMessage )
 	{
+		// Early out if no package is provided
+		if (TopLevelPackages.IsEmpty())
+		{
+			return true;
+		}
+
 		bool bResult = false;
 
 		// Get outermost packages, in case groups were selected.
