Moving WorldPartitionBuildHLODsCommandlet.h from the /Private folder to /Classes

diff --git a/Engine/Source/Editor/UnrealEd/Classes/Commandlets/WorldPartitionBuildHLODsCommandlet.h b/Engine/Source/Editor/UnrealEd/Classes/Commandlets/WorldPartitionBuildHLODsCommandlet.h
new file mode 100644
--- /dev/null
+++ b/Engine/Source/Editor/UnrealEd/Classes/Commandlets/WorldPartitionBuildHLODsCommandlet.h
@@ -0,0 +1,20 @@
+// Copyright Epic Games, Inc. All Rights Reserved.
+
+#pragma once
+
+#include "Commandlets/Commandlet.h"
+#include "Commandlets/CommandletPackageHelper.h"
+#include "WorldPartitionBuildHLODsCommandlet.generated.h"
+
+UCLASS()
+class UWorldPartitionBuildHLODsCommandlet : public UCommandlet
+{
+	GENERATED_UCLASS_BODY()
+
+	//~ Begin UCommandlet Interface
+	virtual int32 Main(const FString& Params) override;
+	//~ End UCommandlet Interface
+
+private:
+	FCommandletPackageHelper PackageHelper;
+};
\ No newline at end of file
diff --git a/Engine/Source/Editor/UnrealEd/Private/Commandlets/WorldPartitionBuildHLODsCommandlet.h b/Engine/Source/Editor/UnrealEd/Private/Commandlets/WorldPartitionBuildHLODsCommandlet.h
deleted file mode 100644
