Fixing a crash on shutdown in commandlet runs of the editor.

#rnx

diff --git a/Engine/Source/Editor/UnrealEd/Private/Subsystems/AssetEditorSubsystem.cpp b/Engine/Source/Editor/UnrealEd/Private/Subsystems/AssetEditorSubsystem.cpp
--- a/Engine/Source/Editor/UnrealEd/Private/Subsystems/AssetEditorSubsystem.cpp
+++ b/Engine/Source/Editor/UnrealEd/Private/Subsystems/AssetEditorSubsystem.cpp
@@ -51,7 +51,7 @@ void UAssetEditorSubsystem::Initialize(FSubsystemCollectionBase& Collection)
 	FCoreUObjectDelegates::OnPackageReloaded.AddUObject(this, &UAssetEditorSubsystem::HandlePackageReloaded);
 
 	GEditor->OnEditorClose().AddUObject(this, &UAssetEditorSubsystem::OnEditorClose);
-
+	FCoreDelegates::OnPreExit.AddUObject(this, &UAssetEditorSubsystem::UnregisterEditorModes);
 	FCoreDelegates::OnPostEngineInit.AddUObject(this, &UAssetEditorSubsystem::RegisterEditorModes);
 }
 
@@ -59,6 +59,7 @@ void UAssetEditorSubsystem::Deinitialize()
 {
 	FCoreUObjectDelegates::OnPackageReloaded.RemoveAll(this);
 	GEditor->OnEditorClose().RemoveAll(this);
+	FCoreDelegates::OnPreExit.RemoveAll(this);
 
 	// Don't attempt to report usage stats if analytics isn't available
 	if (FEngineAnalytics::IsAvailable())
@@ -83,7 +84,6 @@ void UAssetEditorSubsystem::OnEditorClose()
 	SaveOpenAssetEditors(true);
 	TGuardValue<bool> GuardOnShutdown(bSavingOnShutdown, true);
 	CloseAllAssetEditors();
-	UnregisterEditorModes();
 }
 
 IAssetEditorInstance* UAssetEditorSubsystem::FindEditorForAsset(UObject* Asset, bool bFocusIfOpen)
diff --git a/Engine/Source/Editor/UnrealEd/Private/UnrealEdGlobals.cpp b/Engine/Source/Editor/UnrealEd/Private/UnrealEdGlobals.cpp
--- a/Engine/Source/Editor/UnrealEd/Private/UnrealEdGlobals.cpp
+++ b/Engine/Source/Editor/UnrealEd/Private/UnrealEdGlobals.cpp
@@ -119,6 +119,15 @@ int32 EditorInit( IEngineLoop& EngineLoop )
 
 	// Initialize the misc editor
 	FUnrealEdMisc::Get().OnInit();
+	FCoreDelegates::OnExit.AddLambda([]()
+	{
+		// Shutdown the global static mode manager
+		if (Internal::GetGlobalModeManager().IsValid())
+		{
+			GLevelEditorModeTools().SetDefaultMode(FBuiltinEditorModes::EM_Default);
+			Internal::GetGlobalModeManager().Reset();
+		}
+	});
 	
 	SlowTask.EnterProgressFrame(10);
 
@@ -217,10 +226,6 @@ void EditorExit()
 	GEditor->SaveConfig();
 	GLevelEditorModeTools().SaveConfig();
 
-	// Shutdown the global static mode manager
-	GLevelEditorModeTools().SetDefaultMode(FBuiltinEditorModes::EM_Default);
-	Internal::GetGlobalModeManager().Reset();
-
 	// Clean up the actor folders singleton
 	FActorFolders::Cleanup();
 
