Always clear the transaction buffer when saving levels or external actors to avoid known issues with deleted actors when the transaction buffer persists between saves.

diff --git a/Engine/Source/Editor/UnrealEd/Private/EditorEngine.cpp b/Engine/Source/Editor/UnrealEd/Private/EditorEngine.cpp
--- a/Engine/Source/Editor/UnrealEd/Private/EditorEngine.cpp
+++ b/Engine/Source/Editor/UnrealEd/Private/EditorEngine.cpp
@@ -4529,6 +4529,13 @@ FSavePackageResultStruct UEditorEngine::Save( UPackage* InOuter, UObject* InBase
 
 	SlowTask.EnterProgressFrame(10);
 
+	const bool bAutosave = (SaveFlags & SAVE_FromAutosave) != 0;
+	if (!bSavingConcurrent && !IsRunningCommandlet() && !bAutosave && (World || Cast<AActor>(Base)))
+	{
+		// Always reset the transaction buffer on level/actor save to avoid problems with deleted actors (marked pending kill) that gets marked transient by the saving code
+		ResetTransaction( World ? NSLOCTEXT("UnrealEd", "MapSaved", "Map Saved") : NSLOCTEXT("UnrealEd", "ActorSaved", "Actor Saved"));
+	}
+
 	if ( World )
 	{
 		if (OriginalOwningWorld)
@@ -4587,8 +4594,8 @@ void UEditorEngine::OnPreSaveWorld(uint32 SaveFlags, UWorld* World)
 
 	if ( !IsRunningCommandlet() )
 	{
-		const bool bAutosaveOrPIE = (SaveFlags & SAVE_FromAutosave) != 0;
-		if ( bAutosaveOrPIE )
+		const bool bAutosave = (SaveFlags & SAVE_FromAutosave) != 0;
+		if ( bAutosave )
 		{
 			// Temporarily flag packages saved under a PIE filename as PKG_PlayInEditor for serialization so loading
 			// them will have the flag set. We need to undo this as the object flagged isn't actually the PIE package, 
@@ -4660,8 +4667,8 @@ void UEditorEngine::OnPostSaveWorld(uint32 SaveFlags, UWorld* World, uint32 Orig
 	if ( !IsRunningCommandlet() )
 	{
 		UPackage* WorldPackage = World->GetOutermost();
-		const bool bAutosaveOrPIE = (SaveFlags & SAVE_FromAutosave) != 0;
-		if ( bAutosaveOrPIE )
+		const bool bAutosave = (SaveFlags & SAVE_FromAutosave) != 0;
+		if ( bAutosave )
 		{
 			// Restore original value of PKG_PlayInEditor if we changed it during PIE saving
 			const bool bOriginallyPIE = (OriginalPackageFlags & PKG_PlayInEditor) != 0;
@@ -4709,7 +4716,6 @@ void UEditorEngine::OnPostSaveWorld(uint32 SaveFlags, UWorld* World, uint32 Orig
 
 			if ( bIsPersistentLevel )
 			{
-				ResetTransaction( NSLOCTEXT("UnrealEd", "MapSaved", "Map Saved") );
 				FPlatformProcess::SetCurrentWorkingDirectoryToBaseDir();
 			}
 		}
