Fix Foliage Conversion (re-add accidentally removed OverrideLevel + clear all Foliage Actor Tags)

diff --git a/Engine/Source/Editor/UnrealEd/Private/Commandlets/WorldPartitionConvertCommandlet.cpp b/Engine/Source/Editor/UnrealEd/Private/Commandlets/WorldPartitionConvertCommandlet.cpp
--- a/Engine/Source/Editor/UnrealEd/Private/Commandlets/WorldPartitionConvertCommandlet.cpp
+++ b/Engine/Source/Editor/UnrealEd/Private/Commandlets/WorldPartitionConvertCommandlet.cpp
@@ -935,6 +935,12 @@ int32 UWorldPartitionConvertCommandlet::Main(const FString& Params)
 				FArchiveGatherPrivateImports Ar(Actor, PrivateRefsMap, ActorsReferencesToActors);
 				Actor->Serialize(Ar);
 
+				// Even after Foliage Partitioning it is possible some Actors still have a FoliageTag. Make sure it is removed.
+				if (FFoliageHelper::IsOwnedByFoliage(Actor))
+				{
+					FFoliageHelper::SetIsOwnedByFoliage(Actor, false);
+				}
+
 				ChangeObjectOuter(Actor, MainLevel);
 
 				// Migrate blueprint classes
diff --git a/Engine/Source/Runtime/Engine/Private/ActorPartition/ActorPartitionSubsystem.cpp b/Engine/Source/Runtime/Engine/Private/ActorPartition/ActorPartitionSubsystem.cpp
--- a/Engine/Source/Runtime/Engine/Private/ActorPartition/ActorPartitionSubsystem.cpp
+++ b/Engine/Source/Runtime/Engine/Private/ActorPartition/ActorPartitionSubsystem.cpp
@@ -176,6 +176,7 @@ public:
 		if (!FoundActor && bInCreate)
 		{
 			FActorSpawnParameters SpawnParams;
+			SpawnParams.OverrideLevel = InCellCoord.Level;
 			SpawnParams.Name = FName(*FString::Printf(TEXT("%s_%d_%d_%d"), *InActorClass->GetName(), InCellCoord.X, InCellCoord.Y, InCellCoord.Z));
 			SpawnParams.NameMode = FActorSpawnParameters::ESpawnActorNameMode::Requested;
 							
