- Moved File | Package and Cook into Launch menu, which may become the Platforms menu (or similar)
- Can now cook for any platform (and will soon have a Prepare for Debugging option)
- All the packaging and cooking stuff are now in TurnkeySupportModule
- Allow per-platform configuration/cook target
- Added a temporary inactive menu item to point users to the Launch menu for packaging
- This is first pass on this, it's expected to change, even move locations, etc as feedback/needed dictate

diff --git a/Engine/Source/Developer/DesktopPlatform/Private/PlatformInfo.cpp b/Engine/Source/Developer/DesktopPlatform/Private/PlatformInfo.cpp
--- a/Engine/Source/Developer/DesktopPlatform/Private/PlatformInfo.cpp
+++ b/Engine/Source/Developer/DesktopPlatform/Private/PlatformInfo.cpp
@@ -68,7 +68,7 @@ FTargetPlatformInfo::FTargetPlatformInfo(const FString& InIniPlatformName, EBuil
 		// append UAT commandline with cook flavor
 		UATCommandLine += FString::Printf(TEXT(" -cookflavor=%s"), *InCookFlavor);
 
-		VanillaInfo->Flavors.Add(this);
+		VanillaInfo->Flavors.AddUnique(this);
 	}
 
 	// now append the build type (game type has no type suffix, all others do)
@@ -85,6 +85,8 @@ FTargetPlatformInfo::FTargetPlatformInfo(const FString& InIniPlatformName, EBuil
 		{
 			UATCommandLine += TEXT(" -client");
 		}
+
+		VanillaInfo->Flavors.AddUnique(this);
 	}
 
 	// now we can store the final values
diff --git a/Engine/Source/Editor/MainFrame/Private/Frame/MainFrameActions.cpp b/Engine/Source/Editor/MainFrame/Private/Frame/MainFrameActions.cpp
--- a/Engine/Source/Editor/MainFrame/Private/Frame/MainFrameActions.cpp
+++ b/Engine/Source/Editor/MainFrame/Private/Frame/MainFrameActions.cpp
@@ -37,11 +37,9 @@
 #include "EditorAnalytics.h"
 #include "LevelEditor.h"
 #include "Interfaces/IProjectTargetPlatformEditorModule.h"
-#include "InstalledPlatformInfo.h"
 #include "Misc/ConfigCacheIni.h"
 #include "MainFrameModule.h"
 #include "Widgets/Docking/SDockTab.h"
-#include "Widgets/Notifications/SNotificationList.h"
 #include "Framework/Commands/GenericCommands.h"
 #include "Dialogs/SOutputLogDialog.h"
 #include "Dialogs/Dialogs.h"
@@ -125,9 +123,6 @@ void FMainFrameCommands::RegisterCommands()
 	UI_COMMAND( ZipUpProject, "Zip Up Project", "Zips up the project into a zip file.", EUserInterfaceActionType::Button, FInputChord() );
 	ActionList->MapAction(ZipUpProject, FExecuteAction::CreateStatic( &FMainFrameActionCallbacks::ZipUpProject ), DefaultExecuteAction);
 
-	UI_COMMAND( PackagingSettings, "Packaging Settings...", "Opens the settings for project packaging", EUserInterfaceActionType::Button, FInputChord() );
-	ActionList->MapAction( PackagingSettings, FExecuteAction::CreateStatic( &FMainFrameActionCallbacks::PackagingSettings ), DefaultExecuteAction );
-
 	//UI_COMMAND( LocalizeProject, "Localize Project...", "Opens the dashboard for managing project localization data.", EUserInterfaceActionType::Button, FInputChord() );
 	//ActionList->MapAction( LocalizeProject, FExecuteAction::CreateStatic( &FMainFrameActionCallbacks::LocalizeProject ), DefaultExecuteAction );
 
@@ -359,525 +354,6 @@ void FMainFrameActionCallbacks::AddCodeToProject()
 	FGameProjectGenerationModule::Get().OpenAddCodeToProjectDialog();
 }
 
-/**
- * Gets compilation flags for UAT for this system.
- */
-const TCHAR* GetUATCompilationFlags()
-{
-	// We never want to compile editor targets when invoking UAT in this context.
-	// If we are installed or don't have a compiler, we must assume we have a precompiled UAT.
-	return TEXT("-nocompileeditor");
-}
-
-FString GetCookingOptionalParams()
-{
-	FString OptionalParams;
-	const UProjectPackagingSettings* const PackagingSettings = GetDefault<UProjectPackagingSettings>();
-	
-	if (PackagingSettings->bSkipEditorContent)
-	{
-		OptionalParams += TEXT(" -SkipCookingEditorContent");
-	}
-	return OptionalParams;
-}
-
-
-static bool ShowBadSDKDialog(FName IniPlatformName)
-{
-	// Don't show the warning during automation testing; the dlg is modal and blocks
-	if (!GIsAutomationTesting)
-	{
-		FFormatNamedArguments Args;
-		//		Args.Add(TEXT("DisplayName"), PlatformInfo->DisplayName);
-		Args.Add(TEXT("DisplayName"), FText::FromName(IniPlatformName));
-		FText WarningText = FText::Format(LOCTEXT("BadSDK_Message", "The SDK for {DisplayName} appears to not be installed properly, which is needed to generate data. Check the SDK section of the Launch On menu in the main toolbar to update SDK.\n\nAlternatively, you may continue anyway, which will attempt to install/update the SDK before cookingt the content. This is a new feature, and is not guaranteed to work.\n\nWould you like to continue?"), Args);
-
-		FSuppressableWarningDialog::FSetupInfo Info(
-			WarningText,
-			LOCTEXT("BadSDK_Title", "SDK Not Setup"),
-			TEXT("BadSDKDialog")
-		);
-		Info.ConfirmText = LOCTEXT("BadSDK_Confirm", "Continue");
-		Info.CancelText = LOCTEXT("BadSDK_Cancel", "Cancel");
-		FSuppressableWarningDialog BadSDKgDialog(Info);
-
-		return BadSDKgDialog.ShowModal() != FSuppressableWarningDialog::EResult::Cancel;
-	}
-
-	return true;
-}
-
-void FMainFrameActionCallbacks::CookContent(const FName InPlatformInfoName)
-{
-	const PlatformInfo::FTargetPlatformInfo* const PlatformInfo = PlatformInfo::FindPlatformInfo(InPlatformInfoName);
-	check(PlatformInfo);
-
-	if (FInstalledPlatformInfo::Get().IsPlatformMissingRequiredFile(PlatformInfo->DataDrivenPlatformInfo->UBTPlatformString))
-	{
-		if (!FInstalledPlatformInfo::OpenInstallerOptions())
-		{
-			FMessageDialog::Open(EAppMsgType::Ok, LOCTEXT("MissingPlatformFilesCook", "Missing required files to cook for this platform."));
-		}
-		return;
-	}
-
-	FString OptionalParams;
-
-	if (!FModuleManager::LoadModuleChecked<IProjectTargetPlatformEditorModule>("ProjectTargetPlatformEditor").ShowUnsupportedTargetWarning(PlatformInfo->VanillaInfo->Name))
-	{
-		return;
-	}
-
-	// @todo turnkey: run Turnkey here instead of showing SDK!
-	if (PlatformInfo->DataDrivenPlatformInfo->GetSdkStatus() != DDPIPlatformSdkStatus::Valid && !ShowBadSDKDialog(PlatformInfo->IniPlatformName))
-	{
-// 		IMainFrameModule& MainFrameModule = FModuleManager::GetModuleChecked<IMainFrameModule>(TEXT("MainFrame"));
-// 		MainFrameModule.BroadcastMainFrameSDKNotInstalled(PlatformInfo->Name.ToString(), PlatformInfo->DataDrivenPlatformInfo->SDKTutorial);
-		return;
-	}
-
-	// Append any extra UAT flags specified for this platform flavor
-	if (!PlatformInfo->UATCommandLine.IsEmpty())
-	{
-		OptionalParams += TEXT(" ");
-		OptionalParams += PlatformInfo->UATCommandLine;
-	}
-	else
-	{
-		OptionalParams += TEXT(" -platform=");
-		OptionalParams += *PlatformInfo->DataDrivenPlatformInfo->UBTPlatformString;
-	}
-
-	OptionalParams += GetCookingOptionalParams();
-
-	UCookerSettings const* CookerSettings = GetDefault<UCookerSettings>();
-	if (CookerSettings->bIterativeCookingForFileCookContent)
-	{
-		OptionalParams += TEXT(" -iterate");
-	}
-
-	FString ProjectPath = FPaths::IsProjectFilePathSet() ? FPaths::ConvertRelativePathToFull(FPaths::GetProjectFilePath()) : FPaths::RootDir() / FApp::GetProjectName() / FApp::GetProjectName() + TEXT(".uproject");
-	FString CommandLine = FString::Printf(TEXT("-ScriptsForProject=\"%s\" BuildCookRun %s%s -nop4 -project=\"%s\" -cook -skipstage -ue4exe=\"%s\" %s -utf8output -WaitForUATMutex"),
-		*ProjectPath,
-		GetUATCompilationFlags(),
-		FApp::IsEngineInstalled() ? TEXT(" -installed") : TEXT(""),
-		*ProjectPath,
-		*FUnrealEdMisc::Get().GetExecutableForCommandlets(),
-		*OptionalParams
-	);
-	
-	IUATHelperModule::Get().CreateUatTask(CommandLine, PlatformInfo->DisplayName, LOCTEXT("CookingContentTaskName", "Cooking content"), LOCTEXT("CookingTaskName", "Cooking"), FEditorStyle::GetBrush(TEXT("MainFrame.CookContent")));
-}
-
-bool FMainFrameActionCallbacks::CookContentCanExecute( const FName PlatformInfoName )
-{
-	return true;
-}
-
-void FMainFrameActionCallbacks::PackageBuildConfiguration( EProjectPackagingBuildConfigurations BuildConfiguration )
-{
-	UProjectPackagingSettings* PackagingSettings = Cast<UProjectPackagingSettings>(UProjectPackagingSettings::StaticClass()->GetDefaultObject());
-	PackagingSettings->BuildConfiguration = BuildConfiguration;
-}
-
-bool FMainFrameActionCallbacks::CanPackageBuildConfiguration( EProjectPackagingBuildConfigurations BuildConfiguration )
-{
-	return true;
-}
-
-bool FMainFrameActionCallbacks::PackageBuildConfigurationIsChecked( EProjectPackagingBuildConfigurations BuildConfiguration )
-{
-	return (GetDefault<UProjectPackagingSettings>()->BuildConfiguration == BuildConfiguration);
-}
-
-void FMainFrameActionCallbacks::PackageBuildTarget( FString TargetName )
-{
-	UProjectPackagingSettings* PackagingSettings = GetMutableDefault<UProjectPackagingSettings>();
-	PackagingSettings->BuildTarget = TargetName;
-}
-
-bool FMainFrameActionCallbacks::PackageBuildTargetIsChecked( FString TargetName )
-{
-	const FTargetInfo* Target = GetDefault<UProjectPackagingSettings>()->GetBuildTargetInfo();
-	return (Target != nullptr && Target->Name == TargetName);
-}
-
-void FMainFrameActionCallbacks::PackageProject( const FName InPlatformInfoName )
-{
-	GUnrealEd->CancelPlayingViaLauncher();
-	SaveAll();
-
-	// does the project have any code?
-	FGameProjectGenerationModule& GameProjectModule = FModuleManager::LoadModuleChecked<FGameProjectGenerationModule>(TEXT("GameProjectGeneration"));
-	bool bProjectHasCode = GameProjectModule.Get().ProjectHasCodeFiles();
-	
-	const PlatformInfo::FTargetPlatformInfo* const PlatformInfo = PlatformInfo::FindPlatformInfo(InPlatformInfoName);
-	check(PlatformInfo);
-
-	if (FInstalledPlatformInfo::Get().IsPlatformMissingRequiredFile(PlatformInfo->DataDrivenPlatformInfo->UBTPlatformString))
-	{
-		if (!FInstalledPlatformInfo::OpenInstallerOptions())
-		{
-			FMessageDialog::Open(EAppMsgType::Ok, LOCTEXT("MissingPlatformFilesPackage", "Missing required files to package this platform."));
-		}
-		return;
-	}
-
-	if (UGameMapsSettings::GetGameDefaultMap().IsEmpty())
-	{
-		FMessageDialog::Open(EAppMsgType::Ok, LOCTEXT("MissingGameDefaultMap", "No Game Default Map specified in Project Settings > Maps & Modes."));
-		return;
-	}
-
-	// @todo turnkey : run turnkey
-	if (PlatformInfo->DataDrivenPlatformInfo->GetSdkStatus() != DDPIPlatformSdkStatus::Valid || (bProjectHasCode && PlatformInfo->DataDrivenPlatformInfo->bUsesHostCompiler && !FSourceCodeNavigation::IsCompilerAvailable()))
-	{
-		if (!ShowBadSDKDialog(PlatformInfo->IniPlatformName))
-		{
-// 			IMainFrameModule& MainFrameModule = FModuleManager::GetModuleChecked<IMainFrameModule>(TEXT("MainFrame"));
-// 			MainFrameModule.BroadcastMainFrameSDKNotInstalled(PlatformInfo->Name.ToString(), PlatformInfo->DataDrivenPlatformInfo->SDKTutorial);
-			TArray<FAnalyticsEventAttribute> ParamArray;
-			ParamArray.Add(FAnalyticsEventAttribute(TEXT("Time"), 0.0));
-			FEditorAnalytics::ReportEvent(TEXT("Editor.Package.Failed"), PlatformInfo->Name.ToString(), bProjectHasCode, EAnalyticsErrorCodes::SDKNotFound, ParamArray);
-			return;
-		}
-	}
-
-	UProjectPackagingSettings* PackagingSettings = Cast<UProjectPackagingSettings>(UProjectPackagingSettings::StaticClass()->GetDefaultObject());
-	const UProjectPackagingSettings::FConfigurationInfo& ConfigurationInfo = UProjectPackagingSettings::ConfigurationInfo[PackagingSettings->BuildConfiguration];
-	bool bAssetNativizationEnabled = (PackagingSettings->BlueprintNativizationMethod != EProjectPackagingBlueprintNativizationMethod::Disabled);
-
-	const ITargetPlatform* const Platform = GetTargetPlatformManager()->FindTargetPlatform(PlatformInfo->Name.ToString());
-	{
-		if (Platform)
-		{
-			FString NotInstalledTutorialLink;
-			FString DocumentationLink;
-			FText CustomizedLogMessage;
-
-			int32 Result = Platform->CheckRequirements(bProjectHasCode, ConfigurationInfo.Configuration, bAssetNativizationEnabled, NotInstalledTutorialLink, DocumentationLink, CustomizedLogMessage);
-
-			// report to analytics
-			FEditorAnalytics::ReportBuildRequirementsFailure(TEXT("Editor.Package.Failed"), PlatformInfo->Name.ToString(), bProjectHasCode, Result);
-
-			// report to main frame
-			bool UnrecoverableError = false;
-
-			// report to message log
-			if ((Result & ETargetPlatformReadyStatus::SDKNotFound) != 0)
-			{
-				AddMessageLog(
-					LOCTEXT("SdkNotFoundMessage", "Software Development Kit (SDK) not found."),
-					CustomizedLogMessage.IsEmpty() ? FText::Format(LOCTEXT("SdkNotFoundMessageDetail", "Please install the SDK for the {0} target platform!"), Platform->DisplayName()) : CustomizedLogMessage,
-					NotInstalledTutorialLink,
-					DocumentationLink
-				);
-				UnrecoverableError = true;
-			}
-
-			if ((Result & ETargetPlatformReadyStatus::LicenseNotAccepted) != 0)
-			{
-				AddMessageLog(
-					LOCTEXT("LicenseNotAcceptedMessage", "License not accepted."),
-					CustomizedLogMessage.IsEmpty() ? LOCTEXT("LicenseNotAcceptedMessageDetail", "License must be accepted in project settings to deploy your app to the device.") : CustomizedLogMessage,
-					NotInstalledTutorialLink,
-					DocumentationLink
-				);
-
-				UnrecoverableError = true;
-			}
-
-			if ((Result & ETargetPlatformReadyStatus::ProvisionNotFound) != 0)
-			{
-				AddMessageLog(
-					LOCTEXT("ProvisionNotFoundMessage", "Provision not found."),
-					CustomizedLogMessage.IsEmpty() ? LOCTEXT("ProvisionNotFoundMessageDetail", "A provision is required for deploying your app to the device.") : CustomizedLogMessage,
-					NotInstalledTutorialLink,
-					DocumentationLink
-				);
-				UnrecoverableError = true;
-			}
-
-			if ((Result & ETargetPlatformReadyStatus::SigningKeyNotFound) != 0)
-			{
-				AddMessageLog(
-					LOCTEXT("SigningKeyNotFoundMessage", "Signing key not found."),
-					CustomizedLogMessage.IsEmpty() ? LOCTEXT("SigningKeyNotFoundMessageDetail", "The app could not be digitally signed, because the signing key is not configured.") : CustomizedLogMessage,
-					NotInstalledTutorialLink,
-					DocumentationLink
-				);
-				UnrecoverableError = true;
-			}
-
-			if ((Result & ETargetPlatformReadyStatus::ManifestNotFound) != 0)
-			{
-				AddMessageLog(
-					LOCTEXT("ManifestNotFound", "Manifest not found."),
-					CustomizedLogMessage.IsEmpty() ? LOCTEXT("ManifestNotFoundMessageDetail", "The generated application manifest could not be found.") : CustomizedLogMessage,
-					NotInstalledTutorialLink,
-					DocumentationLink
-				);
-				UnrecoverableError = true;
-			}
-
-			if ((Result & ETargetPlatformReadyStatus::RemoveServerNameEmpty) != 0
-					&& (bProjectHasCode || (Result & ETargetPlatformReadyStatus::CodeBuildRequired)
-						|| (!FApp::GetEngineIsPromotedBuild() && !FApp::IsEngineInstalled())))
-			{
-				AddMessageLog(
-					LOCTEXT("RemoveServerNameNotFound", "Remote compiling requires a server name. "),
-					CustomizedLogMessage.IsEmpty() ? LOCTEXT("RemoveServerNameNotFoundDetail", "Please specify one in the Remote Server Name settings field.") : CustomizedLogMessage,
-					NotInstalledTutorialLink,
-					DocumentationLink
-				);
-				UnrecoverableError = true;
-			}
-
-			if ((Result & ETargetPlatformReadyStatus::CodeUnsupported) != 0)
-			{
-				FMessageDialog::Open(EAppMsgType::Ok, LOCTEXT("NotSupported_SelectedPlatform", "Sorry, packaging a code-based project for the selected platform is currently not supported. This feature may be available in a future release."));
-				UnrecoverableError = true;
-			}
-			else if ((Result & ETargetPlatformReadyStatus::PluginsUnsupported) != 0)
-			{
-				FMessageDialog::Open(EAppMsgType::Ok, LOCTEXT("NotSupported_ThirdPartyPlugins", "Sorry, packaging a project with third-party plugins is currently not supported for the selected platform. This feature may be available in a future release."));
-				UnrecoverableError = true;
-			}
-
-			if (UnrecoverableError)
-			{
-				return;
-			}
-		}
-	}
-
-	if (!FModuleManager::LoadModuleChecked<IProjectTargetPlatformEditorModule>("ProjectTargetPlatformEditor").ShowUnsupportedTargetWarning(PlatformInfo->VanillaInfo->Name))
-	{
-		return;
-	}
-
-	// let the user pick a target directory
-	if (PackagingSettings->StagingDirectory.Path.IsEmpty())
-	{
-		PackagingSettings->StagingDirectory.Path = FPaths::ProjectDir();
-	}
-
-	FString OutFolderName;
-
-	void* ParentWindowWindowHandle = nullptr;
-	IMainFrameModule& MainFrameModule = FModuleManager::LoadModuleChecked<IMainFrameModule>(TEXT("MainFrame"));
-	const TSharedPtr<SWindow>& MainFrameParentWindow = MainFrameModule.GetParentWindow();
-	if ( MainFrameParentWindow.IsValid() && MainFrameParentWindow->GetNativeWindow().IsValid() )
-	{
-		ParentWindowWindowHandle = MainFrameParentWindow->GetNativeWindow()->GetOSWindowHandle();
-	}
-	
-	if (!FDesktopPlatformModule::Get()->OpenDirectoryDialog(ParentWindowWindowHandle, LOCTEXT("PackageDirectoryDialogTitle", "Package project...").ToString(), PackagingSettings->StagingDirectory.Path, OutFolderName))
-	{
-		return;
-	}
-
-	PackagingSettings->StagingDirectory.Path = OutFolderName;
-	PackagingSettings->SaveConfig();
-
-	// create the packager process
-	FString OptionalParams;
-	
-	if (PackagingSettings->FullRebuild)
-	{
-		OptionalParams += TEXT(" -clean");
-	}
-
-	if ( PackagingSettings->bCompressed )
-	{
-		OptionalParams += TEXT(" -compressed");
-	}
-
-	OptionalParams += GetCookingOptionalParams();
-
-	if (PackagingSettings->bUseIoStore)
-	{
-		OptionalParams += TEXT(" -iostore");
-
-		// Pak file(s) must be used when using container file(s)
-		PackagingSettings->UsePakFile = true;
-	}
-
-	if (PackagingSettings->UsePakFile)
-	{
-		OptionalParams += TEXT(" -pak");
-	}
-
-	if (PackagingSettings->bUseIoStore)
-	{
-		OptionalParams += TEXT(" -iostore");
-	}
-
-	if (PackagingSettings->IncludePrerequisites)
-	{
-		OptionalParams += TEXT(" -prereqs");
-	}
-
-	if (!PackagingSettings->ApplocalPrerequisitesDirectory.Path.IsEmpty())
-	{
-		OptionalParams += FString::Printf(TEXT(" -applocaldirectory=\"%s\""), *(PackagingSettings->ApplocalPrerequisitesDirectory.Path));
-	}
-	else if (PackagingSettings->IncludeAppLocalPrerequisites)
-	{
-		OptionalParams += TEXT(" -applocaldirectory=\"$(EngineDir)/Binaries/ThirdParty/AppLocalDependencies\"");
-	}
-
-	if (PackagingSettings->ForDistribution)
-	{
-		OptionalParams += TEXT(" -distribution");
-	}
-
-	if (!PackagingSettings->IncludeDebugFiles)
-	{
-		OptionalParams += TEXT(" -nodebuginfo");
-	}
-
-	if (PackagingSettings->bGenerateChunks)
-	{
-		OptionalParams += TEXT(" -manifests");
-	}
-
-	OptionalParams += TEXT(" -targetplatform=");
-	OptionalParams += PlatformInfo->DataDrivenPlatformInfo->UBTPlatformString;
-
-	// Append any extra UAT flags specified for this platform flavor
-	if (!PlatformInfo->UATCommandLine.IsEmpty())
-	{
-		OptionalParams += TEXT(" ");
-		OptionalParams += PlatformInfo->UATCommandLine;
-	}
-
-
-
-	// Get the target to build
-	const FTargetInfo* Target = PackagingSettings->GetBuildTargetInfo();
-
-	// Only build if the user elects to do so
-	bool bBuild = false;
-	if(PackagingSettings->Build == EProjectPackagingBuild::Always)
-	{
-		bBuild = true;
-	}
-	else if(PackagingSettings->Build == EProjectPackagingBuild::Never)
-	{
-		bBuild = false;
-	}
-	else if(PackagingSettings->Build == EProjectPackagingBuild::IfProjectHasCode)
-	{
-		bBuild = true;
-		if (FApp::GetEngineIsPromotedBuild() && !bAssetNativizationEnabled)
-		{
-			FString BaseDir;
-
-			// Get the target name
-			FString TargetName;
-			if (Target == nullptr)
-			{
-				TargetName = TEXT("UnrealGame");
-			}
-			else
-			{
-				TargetName = Target->Name;
-			}
-
-			// Get the directory containing the receipt for this target, depending on whether the project needs to be built or not
-			FString ProjectDir = FPaths::GetPath(FPaths::GetProjectFilePath());
-			if (Target != nullptr && FPaths::IsUnderDirectory(Target->Path, ProjectDir))
-			{
-				UE_LOG(LogMainFrame, Log, TEXT("Selected target: %s"), *Target->Name);
-				BaseDir = ProjectDir;
-			}
-			else
-			{
-				FText Reason;
-				if (Platform->RequiresTempTarget(bProjectHasCode, ConfigurationInfo.Configuration, false, Reason))
-				{
-					UE_LOG(LogMainFrame, Log, TEXT("Project requires temp target (%s)"), *Reason.ToString());
-					BaseDir = ProjectDir;
-				}
-				else
-				{
-					UE_LOG(LogMainFrame, Log, TEXT("Project does not require temp target"));
-					BaseDir = FPaths::EngineDir();
-				}
-			}
-
-			// Check if the receipt is for a matching promoted target
-			FString UBTPlatformName = Platform->GetTargetPlatformInfo().DataDrivenPlatformInfo->UBTPlatformString;
-
-			extern LAUNCHERSERVICES_API bool HasPromotedTarget(const TCHAR* BaseDir, const TCHAR* TargetName, const TCHAR* Platform, EBuildConfiguration Configuration, const TCHAR* Architecture);
-			if (HasPromotedTarget(*BaseDir, *TargetName, *UBTPlatformName, ConfigurationInfo.Configuration, nullptr))
-			{
-				bBuild = false;
-			}
-		}
-	}
-	else if(PackagingSettings->Build == EProjectPackagingBuild::IfEditorWasBuiltLocally)
-	{
-		bBuild = !FApp::GetEngineIsPromotedBuild();
-	}
-	if(bBuild)
-	{
-		OptionalParams += TEXT(" -build");
-	}
-
-	// Whether to include the crash reporter.
-	if (PackagingSettings->IncludeCrashReporter && PlatformInfo->DataDrivenPlatformInfo->bCanUseCrashReporter)
-	{
-		OptionalParams += TEXT( " -CrashReporter" );
-	}
-
-	if (PackagingSettings->bBuildHttpChunkInstallData)
-	{
-		OptionalParams += FString::Printf(TEXT(" -manifests -createchunkinstall -chunkinstalldirectory=\"%s\" -chunkinstallversion=%s"), *(PackagingSettings->HttpChunkInstallDataDirectory.Path), *(PackagingSettings->HttpChunkInstallDataVersion));
-	}
-
-	int32 NumCookers = GetDefault<UEditorExperimentalSettings>()->MultiProcessCooking;
-	if (NumCookers > 0 )
-	{
-		OptionalParams += FString::Printf(TEXT(" -NumCookersToSpawn=%d"), NumCookers); 
-	}
-
-	if (Target == nullptr)
-	{
-		OptionalParams += FString::Printf(TEXT(" -clientconfig=%s"), LexToString(ConfigurationInfo.Configuration));
-	}
-	else if(Target->Type == EBuildTargetType::Server)
-	{
-		OptionalParams += FString::Printf(TEXT(" -target=%s -serverconfig=%s"), *Target->Name, LexToString(ConfigurationInfo.Configuration));
-	}
-	else
-	{
-		OptionalParams += FString::Printf(TEXT(" -target=%s -clientconfig=%s"), *Target->Name, LexToString(ConfigurationInfo.Configuration));
-	}
-
-	FString TurnkeyCommandline = FString::Printf(TEXT("Turnkey -command=VerifySdk -platform=%s -UpdateIfNeeded -EditorIO"), *PlatformInfo->DataDrivenPlatformInfo->UBTPlatformString);
-
-	FString ProjectPath = FPaths::IsProjectFilePathSet() ? FPaths::ConvertRelativePathToFull(FPaths::GetProjectFilePath()) : FPaths::RootDir() / FApp::GetProjectName() / FApp::GetProjectName() + TEXT(".uproject");
-	FString CommandLine = FString::Printf(TEXT("-ScriptsForProject=\"%s\" %s BuildCookRun %s%s -nop4 -project=\"%s\" -cook -stage -archive -archivedirectory=\"%s\" -package -ue4exe=\"%s\" %s -utf8output -WaitForUATMutex"),
-		*ProjectPath,
-		*TurnkeyCommandline,
-		GetUATCompilationFlags(),
-		FApp::IsEngineInstalled() ? TEXT(" -installed") : TEXT(""),
-		*ProjectPath,
-		*PackagingSettings->StagingDirectory.Path,
-		*FUnrealEdMisc::Get().GetExecutableForCommandlets(),
-		*OptionalParams
-	);
-
-	IUATHelperModule::Get().CreateUatTask( CommandLine, PlatformInfo->DisplayName, LOCTEXT("PackagingProjectTaskName", "Packaging project"), LOCTEXT("PackagingTaskName", "Packaging"), FEditorStyle::GetBrush(TEXT("MainFrame.PackageProject")) );
-}
-
-bool FMainFrameActionCallbacks::PackageProjectCanExecute( const FName PlatformInfoName )
-{
-	return true;
-}
-
 void FMainFrameActionCallbacks::RefreshCodeProject()
 {
 	if ( !FSourceCodeNavigation::IsCompilerAvailable() )
@@ -923,16 +399,6 @@ bool FMainFrameActionCallbacks::CanOpenIDE()
 
 void FMainFrameActionCallbacks::ZipUpProject()
 {
-#if PLATFORM_WINDOWS
-	FText PlatformName = LOCTEXT("PlatformName_Windows", "Windows");
-#elif PLATFORM_MAC
-	FText PlatformName = LOCTEXT("PlatformName_Mac", "Mac");
-#elif PLATFORM_LINUX
-	FText PlatformName = LOCTEXT("PlatformName_Linux", "Linux");
-#else
-	FText PlatformName = LOCTEXT("PlatformName_Other", "Other OS");
-#endif
-
 	bool bOpened = false;
 	TArray<FString> SaveFilenames;
 	IDesktopPlatform* DesktopPlatform = FDesktopPlatformModule::Get();
@@ -956,19 +422,14 @@ void FMainFrameActionCallbacks::ZipUpProject()
 			FString FinalFileName = FPaths::ConvertRelativePathToFull(FileName);
 			FString ProjectPath = FPaths::IsProjectFilePathSet() ? FPaths::ConvertRelativePathToFull(FPaths::ProjectDir()) : FPaths::RootDir() / FApp::GetProjectName();
 
-			FString CommandLine = FString::Printf(TEXT("ZipProjectUp %s -project=\"%s\" -install=\"%s\""), GetUATCompilationFlags(), *ProjectPath, *FinalFileName);
+			FString CommandLine = FString::Printf(TEXT("ZipProjectUp -nocompileeditor -project=\"%s\" -install=\"%s\""), *ProjectPath, *FinalFileName);
 
-			IUATHelperModule::Get().CreateUatTask( CommandLine, PlatformName, LOCTEXT("ZipTaskName", "Zipping Up Project"),
+			IUATHelperModule::Get().CreateUatTask( CommandLine, GetTargetPlatformManager()->GetRunningTargetPlatform()->DisplayName(), LOCTEXT("ZipTaskName", "Zipping Up Project"),
 				LOCTEXT("ZipTaskShortName", "Zip Project Task"), FEditorStyle::GetBrush(TEXT("MainFrame.CookContent")), IUATHelperModule::UatTaskResultCallack(), FPaths::GetPath(FinalFileName));
 		}
 	}
 }
 
-void FMainFrameActionCallbacks::PackagingSettings()
-{
-	FModuleManager::LoadModuleChecked<ISettingsModule>("Settings").ShowViewer("Project", "Project", "Packaging");
-}
-
 //void FMainFrameActionCallbacks::LocalizeProject()
 //{
 //	FModuleManager::LoadModuleChecked<ILocalizationDashboardModule>("LocalizationDashboard").Show();
diff --git a/Engine/Source/Editor/MainFrame/Private/Frame/MainFrameActions.h b/Engine/Source/Editor/MainFrame/Private/Frame/MainFrameActions.h
--- a/Engine/Source/Editor/MainFrame/Private/Frame/MainFrameActions.h
+++ b/Engine/Source/Editor/MainFrame/Private/Frame/MainFrameActions.h
@@ -7,7 +7,6 @@
 #include "HAL/IConsoleManager.h"
 #include "Input/Reply.h"
 #include "Framework/Commands/Commands.h"
-#include "Settings/ProjectPackagingSettings.h"
 
 class FUICommandList;
 
@@ -35,7 +34,6 @@ public:
 	TSharedPtr< FUICommandInfo > OpenIDE;
 	TSharedPtr< FUICommandInfo > RefreshCodeProject;
 	TSharedPtr< FUICommandInfo > ZipUpProject;
-	TSharedPtr< FUICommandInfo > PackagingSettings;
 	TSharedPtr< FUICommandInfo > LocalizeProject;
 	TArray< TSharedPtr< FUICommandInfo > > SwitchProjectCommands;
 
@@ -175,30 +173,6 @@ public:
 	/** Adds code to the current project if it does not already have any */
 	static void AddCodeToProject();
 
-	/** Cooks the project's content for the specified platform. */
-	static void CookContent( const FName InPlatformInfoName );
-
-	/** Checks whether a menu action for cooking the project's content can execute. */
-	static bool CookContentCanExecute( const FName PlatformInfoName );
-
-	/** Sets the project packaging build configuration. */
-	static void PackageBuildConfiguration( EProjectPackagingBuildConfigurations BuildConfiguration );
-
-	/** Determines if the packaging build configuration can be used. */
-	static bool CanPackageBuildConfiguration( EProjectPackagingBuildConfigurations BuildConfiguration );
-
-	/** Determines whether the specified build configuration option is checked. */
-	static bool PackageBuildConfigurationIsChecked( EProjectPackagingBuildConfigurations BuildConfiguration );
-
-	/** Sets the project packaging build configuration. */
-	static void PackageBuildTarget( FString TargetName );
-
-	/** Determines whether the specified build configuration option is checked. */
-	static bool PackageBuildTargetIsChecked( FString TargetName );
-
-	/** Packages the project for the specified platform. */
-	static void PackageProject( const FName InPlatformInfoName );
-
 	/** Checks whether a menu action for packaging the project can execute. */
 	static bool PackageProjectCanExecute( const FName PlatformInfoName );
 
@@ -217,9 +191,6 @@ public:
 	/** Zips up the project */
 	static void ZipUpProject();
 
-	/** Opens the Packaging settings tab */
-	static void PackagingSettings();
-
 	/** Opens the Project Localization Dashboard */
 	static void LocalizeProject();
 	
diff --git a/Engine/Source/Editor/MainFrame/Private/Menus/MainMenu.cpp b/Engine/Source/Editor/MainFrame/Private/Menus/MainMenu.cpp
--- a/Engine/Source/Editor/MainFrame/Private/Menus/MainMenu.cpp
+++ b/Engine/Source/Editor/MainFrame/Private/Menus/MainMenu.cpp
@@ -19,7 +19,6 @@
 #include "UnrealEdGlobals.h"
 #include "Frame/MainFrameActions.h"
 #include "Menus/LayoutsMenu.h"
-#include "Menus/PackageProjectMenu.h"
 #include "Menus/RecentProjectsMenu.h"
 #include "Menus/SettingsMenu.h"
 #include "Menus/MainFrameTranslationEditorMenu.h"
@@ -451,13 +450,6 @@ void FMainMenu::RegisterFileProjectMenu()
 		FText::Format(LOCTEXT("AddCodeToProjectTooltip", "Adds C++ code to the project. The code can only be compiled if you have {0} installed."), ShortIDEName)
 	);
 
-	Section.AddSubMenu(
-		"PackageProject",
-		LOCTEXT("PackageProjectSubMenuLabel", "Package Project"),
-		LOCTEXT("PackageProjectSubMenuToolTip", "Compile, cook and package your project and its content for distribution."),
-		FNewMenuDelegate::CreateStatic( &FPackageProjectMenu::MakeMenu ), false, FSlateIcon(FEditorStyle::GetStyleSetName(), "MainFrame.PackageProject")
-	);
-
 	/*
 	MenuBuilder.AddMenuEntry( FMainFrameCommands::Get().LocalizeProject,
 		NAME_None,
@@ -499,30 +491,9 @@ void FMainMenu::RegisterFileProjectMenu()
 		OpenIDEIcon
 	);
 
-	Section.AddDynamicEntry("CookContentForPlatform", FNewToolMenuSectionDelegate::CreateLambda([](FToolMenuSection& InSection)
-	{
-		// @hack GDC: this should be moved somewhere else and be less hacky
-		ITargetPlatform* RunningTargetPlatform = GetTargetPlatformManager()->GetRunningTargetPlatform();
-
-		if (RunningTargetPlatform != nullptr)
-		{
-			const FName CookedPlatformName = RunningTargetPlatform->GetTargetPlatformInfo().VanillaInfo->Name;
-			const FText CookedPlatformText = FText::FromString(CookedPlatformName.ToString());
-
-			FUIAction Action(
-				FExecuteAction::CreateStatic(&FMainFrameActionCallbacks::CookContent, CookedPlatformName),
-				FCanExecuteAction::CreateStatic(&FMainFrameActionCallbacks::CookContentCanExecute, CookedPlatformName)
-			);
+	Section.AddMenuEntry(FMainFrameCommands::Get().ZipUpProject);
 
-			InSection.AddMenuEntry(
-				"CookContentForPlatform",
-				FText::Format(LOCTEXT("CookContentForPlatform", "Cook Content for {0}"), CookedPlatformText),
-				FText::Format(LOCTEXT("CookContentForPlatformTooltip", "Cook your game content for debugging on the {0} platform"), CookedPlatformText),
-				FSlateIcon(FEditorStyle::GetStyleSetName(), "MainFrame.CookContent"),
-				Action
-			);
-		}
-	}));
+	Section.AddMenuEntry("TempNote", FText::FromString(TEXT("[See Launch for Packaging]")), FText(), FSlateIcon(), FUIAction(FExecuteAction(), FCanExecuteAction::CreateLambda([] { return false; })), EUserInterfaceActionType::None);
 }
 
 void FMainMenu::RegisterRecentFileAndExitMenuItems()
diff --git a/Engine/Source/Editor/MainFrame/Private/Menus/PackageProjectMenu.h b/Engine/Source/Editor/MainFrame/Private/Menus/PackageProjectMenu.h
deleted file mode 100644
diff --git a/Engine/Source/Editor/TurnkeySupport/Private/TurnkeySupportModule.cpp b/Engine/Source/Editor/TurnkeySupport/Private/TurnkeySupportModule.cpp
new file mode 100644
--- /dev/null
+++ b/Engine/Source/Editor/TurnkeySupport/Private/TurnkeySupportModule.cpp
@@ -0,0 +1,1476 @@
+// Copyright Epic Games, Inc. All Rights Reserved.
+
+#include "TurnkeySupportModule.h"
+
+#include "SlateOptMacros.h"
+#include "Widgets/Images/SImage.h"
+#include "Widgets/Layout/SBox.h"
+#include "Widgets/Input/SButton.h"
+#include "GameProjectGenerationModule.h"
+#include "MessageLogModule.h"
+#include "EditorStyleSet.h"
+#include "Framework/Notifications/NotificationManager.h"
+#include "Widgets/Notifications/SNotificationList.h"
+#include "Widgets/Input/SEditableTextBox.h"
+#include "Misc/MessageDialog.h"
+#include "HAL/PlatformFileManager.h"
+#include "UnrealEdMisc.h"
+#include "ITargetDeviceServicesModule.h"
+#include "Interfaces/IProjectTargetPlatformEditorModule.h"
+#include "Dialogs/Dialogs.h"
+#include "Misc/DataDrivenPlatformInfoRegistry.h"
+#include "IUATHelperModule.h"
+#include "ISettingsEditorModule.h"
+#include "Async/Async.h"
+#include "Misc/FileHelper.h"
+#include "Interfaces/TargetDeviceId.h"
+#include "Interfaces/ITargetPlatform.h"
+#include "Interfaces/ITargetPlatformManagerModule.h"
+#include "Interfaces/IProjectManager.h"
+#include "ISettingsModule.h"
+#include "SourceControlHelpers.h"
+#include "ISourceControlModule.h"
+#include "ITargetDeviceProxy.h"
+#include "ITargetDeviceServicesModule.h"
+#include "Settings/ProjectPackagingSettings.h" // can go away once turnkey menu handles Package
+#include "IDesktopPlatform.h"
+#include "DesktopPlatformModule.h"
+#include "PlatformInfo.h"
+#include "InstalledPlatformInfo.h"
+#include "Settings/EditorExperimentalSettings.h"
+#include "CookerSettings.h"
+#include "Editor/UnrealEdEngine.h"
+#include "UnrealEdGlobals.h"
+#include "Framework/Commands/Commands.h"
+#include "Framework/Docking/TabManager.h"
+#include "Framework/MultiBox/MultiBoxBuilder.h"
+
+
+
+
+DEFINE_LOG_CATEGORY(LogTurnkeySupport);
+#define LOCTEXT_NAMESPACE "FTurnkeySupportModule"
+
+
+enum class EPrepareContentMode : uint8
+{
+	CookOnly,
+	Package,
+	PrepareForDebugging,
+};
+
+class FTurnkeySupportCallbacks
+{
+protected:
+	static const TCHAR* GetUATCompilationFlags()
+	{
+		// We never want to compile editor targets when invoking UAT in this context.
+		// If we are installed or don't have a compiler, we must assume we have a precompiled UAT.
+		return TEXT("-nocompileeditor");
+	}
+
+	static bool ShowBadSDKDialog(FName IniPlatformName)
+	{
+		// Don't show the warning during automation testing; the dlg is modal and blocks
+		if (!GIsAutomationTesting)
+		{
+			FFormatNamedArguments Args;
+			//		Args.Add(TEXT("DisplayName"), PlatformInfo->DisplayName);
+			Args.Add(TEXT("DisplayName"), FText::FromName(IniPlatformName));
+			FText WarningText = FText::Format(LOCTEXT("BadSDK_Message", "The SDK for {DisplayName} is not installed properly, which is needed to generate data. Check the SDK section of the Launch On menu in the main toolbar to update SDK.\n\nWould you like to attempt to continue anyway?"), Args);
+
+			FSuppressableWarningDialog::FSetupInfo Info(
+				WarningText,
+				LOCTEXT("BadSDK_Title", "SDK Not Setup"),
+				TEXT("BadSDKDialog")
+			);
+			Info.ConfirmText = LOCTEXT("BadSDK_Confirm", "Continue");
+			Info.CancelText = LOCTEXT("BadSDK_Cancel", "Cancel");
+			FSuppressableWarningDialog BadSDKgDialog(Info);
+
+			return BadSDKgDialog.ShowModal() != FSuppressableWarningDialog::EResult::Cancel;
+		}
+
+		return true;
+	}
+
+
+	static bool ShouldBuildProject(UProjectPackagingSettings* PackagingSettings, const ITargetPlatform* TargetPlatform)
+	{
+		const UProjectPackagingSettings::FConfigurationInfo& ConfigurationInfo = UProjectPackagingSettings::ConfigurationInfo[(int)PackagingSettings->BuildConfiguration];
+		bool bAssetNativizationEnabled = (PackagingSettings->BlueprintNativizationMethod != EProjectPackagingBlueprintNativizationMethod::Disabled);
+
+		// Get the target to build
+		const FTargetInfo* Target = PackagingSettings->GetBuildTargetInfo();
+
+		// Only build if the user elects to do so
+		bool bBuild = false;
+		if (PackagingSettings->Build == EProjectPackagingBuild::Always)
+		{
+			bBuild = true;
+		}
+		else if (PackagingSettings->Build == EProjectPackagingBuild::Never)
+		{
+			bBuild = false;
+		}
+		else if (PackagingSettings->Build == EProjectPackagingBuild::IfProjectHasCode)
+		{
+			bBuild = true;
+			if (FApp::GetEngineIsPromotedBuild() && !bAssetNativizationEnabled)
+			{
+				FString BaseDir;
+
+				// Get the target name
+				FString TargetName;
+				if (Target == nullptr)
+				{
+					TargetName = TEXT("UE4Game");
+				}
+				else
+				{
+					TargetName = Target->Name;
+				}
+
+				// Get the directory containing the receipt for this target, depending on whether the project needs to be built or not
+				FString ProjectDir = FPaths::GetPath(FPaths::GetProjectFilePath());
+				if (Target != nullptr && FPaths::IsUnderDirectory(Target->Path, ProjectDir))
+				{
+					UE_LOG(LogTurnkeySupport, Log, TEXT("Selected target: %s"), *Target->Name);
+					BaseDir = ProjectDir;
+				}
+				else
+				{
+					FText Reason;
+
+					FGameProjectGenerationModule& GameProjectModule = FModuleManager::LoadModuleChecked<FGameProjectGenerationModule>(TEXT("GameProjectGeneration"));
+					bool bProjectHasCode = GameProjectModule.Get().ProjectHasCodeFiles();
+
+					if (TargetPlatform->RequiresTempTarget(bProjectHasCode, ConfigurationInfo.Configuration, false, Reason))
+					{
+						UE_LOG(LogTurnkeySupport, Log, TEXT("Project requires temp target (%s)"), *Reason.ToString());
+						BaseDir = ProjectDir;
+					}
+					else
+					{
+						UE_LOG(LogTurnkeySupport, Log, TEXT("Project does not require temp target"));
+						BaseDir = FPaths::EngineDir();
+					}
+				}
+
+				// Check if the receipt is for a matching promoted target
+				FString UBTPlatformName = TargetPlatform->GetTargetPlatformInfo().DataDrivenPlatformInfo->UBTPlatformString;
+
+				extern LAUNCHERSERVICES_API bool HasPromotedTarget(const TCHAR * BaseDir, const TCHAR * TargetName, const TCHAR * Platform, EBuildConfiguration Configuration, const TCHAR * Architecture);
+				if (HasPromotedTarget(*BaseDir, *TargetName, *UBTPlatformName, ConfigurationInfo.Configuration, nullptr))
+				{
+					bBuild = false;
+				}
+			}
+		}
+		else if (PackagingSettings->Build == EProjectPackagingBuild::IfEditorWasBuiltLocally)
+		{
+			bBuild = !FApp::GetEngineIsPromotedBuild();
+		}
+
+		return bBuild;
+	}
+
+
+public:
+
+	/** Opens the Packaging settings tab */
+	static void PackagingSettings()
+	{
+		FModuleManager::LoadModuleChecked<ISettingsModule>("Settings").ShowViewer("Project", "Project", "Packaging");
+	}
+
+	static void OpenProjectLauncher()
+	{
+		FGlobalTabmanager::Get()->TryInvokeTab(FTabId("ProjectLauncher"));
+	}
+
+	static void OpenDeviceManager()
+	{
+		FGlobalTabmanager::Get()->TryInvokeTab(FTabId("DeviceManager"));
+	}
+
+	static bool CanCookOrPackage(FName IniPlatformName, EPrepareContentMode Mode)
+	{
+		return GetTargetPlatformManager()->FindTargetPlatform(IniPlatformName.ToString()) != nullptr;
+	}
+
+	static void CookOrPackage(FName IniPlatformName, EPrepareContentMode Mode)
+	{
+		const PlatformInfo::FTargetPlatformInfo* PlatformInfo = PlatformInfo::FindPlatformInfo(GetDefault<UProjectPackagingSettings>()->GetTargetPlatformForPlatform(IniPlatformName));
+		
+		// this is unexpected to be able to happen, but it could if there was a bad value saved in the UProjectPackagingSettings - if this trips, we should handle errors
+		check(PlatformInfo != nullptr);
+
+
+		// get all the helper objects
+		const FString UBTPlatformString = PlatformInfo->DataDrivenPlatformInfo->UBTPlatformString;
+		UProjectPackagingSettings* PackagingSettings = Cast<UProjectPackagingSettings>(UProjectPackagingSettings::StaticClass()->GetDefaultObject());
+		const FString ProjectPath = FPaths::IsProjectFilePathSet() ? FPaths::ConvertRelativePathToFull(FPaths::GetProjectFilePath()) : FPaths::RootDir() / FApp::GetProjectName() / FApp::GetProjectName() + TEXT(".uproject");
+
+
+		// check that we can proceed
+		{
+			if (FInstalledPlatformInfo::Get().IsPlatformMissingRequiredFile(UBTPlatformString))
+			{
+				if (!FInstalledPlatformInfo::OpenInstallerOptions())
+				{
+					FMessageDialog::Open(EAppMsgType::Ok, LOCTEXT("MissingPlatformFilesCook", "Missing required files to cook for this platform."));
+				}
+				return;
+			}
+
+			if (!FModuleManager::LoadModuleChecked<IProjectTargetPlatformEditorModule>("ProjectTargetPlatformEditor").ShowUnsupportedTargetWarning(IniPlatformName))
+			{
+				return;
+			}
+
+			if (PlatformInfo->DataDrivenPlatformInfo->GetSdkStatus() != DDPIPlatformSdkStatus::Valid && ShowBadSDKDialog(IniPlatformName) == false)
+			{
+				return;
+			}
+		}
+
+
+		// basic BuildCookRun params we always want
+		FString BuildCookRunParams = FString::Printf(TEXT("-nop4 -utf8output %s -cook "), GetUATCompilationFlags());
+
+
+		// set locations to engine and project
+		{
+			BuildCookRunParams += FString::Printf(TEXT(" -project=\"%s\" -ue4exe=\"%s\""), *ProjectPath, *FUnrealEdMisc::Get().GetExecutableForCommandlets());
+		}
+
+		// set the platform we are preparing content for
+		{
+			BuildCookRunParams += FString::Printf(TEXT(" -platform=%s"), *UBTPlatformString);
+		}
+
+		// Append any extra UAT flags specified for this platform flavor
+		if (!PlatformInfo->UATCommandLine.IsEmpty())
+		{
+			BuildCookRunParams += FString::Printf(TEXT(" %s"), *PlatformInfo->UATCommandLine);
+		}
+
+		// optional settings
+		if (PackagingSettings->bSkipEditorContent)
+		{
+			BuildCookRunParams += TEXT(" -SkipCookingEditorContent");
+		}
+		if (FApp::IsEngineInstalled())
+		{
+			BuildCookRunParams += TEXT(" -installed");
+		}
+		int32 NumCookers = GetDefault<UEditorExperimentalSettings>()->MultiProcessCooking;
+		if (NumCookers > 0)
+		{
+			BuildCookRunParams += FString::Printf(TEXT(" -NumCookersToSpawn=%d"), NumCookers);
+		}
+
+
+
+		// per mode settings
+		FText ContentPrepDescription;
+		FText ContentPrepTaskName;
+		const FSlateBrush* ContentPrepIcon = nullptr;
+		if (Mode == EPrepareContentMode::Package)
+		{
+			ContentPrepDescription = LOCTEXT("PackagingProjectTaskName", "Packaging project");
+			ContentPrepTaskName = LOCTEXT("PackagingTaskName", "Packaging");
+			ContentPrepIcon = FEditorStyle::GetBrush(TEXT("MainFrame.PackageProject"));
+
+			// let the user pick a target directory
+			if (PackagingSettings->StagingDirectory.Path.IsEmpty())
+			{
+				PackagingSettings->StagingDirectory.Path = FPaths::ProjectDir();
+			}
+
+			FString OutFolderName;
+
+			if (!FDesktopPlatformModule::Get()->OpenDirectoryDialog(FSlateApplication::Get().FindBestParentWindowHandleForDialogs(nullptr), LOCTEXT("PackageDirectoryDialogTitle", "Package project...").ToString(), PackagingSettings->StagingDirectory.Path, OutFolderName))
+			{
+				return;
+			}
+
+			PackagingSettings->StagingDirectory.Path = OutFolderName;
+			PackagingSettings->SaveConfig();
+
+
+			BuildCookRunParams += TEXT(" -stage -archive -package");
+
+			const ITargetPlatform* TargetPlatform = GetTargetPlatformManager()->FindTargetPlatform(PlatformInfo->Name);
+			if (ShouldBuildProject(PackagingSettings, TargetPlatform))
+			{
+				BuildCookRunParams += TEXT(" -build");
+			}
+
+			if (PackagingSettings->FullRebuild)
+			{
+				BuildCookRunParams += TEXT(" -clean");
+			}
+
+			if (PackagingSettings->bCompressed)
+			{
+				BuildCookRunParams += TEXT(" -compressed");
+			}
+
+			if (PackagingSettings->bUseIoStore)
+			{
+				BuildCookRunParams += TEXT(" -iostore");
+
+				// Pak file(s) must be used when using container file(s)
+				PackagingSettings->UsePakFile = true;
+			}
+
+			if (PackagingSettings->UsePakFile)
+			{
+				BuildCookRunParams += TEXT(" -pak");
+			}
+
+			if (PackagingSettings->bUseIoStore)
+			{
+				BuildCookRunParams += TEXT(" -iostore");
+			}
+
+			if (PackagingSettings->IncludePrerequisites)
+			{
+				BuildCookRunParams += TEXT(" -prereqs");
+			}
+
+			if (!PackagingSettings->ApplocalPrerequisitesDirectory.Path.IsEmpty())
+			{
+				BuildCookRunParams += FString::Printf(TEXT(" -applocaldirectory=\"%s\""), *(PackagingSettings->ApplocalPrerequisitesDirectory.Path));
+			}
+			else if (PackagingSettings->IncludeAppLocalPrerequisites)
+			{
+				BuildCookRunParams += TEXT(" -applocaldirectory=\"$(EngineDir)/Binaries/ThirdParty/AppLocalDependencies\"");
+			}
+
+			BuildCookRunParams += FString::Printf(TEXT(" -archivedirectory=\"%s\""), *PackagingSettings->StagingDirectory.Path);
+
+			if (PackagingSettings->ForDistribution)
+			{
+				BuildCookRunParams += TEXT(" -distribution");
+			}
+
+			if (!PackagingSettings->IncludeDebugFiles)
+			{
+				BuildCookRunParams += TEXT(" -nodebuginfo");
+			}
+
+			if (PackagingSettings->bGenerateChunks)
+			{
+				BuildCookRunParams += TEXT(" -manifests");
+			}
+
+			// Whether to include the crash reporter.
+			if (PackagingSettings->IncludeCrashReporter && PlatformInfo->DataDrivenPlatformInfo->bCanUseCrashReporter)
+			{
+				BuildCookRunParams += TEXT(" -CrashReporter");
+			}
+
+			if (PackagingSettings->bBuildHttpChunkInstallData)
+			{
+				BuildCookRunParams += FString::Printf(TEXT(" -manifests -createchunkinstall -chunkinstalldirectory=\"%s\" -chunkinstallversion=%s"), *(PackagingSettings->HttpChunkInstallDataDirectory.Path), *(PackagingSettings->HttpChunkInstallDataVersion));
+			}
+
+			const UProjectPackagingSettings::FConfigurationInfo& ConfigurationInfo = UProjectPackagingSettings::ConfigurationInfo[(int)PackagingSettings->GetBuildConfigurationForPlatform(IniPlatformName)];
+			if (PlatformInfo->PlatformType == EBuildTargetType::Server)
+			{
+				BuildCookRunParams += FString::Printf(TEXT(" -serverconfig=%s"), LexToString(ConfigurationInfo.Configuration));
+			}
+			else
+			{
+				BuildCookRunParams += FString::Printf(TEXT(" -clientconfig=%s"), LexToString(ConfigurationInfo.Configuration));
+			}
+		}
+		else if (Mode == EPrepareContentMode::CookOnly)
+		{
+			ContentPrepDescription = LOCTEXT("CookingContentTaskName", "Cooking content");
+			ContentPrepTaskName = LOCTEXT("CookingTaskName", "Cooking");
+			ContentPrepIcon = FEditorStyle::GetBrush(TEXT("MainFrame.CookContent"));
+
+
+			UCookerSettings const* CookerSettings = GetDefault<UCookerSettings>();
+			if (CookerSettings->bIterativeCookingForFileCookContent)
+			{
+				BuildCookRunParams += TEXT(" -iterate");
+			}
+
+			BuildCookRunParams += TEXT(" -skipstage");
+		}
+
+
+		FString TurnkeyParams = FString::Printf(TEXT(" -command=VerifySdk -platform=%s -UpdateIfNeeded -EditorIO"), *UBTPlatformString);
+
+		FString CommandLine = FString::Printf(TEXT("-ScriptsForProject=\"%s\" Turnkey %s BuildCookRun %s"),
+			*ProjectPath,
+			*TurnkeyParams,
+			*BuildCookRunParams
+		);
+
+		IUATHelperModule::Get().CreateUatTask(CommandLine, PlatformInfo->DisplayName, ContentPrepDescription, ContentPrepTaskName, ContentPrepIcon);
+	}
+
+	static void PackageBuildConfiguration(const PlatformInfo::FTargetPlatformInfo* Info, EProjectPackagingBuildConfigurations BuildConfiguration)
+	{
+		UProjectPackagingSettings* PackagingSettings = GetMutableDefault<UProjectPackagingSettings>();
+		PackagingSettings->SetBuildConfigurationForPlatform(Info->IniPlatformName, BuildConfiguration);
+		PackagingSettings->SaveConfig();
+	}
+
+	static bool CanPackageBuildConfiguration(const PlatformInfo::FTargetPlatformInfo* Info, EProjectPackagingBuildConfigurations BuildConfiguration)
+	{
+		return true;
+	}
+
+	static bool PackageBuildConfigurationIsChecked(const PlatformInfo::FTargetPlatformInfo* Info, EProjectPackagingBuildConfigurations BuildConfiguration)
+	{
+		return GetDefault<UProjectPackagingSettings>()->GetBuildConfigurationForPlatform(Info->IniPlatformName) == BuildConfiguration;
+	}
+
+	static void SetActiveTargetPlatform(const PlatformInfo::FTargetPlatformInfo* Info)
+	{
+		UProjectPackagingSettings* PackagingSettings = GetMutableDefault<UProjectPackagingSettings>();
+		PackagingSettings->SetTargetPlatformForPlatform(Info->IniPlatformName, Info->Name);
+		PackagingSettings->SaveConfig();
+	}
+
+	static bool CanSetActiveTargetPlatform(const PlatformInfo::FTargetPlatformInfo* Info)
+	{
+		return true;
+	}
+
+	static bool SetActiveTargetPlatformIsChecked(const PlatformInfo::FTargetPlatformInfo* Info)
+	{
+		return GetDefault<UProjectPackagingSettings>()->GetTargetPlatformForPlatform(Info->IniPlatformName) == Info->Name;
+	}
+
+	static void SetCookOnTheFly()
+	{
+		UCookerSettings* CookerSettings = GetMutableDefault<UCookerSettings>();
+
+		CookerSettings->bCookOnTheFlyForLaunchOn = !CookerSettings->bCookOnTheFlyForLaunchOn;
+		CookerSettings->Modify(true);
+
+		// Update source control
+
+		FString ConfigPath = FPaths::ConvertRelativePathToFull(CookerSettings->GetDefaultConfigFilename());
+
+		if (FPlatformFileManager::Get().GetPlatformFile().FileExists(*ConfigPath))
+		{
+			if (ISourceControlModule::Get().IsEnabled())
+			{
+				FText ErrorMessage;
+
+				if (!SourceControlHelpers::CheckoutOrMarkForAdd(ConfigPath, FText::FromString(ConfigPath), NULL, ErrorMessage))
+				{
+					FNotificationInfo Info(ErrorMessage);
+					Info.ExpireDuration = 3.0f;
+					FSlateNotificationManager::Get().AddNotification(Info);
+				}
+			}
+			else
+			{
+				if (!FPlatformFileManager::Get().GetPlatformFile().SetReadOnly(*ConfigPath, false))
+				{
+					FNotificationInfo Info(FText::Format(LOCTEXT("FailedToMakeWritable", "Could not make {0} writable."), FText::FromString(ConfigPath)));
+					Info.ExpireDuration = 3.0f;
+					FSlateNotificationManager::Get().AddNotification(Info);
+				}
+			}
+		}
+
+		// Save settings
+		CookerSettings->UpdateSinglePropertyInConfigFile(CookerSettings->GetClass()->FindPropertyByName(GET_MEMBER_NAME_CHECKED(UCookerSettings, bCookOnTheFlyForLaunchOn)), CookerSettings->GetDefaultConfigFilename());
+
+	}
+
+	static bool CanSetCookOnTheFly()
+	{
+		return true;
+	}
+
+	static bool SetCookOnTheFlyIsChecked()
+	{
+		return GetDefault<UCookerSettings>()->bCookOnTheFlyForLaunchOn;
+	}
+};
+
+class FTurnkeySupportCommands : public TCommands<FTurnkeySupportCommands>
+{
+public:
+	void RegisterCommands()
+	{
+		UI_COMMAND(PackagingSettings, "Packaging Settings...", "Opens the settings for project packaging", EUserInterfaceActionType::Button, FInputChord());
+		ActionList->MapAction(PackagingSettings, FExecuteAction::CreateStatic(&FTurnkeySupportCallbacks::PackagingSettings));
+
+	}
+
+	TSharedPtr< FUICommandInfo > PackagingSettings;
+
+private:
+
+	friend class TCommands<FTurnkeySupportCommands>;
+
+	FTurnkeySupportCommands()
+		: TCommands<FTurnkeySupportCommands>("TurnkeySupport", LOCTEXT("TurnkeySupport", "Turnkey and General Platform Options"), "MainFrame", FEditorStyle::GetStyleSetName())
+	{
+
+	}
+
+public:
+	/** List of all of the main frame commands */
+	static TSharedRef<FUICommandList> ActionList;
+
+};
+
+TSharedRef<FUICommandList> FTurnkeySupportCommands::ActionList = MakeShareable(new FUICommandList);
+
+
+
+
+static void TurnkeyInstallSdk(FString PlatformName, bool bPreferFull, bool bForceInstall, const TCHAR* DeviceName = nullptr)
+{
+	//	FString CommandLine = FString::Printf(TEXT("Turnkey -command=InstallSdk -platform=%s -BestAvailable -AllowAutoSdk -EditorIO -noturnkeyvariables -utf8output -WaitForUATMutex"), *PlatformName);
+
+	FString OptionalOptions;
+	if (bPreferFull)
+	{
+		OptionalOptions += TEXT(" -PreferFull");
+	}
+	if (bForceInstall)
+	{
+		OptionalOptions += DeviceName != nullptr ? TEXT(" -ForceDeviceInstall") : TEXT(" -ForceSdkInstall");
+	}
+	if (DeviceName != nullptr)
+	{
+		OptionalOptions += FString::Printf(TEXT(" -Device=%s"), DeviceName);
+	}
+
+	FString CommandLine = FString::Printf(TEXT("Turnkey -command=VerifySdk -UpdateIfNeeded -platform=%s %s -EditorIO -noturnkeyvariables -utf8output -WaitForUATMutex"), *PlatformName, *OptionalOptions);
+
+	FText TaskName = LOCTEXT("InstallingSdk", "Installing Sdk");
+	IUATHelperModule::Get().CreateUatTask(CommandLine, FText::FromString(PlatformName), TaskName, TaskName, FEditorStyle::GetBrush(TEXT("MainFrame.PackageProject")),
+		[PlatformName](FString, double)
+	{
+		AsyncTask(ENamedThreads::GameThread, [PlatformName]()
+		{
+
+			// read in env var changes
+			// @todo turnkey move this and make it mac/linux aware
+			FString TurnkeyEnvVarsFilename = FPaths::Combine(FPaths::EngineIntermediateDir(), TEXT("Turnkey/PostTurnkeyVariables.bat"));
+
+			if (IFileManager::Get().FileExists(*TurnkeyEnvVarsFilename))
+			{
+				TArray<FString> Contents;
+				if (FFileHelper::LoadFileToStringArray(Contents, *TurnkeyEnvVarsFilename))
+				{
+					for (const FString& Line : Contents)
+					{
+						if (Line.StartsWith(TEXT("set ")))
+						{
+							// split the line
+							FString VariableLine = Line.Mid(4);
+							int Equals;
+							if (VariableLine.FindChar('=', Equals))
+							{
+								// set the key/value
+								FString Key = VariableLine.Mid(0, Equals);
+								FString Value = VariableLine.Mid(Equals + 1);
+
+								FPlatformMisc::SetEnvironmentVar(*Key, *Value);
+
+								UE_LOG(LogTurnkeySupport, Log, TEXT("Turnkey setting env var: %s = %s"), *Key, *Value);
+							}
+						}
+					}
+				}
+			}
+
+			// update the Sdk status
+//			FDataDrivenPlatformInfoRegistry::UpdateSdkStatus();
+			GetTargetPlatformManager()->UpdateAfterSDKInstall(*PlatformName);
+			RenderUtilsInit();
+
+
+			// show restart dialog
+			FModuleManager::GetModuleChecked<ISettingsEditorModule>("SettingsEditor").OnApplicationRestartRequired();
+		});
+	}
+	);
+}
+
+static TAttribute<FText> MakeSdkStatusAttribute(FName IniPlatformName, TSharedPtr< ITargetDeviceProxy> DeviceProxy)
+{
+	// 	FFormatNamedArguments LabelArguments;
+	// 	LabelArguments.Add(TEXT("DeviceName"), FText::FromString(DeviceProxy->GetName()));
+	// 
+	// 	if (!DeviceProxy->IsConnected())
+	// 	{
+	// 		LabelArguments.Add(TEXT("HostUser"), LOCTEXT("DisconnectedHint", " [Disconnected]"));
+	// 	}
+	// 	else if (DeviceProxy->GetHostUser() != FPlatformProcess::UserName(false))
+	// 	{
+	// 		LabelArguments.Add(TEXT("HostUser"), FText::FromString(DeviceProxy->GetHostUser()));
+	// 	}
+	// 	else
+	// 	{
+	// 		LabelArguments.Add(TEXT("HostUser"), FText::GetEmpty());
+	// 	}
+	// 
+	// 	OutLabel = FText::Format(LOCTEXT("LaunchDeviceLabel", "{DeviceName}{HostUser}"), LabelArguments);
+	// 	TAttribute<FText> = Make
+
+
+	FString DisplayString = DeviceProxy ? DeviceProxy->GetName() : IniPlatformName.ToString();
+	FString DeviceId = DeviceProxy ? DeviceProxy->GetTargetDeviceId(NAME_None) : FString();
+
+	return TAttribute<FText>::Create(TAttribute<FText>::FGetter::CreateLambda([IniPlatformName, DisplayString, DeviceId]()
+	{
+		// get the status, or Unknown if it's not there
+		const FDataDrivenPlatformInfo& Info = FDataDrivenPlatformInfoRegistry::GetPlatformInfo(IniPlatformName);
+		DDPIPlatformSdkStatus Status = DeviceId.Len() ? Info.GetStatusForDeviceId(DeviceId) : Info.GetSdkStatus(false);
+
+		// @todo turnkey: Have premade FText's by SdkStatus for speed
+		const TCHAR* Desc =
+			Status == DDPIPlatformSdkStatus::Querying ? TEXT("Querying...") :
+			Status == DDPIPlatformSdkStatus::Valid ? TEXT("Valid Sdk") :
+			Status == DDPIPlatformSdkStatus::OutOfDate ? TEXT("Outdated Sdk") :
+			Status == DDPIPlatformSdkStatus::NoSdk ? TEXT("No Sdk") :
+			Status == DDPIPlatformSdkStatus::FlashValid ? TEXT("Valid Flash") :
+			Status == DDPIPlatformSdkStatus::FlashOutOfDate ? TEXT("Outdated Flash") :
+			TEXT("???");
+		return FText::FromString(FString::Printf(TEXT("%s (%s)"), *DisplayString, Desc));
+	}));
+}
+
+
+
+
+
+
+static void MakeTurnkeyPlatformMenu(FMenuBuilder& MenuBuilder, FName IniPlatformName, ITargetDeviceServicesModule* TargetDeviceServicesModule)
+{
+	const FDataDrivenPlatformInfo& DDPI = FDataDrivenPlatformInfoRegistry::GetPlatformInfo(IniPlatformName);
+	FString UBTPlatformString = DDPI.UBTPlatformString;
+
+	const PlatformInfo::FTargetPlatformInfo* VanillaInfo = PlatformInfo::FindVanillaPlatformInfo(IniPlatformName);
+
+	if (VanillaInfo != nullptr)
+	{
+
+		MenuBuilder.BeginSection("ContentManagement", LOCTEXT("TurnkeySection_Content", "Content Management"));
+
+
+		MenuBuilder.AddMenuEntry(
+			LOCTEXT("Turnkey_CookContent", "Cook Content"),
+			LOCTEXT("TurnkeyTooltip_CookContent", "Cook this project for the selected configuration and target"),
+			FSlateIcon(),
+			FUIAction(
+				FExecuteAction::CreateStatic(&FTurnkeySupportCallbacks::CookOrPackage, IniPlatformName, EPrepareContentMode::CookOnly),
+				FCanExecuteAction::CreateStatic(&FTurnkeySupportCallbacks::CanCookOrPackage, IniPlatformName, EPrepareContentMode::CookOnly)
+			)
+		);
+
+		MenuBuilder.AddMenuEntry(
+			LOCTEXT("Turnkey_PrepareForDebugging", "Prepare For Debugging"),
+			LOCTEXT("TurnkeyTooltip_PrepareForDebugging", "Prepare this project for debugging"),
+			FSlateIcon(),
+			FUIAction(
+				FExecuteAction::CreateStatic(&FTurnkeySupportCallbacks::CookOrPackage, IniPlatformName, EPrepareContentMode::PrepareForDebugging),
+				FCanExecuteAction::CreateStatic(&FTurnkeySupportCallbacks::CanCookOrPackage, IniPlatformName, EPrepareContentMode::PrepareForDebugging)
+			)
+		);
+
+		MenuBuilder.AddMenuEntry(
+			LOCTEXT("Turnkey_PackageProject", "Package Project"),
+			LOCTEXT("TurnkeyTooltip_PackageProject", "Package this project and archive it to a user-selected directory. This can then be used to install and run."),
+			FSlateIcon(),
+			FUIAction(
+				FExecuteAction::CreateStatic(&FTurnkeySupportCallbacks::CookOrPackage, IniPlatformName, EPrepareContentMode::Package),
+				FCanExecuteAction::CreateStatic(&FTurnkeySupportCallbacks::CanCookOrPackage, IniPlatformName, EPrepareContentMode::Package)
+			)
+		);
+
+		MenuBuilder.EndSection();
+
+
+		MenuBuilder.BeginSection("BuildConfig", LOCTEXT("TurnkeySection_BuildConfig", "Binary Configuration"));
+			EProjectType ProjectType = FGameProjectGenerationModule::Get().ProjectHasCodeFiles() ? EProjectType::Code : EProjectType::Content;
+			TArray<EProjectPackagingBuildConfigurations> PackagingConfigurations = UProjectPackagingSettings::GetValidPackageConfigurations();
+
+			for (EProjectPackagingBuildConfigurations PackagingConfiguration : PackagingConfigurations)
+			{
+				const UProjectPackagingSettings::FConfigurationInfo& ConfigurationInfo = UProjectPackagingSettings::ConfigurationInfo[(int)PackagingConfiguration];
+				if (FInstalledPlatformInfo::Get().IsValid(TOptional<EBuildTargetType>(), TOptional<FString>(), ConfigurationInfo.Configuration, ProjectType, EInstalledPlatformState::Downloaded))
+				{
+					MenuBuilder.AddMenuEntry(
+						ConfigurationInfo.Name,
+						ConfigurationInfo.ToolTip,
+						FSlateIcon(),
+						FUIAction(
+							FExecuteAction::CreateStatic(&FTurnkeySupportCallbacks::PackageBuildConfiguration, VanillaInfo, PackagingConfiguration),
+							FCanExecuteAction::CreateStatic(&FTurnkeySupportCallbacks::CanPackageBuildConfiguration, VanillaInfo, PackagingConfiguration),
+							FIsActionChecked::CreateStatic(&FTurnkeySupportCallbacks::PackageBuildConfigurationIsChecked, VanillaInfo, PackagingConfiguration)
+						),
+						NAME_None,
+						EUserInterfaceActionType::RadioButton
+					);
+				}
+			}
+		MenuBuilder.EndSection();
+
+		MenuBuilder.BeginSection("TargetSelection", LOCTEXT("TurnkeySection_TargetSelection", "Target Selection"));
+
+			// gather all platform infos
+			TArray<const PlatformInfo::FTargetPlatformInfo*> AllTargets = { VanillaInfo };
+			AllTargets.Append(VanillaInfo->Flavors);
+
+			for (const PlatformInfo::FTargetPlatformInfo* Info : AllTargets)
+			{
+				MenuBuilder.AddMenuEntry(
+					Info->DisplayName,
+					FText(),
+					FSlateIcon(),
+					FUIAction(
+						FExecuteAction::CreateStatic(&FTurnkeySupportCallbacks::SetActiveTargetPlatform, Info),
+						FCanExecuteAction::CreateStatic(&FTurnkeySupportCallbacks::CanSetActiveTargetPlatform, Info),
+						FIsActionChecked::CreateStatic(&FTurnkeySupportCallbacks::SetActiveTargetPlatformIsChecked, Info)
+					),
+					NAME_None,
+					EUserInterfaceActionType::RadioButton
+				);
+			}
+		MenuBuilder.EndSection();
+
+		MenuBuilder.BeginSection("AllDevices", LOCTEXT("TurnkeySection_AllDevices", "All Devices"));
+
+		TArray<TSharedPtr<ITargetDeviceProxy>> DeviceProxies;
+		TargetDeviceServicesModule->GetDeviceProxyManager()->GetAllProxies(IniPlatformName, DeviceProxies);
+
+		for (const TSharedPtr<ITargetDeviceProxy> Proxy : DeviceProxies)
+		{
+			FString DeviceName = Proxy->GetName();
+			MenuBuilder.AddSubMenu(
+				MakeSdkStatusAttribute(IniPlatformName, Proxy),
+				FText(),
+				FNewMenuDelegate::CreateLambda([UBTPlatformString, DeviceName](FMenuBuilder& SubMenuBuilder) {
+				SubMenuBuilder.AddMenuEntry(
+					LOCTEXT("Turnkey_RepairDevice", "Repair Device as Needed"),
+					LOCTEXT("TurnkeyTooltip_RepairDevice", "Perform any fixup that may be needed on this device. If up to date already, nothing will be done."),
+					FSlateIcon(),
+					FExecuteAction::CreateStatic(TurnkeyInstallSdk, UBTPlatformString, false, false, *DeviceName)
+				);
+
+				SubMenuBuilder.AddMenuEntry(
+					LOCTEXT("Turnkey_ForceRepairDevice", "Force Repair Device"),
+					LOCTEXT("TurnkeyTooltip_ForceRepairDevice", "Force repairing anything on the device needed (update firmware, etc). Will perform all steps possible, even if not needed."),
+					FSlateIcon(),
+					FExecuteAction::CreateStatic(TurnkeyInstallSdk, UBTPlatformString, true, false, *DeviceName)
+				);
+			})
+			);
+		}
+
+		MenuBuilder.EndSection();
+		}
+
+	MenuBuilder.BeginSection("AllDevices", LOCTEXT("TurnkeySection_Sdks", "Sdk Managment"));
+
+	const TCHAR* NoDevice = nullptr;
+	if (DDPI.GetSdkStatus() == DDPIPlatformSdkStatus::OutOfDate)
+	{
+		MenuBuilder.AddMenuEntry(
+			LOCTEXT("Turnkey_InstallSdkMinimal", "Update Sdk (Prefer Minimal)"),
+			LOCTEXT("TurnkeyTooltip_InstallSdkMinimal", "Attempt to update an Sdk, as hosted by your studio. Will attempt to install a minimal Sdk (useful for building/running only)"),
+			FSlateIcon(),
+			FExecuteAction::CreateStatic(TurnkeyInstallSdk, UBTPlatformString, false, false, NoDevice)
+		);
+
+		MenuBuilder.AddMenuEntry(
+			LOCTEXT("Turnkey_InstallSdkFull", "Update Sdk (Prefer Full)"),
+			LOCTEXT("TurnkeyTooltip_InstallSdkMinimal", "Attempt to update an Sdk, as hosted by your studio. Will attempt to install a full Sdk (useful profiling or other use cases)"),
+			FSlateIcon(),
+			FExecuteAction::CreateStatic(TurnkeyInstallSdk, UBTPlatformString, true, false, NoDevice)
+		);
+	}
+	else if (DDPI.GetSdkStatus() == DDPIPlatformSdkStatus::Valid)
+	{
+		MenuBuilder.AddMenuEntry(
+			LOCTEXT("Turnkey_InstallSdkMinimal", "Force Reinstall Sdk (Prefer Minimal)"),
+			LOCTEXT("TurnkeyTooltip_InstallSdkMinimal", "Attempt to force re-install an Sdk, as hosted by your studio. Will attempt to install a minimal Sdk (useful for building/running only)"),
+			FSlateIcon(),
+			FExecuteAction::CreateStatic(TurnkeyInstallSdk, UBTPlatformString, false, true, NoDevice)
+		);
+
+		MenuBuilder.AddMenuEntry(
+			LOCTEXT("Turnkey_InstallSdkFull", "Force Reinstall (Prefer Full)"),
+			LOCTEXT("TurnkeyTooltip_InstallSdkMinimal", "Attempt to force re-install an Sdk, as hosted by your studio. Will attempt to install a full Sdk (useful profiling or other use cases)"),
+			FSlateIcon(),
+			FExecuteAction::CreateStatic(TurnkeyInstallSdk, UBTPlatformString, true, true, NoDevice)
+		);
+	}
+	else
+	{
+		MenuBuilder.AddMenuEntry(
+			LOCTEXT("Turnkey_InstallSdkMinimal", "Install Sdk (Prefer Minimal)"),
+			LOCTEXT("TurnkeyTooltip_InstallSdkMinimal", "Attempt to install an Sdk, as hosted by your studio. Will attempt to install a minimal Sdk (useful for building/running only)"),
+			FSlateIcon(),
+			FExecuteAction::CreateStatic(TurnkeyInstallSdk, UBTPlatformString, false, false, NoDevice)
+		);
+
+		MenuBuilder.AddMenuEntry(
+			LOCTEXT("Turnkey_InstallSdkFull", "Install Sdk (Prefer Full)"),
+			LOCTEXT("TurnkeyTooltip_InstallSdkMinimal", "Attempt to install an Sdk, as hosted by your studio. Will attempt to install a full Sdk (useful profiling or other use cases)"),
+			FSlateIcon(),
+			FExecuteAction::CreateStatic(TurnkeyInstallSdk, UBTPlatformString, true, false, NoDevice)
+		);
+	}
+}
+
+
+// Launch On
+
+bool CanLaunchOnDevice(const FString& DeviceName)
+{
+	if (!GUnrealEd->IsPlayingViaLauncher())
+	{
+		static TWeakPtr<ITargetDeviceProxyManager> DeviceProxyManagerPtr;
+
+		if (!DeviceProxyManagerPtr.IsValid())
+		{
+			ITargetDeviceServicesModule* TargetDeviceServicesModule = FModuleManager::Get().LoadModulePtr<ITargetDeviceServicesModule>(TEXT("TargetDeviceServices"));
+			if (TargetDeviceServicesModule)
+			{
+				DeviceProxyManagerPtr = TargetDeviceServicesModule->GetDeviceProxyManager();
+			}
+		}
+
+		TSharedPtr<ITargetDeviceProxyManager> DeviceProxyManager = DeviceProxyManagerPtr.Pin();
+		if (DeviceProxyManager.IsValid())
+		{
+			TSharedPtr<ITargetDeviceProxy> DeviceProxy = DeviceProxyManager->FindProxy(DeviceName);
+			if (DeviceProxy.IsValid() && DeviceProxy->IsConnected() && DeviceProxy->IsAuthorized())
+			{
+				return true;
+			}
+
+			// check if this is an aggregate proxy
+			TArray<TSharedPtr<ITargetDeviceProxy>> Devices;
+			DeviceProxyManager->GetProxies(FName(*DeviceName), false, Devices);
+
+			// returns true if the game can be launched al least on 1 device
+			for (auto DevicesIt = Devices.CreateIterator(); DevicesIt; ++DevicesIt)
+			{
+				TSharedPtr<ITargetDeviceProxy> DeviceAggregateProxy = *DevicesIt;
+				if (DeviceAggregateProxy.IsValid() && DeviceAggregateProxy->IsConnected() && DeviceAggregateProxy->IsAuthorized())
+				{
+					return true;
+				}
+			}
+
+		}
+	}
+
+	return false;
+}
+
+
+void LaunchOnDevice(const FString& DeviceId, const FString& DeviceName, bool bUseTurnkey)
+{
+	FTargetDeviceId TargetDeviceId;
+	if (FTargetDeviceId::Parse(DeviceId, TargetDeviceId))
+	{
+		const PlatformInfo::FTargetPlatformInfo* const PlatformInfo = PlatformInfo::FindPlatformInfo(*TargetDeviceId.GetPlatformName());
+		FString UBTPlatformName = PlatformInfo->DataDrivenPlatformInfo->UBTPlatformString;
+		FString IniPlatformName = PlatformInfo->IniPlatformName.ToString();
+
+		check(PlatformInfo);
+
+		if (FInstalledPlatformInfo::Get().IsPlatformMissingRequiredFile(UBTPlatformName))
+		{
+			if (!FInstalledPlatformInfo::OpenInstallerOptions())
+			{
+				FMessageDialog::Open(EAppMsgType::Ok, LOCTEXT("MissingPlatformFilesLaunch", "Missing required files to launch on this platform."));
+			}
+			return;
+		}
+
+		if (FModuleManager::LoadModuleChecked<IProjectTargetPlatformEditorModule>("ProjectTargetPlatformEditor").ShowUnsupportedTargetWarning(*TargetDeviceId.GetPlatformName()))
+		{
+			GUnrealEd->CancelPlayingViaLauncher();
+
+			FRequestPlaySessionParams::FLauncherDeviceInfo DeviceInfo;
+			DeviceInfo.DeviceId = DeviceId;
+			DeviceInfo.DeviceName = DeviceName;
+			// @todo turnkey: we set this to false because we will kick off a Turnkey run before cooking, etc, to get an early warning. however, if it's too difficult
+			// to get an error back from CreateUatTask, then we should set this to bUseTurnkey and remove the block below, and let the code in FLauncherWorker::CreateAndExecuteTasks handle it
+			DeviceInfo.bUpdateDeviceFlash = false;
+
+			FRequestPlaySessionParams SessionParams;
+			SessionParams.SessionDestination = EPlaySessionDestinationType::Launcher;
+			SessionParams.LauncherTargetDevice = DeviceInfo;
+
+			// if we want to check device flash before we start cooking, kick it off now. we could delay this 
+			if (bUseTurnkey)
+			{
+				FString CommandLine = FString::Printf(TEXT("Turnkey -command=VerifySdk -UpdateIfNeeded -platform=%s -EditorIO -noturnkeyvariables -device=%s -utf8output -WaitForUATMutex"), *UBTPlatformName, *DeviceName);
+				FText TaskName = LOCTEXT("VerifyingSDK", "Verifying SDK and Device");
+
+				IUATHelperModule::Get().CreateUatTask(CommandLine, FText::FromString(IniPlatformName), TaskName, TaskName, FEditorStyle::GetBrush(TEXT("MainFrame.PackageProject")),
+					[SessionParams](FString Result, double)
+				{
+					// unfortunate string comparison for success
+					bool bWasSuccessful = Result == TEXT("Completed");
+					AsyncTask(ENamedThreads::GameThread, [SessionParams, bWasSuccessful]()
+					{
+						if (bWasSuccessful)
+						{
+							GUnrealEd->RequestPlaySession(SessionParams);
+						}
+						else
+						{
+							TSharedRef<SWindow> Win = OpenMsgDlgInt_NonModal(EAppMsgType::YesNo, LOCTEXT("SDKCheckFailed", "SDK Verification failed. Would you like to attempt the Launch On anyway?"), LOCTEXT("SDKCheckFailedTitle", "SDK Verification"),
+								FOnMsgDlgResult::CreateLambda([SessionParams](const TSharedRef<SWindow>&, EAppReturnType::Type Choice)
+							{
+								if (Choice == EAppReturnType::Yes)
+								{
+									GUnrealEd->RequestPlaySession(SessionParams);
+								}
+							}));
+							Win->ShowWindow();
+						}
+					});
+				});
+			}
+			else
+			{
+				GUnrealEd->RequestPlaySession(SessionParams);
+			}
+		}
+	}
+}
+
+static void PrepareLaunchOn(FString DeviceId, FString DeviceName)
+{
+	ULevelEditorPlaySettings* PlaySettings = GetMutableDefault<ULevelEditorPlaySettings>();
+
+	PlaySettings->LastExecutedLaunchModeType = LaunchMode_OnDevice;
+	PlaySettings->LastExecutedLaunchDevice = DeviceId;
+	PlaySettings->LastExecutedLaunchName = DeviceName;
+
+	PlaySettings->PostEditChange();
+
+	PlaySettings->SaveConfig();
+}
+
+void HandleLaunchOnDeviceActionExecute(FString DeviceId, FString DeviceName, bool bUseTurnkey)
+{
+	PrepareLaunchOn(DeviceId, DeviceName);
+	LaunchOnDevice(DeviceId, DeviceName, bUseTurnkey);
+}
+
+
+bool HandleLaunchOnDeviceActionCanExecute(FString DeviceName)
+{
+	return CanLaunchOnDevice(DeviceName);
+}
+
+bool HandleLaunchOnDeviceActionIsChecked(FString DeviceName)
+{
+	return (DeviceName == GetDefault<ULevelEditorPlaySettings>()->LastExecutedLaunchName);
+}
+
+static void GenerateDeviceProxyMenuParams(TSharedPtr<ITargetDeviceProxy> DeviceProxy, FName PlatformName, FUIAction& OutAction, FText& OutTooltip)
+{
+	// 	// create an All_<platform>_devices_on_<host> submenu
+	// 	if (DeviceProxy->IsAggregated())
+	// 	{
+	// 		FString AggregateDevicedName(FString::Printf(TEXT("  %s"), *DeviceProxy->GetName())); //align with the other menu entries
+	// 		FSlateIcon AggregateDeviceIcon(FEditorStyle::GetStyleSetName(), EditorPlatformInfo->GetIconStyleName(PlatformInfo::EPlatformIconSize::Normal));
+	// 
+	// 		MenuBuilder.AddSubMenu(
+	// 			FText::FromString(AggregateDevicedName),
+	// 			FText::FromString(AggregateDevicedName),
+	// 			FNewMenuDelegate::CreateStatic(&MakeAllDevicesSubMenu, EditorPlatformInfo, DeviceProxy),
+	// 			false, AggregateDeviceIcon, true
+	// 		);
+	// 		continue;
+	// 	}
+
+		// ... create an action...
+	OutAction = FUIAction(
+		FExecuteAction::CreateStatic(&HandleLaunchOnDeviceActionExecute, DeviceProxy->GetTargetDeviceId(NAME_None), DeviceProxy->GetName(), true)
+		//		, FCanExecuteAction::CreateStatic(&HandleLaunchOnDeviceActionCanExecute, DeviceProxy->GetName())
+		//		, FIsActionChecked::CreateStatic(&HandleLaunchOnDeviceActionIsChecked, DeviceProxy->GetName())
+	);
+
+	// ... generate tooltip text
+	FFormatNamedArguments TooltipArguments;
+	TooltipArguments.Add(TEXT("DeviceID"), FText::FromString(DeviceProxy->GetName()));
+	TooltipArguments.Add(TEXT("DisplayName"), FText::FromName(PlatformName));
+	OutTooltip = FText::Format(LOCTEXT("LaunchDeviceToolTipText_ThisDevice", "Launch the game on this {DisplayName} device ({DeviceID})"), TooltipArguments);
+	if (!DeviceProxy->IsAuthorized())
+	{
+		OutTooltip = FText::Format(LOCTEXT("LaunchDeviceToolTipText_UnauthorizedOrLocked", "{DisplayName} device ({DeviceID}) is unauthorized or locked"), TooltipArguments);
+	}
+
+	FProjectStatus ProjectStatus;
+	if (IProjectManager::Get().QueryStatusForCurrentProject(ProjectStatus) && !ProjectStatus.IsTargetPlatformSupported(PlatformName))
+	{
+		FText TooltipLine2 = FText::Format(LOCTEXT("LaunchDevicePlatformWarning", "{DisplayName} is not listed as a target platform for this project, so may not run as expected."), TooltipArguments);
+		OutTooltip = FText::Format(FText::FromString(TEXT("{0}\n\n{1}")), OutTooltip, TooltipLine2);
+	}
+
+}
+
+
+
+
+
+
+
+
+
+
+TSharedRef<SWidget> FTurnkeySupportModule::MakeTurnkeyMenu() const
+{
+	const FTurnkeySupportCommands& Commands = FTurnkeySupportCommands::Get();
+
+	const bool bShouldCloseWindowAfterMenuSelection = true;
+	FMenuBuilder MenuBuilder(bShouldCloseWindowAfterMenuSelection, FTurnkeySupportCommands::ActionList);
+
+	// shared devices section
+	ITargetDeviceServicesModule* TargetDeviceServicesModule = static_cast<ITargetDeviceServicesModule*>(FModuleManager::Get().LoadModule(TEXT("TargetDeviceServices")));
+	IProjectTargetPlatformEditorModule& ProjectTargetPlatformEditorModule = FModuleManager::LoadModuleChecked<IProjectTargetPlatformEditorModule>("ProjectTargetPlatformEditor");
+
+	TArray<FString> DeviceIdsToQuery;
+
+	MenuBuilder.BeginSection("LevelEditorLaunchDevices", LOCTEXT("TurnkeySection_LaunchButtonDevices", "Quick Launch"));
+	{
+		for (const auto Pair : FDataDrivenPlatformInfoRegistry::GetAllPlatformInfos())
+		{
+			if (Pair.Value.bIsFakePlatform)
+			{
+				continue;
+			}
+
+			FName PlatformName = Pair.Key;
+			const FDataDrivenPlatformInfo& Info = Pair.Value;
+
+			// look for devices for all platforms, even if the platform isn't installed - Turnkey can install Sdk after selecting LaunchOn
+			TArray<TSharedPtr<ITargetDeviceProxy>> DeviceProxies;
+			TargetDeviceServicesModule->GetDeviceProxyManager()->GetAllProxies(PlatformName, DeviceProxies);
+
+			if (DeviceProxies.Num() > 0)
+			{
+				// 					Algo::Sort(DeviceProxies, [](TSharedPtr<ITargetDeviceProxy> A, TSharedPtr<ITargetDeviceProxy> B)
+				// 					{ return A->Name == LastChosen || B->Name != LastChosen && A->IsDefault(); });
+
+
+				// always use the first one, after sorting
+				FUIAction Action;
+				FText Tooltip;
+				GenerateDeviceProxyMenuParams(DeviceProxies[0], PlatformName, Action, Tooltip);
+
+				if (DeviceProxies.Num() == 1)
+				{
+					MenuBuilder.AddMenuEntry(
+						MakeSdkStatusAttribute(PlatformName, DeviceProxies[0]),
+						Tooltip,
+						FSlateIcon(FEditorStyle::GetStyleSetName(), Pair.Value.GetIconStyleName(EPlatformIconSize::Normal)),
+						Action,
+						NAME_None,
+						EUserInterfaceActionType::Button
+					);
+				}
+				else
+				{
+					MenuBuilder.AddSubMenu(
+						MakeSdkStatusAttribute(PlatformName, DeviceProxies[0]),
+						Tooltip,
+						FNewMenuDelegate::CreateLambda([TargetDeviceServicesModule, PlatformName, IconStyle = Pair.Value.GetIconStyleName(EPlatformIconSize::Normal)](FMenuBuilder& SubMenuBuilder) 
+						{
+							// re-get the proxies, just in case they changed
+							TArray<TSharedPtr<ITargetDeviceProxy>> DeviceProxies;
+							TargetDeviceServicesModule->GetDeviceProxyManager()->GetAllProxies(PlatformName, DeviceProxies);
+							// for each one, put an entry (even the one that was in the outer menu, for less confusion)
+							for (const TSharedPtr<ITargetDeviceProxy> Proxy : DeviceProxies)
+							{
+								FUIAction SubAction;
+								FText SubTooltip;
+								GenerateDeviceProxyMenuParams(Proxy, PlatformName, SubAction, SubTooltip);
+								SubMenuBuilder.AddMenuEntry(
+									MakeSdkStatusAttribute(PlatformName, Proxy),
+									SubTooltip,
+									FSlateIcon(FEditorStyle::GetStyleSetName(), IconStyle),
+									SubAction,
+									NAME_None,
+									EUserInterfaceActionType::Button
+								);
+							}
+						}),
+						Action,
+						NAME_None,
+						EUserInterfaceActionType::Check,
+						false,
+						FSlateIcon(FEditorStyle::GetStyleSetName(), Pair.Value.GetIconStyleName(EPlatformIconSize::Normal)),
+						true
+						);
+				}
+
+				// gather any unknown status devices to query at the end
+				for (const TSharedPtr<ITargetDeviceProxy> Proxy : DeviceProxies)
+				{
+					FString DeviceId = Proxy->GetTargetDeviceId(NAME_None);
+					if (Pair.Value.GetStatusForDeviceId(DeviceId) == DDPIPlatformSdkStatus::Unknown)
+					{
+						DeviceIdsToQuery.Add(DeviceId);
+					}
+				}
+			}
+		}
+	}
+	MenuBuilder.EndSection();
+
+
+	MenuBuilder.BeginSection("CookerSettings");
+
+	MenuBuilder.AddMenuEntry(
+		LOCTEXT("CookOnTheFlyOnLaunch", "Enable cooking on the fly"),
+		LOCTEXT("CookOnTheFlyOnLaunchDescription", "Cook on the fly instead of cooking upfront when launching"),
+		FSlateIcon(),
+		FUIAction(
+			FExecuteAction::CreateStatic(&FTurnkeySupportCallbacks::SetCookOnTheFly),
+			FCanExecuteAction::CreateStatic(&FTurnkeySupportCallbacks::CanSetCookOnTheFly),
+			FIsActionChecked::CreateStatic(&FTurnkeySupportCallbacks::SetCookOnTheFlyIsChecked)
+		),
+		NAME_None,
+		EUserInterfaceActionType::ToggleButton
+	);
+
+	{
+		MenuBuilder.AddWidget(
+			SNew(STextBlock)
+			.ColorAndOpacity(FSlateColor::UseSubduedForeground())
+			.Text(LOCTEXT("ZoomToFitHorizontal", "Launching a game on a different device will change your default 'Launch' device in the toolbar"))
+			.WrapTextAt(300),
+			FText::GetEmpty()
+		);
+	}
+	MenuBuilder.EndSection();
+
+	MenuBuilder.BeginSection("AllPlatforms", LOCTEXT("TurnkeyMenu_ManagePlatforms", "Content/Sdk/Device Management"));
+	TMap<FName, const FDataDrivenPlatformInfo*> UncompiledPlatforms;
+	for (const auto Pair : FDataDrivenPlatformInfoRegistry::GetAllPlatformInfos())
+	{
+		if (Pair.Value.bIsFakePlatform || Pair.Value.bEnabledForUse == false)
+		{
+			continue;
+		}
+
+		FName PlatformName = Pair.Key;
+		const FDataDrivenPlatformInfo& Info = Pair.Value;
+
+		if (!FDataDrivenPlatformInfoRegistry::HasCompiledSupportForPlatform(PlatformName, FDataDrivenPlatformInfoRegistry::EPlatformNameType::Ini))
+		{
+			UncompiledPlatforms.Add(PlatformName, &Info);
+			continue;
+		}
+
+		MenuBuilder.AddSubMenu(
+			MakeSdkStatusAttribute(PlatformName, nullptr),
+			FText::FromString(PlatformName.ToString()),
+			FNewMenuDelegate::CreateStatic(&MakeTurnkeyPlatformMenu, PlatformName, TargetDeviceServicesModule),
+			false,
+			FSlateIcon(FEditorStyle::GetStyleSetName(), Pair.Value.GetIconStyleName(EPlatformIconSize::Normal)),
+			true
+		);
+	}
+
+	if (UncompiledPlatforms.Num() != 0)
+	{
+		MenuBuilder.AddSubMenu(
+			LOCTEXT("Turnkey_UncompiledPlatforms", "Platforms With No Compiled Support"),
+			LOCTEXT("Turnkey_UncompiledPlatformsToolTip", "List of platforms that you have access to, but support is not compiled in to the editor. It may be caused by missing an SDK, so you attempt to install an SDK here."),
+			FNewMenuDelegate::CreateLambda([UncompiledPlatforms, TargetDeviceServicesModule](FMenuBuilder& SubMenuBuilder)
+			{
+				for (const auto It : UncompiledPlatforms)
+				{
+					SubMenuBuilder.AddSubMenu(
+						MakeSdkStatusAttribute(It.Key, nullptr),
+						FText::FromString(It.Key.ToString()),
+						FNewMenuDelegate::CreateStatic(&MakeTurnkeyPlatformMenu, It.Key, TargetDeviceServicesModule),
+						false,
+						FSlateIcon(FEditorStyle::GetStyleSetName(), It.Value->GetIconStyleName(EPlatformIconSize::Normal)),
+						true
+					);
+				}
+			})
+		);
+	}
+
+	MenuBuilder.EndSection();
+
+
+	// options section
+	MenuBuilder.BeginSection("TurnkeyOptions", LOCTEXT("TurnkeySection_Options", "Options and Settings"));
+	{
+		MenuBuilder.AddMenuEntry(
+			LOCTEXT("OpenProjectLauncher", "Project Launcher..."),
+			LOCTEXT("OpenProjectLauncher_ToolTip", "Open the Project Launcher for advanced packaging, deploying and launching of your projects"),
+			FSlateIcon(FEditorStyle::GetStyleSetName(), "Launcher.TabIcon"),
+			FUIAction(FExecuteAction::CreateStatic(&FTurnkeySupportCallbacks::OpenProjectLauncher))
+		);
+
+		MenuBuilder.AddMenuEntry(
+			LOCTEXT("OpenDeviceManager", "Device Manager..."),
+			LOCTEXT("OpenDeviceManager_ToolTip", "View and manage connected devices."),
+			FSlateIcon(FEditorStyle::GetStyleSetName(), "DeviceDetails.TabIcon"),
+			FUIAction(FExecuteAction::CreateStatic(&FTurnkeySupportCallbacks::OpenDeviceManager))
+		);
+
+		MenuBuilder.AddSeparator();
+
+		MenuBuilder.AddMenuEntry(
+			LOCTEXT("PackagingSettings", "Packaging Settings..."),
+			LOCTEXT("PackagingSettings_ToolTip", "Opens the settings for project packaging"),
+			FSlateIcon(),
+			FUIAction(FExecuteAction::CreateStatic(&FTurnkeySupportCallbacks::PackagingSettings))
+		);
+
+// 		MenuBuilder.AddMenuEntry(Commands.ZipUpProject,
+// 			NAME_None,
+// 			TAttribute<FText>(),
+// 			TAttribute<FText>(),
+// 			FSlateIcon()
+// 		);
+// 		MenuBuilder.AddMenuEntry(Commands.PackagingSettings,
+// 			NAME_None,
+// 			TAttribute<FText>(),
+// 			TAttribute<FText>(),
+// 			FSlateIcon()
+// 		);
+
+		ProjectTargetPlatformEditorModule.AddOpenProjectTargetPlatformEditorMenuItem(MenuBuilder);
+	}
+	MenuBuilder.EndSection();
+
+	// now kick-off any devices that need to be updated
+	if (DeviceIdsToQuery.Num() > 0)
+	{
+		FDataDrivenPlatformInfoRegistry::UpdateDeviceSdkStatus(DeviceIdsToQuery);
+	}
+
+	return MenuBuilder.MakeWidget();
+}
+
+/* IModuleInterface implementation
+ *****************************************************************************/
+
+void FTurnkeySupportModule::StartupModule( )
+{
+	FTurnkeySupportCommands::Register();
+}
+
+
+void FTurnkeySupportModule::ShutdownModule( )
+{
+}
+
+
+
+#undef LOCTEXT_NAMESPACE
+
+
+
+
+
+
+
+
+
+
+#if 0
+
+// PACKAGING STUFF THAT HASN'T COME OVER
+
+
+GUnrealEd->CancelPlayingViaLauncher();
+SaveAll();
+
+// does the project have any code?
+FGameProjectGenerationModule& GameProjectModule = FModuleManager::LoadModuleChecked<FGameProjectGenerationModule>(TEXT("GameProjectGeneration"));
+bool bProjectHasCode = GameProjectModule.Get().ProjectHasCodeFiles();
+
+const PlatformInfo::FTargetPlatformInfo* const PlatformInfo = PlatformInfo::FindPlatformInfo(InPlatformInfoName);
+check(PlatformInfo);
+
+if (FInstalledPlatformInfo::Get().IsPlatformMissingRequiredFile(PlatformInfo->DataDrivenPlatformInfo->UBTPlatformString))
+{
+	if (!FInstalledPlatformInfo::OpenInstallerOptions())
+	{
+		FMessageDialog::Open(EAppMsgType::Ok, LOCTEXT("MissingPlatformFilesPackage", "Missing required files to package this platform."));
+	}
+	return;
+}
+
+if (UGameMapsSettings::GetGameDefaultMap().IsEmpty())
+{
+	FMessageDialog::Open(EAppMsgType::Ok, LOCTEXT("MissingGameDefaultMap", "No Game Default Map specified in Project Settings > Maps & Modes."));
+	return;
+}
+
+if (/*PlatformInfo->DataDrivenPlatformInfo->GetSdkStatus() != DDPIPlatformSdkStatus::Valid ||*/ (bProjectHasCode && PlatformInfo->DataDrivenPlatformInfo->bUsesHostCompiler && !FSourceCodeNavigation::IsCompilerAvailable()))
+{
+	if (!ShowBadSDKDialog(PlatformInfo->IniPlatformName))
+	{
+		// 			IMainFrameModule& MainFrameModule = FModuleManager::GetModuleChecked<IMainFrameModule>(TEXT("MainFrame"));
+		// 			MainFrameModule.BroadcastMainFrameSDKNotInstalled(PlatformInfo->Name.ToString(), PlatformInfo->DataDrivenPlatformInfo->SDKTutorial);
+		TArray<FAnalyticsEventAttribute> ParamArray;
+		ParamArray.Add(FAnalyticsEventAttribute(TEXT("Time"), 0.0));
+		FEditorAnalytics::ReportEvent(TEXT("Editor.Package.Failed"), PlatformInfo->Name.ToString(), bProjectHasCode, EAnalyticsErrorCodes::SDKNotFound, ParamArray);
+		return;
+	}
+}
+
+UProjectPackagingSettings* PackagingSettings = Cast<UProjectPackagingSettings>(UProjectPackagingSettings::StaticClass()->GetDefaultObject());
+const UProjectPackagingSettings::FConfigurationInfo& ConfigurationInfo = UProjectPackagingSettings::ConfigurationInfo[(int)PackagingSettings->BuildConfiguration];
+bool bAssetNativizationEnabled = (PackagingSettings->BlueprintNativizationMethod != EProjectPackagingBlueprintNativizationMethod::Disabled);
+
+const ITargetPlatform* const Platform = GetTargetPlatformManager()->FindTargetPlatform(PlatformInfo->Name.ToString());
+{
+	if (Platform)
+	{
+		FString NotInstalledTutorialLink;
+		FString DocumentationLink;
+		FText CustomizedLogMessage;
+
+		int32 Result = Platform->CheckRequirements(bProjectHasCode, ConfigurationInfo.Configuration, bAssetNativizationEnabled, NotInstalledTutorialLink, DocumentationLink, CustomizedLogMessage);
+
+		// report to analytics
+		FEditorAnalytics::ReportBuildRequirementsFailure(TEXT("Editor.Package.Failed"), PlatformInfo->Name.ToString(), bProjectHasCode, Result);
+
+		// report to main frame
+		bool UnrecoverableError = false;
+
+		// report to message log
+		if ((Result & ETargetPlatformReadyStatus::SDKNotFound) != 0)
+		{
+			AddMessageLog(
+				LOCTEXT("SdkNotFoundMessage", "Software Development Kit (SDK) not found."),
+				CustomizedLogMessage.IsEmpty() ? FText::Format(LOCTEXT("SdkNotFoundMessageDetail", "Please install the SDK for the {0} target platform!"), Platform->DisplayName()) : CustomizedLogMessage,
+				NotInstalledTutorialLink,
+				DocumentationLink
+			);
+			UnrecoverableError = true;
+		}
+
+		if ((Result & ETargetPlatformReadyStatus::LicenseNotAccepted) != 0)
+		{
+			AddMessageLog(
+				LOCTEXT("LicenseNotAcceptedMessage", "License not accepted."),
+				CustomizedLogMessage.IsEmpty() ? LOCTEXT("LicenseNotAcceptedMessageDetail", "License must be accepted in project settings to deploy your app to the device.") : CustomizedLogMessage,
+				NotInstalledTutorialLink,
+				DocumentationLink
+			);
+
+			UnrecoverableError = true;
+		}
+
+		if ((Result & ETargetPlatformReadyStatus::ProvisionNotFound) != 0)
+		{
+			AddMessageLog(
+				LOCTEXT("ProvisionNotFoundMessage", "Provision not found."),
+				CustomizedLogMessage.IsEmpty() ? LOCTEXT("ProvisionNotFoundMessageDetail", "A provision is required for deploying your app to the device.") : CustomizedLogMessage,
+				NotInstalledTutorialLink,
+				DocumentationLink
+			);
+			UnrecoverableError = true;
+		}
+
+		if ((Result & ETargetPlatformReadyStatus::SigningKeyNotFound) != 0)
+		{
+			AddMessageLog(
+				LOCTEXT("SigningKeyNotFoundMessage", "Signing key not found."),
+				CustomizedLogMessage.IsEmpty() ? LOCTEXT("SigningKeyNotFoundMessageDetail", "The app could not be digitally signed, because the signing key is not configured.") : CustomizedLogMessage,
+				NotInstalledTutorialLink,
+				DocumentationLink
+			);
+			UnrecoverableError = true;
+		}
+
+		if ((Result & ETargetPlatformReadyStatus::ManifestNotFound) != 0)
+		{
+			AddMessageLog(
+				LOCTEXT("ManifestNotFound", "Manifest not found."),
+				CustomizedLogMessage.IsEmpty() ? LOCTEXT("ManifestNotFoundMessageDetail", "The generated application manifest could not be found.") : CustomizedLogMessage,
+				NotInstalledTutorialLink,
+				DocumentationLink
+			);
+			UnrecoverableError = true;
+		}
+
+		if ((Result & ETargetPlatformReadyStatus::RemoveServerNameEmpty) != 0
+			&& (bProjectHasCode || (Result & ETargetPlatformReadyStatus::CodeBuildRequired)
+				|| (!FApp::GetEngineIsPromotedBuild() && !FApp::IsEngineInstalled())))
+		{
+			AddMessageLog(
+				LOCTEXT("RemoveServerNameNotFound", "Remote compiling requires a server name. "),
+				CustomizedLogMessage.IsEmpty() ? LOCTEXT("RemoveServerNameNotFoundDetail", "Please specify one in the Remote Server Name settings field.") : CustomizedLogMessage,
+				NotInstalledTutorialLink,
+				DocumentationLink
+			);
+			UnrecoverableError = true;
+		}
+
+		if ((Result & ETargetPlatformReadyStatus::CodeUnsupported) != 0)
+		{
+			FMessageDialog::Open(EAppMsgType::Ok, LOCTEXT("NotSupported_SelectedPlatform", "Sorry, packaging a code-based project for the selected platform is currently not supported. This feature may be available in a future release."));
+			UnrecoverableError = true;
+		}
+		else if ((Result & ETargetPlatformReadyStatus::PluginsUnsupported) != 0)
+		{
+			FMessageDialog::Open(EAppMsgType::Ok, LOCTEXT("NotSupported_ThirdPartyPlugins", "Sorry, packaging a project with third-party plugins is currently not supported for the selected platform. This feature may be available in a future release."));
+			UnrecoverableError = true;
+		}
+
+		if (UnrecoverableError)
+		{
+			return;
+		}
+	}
+}
+
+if (!FModuleManager::LoadModuleChecked<IProjectTargetPlatformEditorModule>("ProjectTargetPlatformEditor").ShowUnsupportedTargetWarning(PlatformInfo->VanillaInfo->Name))
+{
+	return;
+}
+
+
+#endif
+
+
+
+
+
+
+
+
+
+
+
+
+
+
+
+
+
+
+
+IMPLEMENT_MODULE(FTurnkeySupportModule, TurnkeySupport);
diff --git a/Engine/Source/Editor/TurnkeySupport/Private/TurnkeySupportModule.h b/Engine/Source/Editor/TurnkeySupport/Private/TurnkeySupportModule.h
new file mode 100644
--- /dev/null
+++ b/Engine/Source/Editor/TurnkeySupport/Private/TurnkeySupportModule.h
@@ -0,0 +1,40 @@
+// Copyright Epic Games, Inc. All Rights Reserved.
+
+#pragma once
+
+#include "CoreMinimal.h"
+#include "Widgets/SWidget.h"
+#include "Interfaces/ITurnkeySupportModule.h"
+
+
+DECLARE_LOG_CATEGORY_EXTERN(LogTurnkeySupport, Log, All);
+
+/**
+ * Editor main frame module
+ */
+class FTurnkeySupportModule	: public ITurnkeySupportModule
+{
+public:
+
+	/**
+	 *
+	 * @return	The newly-created menu widget
+	 */
+	virtual TSharedRef<SWidget> MakeTurnkeyMenu() const override;
+
+public:
+
+	// IModuleInterface interface
+
+	virtual void StartupModule( ) override;
+	virtual void ShutdownModule( ) override;
+
+	virtual bool SupportsDynamicReloading( ) override
+	{
+		return true; // @todo: Eventually, this should probably not be allowed.
+	}
+
+
+private:
+
+};
diff --git a/Engine/Source/Editor/TurnkeySupport/Public/Interfaces/ITurnkeySupportModule.h b/Engine/Source/Editor/TurnkeySupport/Public/Interfaces/ITurnkeySupportModule.h
new file mode 100644
--- /dev/null
+++ b/Engine/Source/Editor/TurnkeySupport/Public/Interfaces/ITurnkeySupportModule.h
@@ -0,0 +1,48 @@
+// Copyright Epic Games, Inc. All Rights Reserved.
+
+#pragma once
+
+#include "CoreMinimal.h"
+#include "SlateFwd.h"
+#include "Modules/ModuleInterface.h"
+#include "Modules/ModuleManager.h"
+#include "Widgets/SWidget.h"
+#include "Framework/Commands/UICommandList.h"
+#include "Framework/MultiBox/MultiBoxExtender.h"
+
+
+/**
+ * Interface for turnkey support module
+ */
+class ITurnkeySupportModule
+	: public IModuleInterface
+{
+public:
+
+	/**
+	 *
+	 * @return	The newly-created menu widget
+	 */
+	virtual TSharedRef<SWidget> MakeTurnkeyMenu() const = 0;
+
+public:
+
+	/**
+	 * Gets a reference to the search module instance.
+	 *
+	 * @todo gmp: better implementation using dependency injection.
+	 * @return A reference to the MainFrame module.
+	 */
+	static ITurnkeySupportModule& Get( )
+	{
+		static const FName TurnkeySupportModuleName = "TurnkeySupport";
+		return FModuleManager::LoadModuleChecked<ITurnkeySupportModule>(TurnkeySupportModuleName);
+	}
+
+public:
+
+	/**
+	 * Virtual destructor.
+	 */
+	virtual ~ITurnkeySupportModule( ) { }
+};
diff --git a/Engine/Source/Editor/TurnkeySupport/Public/TurnkeySupport.h b/Engine/Source/Editor/TurnkeySupport/Public/TurnkeySupport.h
new file mode 100644
--- /dev/null
+++ b/Engine/Source/Editor/TurnkeySupport/Public/TurnkeySupport.h
@@ -0,0 +1,23 @@
+// Copyright Epic Games, Inc. All Rights Reserved.
+
+#pragma once
+
+
+/* Boilerplate
+ *****************************************************************************/
+
+#include "Misc/MonolithicHeaderBoilerplate.h"
+MONOLITHIC_HEADER_BOILERPLATE()
+
+
+/* Public Dependencies
+ *****************************************************************************/
+
+#include "Modules/ModuleInterface.h"
+#include "SlateBasics.h"
+
+
+/* Public Includes
+ *****************************************************************************/
+
+#include "Interfaces/ITurnkeySupportModule.h"
diff --git a/Engine/Source/Editor/TurnkeySupport/TurnkeySupport.Build.cs b/Engine/Source/Editor/TurnkeySupport/TurnkeySupport.Build.cs
new file mode 100644
--- /dev/null
+++ b/Engine/Source/Editor/TurnkeySupport/TurnkeySupport.Build.cs
@@ -0,0 +1,64 @@
+// Copyright Epic Games, Inc. All Rights Reserved.
+
+using UnrealBuildTool;
+
+public class TurnkeySupport : ModuleRules
+{
+	public TurnkeySupport(ReadOnlyTargetRules Target) : base(Target)
+	{
+
+		PrivateDependencyModuleNames.AddRange(
+			new string[] {
+				"Core",
+				"CoreUObject",
+				"ApplicationCore",
+				"Engine",
+				"EngineSettings",
+				"InputCore",
+				"RHI",
+				"RenderCore",
+				"Slate",
+				"SlateCore",
+				"EditorStyle",
+				"SourceControl",
+				"SourceControlWindows",
+				"TargetPlatform",
+				"DesktopPlatform",
+				"EditorFramework",
+				"UnrealEd",
+				"WorkspaceMenuStructure",
+				"MessageLog",
+				"UATHelper",
+				"TranslationEditor",
+				"Projects",
+				"DeviceProfileEditor",
+				"UndoHistory",
+				"Analytics",
+				"ToolMenus",
+				"LauncherServices",
+				
+				"SettingsEditor"
+			}
+		);
+
+		PrivateIncludePathModuleNames.AddRange(
+			new string[] {
+				"AssetTools",
+				"DesktopPlatform",
+				"GameProjectGeneration",
+				"ProjectTargetPlatformEditor",
+				"LevelEditor",
+				"Settings",
+				"SourceCodeAccess",
+				"Toolbox",
+				"LocalizationDashboard",
+			}
+		);
+
+		PrivateIncludePaths.AddRange(
+			new string[] {
+				"Editor/TurnkeySupport/Private",
+			}
+		);
+	}
+}
diff --git a/Engine/Source/Editor/UnrealEd/Classes/Settings/ProjectPackagingSettings.h b/Engine/Source/Editor/UnrealEd/Classes/Settings/ProjectPackagingSettings.h
--- a/Engine/Source/Editor/UnrealEd/Classes/Settings/ProjectPackagingSettings.h
+++ b/Engine/Source/Editor/UnrealEd/Classes/Settings/ProjectPackagingSettings.h
@@ -14,7 +14,7 @@ struct FTargetInfo;
  * Enumerates the available build configurations for project packaging.
  */
 UENUM()
-enum EProjectPackagingBuildConfigurations
+enum class EProjectPackagingBuildConfigurations : uint8
 {
 	/** Debug configuration. */
 	PPBC_Debug UMETA(DisplayName="Debug"),
@@ -115,7 +115,7 @@ public:
 	/**
 	 * Static array of information about each configuration
 	 */
-	static const FConfigurationInfo ConfigurationInfo[PPBC_MAX];
+	static const FConfigurationInfo ConfigurationInfo[(int)EProjectPackagingBuildConfigurations::PPBC_MAX];
 
 	/** Specifies whether to build the game executable during packaging. */
 	UPROPERTY(config, EditAnywhere, Category=Project)
@@ -123,7 +123,7 @@ public:
 
 	/** The build configuration for which the project is packaged. */
 	UPROPERTY(config, EditAnywhere, Category=Project)
-	TEnumAsByte<EProjectPackagingBuildConfigurations> BuildConfiguration;
+	EProjectPackagingBuildConfigurations BuildConfiguration;
 
 	/** Name of the target to build */
 	UPROPERTY(config, EditAnywhere, Category=Project)
@@ -443,8 +443,27 @@ public:
 	UPROPERTY(config, EditAnywhere, Category=Packaging, AdvancedDisplay, meta=(DisplayName="Additional Non-Asset Directories To Copy for dedicated server only", RelativeToGameContentDir))
 	TArray<FDirectoryPath> DirectoriesToAlwaysStageAsNonUFSServer;	
 
+
+
+	/**
+	 * Get and set the per-platform build config and targetplatform settings for the Turnkey/Launch on menu
+	 */
+	EProjectPackagingBuildConfigurations GetBuildConfigurationForPlatform(FName PlatformName) const;
+	void SetBuildConfigurationForPlatform(FName PlatformName, EProjectPackagingBuildConfigurations Configuration);
+
+	FName GetTargetPlatformForPlatform(FName PlatformName) const;
+	void SetTargetPlatformForPlatform(FName PlatformName, FName TargetPlatformName);
+
 private:
-	/** Helper array used to mirror Blueprint asset selections across edits */
+	/** Per platform build configuration */
+	UPROPERTY(config)
+	TMap<FName, EProjectPackagingBuildConfigurations> PerPlatformBuildConfig;
+
+	/** Per platform targetplatform cooking target */
+	UPROPERTY(config)
+	TMap<FName, FName> PerPlatformTargetPlatformName;
+
+		/** Helper array used to mirror Blueprint asset selections across edits */
 	TArray<FFilePath> CachedNativizeBlueprintAssets;
 	
 public:
diff --git a/Engine/Source/Editor/UnrealEd/Private/Kismet2/DebuggerCommands.cpp b/Engine/Source/Editor/UnrealEd/Private/Kismet2/DebuggerCommands.cpp
--- a/Engine/Source/Editor/UnrealEd/Private/Kismet2/DebuggerCommands.cpp
+++ b/Engine/Source/Editor/UnrealEd/Private/Kismet2/DebuggerCommands.cpp
@@ -79,10 +79,14 @@
 #include "ISettingsEditorModule.h"
 #include "Async/Async.h"
 #include "Misc/FileHelper.h"
+#include "Interfaces/ITurnkeySupportModule.h"
+#include "Settings/ProjectPackagingSettings.h"
 
 
 #define LOCTEXT_NAMESPACE "DebuggerCommands"
 
+DEFINE_LOG_CATEGORY_STATIC(LogDebuggerCommands, Log, All);
+
 void SGlobalPlayWorldActions::Construct(const FArguments& InArgs)
 {
 	// Always keep track of the current active play world actions widget so we later set user focus on it
@@ -165,8 +169,6 @@ public:
 	static bool RepeatLastLaunch_CanExecute();
 	static FText GetRepeatLastLaunchToolTip();
 	static FSlateIcon GetRepeatLastLaunchIcon();
-	static void OpenProjectLauncher_Clicked();
-	static void OpenDeviceManager_Clicked();
 
 	static FSlateIcon GetResumePlaySessionImage();
 	static FText GetResumePlaySessionToolTip();
@@ -362,8 +364,6 @@ void FPlayWorldCommands::RegisterCommands()
 
 	// Launch
 	UI_COMMAND(RepeatLastLaunch, "Launch", "Launches the game on the device as the last session launched from the dropdown next to the Play on Device button on the level editor toolbar", EUserInterfaceActionType::Button, FInputChord(EKeys::P, EModifierKey::Alt | EModifierKey::Shift))
-		UI_COMMAND(OpenProjectLauncher, "Project Launcher...", "Open the Project Launcher for advanced packaging, deploying and launching of your projects", EUserInterfaceActionType::Button, FInputChord());
-	UI_COMMAND(OpenDeviceManager, "Device Manager...", "View and manage connected devices.", EUserInterfaceActionType::Button, FInputChord());
 
 	// PIE mobile preview devices.
 	AddPIEPreviewDeviceCommands();
@@ -483,13 +483,6 @@ void FPlayWorldCommands::BindGlobalPlayWorldCommands()
 	);
 
 	// Launch
-	ActionList.MapAction(Commands.OpenProjectLauncher,
-		FExecuteAction::CreateStatic(&FInternalPlayWorldCommandCallbacks::OpenProjectLauncher_Clicked)
-	);
-
-	ActionList.MapAction(Commands.OpenDeviceManager,
-		FExecuteAction::CreateStatic(&FInternalPlayWorldCommandCallbacks::OpenDeviceManager_Clicked)
-	);
 
 	ActionList.MapAction(Commands.RepeatLastLaunch,
 		FExecuteAction::CreateStatic(&FInternalPlayWorldCommandCallbacks::RepeatLastLaunch_Clicked),
@@ -663,30 +656,15 @@ void FPlayWorldCommands::BuildToolbar(FToolMenuSection& InSection, bool bInclude
 				FUIAction LaunchMenuAction;
 				LaunchMenuAction.IsActionVisibleDelegate = FIsActionButtonVisible::CreateStatic(&FInternalPlayWorldCommandCallbacks::CanShowNonPlayWorldOnlyActions);
 
-				if (FParse::Param(FCommandLine::Get(), TEXT("OldLaunchOn")))
-				{
-					InDynamicSection.AddEntry(FToolMenuEntry::InitComboButton(
-						"LaunchCombo",
-						LaunchMenuAction,
-						FOnGetContent::CreateStatic(&GenerateLaunchMenuContent, GlobalPlayWorldActions.ToSharedRef()),
-						LOCTEXT("LaunchCombo_Label", "Launch Options"),
-						LOCTEXT("PODComboToolTip", "Options for launching on a device"),
-						FSlateIcon(FEditorStyle::GetStyleSetName(), "LevelEditor.RepeatLastLaunch"),
-						true
-					));
-				}
-				else
-				{
-					InDynamicSection.AddEntry(FToolMenuEntry::InitComboButton(
-						"TurnkeyCombo",
-						LaunchMenuAction,
-						FOnGetContent::CreateStatic(&GenerateTurnkeyMenuContent, GlobalPlayWorldActions.ToSharedRef()),
-						LOCTEXT("LaunchCombo_Label", "Launch Options"),
-						LOCTEXT("PlayWithTurnkeyComboToolTip", "Options for launching on a device and manage platforms"),
-						FSlateIcon(FEditorStyle::GetStyleSetName(), "LevelEditor.RepeatLastLaunch"),
-						true
-					));
-				}
+				InDynamicSection.AddEntry(FToolMenuEntry::InitComboButton(
+					"TurnkeyCombo",
+					LaunchMenuAction,
+					FOnGetContent::CreateLambda([] { return ITurnkeySupportModule::Get().MakeTurnkeyMenu(); }),
+					LOCTEXT("LaunchCombo_Label", "Launch Options"),
+					LOCTEXT("PlayWithTurnkeyComboToolTip", "Options for launching on a device and manage platforms"),
+					FSlateIcon(FEditorStyle::GetStyleSetName(), "LevelEditor.RepeatLastLaunch"),
+					true
+				));
 			}
 		}));
 	}
@@ -997,1126 +975,256 @@ TSharedRef< SWidget > FPlayWorldCommands::GeneratePlayMenuContent(TSharedRef<FUI
 	return UToolMenus::Get()->GenerateWidget(MenuName, MenuContext);
 }
 
-/*
- * Create an All_<platform>_devices_on_<host> submenu
- * can be extended to any othe All <Platform> aggregate proxy
-*/
-static void MakeAllDevicesSubMenu(FMenuBuilder& InMenuBuilder, const PlatformInfo::FTargetPlatformInfo* InPlatformInfo, const TSharedPtr<ITargetDeviceProxy> DeviceProxy)
-{
-	ITargetDeviceServicesModule* TargetDeviceServicesModule = static_cast<ITargetDeviceServicesModule*>(FModuleManager::Get().LoadModule(TEXT("TargetDeviceServices")));
-	IProjectTargetPlatformEditorModule& ProjectTargetPlatformEditorModule = FModuleManager::LoadModuleChecked<IProjectTargetPlatformEditorModule>("ProjectTargetPlatformEditor");
-
-	TArray<FName> PlatformVariants;
-	DeviceProxy->GetVariants(PlatformVariants);
-	for (auto It = PlatformVariants.CreateIterator(); It; ++It)
-	{
-		FName Variant = *It;
-
-		// for an aggregate (All_<platform>_devices_on_<host>) proxy, allow only the "Android_<texture_compression>" variants
-		const PlatformInfo::FTargetPlatformInfo* PlatformInfo = PlatformInfo::FindPlatformInfo(Variant);
-		if (DeviceProxy->IsAggregated() && PlatformInfo != NULL &&
-			(Variant == PlatformInfo->VanillaInfo->Name || PlatformInfo->PlatformType != EBuildTargetType::Game))
-		{
-			continue;
-		}
 
-		FString DeviceListStr;
-		bool bVariantHasDevices = false;
+//////////////////////////////////////////////////////////////////////////
+// FPlayWorldCommandCallbacks
 
-		const TSet<FString>& TargetDeviceIds = DeviceProxy->GetTargetDeviceIds(Variant);
-		for (TSet<FString>::TConstIterator ItDeviceId(TargetDeviceIds); ItDeviceId; ++ItDeviceId)
-		{
-			TSharedPtr<ITargetDeviceProxy> PhysicalDeviceProxy = TargetDeviceServicesModule->GetDeviceProxyManager()->FindProxyDeviceForTargetDevice(*ItDeviceId);
+void FPlayWorldCommandCallbacks::StartPlayFromHere()
+{
+	// Is a PIE session already running?  If so we close it first
+	if (GUnrealEd->PlayWorld != NULL)
+	{
+		GUnrealEd->EndPlayMap();
+	}
 
-			if (PhysicalDeviceProxy.IsValid())
-			{
-				DeviceListStr.AppendChar('\n');
-				DeviceListStr.Append(*PhysicalDeviceProxy->GetName());
-				bVariantHasDevices = true;
-			}
-		}
+	FRequestPlaySessionParams SessionParams;
 
-		if (!bVariantHasDevices)
-		{
-			continue;
-		}
+	UClass* const PlayerStartClass = GUnrealEd->PlayFromHerePlayerStartClass ? (UClass*)GUnrealEd->PlayFromHerePlayerStartClass : APlayerStart::StaticClass();
 
-		FString PlatformVariantStr = Variant.ToString();
-		FString PlatformId = PlatformVariantStr + TEXT("@") + PlatformVariantStr;
+	// Figure out the start location of the player
+	UCapsuleComponent*	DefaultCollisionComponent = CastChecked<UCapsuleComponent>(PlayerStartClass->GetDefaultObject<AActor>()->GetRootComponent());
+	FVector	CollisionExtent = FVector(DefaultCollisionComponent->GetScaledCapsuleRadius(), DefaultCollisionComponent->GetScaledCapsuleRadius(), DefaultCollisionComponent->GetScaledCapsuleHalfHeight());
+	SessionParams.StartLocation = GEditor->UnsnappedClickLocation + GEditor->ClickPlane * (FVector::BoxPushOut(GEditor->ClickPlane, CollisionExtent) + 0.1f);
 
-		// create an action
-		FUIAction LaunchDeviceAction(
-			FExecuteAction::CreateStatic(&FInternalPlayWorldCommandCallbacks::HandleLaunchOnDeviceActionExecute, PlatformId, PlatformVariantStr, false),
-			FCanExecuteAction::CreateStatic(&FInternalPlayWorldCommandCallbacks::HandleLaunchOnDeviceActionCanExecute, PlatformVariantStr),
-			FIsActionChecked::CreateStatic(&FInternalPlayWorldCommandCallbacks::HandleLaunchOnDeviceActionIsChecked, PlatformVariantStr)
-		);
+	FLevelEditorModule& LevelEditorModule = FModuleManager::GetModuleChecked<FLevelEditorModule>("LevelEditor");
 
-		// generate display label
-		FText Label = FText::FromString(PlatformVariantStr);
+	TSharedPtr<IAssetViewport> ActiveLevelViewport = LevelEditorModule.GetFirstActiveViewport();
 
-		// generate tooltip text with the devices' list
-		FFormatNamedArguments TooltipArguments;
-		TooltipArguments.Add(TEXT("DeviceList"), FText::FromString(DeviceListStr));
-		FText Tooltip = FText::Format(LOCTEXT("LaunchDeviceToolTipText_LaunchOn", "Launch the game on:\n {DeviceList}"), TooltipArguments);
 
-		// add a submenu entry
-		InMenuBuilder.AddMenuEntry(
-			LaunchDeviceAction,
-			ProjectTargetPlatformEditorModule.MakePlatformMenuItemWidget(*InPlatformInfo, true, Label),
-			NAME_None,
-			Tooltip,
-			EUserInterfaceActionType::Check
-		);
+	if (ActiveLevelViewport.IsValid() && ActiveLevelViewport->GetAssetViewportClient().IsPerspective())
+	{
+		// If there is no level viewport, a new window will be spawned to play in.
+		SessionParams.DestinationSlateViewport = ActiveLevelViewport;
+		SessionParams.StartRotation = ActiveLevelViewport->GetAssetViewportClient().GetViewRotation();
 	}
+
+	GUnrealEd->RequestPlaySession(SessionParams);
 }
 
-static void TurnkeyInstallSdk(FString PlatformName, bool bPreferFull, bool bForceInstall, const TCHAR* DeviceName=nullptr)
-{
-//	FString CommandLine = FString::Printf(TEXT("Turnkey -command=InstallSdk -platform=%s -BestAvailable -AllowAutoSdk -EditorIO -noturnkeyvariables -utf8output -WaitForUATMutex"), *PlatformName);
 
-	FString OptionalOptions;
-	if (bPreferFull)
-	{
-		OptionalOptions += TEXT(" -PreferFull");
-	}
-	if (bForceInstall)
-	{
-		OptionalOptions += DeviceName != nullptr ? TEXT(" -ForceDeviceInstall") : TEXT(" -ForceSdkInstall");
-	}
-	if (DeviceName != nullptr)
+void FPlayWorldCommandCallbacks::ResumePlaySession_Clicked()
+{
+	if (HasPlayWorld())
 	{
-		OptionalOptions += FString::Printf(TEXT(" -Device=%s"), DeviceName);
+		LeaveDebuggingMode();
+		GUnrealEd->PlaySessionResumed();
+		uint32 UserIndex = 0;
+		FSlateApplication::Get().SetUserFocusToGameViewport(UserIndex);
 	}
+}
 
-	FString CommandLine = FString::Printf(TEXT("Turnkey -command=VerifySdk -UpdateIfNeeded -platform=%s %s -EditorIO -noturnkeyvariables -utf8output -WaitForUATMutex"), *PlatformName, *OptionalOptions);
 
-	FText TaskName = LOCTEXT("InstallingSdk", "Installing Sdk");
-	IUATHelperModule::Get().CreateUatTask(CommandLine, FText::FromString(PlatformName), TaskName, TaskName, FEditorStyle::GetBrush(TEXT("MainFrame.PackageProject")),
-		[PlatformName](FString, double)
+void FPlayWorldCommandCallbacks::PausePlaySession_Clicked()
+{
+	if (HasPlayWorld())
 	{
-		AsyncTask(ENamedThreads::GameThread, [PlatformName]()
-		{
-
-			// read in env var changes
-			// @todo turnkey move this and make it mac/linux aware
-			FString TurnkeyEnvVarsFilename = FPaths::Combine(FPaths::EngineIntermediateDir(), TEXT("Turnkey/PostTurnkeyVariables.bat"));
+		GUnrealEd->PlayWorld->bDebugPauseExecution = true;
+		GUnrealEd->PlaySessionPaused();
+		if (IsInPIE()) {
+			FSlateApplication::Get().ClearKeyboardFocus(EFocusCause::SetDirectly);
+			FSlateApplication::Get().ResetToDefaultInputSettings();
 
-			if (IFileManager::Get().FileExists(*TurnkeyEnvVarsFilename))
+			TWeakPtr<SGlobalPlayWorldActions> ActiveGlobalPlayWorldWidget = FPlayWorldCommands::GetActiveGlobalPlayWorldActionsWidget();
+			if (ActiveGlobalPlayWorldWidget.IsValid())
 			{
-				TArray<FString> Contents;
-				if (FFileHelper::LoadFileToStringArray(Contents, *TurnkeyEnvVarsFilename))
-				{
-					for (const FString& Line : Contents)
-					{
-						if (Line.StartsWith(TEXT("set ")))
-						{
-							// split the line
-							FString VariableLine = Line.Mid(4);
-							int Equals;
-							if (VariableLine.FindChar('=', Equals))
-							{
-								// set the key/value
-								FString Key = VariableLine.Mid(0, Equals);
-								FString Value = VariableLine.Mid(Equals + 1);
-
-								FPlatformMisc::SetEnvironmentVar(*Key, *Value);
-
-								UE_LOG(LogTemp, Display, TEXT("Turnkey setting env var: %s = %s"), *Key, *Value);
-							}
-						}
-					}
-				}
+				uint32 UserIndex = 0;
+				FSlateApplication::Get().SetUserFocus(UserIndex, ActiveGlobalPlayWorldWidget.Pin());
 			}
-
-			// update the Sdk status
-//			FDataDrivenPlatformInfoRegistry::UpdateSdkStatus();
-			GetTargetPlatformManager()->UpdateAfterSDKInstall(*PlatformName);
-			RenderUtilsInit();
+		}
+	}
+}
 
 
- 			// show restart dialog
- 			FModuleManager::GetModuleChecked<ISettingsEditorModule>("SettingsEditor").OnApplicationRestartRequired();
-		});
+void FPlayWorldCommandCallbacks::SingleFrameAdvance_Clicked()
+{
+	if (HasPlayWorld())
+	{
+		FInternalPlayWorldCommandCallbacks::SingleFrameAdvance_Clicked();
 	}
-	);
 }
 
-static TAttribute<FText> MakeSdkStatusAttribute(FName IniPlatformName, TSharedPtr< ITargetDeviceProxy> DeviceProxy)
+bool FPlayWorldCommandCallbacks::IsInSIE()
 {
-// 	FFormatNamedArguments LabelArguments;
-// 	LabelArguments.Add(TEXT("DeviceName"), FText::FromString(DeviceProxy->GetName()));
-// 
-// 	if (!DeviceProxy->IsConnected())
-// 	{
-// 		LabelArguments.Add(TEXT("HostUser"), LOCTEXT("DisconnectedHint", " [Disconnected]"));
-// 	}
-// 	else if (DeviceProxy->GetHostUser() != FPlatformProcess::UserName(false))
-// 	{
-// 		LabelArguments.Add(TEXT("HostUser"), FText::FromString(DeviceProxy->GetHostUser()));
-// 	}
-// 	else
-// 	{
-// 		LabelArguments.Add(TEXT("HostUser"), FText::GetEmpty());
-// 	}
-// 
-// 	OutLabel = FText::Format(LOCTEXT("LaunchDeviceLabel", "{DeviceName}{HostUser}"), LabelArguments);
-// 	TAttribute<FText> = Make
+	return GEditor->bIsSimulatingInEditor;
+}
 
 
-	FString DisplayString = DeviceProxy ? DeviceProxy->GetName() : IniPlatformName.ToString();
-	FString DeviceId = DeviceProxy ? DeviceProxy->GetTargetDeviceId(NAME_None) : FString();
+bool FPlayWorldCommandCallbacks::IsInPIE()
+{
+	return (GEditor->PlayWorld != NULL) && (!GEditor->bIsSimulatingInEditor);
+}
 
-	return TAttribute<FText>::Create(TAttribute<FText>::FGetter::CreateLambda([IniPlatformName, DisplayString, DeviceId]()
-	{
-		// get the status, or Unknown if it's not there
-		const FDataDrivenPlatformInfo& Info = FDataDrivenPlatformInfoRegistry::GetPlatformInfo(IniPlatformName);
-		DDPIPlatformSdkStatus Status = DeviceId.Len() ? Info.GetStatusForDeviceId(DeviceId) : Info.GetSdkStatus(false);
 
-		// @todo turnkey: Have premade FText's by SdkStatus for speed
-		const TCHAR* Desc =
-			Status == DDPIPlatformSdkStatus::Querying ? TEXT("Querying...") :
-			Status == DDPIPlatformSdkStatus::Valid ? TEXT("Valid Sdk") :
-			Status == DDPIPlatformSdkStatus::OutOfDate ? TEXT("Outdated Sdk") :
-			Status == DDPIPlatformSdkStatus::NoSdk ? TEXT("No Sdk") :
-			Status == DDPIPlatformSdkStatus::FlashValid ? TEXT("Valid Flash") :
-			Status == DDPIPlatformSdkStatus::FlashOutOfDate ? TEXT("Outdated Flash") :
-			TEXT("???");
-		const TCHAR* CompiledStatus = FDataDrivenPlatformInfoRegistry::HasCompiledSupportForPlatform(IniPlatformName, FDataDrivenPlatformInfoRegistry::EPlatformNameType::Ini) ? TEXT("") : TEXT(" (No Compiled Support)");
-		return FText::FromString(FString::Printf(TEXT("%s (%s)%s"), *DisplayString, Desc, CompiledStatus));
-	}));
+bool FPlayWorldCommandCallbacks::IsInSIE_AndRunning()
+{
+	return IsInSIE() && ((GEditor->PlayWorld == NULL) || !(GEditor->PlayWorld->bDebugPauseExecution));
 }
 
 
-static void MakeTurnkeyPlatformMenu(FMenuBuilder& MenuBuilder, FName IniPlatformName, ITargetDeviceServicesModule* TargetDeviceServicesModule)
+bool FPlayWorldCommandCallbacks::IsInPIE_AndRunning()
 {
-	MenuBuilder.BeginSection("AllDevices", LOCTEXT("TurnkeySection_AllDevices", "All Devices"));
+	return IsInPIE() && ((GEditor->PlayWorld == NULL) || !(GEditor->PlayWorld->bDebugPauseExecution));
+}
 
-	TArray<TSharedPtr<ITargetDeviceProxy>> DeviceProxies;
-	TargetDeviceServicesModule->GetDeviceProxyManager()->GetAllProxies(IniPlatformName, DeviceProxies);
 
-	const FDataDrivenPlatformInfo& Info = FDataDrivenPlatformInfoRegistry::GetPlatformInfo(IniPlatformName);
+bool FPlayWorldCommandCallbacks::HasPlayWorld()
+{
+	return GEditor->PlayWorld != NULL;
+}
 
-	FString UBTPlatformName = Info.UBTPlatformString;
-	for (const TSharedPtr<ITargetDeviceProxy> Proxy : DeviceProxies)
-	{
-		FString DeviceName = Proxy->GetName();
-		MenuBuilder.AddSubMenu(
-			MakeSdkStatusAttribute(IniPlatformName, Proxy),
-			FText(),
-			FNewMenuDelegate::CreateLambda([UBTPlatformName, DeviceName](FMenuBuilder& SubMenuBuilder) {
-				SubMenuBuilder.AddMenuEntry(
-					LOCTEXT("Turnkey_RepairDevice", "Repair Device as Needed"),
-					LOCTEXT("TurnkeyTooltip_RepairDevice", "Perform any fixup that may be needed on this device. If up to date already, nothing will be done."),
-					FSlateIcon(),
-					FExecuteAction::CreateStatic(TurnkeyInstallSdk, UBTPlatformName, false, false, *DeviceName)
-				);
 
-				SubMenuBuilder.AddMenuEntry(
-					LOCTEXT("Turnkey_ForceRepairDevice", "Force Repair Device"),
-					LOCTEXT("TurnkeyTooltip_ForceRepairDevice", "Force repairing anything on the device needed (update firmware, etc). Will perform all steps possible, even if not needed."),
-					FSlateIcon(),
-					FExecuteAction::CreateStatic(TurnkeyInstallSdk, UBTPlatformName, true, false, *DeviceName)
-				);
-			})
-		);
-	}
+bool FPlayWorldCommandCallbacks::HasPlayWorldAndPaused()
+{
+	return HasPlayWorld() && GUnrealEd->PlayWorld->bDebugPauseExecution;
+}
 
-	MenuBuilder.EndSection();
 
-	MenuBuilder.BeginSection("AllDevices", LOCTEXT("TurnkeySection_Sdks", "Sdk Managment"));
+bool FPlayWorldCommandCallbacks::HasPlayWorldAndRunning()
+{
+	return HasPlayWorld() && !GUnrealEd->PlayWorld->bDebugPauseExecution;
+}
 
-	const TCHAR* NoDevice = nullptr;
-	if (Info.GetSdkStatus() == DDPIPlatformSdkStatus::OutOfDate)
-	{
-		MenuBuilder.AddMenuEntry(
-			LOCTEXT("Turnkey_InstallSdkMinimal", "Update Sdk (Prefer Minimal)"),
-			LOCTEXT("TurnkeyTooltip_InstallSdkMinimal", "Attempt to update an Sdk, as hosted by your studio. Will attempt to install a minimal Sdk (useful for building/running only)"),
-			FSlateIcon(),
-			FExecuteAction::CreateStatic(TurnkeyInstallSdk, UBTPlatformName, false, false, NoDevice)
-		);
 
-		MenuBuilder.AddMenuEntry(
-			LOCTEXT("Turnkey_InstallSdkFull", "Update Sdk (Prefer Full)"),
-			LOCTEXT("TurnkeyTooltip_InstallSdkMinimal", "Attempt to update an Sdk, as hosted by your studio. Will attempt to install a full Sdk (useful profiling or other use cases)"),
-			FSlateIcon(),
-			FExecuteAction::CreateStatic(TurnkeyInstallSdk, UBTPlatformName, true, false, NoDevice)
-		);
+//////////////////////////////////////////////////////////////////////////
+// FInternalPlayWorldCommandCallbacks
+
+FText FInternalPlayWorldCommandCallbacks::GetPossessEjectLabel()
+{
+	if (IsInPIE())
+	{
+		return LOCTEXT("EjectLabel", "Eject");
 	}
-	else if (Info.GetSdkStatus() == DDPIPlatformSdkStatus::Valid)
+	else if (IsInSIE())
 	{
-		MenuBuilder.AddMenuEntry(
-			LOCTEXT("Turnkey_InstallSdkMinimal", "Force Reinstall Sdk (Prefer Minimal)"),
-			LOCTEXT("TurnkeyTooltip_InstallSdkMinimal", "Attempt to force re-install an Sdk, as hosted by your studio. Will attempt to install a minimal Sdk (useful for building/running only)"),
-			FSlateIcon(),
-			FExecuteAction::CreateStatic(TurnkeyInstallSdk, UBTPlatformName, false, true, NoDevice)
-		);
-
-		MenuBuilder.AddMenuEntry(
-			LOCTEXT("Turnkey_InstallSdkFull", "Force Reinstall (Prefer Full)"),
-			LOCTEXT("TurnkeyTooltip_InstallSdkMinimal", "Attempt to force re-install an Sdk, as hosted by your studio. Will attempt to install a full Sdk (useful profiling or other use cases)"),
-			FSlateIcon(),
-			FExecuteAction::CreateStatic(TurnkeyInstallSdk, UBTPlatformName, true, true, NoDevice)
-		);
+		return LOCTEXT("PossessLabel", "Possess");
 	}
 	else
 	{
-		MenuBuilder.AddMenuEntry(
-			LOCTEXT("Turnkey_InstallSdkMinimal", "Install Sdk (Prefer Minimal)"),
-			LOCTEXT("TurnkeyTooltip_InstallSdkMinimal", "Attempt to install an Sdk, as hosted by your studio. Will attempt to install a minimal Sdk (useful for building/running only)"),
-			FSlateIcon(),
-			FExecuteAction::CreateStatic(TurnkeyInstallSdk, UBTPlatformName, false, false, NoDevice)
-		);
-
-		MenuBuilder.AddMenuEntry(
-			LOCTEXT("Turnkey_InstallSdkFull", "Install Sdk (Prefer Full)"),
-			LOCTEXT("TurnkeyTooltip_InstallSdkMinimal", "Attempt to install an Sdk, as hosted by your studio. Will attempt to install a full Sdk (useful profiling or other use cases)"),
-			FSlateIcon(),
-			FExecuteAction::CreateStatic(TurnkeyInstallSdk, UBTPlatformName, true, false, NoDevice)
-		);
+		return LOCTEXT("ToggleBetweenPieAndSIELabel", "Toggle Between PIE and SIE");
 	}
 }
 
-static void GenerateDeviceProxyMenuParams(TSharedPtr<ITargetDeviceProxy> DeviceProxy, FName PlatformName, FUIAction& OutAction, FText& OutTooltip)
-{
-	// 	// create an All_<platform>_devices_on_<host> submenu
-	// 	if (DeviceProxy->IsAggregated())
-	// 	{
-	// 		FString AggregateDevicedName(FString::Printf(TEXT("  %s"), *DeviceProxy->GetName())); //align with the other menu entries
-	// 		FSlateIcon AggregateDeviceIcon(FEditorStyle::GetStyleSetName(), EditorPlatformInfo->GetIconStyleName(PlatformInfo::EPlatformIconSize::Normal));
-	// 
-	// 		MenuBuilder.AddSubMenu(
-	// 			FText::FromString(AggregateDevicedName),
-	// 			FText::FromString(AggregateDevicedName),
-	// 			FNewMenuDelegate::CreateStatic(&MakeAllDevicesSubMenu, EditorPlatformInfo, DeviceProxy),
-	// 			false, AggregateDeviceIcon, true
-	// 		);
-	// 		continue;
-	// 	}
-
-		// ... create an action...
-	OutAction = FUIAction(
-		FExecuteAction::CreateStatic(&FInternalPlayWorldCommandCallbacks::HandleLaunchOnDeviceActionExecute, DeviceProxy->GetTargetDeviceId(NAME_None), DeviceProxy->GetName(), true)
-//		, FCanExecuteAction::CreateStatic(&FInternalPlayWorldCommandCallbacks::HandleLaunchOnDeviceActionCanExecute, DeviceProxy->GetName())
-//		, FIsActionChecked::CreateStatic(&FInternalPlayWorldCommandCallbacks::HandleLaunchOnDeviceActionIsChecked, DeviceProxy->GetName())
-	);
 
-	// ... generate tooltip text
-	FFormatNamedArguments TooltipArguments;
-	TooltipArguments.Add(TEXT("DeviceID"), FText::FromString(DeviceProxy->GetName()));
-	TooltipArguments.Add(TEXT("DisplayName"), FText::FromName(PlatformName));
-	OutTooltip = FText::Format(LOCTEXT("LaunchDeviceToolTipText_ThisDevice", "Launch the game on this {DisplayName} device ({DeviceID})"), TooltipArguments);
-	if (!DeviceProxy->IsAuthorized())
+FText FInternalPlayWorldCommandCallbacks::GetPossessEjectTooltip()
+{
+	if (IsInPIE())
 	{
-		OutTooltip = FText::Format(LOCTEXT("LaunchDeviceToolTipText_UnauthorizedOrLocked", "{DisplayName} device ({DeviceID}) is unauthorized or locked"), TooltipArguments);
+		return LOCTEXT("EjectToolTip", "Detaches from the player controller, allowing regular editor controls");
 	}
-
-	FProjectStatus ProjectStatus;
-	if (IProjectManager::Get().QueryStatusForCurrentProject(ProjectStatus) && !ProjectStatus.IsTargetPlatformSupported(PlatformName))
+	else if (IsInSIE())
 	{
-		FText TooltipLine2 = FText::Format(LOCTEXT("LaunchDevicePlatformWarning", "{DisplayName} is not listed as a target platform for this project, so may not run as expected."), TooltipArguments);
-		OutTooltip = FText::Format(FText::FromString(TEXT("{0}\n\n{1}")), OutTooltip, TooltipLine2);
+		return LOCTEXT("PossessToolTip", "Attaches to the player controller, allowing normal gameplay controls");
+	}
+	else
+	{
+		return LOCTEXT("ToggleBetweenPieAndSIEToolTip", "Toggles the current play session between play in editor and simulate in editor");
 	}
-
 }
 
-TSharedRef< SWidget > FPlayWorldCommands::GenerateTurnkeyMenuContent(TSharedRef<FUICommandList> InCommandList)
-{
-	const bool bShouldCloseWindowAfterMenuSelection = true;
-	FMenuBuilder MenuBuilder(bShouldCloseWindowAfterMenuSelection, InCommandList);
-
-	// shared devices section
-	ITargetDeviceServicesModule* TargetDeviceServicesModule = static_cast<ITargetDeviceServicesModule*>(FModuleManager::Get().LoadModule(TEXT("TargetDeviceServices")));
-	IProjectTargetPlatformEditorModule& ProjectTargetPlatformEditorModule = FModuleManager::LoadModuleChecked<IProjectTargetPlatformEditorModule>("ProjectTargetPlatformEditor");
 
-	TArray<FString> DeviceIdsToQuery;
-
-	MenuBuilder.BeginSection("LevelEditorLaunchHint", FText::FromString(TEXT("Turnkey Info")));
+FSlateIcon FInternalPlayWorldCommandCallbacks::GetPossessEjectImage()
+{
+	if (IsInPIE())
 	{
-		MenuBuilder.AddWidget(
-			SNew(STextBlock)
-			.ColorAndOpacity(FSlateColor::UseSubduedForeground())
-			// this is temp and doesn't need to be loc'd
-			.WrapTextAt(300)
-			.Text(FText::FromString(TEXT("The Turnkey SDK management system is now enabled for this menu. If you have any issues, run with \"-OldLaunchOn\"  on the editor commandline to use the old Launch On menu"))),
-			FText::GetEmpty()
-		);
+		return FSlateIcon(FEditorStyle::GetStyleSetName(), "PlayWorld.EjectFromPlayer");
 	}
-	MenuBuilder.EndSection();
-
-	MenuBuilder.BeginSection("LevelEditorLaunchDevices", LOCTEXT("TurnkeySection_LaunchButtonDevices", "Quick Launch"));
+	else if (IsInSIE())
 	{
-		for (const auto Pair : FDataDrivenPlatformInfoRegistry::GetAllPlatformInfos())
-		{
-			if (Pair.Value.bIsFakePlatform)
-			{
-				continue;
-			}
+		return FSlateIcon(FEditorStyle::GetStyleSetName(), "PlayWorld.PossessPlayer");
+	}
+	else
+	{
+		return FSlateIcon();
+	}
+}
 
-			FName PlatformName = Pair.Key;
-			const FDataDrivenPlatformInfo& Info = Pair.Value;
 
-			// look for devices for all platforms, even if the platform isn't installed - Turnkey can install Sdk after selecting LaunchOn
-			TArray<TSharedPtr<ITargetDeviceProxy>> DeviceProxies;
-			TargetDeviceServicesModule->GetDeviceProxyManager()->GetAllProxies(PlatformName, DeviceProxies);
+bool FInternalPlayWorldCommandCallbacks::CanLateJoin()
+{
+	return HasPlayWorld();
+}
 
-			if (DeviceProxies.Num() > 0)
-			{
-				// 					Algo::Sort(DeviceProxies, [](TSharedPtr<ITargetDeviceProxy> A, TSharedPtr<ITargetDeviceProxy> B)
-				// 					{ return A->Name == LastChosen || B->Name != LastChosen && A->IsDefault(); });
+bool FInternalPlayWorldCommandCallbacks::CanShowLateJoinButton()
+{
+	return GetDefault<UEditorExperimentalSettings>()->bAllowLateJoinInPIE && HasPlayWorld();
+}
 
+void SetLastExecutedPlayMode(EPlayModeType PlayMode)
+{
+	ULevelEditorPlaySettings* PlaySettings = GetMutableDefault<ULevelEditorPlaySettings>();
+	PlaySettings->LastExecutedPlayModeType = PlayMode;
 
-				// always use the first one, after sorting
-				FUIAction Action;
-				FText Tooltip;
-				GenerateDeviceProxyMenuParams(DeviceProxies[0], PlatformName, Action, Tooltip);
+	FPropertyChangedEvent PropChangeEvent(ULevelEditorPlaySettings::StaticClass()->FindPropertyByName(GET_MEMBER_NAME_CHECKED(ULevelEditorPlaySettings, LastExecutedPlayModeType)));
+	PlaySettings->PostEditChangeProperty(PropChangeEvent);
 
-				if (DeviceProxies.Num() == 1)
-				{
-					MenuBuilder.AddMenuEntry(
-						MakeSdkStatusAttribute(PlatformName, DeviceProxies[0]),
-						Tooltip,
-						FSlateIcon(FEditorStyle::GetStyleSetName(), Pair.Value.GetIconStyleName(EPlatformIconSize::Normal)),
-						Action,
-						NAME_None,
-						EUserInterfaceActionType::Button
-					);
-				}
-				else
-				{
-					MenuBuilder.AddSubMenu(
-						MakeSdkStatusAttribute(PlatformName, DeviceProxies[0]),
-						Tooltip,
-						FNewMenuDelegate::CreateLambda([TargetDeviceServicesModule, PlatformName, IconStyle=Pair.Value.GetIconStyleName(EPlatformIconSize::Normal)](FMenuBuilder& SubMenuBuilder) {
-							// re-get the proxies, just in case they changed
-							TArray<TSharedPtr<ITargetDeviceProxy>> DeviceProxies;
-							TargetDeviceServicesModule->GetDeviceProxyManager()->GetAllProxies(PlatformName, DeviceProxies);
-							// for each one, put an entry (even the one that was in the outer menu, for less confusion)
-							for (const TSharedPtr<ITargetDeviceProxy> Proxy : DeviceProxies)
-							{
-								FUIAction SubAction;
-								FText SubTooltip;
-								GenerateDeviceProxyMenuParams(Proxy, PlatformName, SubAction, SubTooltip);
-								SubMenuBuilder.AddMenuEntry(
-									MakeSdkStatusAttribute(PlatformName, Proxy),
-									SubTooltip,
-									FSlateIcon(FEditorStyle::GetStyleSetName(), IconStyle),
-									SubAction,
-									NAME_None,
-									EUserInterfaceActionType::Button
-								);
-							}
-						}),
-						Action,
-						NAME_None,
-						EUserInterfaceActionType::Check,
-						false,
-						FSlateIcon(FEditorStyle::GetStyleSetName(), Pair.Value.GetIconStyleName(EPlatformIconSize::Normal)),
-						true
-					);
-				}
+	PlaySettings->SaveConfig();
+}
 
-				// gather any unknown status devices to query at the end
-				for (const TSharedPtr<ITargetDeviceProxy> Proxy : DeviceProxies)
-				{
-					FString DeviceId = Proxy->GetTargetDeviceId(NAME_None);
-					if (Pair.Value.GetStatusForDeviceId(DeviceId) == DDPIPlatformSdkStatus::Unknown)
-					{
-						DeviceIdsToQuery.Add(DeviceId);
-					}
-				}
-			}
-		}
-	}
-	MenuBuilder.EndSection();
 
-	// tip section
-	MenuBuilder.BeginSection("LevelEditorLaunchHint");
+void FInternalPlayWorldCommandCallbacks::Simulate_Clicked()
+{
+	// Is a simulation session already running?  If so, do nothing
+	if (HasPlayWorld() && GUnrealEd->bIsSimulatingInEditor)
 	{
-		MenuBuilder.AddWidget(
-			SNew(STextBlock)
-			.ColorAndOpacity(FSlateColor::UseSubduedForeground())
-			.Text(LOCTEXT("ZoomToFitHorizontal", "Launching a game on a different device will change your default 'Launch' device in the toolbar"))
-			.WrapTextAt(300),
-			FText::GetEmpty()
-		);
+		return;
 	}
-	MenuBuilder.EndSection();
 
-	MenuBuilder.BeginSection("AllPlatforms", LOCTEXT("TurnkeyMenu_ManagePlatforms", "Sdk/Device Management"));
-	for (const auto Pair : FDataDrivenPlatformInfoRegistry::GetAllPlatformInfos())
+	FLevelEditorModule& LevelEditorModule = FModuleManager::GetModuleChecked<FLevelEditorModule>(TEXT("LevelEditor"));
+
+	TSharedPtr<IAssetViewport> ActiveLevelViewport = LevelEditorModule.GetFirstActiveViewport();
+	if (ActiveLevelViewport.IsValid())
 	{
-		if (Pair.Value.bIsFakePlatform)
+		// Start a new simulation session!
+		if (!HasPlayWorld())
 		{
-			continue;
-		}
+			if (FEngineAnalytics::IsAvailable())
+			{
+				FEngineAnalytics::GetProvider().RecordEvent(TEXT("Editor.Usage.SimulateInEditor"));
+			}
+			SetLastExecutedPlayMode(PlayMode_Simulate);
+			FRequestPlaySessionParams SessionParams;
+			SessionParams.WorldType = EPlaySessionWorldType::SimulateInEditor;
+			SessionParams.DestinationSlateViewport = ActiveLevelViewport;
 
-		FName PlatformName = Pair.Key;
-		const FDataDrivenPlatformInfo& Info = Pair.Value;
-
-		FString DeviceNameToCheckForStatus;
-		MenuBuilder.AddSubMenu(
-			MakeSdkStatusAttribute(PlatformName, nullptr),
-			FText::FromString(PlatformName.ToString()),
-			FNewMenuDelegate::CreateStatic(&MakeTurnkeyPlatformMenu, PlatformName, TargetDeviceServicesModule),
-			false, 
-			FSlateIcon(FEditorStyle::GetStyleSetName(), Pair.Value.GetIconStyleName(EPlatformIconSize::Normal)),
-			true
-		);
+			GUnrealEd->RequestPlaySession(SessionParams);
+		}
+		else if (ActiveLevelViewport->HasPlayInEditorViewport())
+		{
+			GUnrealEd->RequestToggleBetweenPIEandSIE();
+		}
 	}
-
-	MenuBuilder.EndSection();
+}
 
 
-	MenuBuilder.BeginSection("CookerSettings");
+bool FInternalPlayWorldCommandCallbacks::Simulate_CanExecute()
+{
+	// Can't simulate while already simulating; PIE is fine as we toggle to simulate
+	return !(HasPlayWorld() && GUnrealEd->bIsSimulatingInEditor) && !GEditor->IsLightingBuildCurrentlyRunning();
+}
 
-	UCookerSettings* CookerSettings = GetMutableDefault<UCookerSettings>();
 
-	{
-		FUIAction UIAction;
-		UIAction.ExecuteAction = FExecuteAction::CreateLambda([CookerSettings]
-		{
-			CookerSettings->bCookOnTheFlyForLaunchOn = !CookerSettings->bCookOnTheFlyForLaunchOn;
-			CookerSettings->Modify(true);
+bool FInternalPlayWorldCommandCallbacks::Simulate_IsChecked()
+{
+	return HasPlayWorld() && GUnrealEd->bIsSimulatingInEditor;
+}
 
-			// Update source control
 
-			FString ConfigPath = FPaths::ConvertRelativePathToFull(CookerSettings->GetDefaultConfigFilename());
+const TSharedRef < FUICommandInfo > GetLastPlaySessionCommand()
+{
+	const ULevelEditorPlaySettings* PlaySettings = GetDefault<ULevelEditorPlaySettings>();
 
-			if (FPlatformFileManager::Get().GetPlatformFile().FileExists(*ConfigPath))
-			{
-				if (ISourceControlModule::Get().IsEnabled())
-				{
-					FText ErrorMessage;
-
-					if (!SourceControlHelpers::CheckoutOrMarkForAdd(ConfigPath, FText::FromString(ConfigPath), NULL, ErrorMessage))
-					{
-						FNotificationInfo Info(ErrorMessage);
-						Info.ExpireDuration = 3.0f;
-						FSlateNotificationManager::Get().AddNotification(Info);
-					}
-				}
-				else
-				{
-					if (!FPlatformFileManager::Get().GetPlatformFile().SetReadOnly(*ConfigPath, false))
-					{
-						FNotificationInfo Info(FText::Format(LOCTEXT("FailedToMakeWritable", "Could not make {0} writable."), FText::FromString(ConfigPath)));
-						Info.ExpireDuration = 3.0f;
-						FSlateNotificationManager::Get().AddNotification(Info);
-					}
-				}
-			}
-
-			// Save settings
-			CookerSettings->UpdateSinglePropertyInConfigFile(CookerSettings->GetClass()->FindPropertyByName(GET_MEMBER_NAME_CHECKED(UCookerSettings, bCookOnTheFlyForLaunchOn)), CookerSettings->GetDefaultConfigFilename());
-		});
-
-		UIAction.GetActionCheckState = FGetActionCheckState::CreateLambda([CookerSettings]
-		{
-			return CookerSettings->bCookOnTheFlyForLaunchOn ? ECheckBoxState::Checked : ECheckBoxState::Unchecked;
-		});
-
-		MenuBuilder.AddMenuEntry(
-			LOCTEXT("CookOnTheFlyOnLaunch", "Enable cooking on the fly"),
-			LOCTEXT("CookOnTheFlyOnLaunchDescription", "Cook on the fly instead of cooking upfront when launching"),
-			FSlateIcon(),
-			UIAction,
-			NAME_None,
-			EUserInterfaceActionType::ToggleButton
-		);
-	}
-
-	MenuBuilder.EndSection();
-
-	// options section
-	MenuBuilder.BeginSection("LevelEditorLaunchOptions");
-	{
-		MenuBuilder.AddMenuEntry(FPlayWorldCommands::Get().OpenProjectLauncher,
-			NAME_None,
-			TAttribute<FText>(),
-			TAttribute<FText>(),
-			FSlateIcon(FEditorStyle::GetStyleSetName(), "Launcher.TabIcon")
-		);
-
-		MenuBuilder.AddMenuEntry(FPlayWorldCommands::Get().OpenDeviceManager,
-			NAME_None,
-			TAttribute<FText>(),
-			TAttribute<FText>(),
-			FSlateIcon(FEditorStyle::GetStyleSetName(), "DeviceDetails.TabIcon")
-		);
-
-		ProjectTargetPlatformEditorModule.AddOpenProjectTargetPlatformEditorMenuItem(MenuBuilder);
-	}
-	MenuBuilder.EndSection();
-
-	// now kick-off any devices that need to be updated
-	if (DeviceIdsToQuery.Num() > 0)
-	{
-		FDataDrivenPlatformInfoRegistry::UpdateDeviceSdkStatus(DeviceIdsToQuery);
-	}
-
-	return MenuBuilder.MakeWidget();
-}
-
-
-
-
-void PopulateLaunchMenu(UToolMenu* Menu)
-{
-	TArray<PlatformInfo::FTargetPlatformInfo*> VanillaPlatforms = PlatformInfo::GetVanillaPlatformInfoArray();
-
-	Algo::Sort(VanillaPlatforms, [](const PlatformInfo::FTargetPlatformInfo* One, const PlatformInfo::FTargetPlatformInfo* Two) -> bool
-	{
-		return One->DisplayName.CompareTo(Two->DisplayName) < 0;
-	});
-
-	// shared devices section
-	ITargetDeviceServicesModule* TargetDeviceServicesModule = static_cast<ITargetDeviceServicesModule*>(FModuleManager::Get().LoadModule(TEXT("TargetDeviceServices")));
-	IProjectTargetPlatformEditorModule& ProjectTargetPlatformEditorModule = FModuleManager::LoadModuleChecked<IProjectTargetPlatformEditorModule>("ProjectTargetPlatformEditor");
-
-	TArray<FString> PlatformsToMaybeInstallLinksFor;
-	PlatformsToMaybeInstallLinksFor.Add(TEXT("Android"));
-	PlatformsToMaybeInstallLinksFor.Add(TEXT("IOS"));
-	PlatformsToMaybeInstallLinksFor.Add(TEXT("Linux"));
-	PlatformsToMaybeInstallLinksFor.Add(TEXT("Lumin"));
-	TArray<FString> PlatformsToCheckFlavorsFor;
-	PlatformsToCheckFlavorsFor.Add(TEXT("Android"));
-	PlatformsToCheckFlavorsFor.Add(TEXT("IOS"));
-	TArray<FName> PlatformsWithNoDevices;
-	TArray<const PlatformInfo::FTargetPlatformInfo*> PlatformsToAddInstallLinksFor;
-	EProjectType ProjectType = FGameProjectGenerationModule::Get().ProjectHasCodeFiles() ? EProjectType::Code : EProjectType::Content;
-
-	{
-		FToolMenuSection& Section = Menu->AddSection("LevelEditorLaunchDevices", LOCTEXT("LaunchButtonDevicesSection", "Devices"));
-		for (const PlatformInfo::FTargetPlatformInfo* VanillaPlatform : VanillaPlatforms)
-		{
-			// for the Editor we are only interested in launching standalone games
-			if (VanillaPlatform->PlatformType != EBuildTargetType::Game || !VanillaPlatform->DataDrivenPlatformInfo->bEnabledForUse || !FInstalledPlatformInfo::Get().CanDisplayPlatform(VanillaPlatform->DataDrivenPlatformInfo->UBTPlatformString, ProjectType))
-			{
-				continue;
-			}
-
-			if (VanillaPlatform->DataDrivenPlatformInfo->GetSdkStatus() == DDPIPlatformSdkStatus::Valid)
-			{
-				// for each platform...
-				TArray<TSharedPtr<ITargetDeviceProxy>> DeviceProxies;
-				// the list of proxies include the "Al_Android" entry
-				TargetDeviceServicesModule->GetDeviceProxyManager()->GetAllProxies(VanillaPlatform->Name, DeviceProxies);
-
-				// if this platform had no devices, but we want to show an extra option if not installed right
-				if (DeviceProxies.Num() == 0)
-				{
-					if (PlatformsWithNoDevices.Find(VanillaPlatform->Name) == INDEX_NONE)
-					{
-						// add an entry with a no devices found
-						PlatformsWithNoDevices.Add(VanillaPlatform->Name);
-					}
-				}
-				else
-				{
-					// for each proxy...
-					for (auto DeviceProxyIt = DeviceProxies.CreateIterator(); DeviceProxyIt; ++DeviceProxyIt)
-					{
-						TSharedPtr<ITargetDeviceProxy> DeviceProxy = *DeviceProxyIt;
-
-						// create an All_<platform>_devices_on_<host> submenu
-						if (DeviceProxy->IsAggregated())
-						{
-							FString AggregateDevicedName(FString::Printf(TEXT("  %s"), *DeviceProxy->GetName())); //align with the other menu entries
-							FSlateIcon AggregateDeviceIcon(FEditorStyle::GetStyleSetName(), VanillaPlatform->GetIconStyleName(EPlatformIconSize::Normal));
-
-							Section.AddSubMenu(
-								NAME_None,
-								FText::FromString(AggregateDevicedName),
-								FText::FromString(AggregateDevicedName),
-								FNewMenuDelegate::CreateStatic(&MakeAllDevicesSubMenu, VanillaPlatform, DeviceProxy),
-								false, AggregateDeviceIcon, true
-							);
-							continue;
-						}
-
-						// ... create an action...
-						FUIAction LaunchDeviceAction(
-							FExecuteAction::CreateStatic(&FInternalPlayWorldCommandCallbacks::HandleLaunchOnDeviceActionExecute, DeviceProxy->GetTargetDeviceId(NAME_None), DeviceProxy->GetName(), false),
-							FCanExecuteAction::CreateStatic(&FInternalPlayWorldCommandCallbacks::HandleLaunchOnDeviceActionCanExecute, DeviceProxy->GetName()),
-							FIsActionChecked::CreateStatic(&FInternalPlayWorldCommandCallbacks::HandleLaunchOnDeviceActionIsChecked, DeviceProxy->GetName())
-						);
-
-						// ... generate display label...
-						FFormatNamedArguments LabelArguments;
-						LabelArguments.Add(TEXT("DeviceName"), FText::FromString(DeviceProxy->GetName()));
-
-						if (!DeviceProxy->IsConnected())
-						{
-							LabelArguments.Add(TEXT("HostUser"), LOCTEXT("DisconnectedHint", " [Disconnected]"));
-						}
-						else if (DeviceProxy->GetHostUser() != FPlatformProcess::UserName(false))
-						{
-							LabelArguments.Add(TEXT("HostUser"), FText::FromString(DeviceProxy->GetHostUser()));
-						}
-						else
-						{
-							LabelArguments.Add(TEXT("HostUser"), FText::GetEmpty());
-						}
-
-						FText Label = FText::Format(LOCTEXT("LaunchDeviceLabel", "{DeviceName}{HostUser}"), LabelArguments);
-
-						// ... generate tooltip text
-						FFormatNamedArguments TooltipArguments;
-						TooltipArguments.Add(TEXT("DeviceID"), FText::FromString(DeviceProxy->GetName()));
-						TooltipArguments.Add(TEXT("DisplayName"), VanillaPlatform->DisplayName);
-						FText Tooltip = FText::Format(LOCTEXT("LaunchDeviceToolTipText_ThisDevice", "Launch the game on this {DisplayName} device ({DeviceID})"), TooltipArguments);
-						if (!DeviceProxy->IsAuthorized())
-						{
-							Tooltip = FText::Format(LOCTEXT("LaunchDeviceToolTipText_UnauthorizedOrLocked", "{DisplayName} device ({DeviceID}) is unauthorized or locked"), TooltipArguments);
-						}
-
-						FProjectStatus ProjectStatus;
-						if (IProjectManager::Get().QueryStatusForCurrentProject(ProjectStatus) && !ProjectStatus.IsTargetPlatformSupported(VanillaPlatform->Name))
-						{
-							FText TooltipLine2 = FText::Format(LOCTEXT("LaunchDevicePlatformWarning", "{DisplayName} is not listed as a target platform for this project, so may not run as expected."), TooltipArguments);
-							Tooltip = FText::Format(FText::FromString(TEXT("{0}\n\n{1}")), Tooltip, TooltipLine2);
-						}
-
-						// ... and add a menu entry
-						FToolMenuEntry& Entry = Section.AddEntry(FToolMenuEntry::InitMenuEntry(
-							NAME_None,
-							LaunchDeviceAction,
-							ProjectTargetPlatformEditorModule.MakePlatformMenuItemWidget(*VanillaPlatform, true, Label)
-						));
-						Entry.ToolTip = Tooltip;
-						Entry.UserInterfaceActionType = EUserInterfaceActionType::Check;
-					}
-				}
-			}
-			else
-			{
-				// if the platform wasn't installed, we'll add a menu item later (we never care about code in this case, since we don't compile)
-				if (PlatformsToMaybeInstallLinksFor.Find(VanillaPlatform->Name.ToString()) != INDEX_NONE)
-				{
-					PlatformsToAddInstallLinksFor.Add(VanillaPlatform);
-				}
-			}
-		}
-	}
-
-	TWeakObjectPtr<UCookerSettings> CookerSettings = GetMutableDefault<UCookerSettings>();
-
-	{
-		FToolMenuSection& Section = Menu->AddSection("CookerSettings");
-
-		FUIAction UIAction;
-		UIAction.ExecuteAction = FExecuteAction::CreateLambda([CookerSettings]
-		{
-			CookerSettings->bCookOnTheFlyForLaunchOn = !CookerSettings->bCookOnTheFlyForLaunchOn;
-			CookerSettings->Modify(true);
-
-			// Update source control
-
-			FString ConfigPath = FPaths::ConvertRelativePathToFull(CookerSettings->GetDefaultConfigFilename());
-
-			if (FPlatformFileManager::Get().GetPlatformFile().FileExists(*ConfigPath))
-			{
-				if (ISourceControlModule::Get().IsEnabled())
-				{
-					FText ErrorMessage;
-
-					if (!SourceControlHelpers::CheckoutOrMarkForAdd(ConfigPath, FText::FromString(ConfigPath), NULL, ErrorMessage))
-					{
-						FNotificationInfo Info(ErrorMessage);
-						Info.ExpireDuration = 3.0f;
-						FSlateNotificationManager::Get().AddNotification(Info);
-					}
-				}
-				else
-				{
-					if (!FPlatformFileManager::Get().GetPlatformFile().SetReadOnly(*ConfigPath, false))
-					{
-						FNotificationInfo Info(FText::Format(LOCTEXT("FailedToMakeWritable", "Could not make {0} writable."), FText::FromString(ConfigPath)));
-						Info.ExpireDuration = 3.0f;
-						FSlateNotificationManager::Get().AddNotification(Info);
-					}
-				}
-			}
-
-			// Save settings
-			CookerSettings->UpdateSinglePropertyInConfigFile(CookerSettings->GetClass()->FindPropertyByName(GET_MEMBER_NAME_CHECKED(UCookerSettings, bCookOnTheFlyForLaunchOn)), CookerSettings->GetDefaultConfigFilename());
-		});
-
-		UIAction.GetActionCheckState = FGetActionCheckState::CreateLambda([CookerSettings]
-		{
-			return CookerSettings->bCookOnTheFlyForLaunchOn ? ECheckBoxState::Checked : ECheckBoxState::Unchecked;
-		});
-
-		Section.AddMenuEntry(
-			"CookOnTheFlyOnLaunch",
-			LOCTEXT("CookOnTheFlyOnLaunch", "Enable cooking on the fly"),
-			LOCTEXT("CookOnTheFlyOnLaunchDescription", "Cook on the fly instead of cooking upfront when launching"),
-			FSlateIcon(),
-			UIAction,
-			EUserInterfaceActionType::ToggleButton
-		);
-	}
-
-	if (PlatformsWithNoDevices.Num() > 0)
-	{
-		{
-			FToolMenuSection& Section = Menu->AddSection("NoDevices");
-			for (int32 PlatformIndex = 0; PlatformIndex < PlatformsWithNoDevices.Num(); PlatformIndex++)
-			{
-				const PlatformInfo::FTargetPlatformInfo* PlatformInfo = PlatformInfo::FindVanillaPlatformInfo(PlatformsWithNoDevices[PlatformIndex]);
-
-				// ... generate display label...
-				FFormatNamedArguments LabelArguments;
-				LabelArguments.Add(TEXT("DisplayName"), PlatformInfo->DisplayName);
-
-				FText Label = FText::Format(LOCTEXT("NoDeviceLabel", "{DisplayName} - No Devices Found"), LabelArguments);
-
-				// ... create an action...
-				FUIAction NoDeviceAction(
-					FExecuteAction::CreateStatic(&FInternalPlayWorldCommandCallbacks::HandleNoDeviceFoundActionExecute),
-					FCanExecuteAction::CreateStatic(&FInternalPlayWorldCommandCallbacks::HandleNoDeviceFoundActionCanExecute)
-				);
-
-				// ... generate tooltip text
-				FFormatNamedArguments TooltipArguments;
-				TooltipArguments.Add(TEXT("DisplayName"), PlatformInfo->DisplayName);
-				FText Tooltip = FText::Format(LOCTEXT("LaunchNoDevicesToolTipText", "Found no connected devices for {DisplayName}"), TooltipArguments);
-
-				// ... and add a menu entry
-				FToolMenuEntry& Entry = Section.AddEntry(FToolMenuEntry::InitMenuEntry(
-					NAME_None,
-					NoDeviceAction,
-					ProjectTargetPlatformEditorModule.MakePlatformMenuItemWidget(*PlatformInfo, true, Label)
-				));
-				Entry.ToolTip = Tooltip;
-				Entry.UserInterfaceActionType = EUserInterfaceActionType::Check;
-			}
-		}
-	}
-
-	// tip section
-	{
-		FToolMenuSection& Section = Menu->AddSection("LevelEditorLaunchHint");
-		Section.AddEntry(FToolMenuEntry::InitWidget(
-			"LevelEditorLaunchHint",
-			SNew(STextBlock)
-			.ColorAndOpacity(FSlateColor::UseSubduedForeground())
-			.Text(LOCTEXT("ZoomToFitHorizontal", "Launching a game on a different device will change your default 'Launch' device in the toolbar"))
-			.WrapTextAt(250),
-			FText::GetEmpty()
-		));
-	}
-
-	if (PlatformsToAddInstallLinksFor.Num() > 0)
-	{
-		{
-			FToolMenuSection& Section = Menu->AddSection("SDKUninstalledTutorials");
-			for (const PlatformInfo::FTargetPlatformInfo* Platform : PlatformsToAddInstallLinksFor)
-			{
-				FUIAction Action(FExecuteAction::CreateStatic(&FInternalPlayWorldCommandCallbacks::HandleShowSDKTutorial, Platform->DisplayName.ToString(), Platform->DataDrivenPlatformInfo->SDKTutorial));
-
-				FFormatNamedArguments LabelArguments;
-				LabelArguments.Add(TEXT("PlatformName"), Platform->DisplayName);
-				FText Label = FText::Format(LOCTEXT("LaunchPlatformLabel", "{PlatformName} Support"), LabelArguments);
-
-
-				Section.AddMenuEntry(
-					NAME_None,
-					Label,
-					LOCTEXT("PlatformSDK", "Show information on setting up the platform tools"),
-					FSlateIcon(FEditorStyle::GetStyleSetName(), "LevelEditor.BrowseDocumentation"),
-					Action,
-					EUserInterfaceActionType::Button);
-			}
-		}
-	}
-
-	// options section
-	{
-		FToolMenuSection& Section = Menu->AddSection("LevelEditorLaunchOptions");
-		Section.AddMenuEntry(FPlayWorldCommands::Get().OpenProjectLauncher,
-			TAttribute<FText>(),
-			TAttribute<FText>(),
-			FSlateIcon(FEditorStyle::GetStyleSetName(), "Launcher.TabIcon")
-		);
-
-		Section.AddMenuEntry(FPlayWorldCommands::Get().OpenDeviceManager,
-			TAttribute<FText>(),
-			TAttribute<FText>(),
-			FSlateIcon(FEditorStyle::GetStyleSetName(), "DeviceDetails.TabIcon")
-		);
-
-		Section.AddDynamicEntry("OpenProjectTargetPlatform", FNewToolMenuDelegateLegacy::CreateLambda([](FMenuBuilder& MenuBuilder, UToolMenu* ToolMenu)
-		{
-			FModuleManager::LoadModuleChecked<IProjectTargetPlatformEditorModule>("ProjectTargetPlatformEditor").AddOpenProjectTargetPlatformEditorMenuItem(MenuBuilder);
-		}));
-	}
-}
-
-TSharedRef< SWidget > FPlayWorldCommands::GenerateLaunchMenuContent(TSharedRef<FUICommandList> InCommandList)
-{
-	static const FName MenuName("UnrealEd.PlayWorldCommands.LaunchMenu");
-
-	if (!UToolMenus::Get()->IsMenuRegistered(MenuName))
-	{
-		UToolMenu* Menu = UToolMenus::Get()->RegisterMenu(MenuName);
-		Menu->AddDynamicSection("DynamicSection", FNewToolMenuDelegate::CreateStatic(&PopulateLaunchMenu));
-	}
-
-	FToolMenuContext MenuContext(InCommandList);
-	return UToolMenus::Get()->GenerateWidget(MenuName, MenuContext);
-}
-
-
-//////////////////////////////////////////////////////////////////////////
-// FPlayWorldCommandCallbacks
-
-void FPlayWorldCommandCallbacks::StartPlayFromHere()
-{
-	// Is a PIE session already running?  If so we close it first
-	if (GUnrealEd->PlayWorld != NULL)
-	{
-		GUnrealEd->EndPlayMap();
-	}
-
-	FRequestPlaySessionParams SessionParams;
-
-	UClass* const PlayerStartClass = GUnrealEd->PlayFromHerePlayerStartClass ? (UClass*)GUnrealEd->PlayFromHerePlayerStartClass : APlayerStart::StaticClass();
-
-	// Figure out the start location of the player
-	UCapsuleComponent*	DefaultCollisionComponent = CastChecked<UCapsuleComponent>(PlayerStartClass->GetDefaultObject<AActor>()->GetRootComponent());
-	FVector	CollisionExtent = FVector(DefaultCollisionComponent->GetScaledCapsuleRadius(), DefaultCollisionComponent->GetScaledCapsuleRadius(), DefaultCollisionComponent->GetScaledCapsuleHalfHeight());
-	SessionParams.StartLocation = GEditor->UnsnappedClickLocation + GEditor->ClickPlane * (FVector::BoxPushOut(GEditor->ClickPlane, CollisionExtent) + 0.1f);
-
-	FLevelEditorModule& LevelEditorModule = FModuleManager::GetModuleChecked<FLevelEditorModule>("LevelEditor");
-
-	TSharedPtr<IAssetViewport> ActiveLevelViewport = LevelEditorModule.GetFirstActiveViewport();
-
-
-	if (ActiveLevelViewport.IsValid() && ActiveLevelViewport->GetAssetViewportClient().IsPerspective())
-	{
-		// If there is no level viewport, a new window will be spawned to play in.
-		SessionParams.DestinationSlateViewport = ActiveLevelViewport;
-		SessionParams.StartRotation = ActiveLevelViewport->GetAssetViewportClient().GetViewRotation();
-	}
-
-	GUnrealEd->RequestPlaySession(SessionParams);
-}
-
-
-void FPlayWorldCommandCallbacks::ResumePlaySession_Clicked()
-{
-	if (HasPlayWorld())
-	{
-		LeaveDebuggingMode();
-		GUnrealEd->PlaySessionResumed();
-		uint32 UserIndex = 0;
-		FSlateApplication::Get().SetUserFocusToGameViewport(UserIndex);
-	}
-}
-
-
-void FPlayWorldCommandCallbacks::PausePlaySession_Clicked()
-{
-	if (HasPlayWorld())
-	{
-		GUnrealEd->PlayWorld->bDebugPauseExecution = true;
-		GUnrealEd->PlaySessionPaused();
-		if (IsInPIE()) {
-			FSlateApplication::Get().ClearKeyboardFocus(EFocusCause::SetDirectly);
-			FSlateApplication::Get().ResetToDefaultInputSettings();
-
-			TWeakPtr<SGlobalPlayWorldActions> ActiveGlobalPlayWorldWidget = FPlayWorldCommands::GetActiveGlobalPlayWorldActionsWidget();
-			if (ActiveGlobalPlayWorldWidget.IsValid())
-			{
-				uint32 UserIndex = 0;
-				FSlateApplication::Get().SetUserFocus(UserIndex, ActiveGlobalPlayWorldWidget.Pin());
-			}
-		}
-	}
-}
-
-
-void FPlayWorldCommandCallbacks::SingleFrameAdvance_Clicked()
-{
-	if (HasPlayWorld())
-	{
-		FInternalPlayWorldCommandCallbacks::SingleFrameAdvance_Clicked();
-	}
-}
-
-bool FPlayWorldCommandCallbacks::IsInSIE()
-{
-	return GEditor->bIsSimulatingInEditor;
-}
-
-
-bool FPlayWorldCommandCallbacks::IsInPIE()
-{
-	return (GEditor->PlayWorld != NULL) && (!GEditor->bIsSimulatingInEditor);
-}
-
-
-bool FPlayWorldCommandCallbacks::IsInSIE_AndRunning()
-{
-	return IsInSIE() && ((GEditor->PlayWorld == NULL) || !(GEditor->PlayWorld->bDebugPauseExecution));
-}
-
-
-bool FPlayWorldCommandCallbacks::IsInPIE_AndRunning()
-{
-	return IsInPIE() && ((GEditor->PlayWorld == NULL) || !(GEditor->PlayWorld->bDebugPauseExecution));
-}
-
-
-bool FPlayWorldCommandCallbacks::HasPlayWorld()
-{
-	return GEditor->PlayWorld != NULL;
-}
-
-
-bool FPlayWorldCommandCallbacks::HasPlayWorldAndPaused()
-{
-	return HasPlayWorld() && GUnrealEd->PlayWorld->bDebugPauseExecution;
-}
-
-
-bool FPlayWorldCommandCallbacks::HasPlayWorldAndRunning()
-{
-	return HasPlayWorld() && !GUnrealEd->PlayWorld->bDebugPauseExecution;
-}
-
-
-//////////////////////////////////////////////////////////////////////////
-// FInternalPlayWorldCommandCallbacks
-
-FText FInternalPlayWorldCommandCallbacks::GetPossessEjectLabel()
-{
-	if (IsInPIE())
-	{
-		return LOCTEXT("EjectLabel", "Eject");
-	}
-	else if (IsInSIE())
-	{
-		return LOCTEXT("PossessLabel", "Possess");
-	}
-	else
-	{
-		return LOCTEXT("ToggleBetweenPieAndSIELabel", "Toggle Between PIE and SIE");
-	}
-}
-
-
-FText FInternalPlayWorldCommandCallbacks::GetPossessEjectTooltip()
-{
-	if (IsInPIE())
-	{
-		return LOCTEXT("EjectToolTip", "Detaches from the player controller, allowing regular editor controls");
-	}
-	else if (IsInSIE())
-	{
-		return LOCTEXT("PossessToolTip", "Attaches to the player controller, allowing normal gameplay controls");
-	}
-	else
-	{
-		return LOCTEXT("ToggleBetweenPieAndSIEToolTip", "Toggles the current play session between play in editor and simulate in editor");
-	}
-}
-
-
-FSlateIcon FInternalPlayWorldCommandCallbacks::GetPossessEjectImage()
-{
-	if (IsInPIE())
-	{
-		return FSlateIcon(FEditorStyle::GetStyleSetName(), "PlayWorld.EjectFromPlayer");
-	}
-	else if (IsInSIE())
-	{
-		return FSlateIcon(FEditorStyle::GetStyleSetName(), "PlayWorld.PossessPlayer");
-	}
-	else
-	{
-		return FSlateIcon();
-	}
-}
-
-
-bool FInternalPlayWorldCommandCallbacks::CanLateJoin()
-{
-	return HasPlayWorld();
-}
-
-bool FInternalPlayWorldCommandCallbacks::CanShowLateJoinButton()
-{
-	return GetDefault<UEditorExperimentalSettings>()->bAllowLateJoinInPIE && HasPlayWorld();
-}
-
-void SetLastExecutedPlayMode(EPlayModeType PlayMode)
-{
-	ULevelEditorPlaySettings* PlaySettings = GetMutableDefault<ULevelEditorPlaySettings>();
-	PlaySettings->LastExecutedPlayModeType = PlayMode;
-
-	FPropertyChangedEvent PropChangeEvent(ULevelEditorPlaySettings::StaticClass()->FindPropertyByName(GET_MEMBER_NAME_CHECKED(ULevelEditorPlaySettings, LastExecutedPlayModeType)));
-	PlaySettings->PostEditChangeProperty(PropChangeEvent);
-
-	PlaySettings->SaveConfig();
-}
-
-
-void FInternalPlayWorldCommandCallbacks::Simulate_Clicked()
-{
-	// Is a simulation session already running?  If so, do nothing
-	if (HasPlayWorld() && GUnrealEd->bIsSimulatingInEditor)
-	{
-		return;
-	}
-
-	FLevelEditorModule& LevelEditorModule = FModuleManager::GetModuleChecked<FLevelEditorModule>(TEXT("LevelEditor"));
-
-	TSharedPtr<IAssetViewport> ActiveLevelViewport = LevelEditorModule.GetFirstActiveViewport();
-	if (ActiveLevelViewport.IsValid())
-	{
-		// Start a new simulation session!
-		if (!HasPlayWorld())
-		{
-			if (FEngineAnalytics::IsAvailable())
-			{
-				FEngineAnalytics::GetProvider().RecordEvent(TEXT("Editor.Usage.SimulateInEditor"));
-			}
-			SetLastExecutedPlayMode(PlayMode_Simulate);
-			FRequestPlaySessionParams SessionParams;
-			SessionParams.WorldType = EPlaySessionWorldType::SimulateInEditor;
-			SessionParams.DestinationSlateViewport = ActiveLevelViewport;
-
-			GUnrealEd->RequestPlaySession(SessionParams);
-		}
-		else if (ActiveLevelViewport->HasPlayInEditorViewport())
-		{
-			GUnrealEd->RequestToggleBetweenPIEandSIE();
-		}
-	}
-}
-
-
-bool FInternalPlayWorldCommandCallbacks::Simulate_CanExecute()
-{
-	// Can't simulate while already simulating; PIE is fine as we toggle to simulate
-	return !(HasPlayWorld() && GUnrealEd->bIsSimulatingInEditor) && !GEditor->IsLightingBuildCurrentlyRunning();
-}
-
-
-bool FInternalPlayWorldCommandCallbacks::Simulate_IsChecked()
-{
-	return HasPlayWorld() && GUnrealEd->bIsSimulatingInEditor;
-}
-
-
-const TSharedRef < FUICommandInfo > GetLastPlaySessionCommand()
-{
-	const ULevelEditorPlaySettings* PlaySettings = GetDefault<ULevelEditorPlaySettings>();
-
-	const FPlayWorldCommands& Commands = FPlayWorldCommands::Get();
-	TSharedRef < FUICommandInfo > Command = Commands.PlayInViewport.ToSharedRef();
+	const FPlayWorldCommands& Commands = FPlayWorldCommands::Get();
+	TSharedRef < FUICommandInfo > Command = Commands.PlayInViewport.ToSharedRef();
 
 	switch (PlaySettings->LastExecutedPlayModeType)
 	{
@@ -2258,7 +1366,7 @@ void FInternalPlayWorldCommandCallbacks::RepeatLastPlay_Clicked()
 
 	// Grab the play command and execute it
 	TSharedRef<FUICommandInfo> LastCommand = GetLastPlaySessionCommand();
-	UE_LOG(LogTemp, Log, TEXT("Repeating last play command: %s"), *LastCommand->GetLabel().ToString());
+	UE_LOG(LogDebuggerCommands, Log, TEXT("Repeating last play command: %s"), *LastCommand->GetLabel().ToString());
 
 	FPlayWorldCommands::GlobalPlayWorldActions->ExecuteAction(LastCommand);
 }
@@ -2560,244 +1668,8 @@ void FInternalPlayWorldCommandCallbacks::PlayInSettings_Clicked()
 	FModuleManager::LoadModuleChecked<ISettingsModule>("Settings").ShowViewer("Editor", "LevelEditor", "PlayIn");
 }
 
-void FInternalPlayWorldCommandCallbacks::OpenProjectLauncher_Clicked()
-{
-	FGlobalTabmanager::Get()->TryInvokeTab(FTabId("ProjectLauncher"));
-}
-
-void FInternalPlayWorldCommandCallbacks::OpenDeviceManager_Clicked()
-{
-	FGlobalTabmanager::Get()->TryInvokeTab(FTabId("DeviceManager"));
-}
-
-void FInternalPlayWorldCommandCallbacks::RepeatLastLaunch_Clicked()
-{
-	const ULevelEditorPlaySettings* PlaySettings = GetDefault<ULevelEditorPlaySettings>();
-
-	switch (PlaySettings->LastExecutedLaunchModeType)
-	{
-	case LaunchMode_OnDevice:
-		if (FParse::Param(FCommandLine::Get(), TEXT("OldLaunchOn")))
-		{
-			if (IsReadyToLaunchOnDevice(PlaySettings->LastExecutedLaunchDevice))
-			{
-				LaunchOnDevice(PlaySettings->LastExecutedLaunchDevice, PlaySettings->LastExecutedLaunchName, false);
-			}
-		}
-		else
-		{
-			// @todo turnkey - remove this function completely and all the platform support, and put it into Turnkey
-//			if (IsReadyToLaunchOnDevice(PlaySettings->LastExecutedLaunchDevice))
-			{
-				LaunchOnDevice(PlaySettings->LastExecutedLaunchDevice, PlaySettings->LastExecutedLaunchName, true);
-			}
-		}
-
-
-		break;
-
-	default:
-		break;
-	}
-}
-
-
-bool FInternalPlayWorldCommandCallbacks::RepeatLastLaunch_CanExecute()
-{
-	const ULevelEditorPlaySettings* PlaySettings = GetDefault<ULevelEditorPlaySettings>();
-
-	switch (PlaySettings->LastExecutedLaunchModeType)
-	{
-	case LaunchMode_OnDevice:
-		return CanLaunchOnDevice(PlaySettings->LastExecutedLaunchName);
-
-	default:
-		return false;
-	}
-}
-
-
-FText FInternalPlayWorldCommandCallbacks::GetRepeatLastLaunchToolTip()
-{
-	const ULevelEditorPlaySettings* PlaySettings = GetDefault<ULevelEditorPlaySettings>();
-
-	switch (PlaySettings->LastExecutedLaunchModeType)
-	{
-	case LaunchMode_OnDevice:
-		if (CanLaunchOnDevice(PlaySettings->LastExecutedLaunchName))
-		{
-			FFormatNamedArguments Arguments;
-			Arguments.Add(TEXT("DeviceName"), FText::FromString(PlaySettings->LastExecutedLaunchName));
-
-			return FText::Format(LOCTEXT("RepeatLaunchTooltip", "Launch this level on {DeviceName}"), Arguments);
-		}
-
-		break;
-
-	default:
-		break;
-	}
-
-	return LOCTEXT("RepeatLaunchSelectOptionToolTip", "Select a play-on target from the combo menu");
-}
-
-
-FSlateIcon FInternalPlayWorldCommandCallbacks::GetRepeatLastLaunchIcon()
-{
-	const ULevelEditorPlaySettings* PlaySettings = GetDefault<ULevelEditorPlaySettings>();
-
-	// @todo gmp: add play mode specific icons
-	switch (PlaySettings->LastExecutedLaunchModeType)
-	{
-	case LaunchMode_OnDevice:
-		break;
-	}
-
-	static FName RepeatLastLaunchIcon("PlayWorld.RepeatLastLaunch");
-
-	return FSlateIcon(FEditorStyle::GetStyleSetName(), RepeatLastLaunchIcon);
-}
-
-bool FInternalPlayWorldCommandCallbacks::IsReadyToLaunchOnDevice(FString DeviceId)
-{
-	int32 Index = 0;
-	DeviceId.FindChar(TEXT('@'), Index);
-	FString PlatformName = DeviceId.Left(Index);
-
-	const PlatformInfo::FTargetPlatformInfo* const PlatformInfo = PlatformInfo::FindPlatformInfo(FName(*PlatformName));
-	checkf(PlatformInfo, TEXT("Unable to find PlatformInfo for %s"), *PlatformName);
-
-	FGameProjectGenerationModule& GameProjectModule = FModuleManager::LoadModuleChecked<FGameProjectGenerationModule>(TEXT("GameProjectGeneration"));
-	bool bHasCode = GameProjectModule.Get().ProjectHasCodeFiles();
-
-	if (PlatformInfo->DataDrivenPlatformInfo->GetSdkStatus() != DDPIPlatformSdkStatus::Valid)
-	{
-		IMainFrameModule& MainFrameModule = FModuleManager::GetModuleChecked<IMainFrameModule>(TEXT("MainFrame"));
-		MainFrameModule.BroadcastMainFrameSDKNotInstalled(PlatformInfo->Name.ToString(), PlatformInfo->DataDrivenPlatformInfo->SDKTutorial);
-		TArray<FAnalyticsEventAttribute> ParamArray;
-		ParamArray.Add(FAnalyticsEventAttribute(TEXT("Time"), 0.0));
-		FEditorAnalytics::ReportEvent(TEXT("Editor.LaunchOn.Failed"), PlatformInfo->Name.ToString(), bHasCode, EAnalyticsErrorCodes::SDKNotFound, ParamArray);
-		return false;
-	}
-
-	const ITargetPlatform* Platform = GetTargetPlatformManager()->FindTargetPlatform(PlatformName);
-	if (Platform)
-	{
-		FString NotInstalledTutorialLink;
-		FString DocumentationLink;
-		FText CustomizedLogMessage;
-
-		EBuildConfiguration BuildConfiguration = GetDefault<ULevelEditorPlaySettings>()->GetLaunchBuildConfiguration();
-		bool bEnableAssetNativization = false;
-		int32 Result = Platform->CheckRequirements(bHasCode, BuildConfiguration, bEnableAssetNativization, NotInstalledTutorialLink, DocumentationLink, CustomizedLogMessage);
-
-		// report to analytics
-		FEditorAnalytics::ReportBuildRequirementsFailure(TEXT("Editor.LaunchOn.Failed"), PlatformName, bHasCode, Result);
-
-		// report to message log
-		bool UnrecoverableError = false;
-
-		if ((Result & ETargetPlatformReadyStatus::SDKNotFound) != 0)
-		{
-			AddMessageLog(
-				LOCTEXT("SdkNotFoundMessage", "Software Development Kit (SDK) not found."),
-				CustomizedLogMessage.IsEmpty() ? FText::Format(LOCTEXT("SdkNotFoundMessageDetail", "Please install the SDK for the {0} target platform!"), Platform->DisplayName()) : CustomizedLogMessage,
-				NotInstalledTutorialLink,
-				DocumentationLink
-			);
-
-			UnrecoverableError = true;
-		}
-
-		if ((Result & ETargetPlatformReadyStatus::LicenseNotAccepted) != 0)
-		{
-			AddMessageLog(
-				LOCTEXT("LicenseNotAcceptedMessage", "License not accepted."),
-				CustomizedLogMessage.IsEmpty() ? LOCTEXT("LicenseNotAcceptedMessageDetail", "License must be accepted in project settings to deploy your app to the device.") : CustomizedLogMessage,
-				NotInstalledTutorialLink,
-				DocumentationLink
-			);
-
-			UnrecoverableError = true;
-		}
-
-		if ((Result & ETargetPlatformReadyStatus::ProvisionNotFound) != 0)
-		{
-			AddMessageLog(
-				LOCTEXT("ProvisionNotFoundMessage", "Provision not found."),
-				CustomizedLogMessage.IsEmpty() ? LOCTEXT("ProvisionNotFoundMessageDetail", "A provision is required for deploying your app to the device.") : CustomizedLogMessage,
-				NotInstalledTutorialLink,
-				DocumentationLink
-			);
-
-			UnrecoverableError = true;
-		}
-
-		if ((Result & ETargetPlatformReadyStatus::SigningKeyNotFound) != 0)
-		{
-			AddMessageLog(
-				LOCTEXT("SigningKeyNotFoundMessage", "Signing key not found."),
-				CustomizedLogMessage.IsEmpty() ? LOCTEXT("SigningKeyNotFoundMessageDetail", "The app could not be digitally signed, because the signing key is not configured.") : CustomizedLogMessage,
-				NotInstalledTutorialLink,
-				DocumentationLink
-			);
-
-			UnrecoverableError = true;
-		}
-
-		if ((Result & ETargetPlatformReadyStatus::ManifestNotFound) != 0)
-		{
-			AddMessageLog(
-				LOCTEXT("ManifestNotFound", "Manifest not found."),
-				CustomizedLogMessage.IsEmpty() ? LOCTEXT("ManifestNotFoundMessageDetail", "The generated application manifest could not be found.") : CustomizedLogMessage,
-				NotInstalledTutorialLink,
-				DocumentationLink
-			);
-
-			UnrecoverableError = true;
-		}
-
-		if ((Result & ETargetPlatformReadyStatus::RemoveServerNameEmpty) != 0
-			&& (bHasCode || (Result & ETargetPlatformReadyStatus::CodeBuildRequired)
-				|| (!FApp::GetEngineIsPromotedBuild() && !FApp::IsEngineInstalled())))
-		{
-			AddMessageLog(
-				LOCTEXT("RemoveServerNameNotFound", "Remote compiling requires a server name. "),
-				CustomizedLogMessage.IsEmpty() ? LOCTEXT("RemoveServerNameNotFoundDetail", "Please specify one in the Remote Server Name settings field.") : CustomizedLogMessage,
-				NotInstalledTutorialLink,
-				DocumentationLink
-			);
-			UnrecoverableError = true;
-		}
-
-		if (UnrecoverableError)
-		{
-			return false;
-		}
 
-		// report to main frame
-		if ((Result & ETargetPlatformReadyStatus::CodeUnsupported) != 0)
-		{
-			// show the message
-			FMessageDialog::Open(EAppMsgType::Ok, LOCTEXT("NotSupported_CodeBased", "Sorry, launching a code-based project for the selected platform is currently not supported. This feature may be available in a future release."));
-			return false;
-		}
-		if ((Result & ETargetPlatformReadyStatus::PluginsUnsupported) != 0)
-		{
-			// show the message
-			FMessageDialog::Open(EAppMsgType::Ok, LOCTEXT("NotSupported_Plugins", "Sorry, launching a project with third-party plugins is currently not supported for the selected platform. This feature may be available in a future release."));
-			return false;
-		}
-	}
-	else
-	{
-		IMainFrameModule& MainFrameModule = FModuleManager::GetModuleChecked<IMainFrameModule>(TEXT("MainFrame"));
-		MainFrameModule.BroadcastMainFrameSDKNotInstalled(PlatformInfo->Name.ToString(), PlatformInfo->DataDrivenPlatformInfo->SDKTutorial);
-		return false;
-	}
 
-	return true;
-}
 
 static void PrepareLaunchOn(FString DeviceId, FString DeviceName)
 {
@@ -3213,4 +2085,261 @@ EPlayModeLocations FInternalPlayWorldCommandCallbacks::GetPlayModeLocation()
 		: PlayLocation_CurrentCameraLocation;
 }
 
+
+
+
+
+
+
+
+
+
+
+
+
+
+
+
+
+/////// Will move these into TurnkeySupportModule soon:
+
+
+void FInternalPlayWorldCommandCallbacks::RepeatLastLaunch_Clicked()
+{
+	const ULevelEditorPlaySettings* PlaySettings = GetDefault<ULevelEditorPlaySettings>();
+
+	switch (PlaySettings->LastExecutedLaunchModeType)
+	{
+	case LaunchMode_OnDevice:
+		if (FParse::Param(FCommandLine::Get(), TEXT("OldLaunchOn")))
+		{
+			if (IsReadyToLaunchOnDevice(PlaySettings->LastExecutedLaunchDevice))
+			{
+				LaunchOnDevice(PlaySettings->LastExecutedLaunchDevice, PlaySettings->LastExecutedLaunchName, false);
+			}
+		}
+		else
+		{
+			// @todo turnkey - remove this function completely and all the platform support, and put it into Turnkey
+//			if (IsReadyToLaunchOnDevice(PlaySettings->LastExecutedLaunchDevice))
+			{
+				LaunchOnDevice(PlaySettings->LastExecutedLaunchDevice, PlaySettings->LastExecutedLaunchName, true);
+			}
+		}
+
+
+		break;
+
+	default:
+		break;
+	}
+}
+
+
+bool FInternalPlayWorldCommandCallbacks::RepeatLastLaunch_CanExecute()
+{
+	const ULevelEditorPlaySettings* PlaySettings = GetDefault<ULevelEditorPlaySettings>();
+
+	switch (PlaySettings->LastExecutedLaunchModeType)
+	{
+	case LaunchMode_OnDevice:
+		return CanLaunchOnDevice(PlaySettings->LastExecutedLaunchName);
+
+	default:
+		return false;
+	}
+}
+
+
+FText FInternalPlayWorldCommandCallbacks::GetRepeatLastLaunchToolTip()
+{
+	const ULevelEditorPlaySettings* PlaySettings = GetDefault<ULevelEditorPlaySettings>();
+
+	switch (PlaySettings->LastExecutedLaunchModeType)
+	{
+	case LaunchMode_OnDevice:
+		if (CanLaunchOnDevice(PlaySettings->LastExecutedLaunchName))
+		{
+			FFormatNamedArguments Arguments;
+			Arguments.Add(TEXT("DeviceName"), FText::FromString(PlaySettings->LastExecutedLaunchName));
+
+			return FText::Format(LOCTEXT("RepeatLaunchTooltip", "Launch this level on {DeviceName}"), Arguments);
+		}
+
+		break;
+
+	default:
+		break;
+	}
+
+	return LOCTEXT("RepeatLaunchSelectOptionToolTip", "Select a play-on target from the combo menu");
+}
+
+
+FSlateIcon FInternalPlayWorldCommandCallbacks::GetRepeatLastLaunchIcon()
+{
+	const ULevelEditorPlaySettings* PlaySettings = GetDefault<ULevelEditorPlaySettings>();
+
+	// @todo gmp: add play mode specific icons
+	switch (PlaySettings->LastExecutedLaunchModeType)
+	{
+	case LaunchMode_OnDevice:
+		break;
+	}
+
+	static FName RepeatLastLaunchIcon("PlayWorld.RepeatLastLaunch");
+
+	return FSlateIcon(FEditorStyle::GetStyleSetName(), RepeatLastLaunchIcon);
+}
+
+bool FInternalPlayWorldCommandCallbacks::IsReadyToLaunchOnDevice(FString DeviceId)
+{
+	int32 Index = 0;
+	DeviceId.FindChar(TEXT('@'), Index);
+	FString PlatformName = DeviceId.Left(Index);
+
+	const PlatformInfo::FTargetPlatformInfo* const PlatformInfo = PlatformInfo::FindPlatformInfo(FName(*PlatformName));
+	checkf(PlatformInfo, TEXT("Unable to find PlatformInfo for %s"), *PlatformName);
+
+	FGameProjectGenerationModule& GameProjectModule = FModuleManager::LoadModuleChecked<FGameProjectGenerationModule>(TEXT("GameProjectGeneration"));
+	bool bHasCode = GameProjectModule.Get().ProjectHasCodeFiles();
+
+	if (PlatformInfo->DataDrivenPlatformInfo->GetSdkStatus() != DDPIPlatformSdkStatus::Valid)
+	{
+		IMainFrameModule& MainFrameModule = FModuleManager::GetModuleChecked<IMainFrameModule>(TEXT("MainFrame"));
+		MainFrameModule.BroadcastMainFrameSDKNotInstalled(PlatformInfo->Name.ToString(), PlatformInfo->DataDrivenPlatformInfo->SDKTutorial);
+		TArray<FAnalyticsEventAttribute> ParamArray;
+		ParamArray.Add(FAnalyticsEventAttribute(TEXT("Time"), 0.0));
+		FEditorAnalytics::ReportEvent(TEXT("Editor.LaunchOn.Failed"), PlatformInfo->Name.ToString(), bHasCode, EAnalyticsErrorCodes::SDKNotFound, ParamArray);
+		return false;
+	}
+
+	const ITargetPlatform* Platform = GetTargetPlatformManager()->FindTargetPlatform(PlatformName);
+	if (Platform)
+	{
+		FString NotInstalledTutorialLink;
+		FString DocumentationLink;
+		FText CustomizedLogMessage;
+
+		EBuildConfiguration BuildConfiguration = GetDefault<ULevelEditorPlaySettings>()->GetLaunchBuildConfiguration();
+		bool bEnableAssetNativization = false;
+		int32 Result = Platform->CheckRequirements(bHasCode, BuildConfiguration, bEnableAssetNativization, NotInstalledTutorialLink, DocumentationLink, CustomizedLogMessage);
+
+		// report to analytics
+		FEditorAnalytics::ReportBuildRequirementsFailure(TEXT("Editor.LaunchOn.Failed"), PlatformName, bHasCode, Result);
+
+		// report to message log
+		bool UnrecoverableError = false;
+
+		if ((Result & ETargetPlatformReadyStatus::SDKNotFound) != 0)
+		{
+			AddMessageLog(
+				LOCTEXT("SdkNotFoundMessage", "Software Development Kit (SDK) not found."),
+				CustomizedLogMessage.IsEmpty() ? FText::Format(LOCTEXT("SdkNotFoundMessageDetail", "Please install the SDK for the {0} target platform!"), Platform->DisplayName()) : CustomizedLogMessage,
+				NotInstalledTutorialLink,
+				DocumentationLink
+			);
+
+			UnrecoverableError = true;
+		}
+
+		if ((Result & ETargetPlatformReadyStatus::LicenseNotAccepted) != 0)
+		{
+			AddMessageLog(
+				LOCTEXT("LicenseNotAcceptedMessage", "License not accepted."),
+				CustomizedLogMessage.IsEmpty() ? LOCTEXT("LicenseNotAcceptedMessageDetail", "License must be accepted in project settings to deploy your app to the device.") : CustomizedLogMessage,
+				NotInstalledTutorialLink,
+				DocumentationLink
+			);
+
+			UnrecoverableError = true;
+		}
+
+		if ((Result & ETargetPlatformReadyStatus::ProvisionNotFound) != 0)
+		{
+			AddMessageLog(
+				LOCTEXT("ProvisionNotFoundMessage", "Provision not found."),
+				CustomizedLogMessage.IsEmpty() ? LOCTEXT("ProvisionNotFoundMessageDetail", "A provision is required for deploying your app to the device.") : CustomizedLogMessage,
+				NotInstalledTutorialLink,
+				DocumentationLink
+			);
+
+			UnrecoverableError = true;
+		}
+
+		if ((Result & ETargetPlatformReadyStatus::SigningKeyNotFound) != 0)
+		{
+			AddMessageLog(
+				LOCTEXT("SigningKeyNotFoundMessage", "Signing key not found."),
+				CustomizedLogMessage.IsEmpty() ? LOCTEXT("SigningKeyNotFoundMessageDetail", "The app could not be digitally signed, because the signing key is not configured.") : CustomizedLogMessage,
+				NotInstalledTutorialLink,
+				DocumentationLink
+			);
+
+			UnrecoverableError = true;
+		}
+
+		if ((Result & ETargetPlatformReadyStatus::ManifestNotFound) != 0)
+		{
+			AddMessageLog(
+				LOCTEXT("ManifestNotFound", "Manifest not found."),
+				CustomizedLogMessage.IsEmpty() ? LOCTEXT("ManifestNotFoundMessageDetail", "The generated application manifest could not be found.") : CustomizedLogMessage,
+				NotInstalledTutorialLink,
+				DocumentationLink
+			);
+
+			UnrecoverableError = true;
+		}
+
+		if ((Result & ETargetPlatformReadyStatus::RemoveServerNameEmpty) != 0
+			&& (bHasCode || (Result & ETargetPlatformReadyStatus::CodeBuildRequired)
+				|| (!FApp::GetEngineIsPromotedBuild() && !FApp::IsEngineInstalled())))
+		{
+			AddMessageLog(
+				LOCTEXT("RemoveServerNameNotFound", "Remote compiling requires a server name. "),
+				CustomizedLogMessage.IsEmpty() ? LOCTEXT("RemoveServerNameNotFoundDetail", "Please specify one in the Remote Server Name settings field.") : CustomizedLogMessage,
+				NotInstalledTutorialLink,
+				DocumentationLink
+			);
+			UnrecoverableError = true;
+		}
+
+		if (UnrecoverableError)
+		{
+			return false;
+		}
+
+		// report to main frame
+		if ((Result & ETargetPlatformReadyStatus::CodeUnsupported) != 0)
+		{
+			// show the message
+			FMessageDialog::Open(EAppMsgType::Ok, LOCTEXT("NotSupported_CodeBased", "Sorry, launching a code-based project for the selected platform is currently not supported. This feature may be available in a future release."));
+			return false;
+		}
+		if ((Result & ETargetPlatformReadyStatus::PluginsUnsupported) != 0)
+		{
+			// show the message
+			FMessageDialog::Open(EAppMsgType::Ok, LOCTEXT("NotSupported_Plugins", "Sorry, launching a project with third-party plugins is currently not supported for the selected platform. This feature may be available in a future release."));
+			return false;
+		}
+	}
+	else
+	{
+		IMainFrameModule& MainFrameModule = FModuleManager::GetModuleChecked<IMainFrameModule>(TEXT("MainFrame"));
+		MainFrameModule.BroadcastMainFrameSDKNotInstalled(PlatformInfo->Name.ToString(), PlatformInfo->DataDrivenPlatformInfo->SDKTutorial);
+		return false;
+	}
+
+	return true;
+}
+
+
+
+
+
+
+
+
+
 #undef LOCTEXT_NAMESPACE
+
diff --git a/Engine/Source/Editor/UnrealEd/Private/Settings/ProjectPackagingSettingsCustomization.h b/Engine/Source/Editor/UnrealEd/Private/Settings/ProjectPackagingSettingsCustomization.h
--- a/Engine/Source/Editor/UnrealEd/Private/Settings/ProjectPackagingSettingsCustomization.h
+++ b/Engine/Source/Editor/UnrealEd/Private/Settings/ProjectPackagingSettingsCustomization.h
@@ -150,7 +150,7 @@ protected:
 		TSharedRef<FPropertyRestriction> BuildConfigurationRestriction = MakeShareable(new FPropertyRestriction(LOCTEXT("ConfigurationRestrictionReason", "This configuration is not valid for this project. DebugGame configurations require a code project, and client/server configurations require the appropriate targets.")));
 
 		const UEnum* const ProjectPackagingBuildConfigurationsEnum = StaticEnum<EProjectPackagingBuildConfigurations>();
-		for (int Idx = 0; Idx < PPBC_MAX; Idx++)
+		for (int Idx = 0; Idx < (int)EProjectPackagingBuildConfigurations::PPBC_MAX; Idx++)
 		{
 			EProjectPackagingBuildConfigurations Configuration = (EProjectPackagingBuildConfigurations)Idx;
 			if (!PackagingConfigurations.Contains(Configuration))
diff --git a/Engine/Source/Editor/UnrealEd/Private/Settings/SettingsClasses.cpp b/Engine/Source/Editor/UnrealEd/Private/Settings/SettingsClasses.cpp
--- a/Engine/Source/Editor/UnrealEd/Private/Settings/SettingsClasses.cpp
+++ b/Engine/Source/Editor/UnrealEd/Private/Settings/SettingsClasses.cpp
@@ -1005,7 +1005,7 @@ void ULevelEditorViewportSettings::PostEditChangeProperty( struct FPropertyChang
 /* UProjectPackagingSettings interface
  *****************************************************************************/
 
-const UProjectPackagingSettings::FConfigurationInfo UProjectPackagingSettings::ConfigurationInfo[PPBC_MAX] = 
+const UProjectPackagingSettings::FConfigurationInfo UProjectPackagingSettings::ConfigurationInfo[(int)EProjectPackagingBuildConfigurations::PPBC_MAX] =
 {
 	/* PPBC_Debug */         { EBuildConfiguration::Debug, LOCTEXT("DebugConfiguration", "Debug"), LOCTEXT("DebugConfigurationTooltip", "Package the game in Debug configuration") },
 	/* PPBC_DebugGame */     { EBuildConfiguration::DebugGame, LOCTEXT("DebugGameConfiguration", "DebugGame"), LOCTEXT("DebugGameConfigurationTooltip", "Package the game in DebugGame configuration") },
@@ -1304,7 +1304,7 @@ TArray<EProjectPackagingBuildConfigurations> UProjectPackagingSettings::GetValid
 
 	// Set up all the configurations
 	TArray<EProjectPackagingBuildConfigurations> Configurations;
-	for (int32 Idx = 0; Idx < PPBC_MAX; Idx++)
+	for (int32 Idx = 0; Idx < (int)EProjectPackagingBuildConfigurations::PPBC_MAX; Idx++)
 	{
 		EProjectPackagingBuildConfigurations PackagingConfiguration = (EProjectPackagingBuildConfigurations)Idx;
 
@@ -1361,6 +1361,32 @@ int32 UProjectPackagingSettings::FindBlueprintInNativizationList(const UBlueprin
 	return ListIndex;
 }
 
+EProjectPackagingBuildConfigurations UProjectPackagingSettings::GetBuildConfigurationForPlatform(FName PlatformName) const
+{
+	const EProjectPackagingBuildConfigurations* Value = PerPlatformBuildConfig.Find(PlatformName);
+
+	return Value == nullptr ? EProjectPackagingBuildConfigurations::PPBC_Development : *Value;
+}
+
+void UProjectPackagingSettings::SetBuildConfigurationForPlatform(FName PlatformName, EProjectPackagingBuildConfigurations Configuration)
+{
+	PerPlatformBuildConfig.Add(PlatformName, Configuration);
+}
+
+FName UProjectPackagingSettings::GetTargetPlatformForPlatform(FName PlatformName) const
+{
+	const FName* Value = PerPlatformTargetPlatformName.Find(PlatformName);
+
+	// the platform name is also the name of the vanilla info
+	return Value == nullptr ? PlatformName : *Value;
+}
+
+void UProjectPackagingSettings::SetTargetPlatformForPlatform(FName PlatformName, FName TargetPlatformName)
+{
+	PerPlatformTargetPlatformName.Add(PlatformName, TargetPlatformName);
+}
+
+
 /* UCrashReporterSettings interface
 *****************************************************************************/
 UCrashReporterSettings::UCrashReporterSettings(const FObjectInitializer& ObjectInitializer)
diff --git a/Engine/Source/Editor/UnrealEd/Public/Kismet2/DebuggerCommands.h b/Engine/Source/Editor/UnrealEd/Public/Kismet2/DebuggerCommands.h
--- a/Engine/Source/Editor/UnrealEd/Public/Kismet2/DebuggerCommands.h
+++ b/Engine/Source/Editor/UnrealEd/Public/Kismet2/DebuggerCommands.h
@@ -110,8 +110,6 @@ public:
 
 	/** Launch on device */
 	TSharedPtr<FUICommandInfo> RepeatLastLaunch;
-	TSharedPtr<FUICommandInfo> OpenProjectLauncher;
-	TSharedPtr<FUICommandInfo> OpenDeviceManager;
 
 protected:
 
@@ -125,20 +123,6 @@ protected:
 	 */
 	static TSharedRef< SWidget > GeneratePlayMenuContent( TSharedRef<FUICommandList> InCommandList );
 
-	/**
-	 * Generates menu content for the Play On combo button drop down menu
-	 *
-	 * @return	Menu content widget
-	 */
-	static TSharedRef< SWidget > GenerateLaunchMenuContent( TSharedRef<FUICommandList> InCommandList );
-
-	/**
-	 * Generates menu content for the Play On combo button drop down menu
-	 *
-	 * @return	Menu content widget
-	 */
-	static TSharedRef< SWidget > GenerateTurnkeyMenuContent(TSharedRef<FUICommandList> InCommandList);
-
 	// Add mobile PIE preview device commands
 	void AddPIEPreviewDeviceCommands();
 
diff --git a/Engine/Source/Editor/UnrealEd/UnrealEd.Build.cs b/Engine/Source/Editor/UnrealEd/UnrealEd.Build.cs
--- a/Engine/Source/Editor/UnrealEd/UnrealEd.Build.cs
+++ b/Engine/Source/Editor/UnrealEd/UnrealEd.Build.cs
@@ -41,6 +41,7 @@ public class UnrealEd : ModuleRules
 				"ProjectTargetPlatformEditor",
 				"ImageWrapper",
 				"MainFrame",
+				"TurnkeySupport",
 				"MaterialEditor",
 				"MergeActors",
 				"MeshUtilities",
@@ -216,6 +217,7 @@ public class UnrealEd : ModuleRules
 				"DetailCustomizations",
 				"ComponentVisualizers",
 				"MainFrame",
+				"TurnkeySupport",
 				"LevelEditor",
 				"PackagesDialog",
 				"Persona",
