World Partition - Only create the UWorldPartitionSubsystem if the world contains a world partition
* Conditionally create the subsystem using UWorldPartitionSubsystem::ShouldCreateSubsystem()
* Removed UWorldPartitionSubsystem::IsEnabled()
* Added UWorld::HasSubsystem()

diff --git a/Engine/Source/Editor/DetailCustomizations/Private/ActorDetails.cpp b/Engine/Source/Editor/DetailCustomizations/Private/ActorDetails.cpp
--- a/Engine/Source/Editor/DetailCustomizations/Private/ActorDetails.cpp
+++ b/Engine/Source/Editor/DetailCustomizations/Private/ActorDetails.cpp
@@ -556,8 +556,7 @@ void FActorDetails::AddActorCategory( IDetailLayoutBuilder& DetailBuilder, const
 	if (SelectedActorInfo.SelectionClass != AWorldSettings::StaticClass())
 	{
 		// Actor Packaging Mode
-		const UWorldPartitionSubsystem* WorldPartitionSubsystem = SelectedActorInfo.SharedWorld ? SelectedActorInfo.SharedWorld->GetSubsystem<UWorldPartitionSubsystem>() : nullptr;
-		const bool bIsPartitionedWorld = WorldPartitionSubsystem && WorldPartitionSubsystem->IsEnabled();
+		const bool bIsPartitionedWorld = UWorld::HasSubsystem<UWorldPartitionSubsystem>(SelectedActorInfo.SharedWorld);
 
 		auto OnGetMenuContent = [=]() -> TSharedRef<SWidget> {
 			FMenuBuilder MenuBuilder(true, nullptr);
diff --git a/Engine/Source/Editor/DetailCustomizations/Private/WorldSettingsDetails.cpp b/Engine/Source/Editor/DetailCustomizations/Private/WorldSettingsDetails.cpp
--- a/Engine/Source/Editor/DetailCustomizations/Private/WorldSettingsDetails.cpp
+++ b/Engine/Source/Editor/DetailCustomizations/Private/WorldSettingsDetails.cpp
@@ -93,8 +93,7 @@ void FWorldSettingsDetails::AddLevelExternalActorsCustomization(IDetailLayoutBui
 
 	if (CustomizedLevel)
 	{
-		const UWorldPartitionSubsystem* WorldPartitionSubsystem = CustomizedLevel->GetWorld()->GetSubsystem<UWorldPartitionSubsystem>();
-		const bool bIsPartitionedWorld = WorldPartitionSubsystem && WorldPartitionSubsystem->IsEnabled();
+		const bool bIsPartitionedWorld = UWorld::HasSubsystem<UWorldPartitionSubsystem>(CustomizedLevel->GetWorld());
 
 		IDetailCategoryBuilder& WorldCategory = DetailBuilder.EditCategory("World");
 		WorldCategory.AddCustomRow(LOCTEXT("LevelUseExternalActorsRow", "LevelUseExternalActors"), true)
diff --git a/Engine/Source/Editor/FoliageEdit/Private/FoliageEdMode.cpp b/Engine/Source/Editor/FoliageEdit/Private/FoliageEdMode.cpp
--- a/Engine/Source/Editor/FoliageEdit/Private/FoliageEdMode.cpp
+++ b/Engine/Source/Editor/FoliageEdit/Private/FoliageEdMode.cpp
@@ -1999,12 +1999,9 @@ void FEdModeFoliage::RemoveSelectedInstances(UWorld* InWorld)
 void FEdModeFoliage::GetFoliageTypeFilters(TArray<const UClass*>& OutFilters) const
 {
 	OutFilters.Add(UFoliageType_InstancedStaticMesh::StaticClass());
-	if (UWorldPartitionSubsystem* WorldPartition = GetWorld()->GetSubsystem<UWorldPartitionSubsystem>())
+	if (GetWorld()->HasSubsystem<UWorldPartitionSubsystem>())
 	{
-		if (WorldPartition->IsEnabled())
-		{
-			return;
-		}
+		return;
 	}
 	// FoliageType Actor only supported in non WorldPartition worlds
 	OutFilters.Add(UFoliageType_Actor::StaticClass());
diff --git a/Engine/Source/Editor/UnrealEd/Private/ActorGroupingUtils.cpp b/Engine/Source/Editor/UnrealEd/Private/ActorGroupingUtils.cpp
--- a/Engine/Source/Editor/UnrealEd/Private/ActorGroupingUtils.cpp
+++ b/Engine/Source/Editor/UnrealEd/Private/ActorGroupingUtils.cpp
@@ -108,8 +108,7 @@ void UActorGroupingUtils::GroupActors(const TArray<AActor*>& ActorsToGroup)
 					SpawnedGroupActor->CenterGroupLocation();
 					SpawnedGroupActor->Lock();
 
-					UWorldPartitionSubsystem* WorldPartitionSubsystem = World->GetSubsystem<UWorldPartitionSubsystem>();
-					if (WorldPartitionSubsystem && WorldPartitionSubsystem->IsEnabled())
+					if (UWorldPartitionSubsystem* WorldPartitionSubsystem = World->GetSubsystem<UWorldPartitionSubsystem>())
 					{
 						WorldPartitionSubsystem->UpdateActorDesc(SpawnedGroupActor);
 					}
@@ -175,8 +174,7 @@ void UActorGroupingUtils::UngroupActors(const TArray<AActor*>& ActorsToUngroup)
 				AGroupActor* GroupActor = OutermostGroupActors[GroupIndex];
 				GroupActor->ClearAndRemove();
 
-				UWorldPartitionSubsystem* WorldPartitionSubsystem = GroupActor->GetWorld()->GetSubsystem<UWorldPartitionSubsystem>();
-				if (WorldPartitionSubsystem && WorldPartitionSubsystem->IsEnabled())
+				if (UWorldPartitionSubsystem* WorldPartitionSubsystem = GroupActor->GetWorld()->GetSubsystem<UWorldPartitionSubsystem>())
 				{
 					WorldPartitionSubsystem->UpdateActorDesc(GroupActor);
 				}
diff --git a/Engine/Source/Editor/UnrealEd/Private/Commandlets/WorldPartitionBuilderCommandlet.cpp b/Engine/Source/Editor/UnrealEd/Private/Commandlets/WorldPartitionBuilderCommandlet.cpp
--- a/Engine/Source/Editor/UnrealEd/Private/Commandlets/WorldPartitionBuilderCommandlet.cpp
+++ b/Engine/Source/Editor/UnrealEd/Private/Commandlets/WorldPartitionBuilderCommandlet.cpp
@@ -87,15 +87,17 @@ int32 UWorldPartitionBuilderCommandlet::Main(const FString& Params)
 		World->UpdateWorldComponents(true, false);
 	}
 
-	// Retrieve the world partition.
-	UWorldPartitionSubsystem* WorldPartitionSubsystem = World->GetSubsystem<UWorldPartitionSubsystem>();
-	UWorldPartition* WorldPartition = WorldPartitionSubsystem && WorldPartitionSubsystem->IsEnabled() ? World->GetWorldPartition() : nullptr;
-	if (!WorldPartition)
+	// Make sure the world is partitioned.
+	if (!World->HasSubsystem<UWorldPartitionSubsystem>())
 	{
 		UE_LOG(LogWorldPartitionBuilderCommandlet, Error, TEXT("Commandlet only works on partitioned maps."));
 		return 1;
 	}
 
+	// Retrieve the world partition.
+	UWorldPartition* WorldPartition = World->GetWorldPartition();
+	check(WorldPartition);
+
 	FWorldContext& WorldContext = GEditor->GetEditorWorldContext(true);
 	WorldContext.SetCurrentWorld(World);
 	GWorld = World;
diff --git a/Engine/Source/Editor/UnrealEd/Private/WorldPartition/WorldPartitionHLODsBuilder.cpp b/Engine/Source/Editor/UnrealEd/Private/WorldPartition/WorldPartitionHLODsBuilder.cpp
--- a/Engine/Source/Editor/UnrealEd/Private/WorldPartition/WorldPartitionHLODsBuilder.cpp
+++ b/Engine/Source/Editor/UnrealEd/Private/WorldPartition/WorldPartitionHLODsBuilder.cpp
@@ -27,7 +27,6 @@ bool UWorldPartitionHLODsBuilder::Run(UWorld* World, FPackageSourceControlHelper
 {
 	UWorldPartitionSubsystem* WorldPartitionSubsystem = World->GetSubsystem<UWorldPartitionSubsystem>();
 	check(WorldPartitionSubsystem);
-	check(WorldPartitionSubsystem->IsEnabled());
 
 	UWorldPartition* WorldPartition = World->GetWorldPartition();
 	check(WorldPartition);
diff --git a/Engine/Source/Runtime/Engine/Classes/Engine/World.h b/Engine/Source/Runtime/Engine/Classes/Engine/World.h
--- a/Engine/Source/Runtime/Engine/Classes/Engine/World.h
+++ b/Engine/Source/Runtime/Engine/Classes/Engine/World.h
@@ -3606,6 +3606,25 @@ public:
 		return nullptr;
 	}
 
+	/**
+	 * Check if world has a subsystem of the specified type
+	 */
+	template <typename TSubsystemClass>
+	bool HasSubsystem() const
+	{
+		return GetSubsystem<TSubsystemClass>() != nullptr;
+	}
+
+	/**
+	 * Check if world has a subsystem of the specified type from the provided GameInstance
+	 * returns false if the Subsystem cannot be found or the GameInstance is null
+	 */
+	template <typename TSubsystemClass>
+	static FORCEINLINE bool HasSubsystem(const UWorld* World)
+	{
+		return GetSubsystem<TSubsystemClass>(World) != nullptr;
+	}
+
 	/**
 	 * Get all Subsystem of specified type, this is only necessary for interfaces that can have multiple implementations instanced at a time.
 	 *
diff --git a/Engine/Source/Runtime/Engine/Private/ActorEditor.cpp b/Engine/Source/Runtime/Engine/Private/ActorEditor.cpp
--- a/Engine/Source/Runtime/Engine/Private/ActorEditor.cpp
+++ b/Engine/Source/Runtime/Engine/Private/ActorEditor.cpp
@@ -61,8 +61,7 @@ bool AActor::CanEditChange(const FProperty* PropertyThatWillChange) const
 	{
 		if (!IsTemplate())
 		{
-			const UWorldPartitionSubsystem* WorldPartitionSubsystem = GetWorld() ? GetWorld()->GetSubsystem<UWorldPartitionSubsystem>() : nullptr;
-			const bool bIsPartitionedWorld = WorldPartitionSubsystem && WorldPartitionSubsystem->IsEnabled();
+			const bool bIsPartitionedWorld = UWorld::HasSubsystem<UWorldPartitionSubsystem>(GetWorld());
 			if (!bIsPartitionedWorld)
 			{
 				return false;
@@ -594,13 +593,9 @@ bool AActor::InternalPostEditUndo()
 
 	if (!IsTemplate())
 	{
-		if (UWorld* World = GetWorld())
+		if (UWorldPartitionSubsystem* WorldPartitionSubsystem = UWorld::GetSubsystem<UWorldPartitionSubsystem>(GetWorld()))
 		{
-			UWorldPartitionSubsystem* WorldPartitionSubsystem = World->GetSubsystem<UWorldPartitionSubsystem>();
-			if (WorldPartitionSubsystem && WorldPartitionSubsystem->IsEnabled())
-			{
-				WorldPartitionSubsystem->UpdateActorDesc(this);
-			}
+			WorldPartitionSubsystem->UpdateActorDesc(this);
 		}
 	}
 
diff --git a/Engine/Source/Runtime/Engine/Private/ActorPartition/ActorPartitionSubsystem.cpp b/Engine/Source/Runtime/Engine/Private/ActorPartition/ActorPartitionSubsystem.cpp
--- a/Engine/Source/Runtime/Engine/Private/ActorPartition/ActorPartitionSubsystem.cpp
+++ b/Engine/Source/Runtime/Engine/Private/ActorPartition/ActorPartitionSubsystem.cpp
@@ -201,8 +201,7 @@ UActorPartitionSubsystem::UActorPartitionSubsystem()
 
 bool UActorPartitionSubsystem::IsLevelPartition() const
 {
-	UWorldPartitionSubsystem* WorldPartitionSubsystem = GetWorld()->GetSubsystem<UWorldPartitionSubsystem>();
-	return WorldPartitionSubsystem == nullptr || !WorldPartitionSubsystem->IsEnabled();
+	return !GetWorld()->HasSubsystem<UWorldPartitionSubsystem>();
 }
 
 #if WITH_EDITOR
@@ -214,9 +213,8 @@ void UActorPartitionSubsystem::InitializeForWorldPartitionConversion()
 
 void UActorPartitionSubsystem::Initialize(FSubsystemCollectionBase& Collection)
 {
-	Collection.InitializeDependency(UWorldPartitionSubsystem::StaticClass());
-
-	if (UWorldPartitionSubsystem* WorldPartitionSubsystem = GetWorld()->GetSubsystem<UWorldPartitionSubsystem>())
+	UWorldPartitionSubsystem* WorldPartitionSubsystem = Collection.InitializeDependency<UWorldPartitionSubsystem>();
+	if (WorldPartitionSubsystem)
 	{
 		WorldPartitionSubsystem->RegisterActorDescFactory(APartitionActor::StaticClass(), &PartitionActorDescFactory);
 	}
diff --git a/Engine/Source/Runtime/Engine/Private/LevelInstance/LevelInstanceActor.cpp b/Engine/Source/Runtime/Engine/Private/LevelInstance/LevelInstanceActor.cpp
--- a/Engine/Source/Runtime/Engine/Private/LevelInstance/LevelInstanceActor.cpp
+++ b/Engine/Source/Runtime/Engine/Private/LevelInstance/LevelInstanceActor.cpp
@@ -550,8 +550,7 @@ void ALevelInstance::OnLevelInstanceLoaded()
 			SetIsTemporarilyHiddenInEditor(bHiddenInEditor);
 		}
 
-		UWorldPartitionSubsystem* WorldPartitionSubsystem = GetWorld()->GetSubsystem<UWorldPartitionSubsystem>();
-		if (WorldPartitionSubsystem && WorldPartitionSubsystem->IsEnabled())
+		if (UWorldPartitionSubsystem* WorldPartitionSubsystem = GetWorld()->GetSubsystem<UWorldPartitionSubsystem>())
 		{
 			WorldPartitionSubsystem->UpdateActorDesc(this);
 		}
diff --git a/Engine/Source/Runtime/Engine/Private/LevelInstance/LevelInstanceSubsystem.cpp b/Engine/Source/Runtime/Engine/Private/LevelInstance/LevelInstanceSubsystem.cpp
--- a/Engine/Source/Runtime/Engine/Private/LevelInstance/LevelInstanceSubsystem.cpp
+++ b/Engine/Source/Runtime/Engine/Private/LevelInstance/LevelInstanceSubsystem.cpp
@@ -55,8 +55,8 @@ void ULevelInstanceSubsystem::Initialize(FSubsystemCollectionBase& Collection)
 #if WITH_EDITOR
 	LevelInstanceActorDescFactory.Reset(new FLevelInstanceActorDescFactory());
 
-	Collection.InitializeDependency(UWorldPartitionSubsystem::StaticClass());
-	if (UWorldPartitionSubsystem* WorldPartitionSubsystem = GetWorld()->GetSubsystem<UWorldPartitionSubsystem>())
+	UWorldPartitionSubsystem* WorldPartitionSubsystem = Collection.InitializeDependency<UWorldPartitionSubsystem>();
+	if (WorldPartitionSubsystem)
 	{
 		RegisterActorDescFactories(WorldPartitionSubsystem);
 	}
diff --git a/Engine/Source/Runtime/Engine/Private/World.cpp b/Engine/Source/Runtime/Engine/Private/World.cpp
--- a/Engine/Source/Runtime/Engine/Private/World.cpp
+++ b/Engine/Source/Runtime/Engine/Private/World.cpp
@@ -7706,8 +7706,7 @@ void UWorld::GetAssetRegistryTags(TArray<FAssetRegistryTag>& OutTags) const
 
 		FBox LevelBounds;
 
-		UWorldPartitionSubsystem* WorldPartitionSubsystem = GetSubsystem<UWorldPartitionSubsystem>();
-		if (WorldPartitionSubsystem && WorldPartitionSubsystem->IsEnabled())
+		if (UWorldPartitionSubsystem* WorldPartitionSubsystem = GetSubsystem<UWorldPartitionSubsystem>())
 		{
 			LevelBounds = WorldPartitionSubsystem->GetWorldBounds();
 		}
diff --git a/Engine/Source/Runtime/Engine/Private/WorldPartition/HLOD/HLODSubsystem.cpp b/Engine/Source/Runtime/Engine/Private/WorldPartition/HLOD/HLODSubsystem.cpp
--- a/Engine/Source/Runtime/Engine/Private/WorldPartition/HLOD/HLODSubsystem.cpp
+++ b/Engine/Source/Runtime/Engine/Private/WorldPartition/HLOD/HLODSubsystem.cpp
@@ -26,6 +26,11 @@ UHLODSubsystem::~UHLODSubsystem()
 
 bool UHLODSubsystem::ShouldCreateSubsystem(UObject* Outer) const
 {
+	if (!Super::ShouldCreateSubsystem(Outer))
+	{
+		return false;
+	}
+
 	if (UWorld* WorldOuter = Cast<UWorld>(Outer))
 	{
 		return WorldOuter->GetWorldPartition() != nullptr;
diff --git a/Engine/Source/Runtime/Engine/Private/WorldPartition/WorldPartition.cpp b/Engine/Source/Runtime/Engine/Private/WorldPartition/WorldPartition.cpp
--- a/Engine/Source/Runtime/Engine/Private/WorldPartition/WorldPartition.cpp
+++ b/Engine/Source/Runtime/Engine/Private/WorldPartition/WorldPartition.cpp
@@ -324,7 +324,8 @@ void UWorldPartition::Initialize(UWorld* InWorld, const FTransform& InTransform)
 
 	InitState = EWorldPartitionInitState::Initialized;
 
-	if (UWorldPartitionSubsystem* WorldPartitionSubsystem = World->GetSubsystem<UWorldPartitionSubsystem>())
+	UWorldPartitionSubsystem* WorldPartitionSubsystem = World->GetSubsystem<UWorldPartitionSubsystem>();
+	if (ensure(WorldPartitionSubsystem))
 	{
 		WorldPartitionSubsystem->RegisterWorldPartition(this);
 	}
@@ -394,7 +395,8 @@ void UWorldPartition::Uninitialize()
 		ActorToActorCluster.Empty();
 #endif		
 
-		if (UWorldPartitionSubsystem* WorldPartitionSubsystem = World->GetSubsystem<UWorldPartitionSubsystem>())
+		UWorldPartitionSubsystem* WorldPartitionSubsystem = World->GetSubsystem<UWorldPartitionSubsystem>();
+		if (ensure(WorldPartitionSubsystem))
 		{
 			WorldPartitionSubsystem->UnregisterWorldPartition(this);
 		}
diff --git a/Engine/Source/Runtime/Engine/Private/WorldPartition/WorldPartitionSubsystem.cpp b/Engine/Source/Runtime/Engine/Private/WorldPartition/WorldPartitionSubsystem.cpp
--- a/Engine/Source/Runtime/Engine/Private/WorldPartition/WorldPartitionSubsystem.cpp
+++ b/Engine/Source/Runtime/Engine/Private/WorldPartition/WorldPartitionSubsystem.cpp
@@ -26,9 +26,19 @@ static FAutoConsoleCommand CVarDrawRuntimeHash2D(
 UWorldPartitionSubsystem::UWorldPartitionSubsystem()
 {}
 
-bool UWorldPartitionSubsystem::IsEnabled() const
+bool UWorldPartitionSubsystem::ShouldCreateSubsystem(UObject* Outer) const
 {
-	return GetMainWorldPartition() != nullptr;
+	if (!Super::ShouldCreateSubsystem(Outer))
+	{
+		return false;
+	}
+
+	if (UWorld* WorldOuter = Cast<UWorld>(Outer))
+	{
+		return WorldOuter->GetWorldPartition() != nullptr;
+	}
+
+	return false;
 }
 
 UWorldPartition* UWorldPartitionSubsystem::GetMainWorldPartition()
@@ -201,10 +211,7 @@ void UWorldPartitionSubsystem::UpdateActorDesc(AActor* Actor)
 
 void UWorldPartitionSubsystem::RegisterActorDescFactory(TSubclassOf<AActor> Class, FWorldPartitionActorDescFactory* Factory)
 {
-	if (IsEnabled())
-	{
-		UWorldPartition::RegisterActorDescFactory(Class, Factory);
-	}
+	UWorldPartition::RegisterActorDescFactory(Class, Factory);
 }
 
 FBox UWorldPartitionSubsystem::GetWorldBounds()
diff --git a/Engine/Source/Runtime/Engine/Public/WorldPartition/WorldPartitionSubsystem.h b/Engine/Source/Runtime/Engine/Public/WorldPartition/WorldPartitionSubsystem.h
--- a/Engine/Source/Runtime/Engine/Public/WorldPartition/WorldPartitionSubsystem.h
+++ b/Engine/Source/Runtime/Engine/Public/WorldPartition/WorldPartitionSubsystem.h
@@ -24,9 +24,8 @@ class ENGINE_API UWorldPartitionSubsystem : public UWorldSubsystem, public FTick
 public:
 	UWorldPartitionSubsystem();
 
-	bool IsEnabled() const;
-
 	//~ Begin USubsystem Interface.
+	virtual bool ShouldCreateSubsystem(UObject* Outer) const override;
 	virtual void Deinitialize() override;
 	//~ End USubsystem Interface.
 
