Fix level BP editors having no toolbar.

diff --git a/Engine/Source/Editor/Kismet/Private/BlueprintEditor.cpp b/Engine/Source/Editor/Kismet/Private/BlueprintEditor.cpp
--- a/Engine/Source/Editor/Kismet/Private/BlueprintEditor.cpp
+++ b/Engine/Source/Editor/Kismet/Private/BlueprintEditor.cpp
@@ -1962,8 +1962,10 @@ void FBlueprintEditor::InitBlueprintEditor(
 		// Flag the blueprint as having been opened
 		Blueprint->bIsNewlyCreated = false;
 
+		const bool bBaseBlueprintType = (Blueprint->GetClass() == UBlueprint::StaticClass() || Blueprint->GetClass() == ULevelScriptBlueprint::StaticClass());
+
 		// @todo reskin, this is temp for now.  Only the base blueprint editor has had a toolbar treatment as part of the reskin.  All other derived types of BP editors need to use their custom toolbars for now.
-		bUsesCustomToolbar &= (Blueprint->GetClass() != UBlueprint::StaticClass());
+		bUsesCustomToolbar &= (!bBaseBlueprintType);
 
 		Objects.Add( Blueprint );
 	}
diff --git a/Engine/Source/Editor/UnrealEd/Private/Toolkits/AssetEditorToolkit.cpp b/Engine/Source/Editor/UnrealEd/Private/Toolkits/AssetEditorToolkit.cpp
--- a/Engine/Source/Editor/UnrealEd/Private/Toolkits/AssetEditorToolkit.cpp
+++ b/Engine/Source/Editor/UnrealEd/Private/Toolkits/AssetEditorToolkit.cpp
@@ -313,10 +313,13 @@ void FAssetEditorToolkit::RegisterTabSpawners(const TSharedRef<class FTabManager
 	const auto& LocalCategories = InTabManager->GetLocalWorkspaceMenuRoot()->GetChildItems();
 	AssetEditorTabsCategory = LocalCategories.Num() > 0 ? LocalCategories[0] : InTabManager->GetLocalWorkspaceMenuRoot();
 
-	InTabManager->RegisterTabSpawner( ToolbarTabId, FOnSpawnTab::CreateSP(this, &FAssetEditorToolkit::SpawnTab_Toolbar) )
-		.SetDisplayName( LOCTEXT("ToolbarTab", "Toolbar") )
-		.SetGroup(AssetEditorTabsCategory.ToSharedRef())
-		.SetIcon( FSlateIcon(FEditorStyle::GetStyleSetName(), "Toolbar.Icon") );
+	if(UsesCustomToolbarPlacement())
+	{
+		InTabManager->RegisterTabSpawner(ToolbarTabId, FOnSpawnTab::CreateSP(this, &FAssetEditorToolkit::SpawnTab_Toolbar))
+			.SetDisplayName(LOCTEXT("ToolbarTab", "Toolbar"))
+			.SetGroup(AssetEditorTabsCategory.ToSharedRef())
+			.SetIcon(FSlateIcon(FEditorStyle::GetStyleSetName(), "Toolbar.Icon"));
+	}
 }
 
 void FAssetEditorToolkit::UnregisterTabSpawners(const TSharedRef<class FTabManager>& InTabManager)
