World Partition - Conversion cmdlet will now create the partition prior to initializing the world, ensuring the World Partition subsystem is properly created

diff --git a/Engine/Source/Editor/UnrealEd/Classes/Commandlets/WorldPartitionConvertCommandlet.h b/Engine/Source/Editor/UnrealEd/Classes/Commandlets/WorldPartitionConvertCommandlet.h
--- a/Engine/Source/Editor/UnrealEd/Classes/Commandlets/WorldPartitionConvertCommandlet.h
+++ b/Engine/Source/Editor/UnrealEd/Classes/Commandlets/WorldPartitionConvertCommandlet.h
@@ -53,7 +53,8 @@ protected:
 	virtual void ReadAdditionalTokensAndSwitches(const TArray<FString>& Tokens, const TArray<FString>& Switches) {}
 
 	UWorldPartition* CreateWorldPartition(class AWorldSettings* MainWorldSettings, UWorldComposition* WorldComposition) const;
-	ULevel* LoadLevel(const FString& LevelToLoad);
+	UWorld* LoadWorld(const FString& LevelToLoad);
+	ULevel* InitWorld(UWorld* World);
 
 	void ChangeObjectOuter(UObject* Object, UObject* NewOuter);
 	void FixupSoftObjectPaths(UPackage* OuterPackage);
diff --git a/Engine/Source/Editor/UnrealEd/Private/Commandlets/WorldPartitionConvertCommandlet.cpp b/Engine/Source/Editor/UnrealEd/Private/Commandlets/WorldPartitionConvertCommandlet.cpp
--- a/Engine/Source/Editor/UnrealEd/Private/Commandlets/WorldPartitionConvertCommandlet.cpp
+++ b/Engine/Source/Editor/UnrealEd/Private/Commandlets/WorldPartitionConvertCommandlet.cpp
@@ -143,7 +143,7 @@ UWorldPartitionConvertCommandlet::UWorldPartitionConvertCommandlet(const FObject
 , ConversionSuffix(TEXT("_WP"))
 {}
 
-ULevel* UWorldPartitionConvertCommandlet::LoadLevel(const FString& LevelToLoad)
+UWorld* UWorldPartitionConvertCommandlet::LoadWorld(const FString& LevelToLoad)
 {
 	SET_WARN_COLOR(COLOR_WHITE);
 	UE_LOG(LogWorldPartitionConvertCommandlet, Log, TEXT("Loading level %s."), *LevelToLoad);
@@ -156,39 +156,43 @@ ULevel* UWorldPartitionConvertCommandlet::LoadLevel(const FString& LevelToLoad)
 		return nullptr;
 	}
 
-	UWorld* World = UWorld::FindWorldInPackage(MapPackage);
-	if (World)
-	{
-		// Setup the world.
-		World->WorldType = EWorldType::Editor;
-		World->AddToRoot();
-		if (!World->bIsWorldInitialized)
-		{
-			UWorld::InitializationValues IVS;
-			IVS.RequiresHitProxies(false);
-			IVS.ShouldSimulatePhysics(false);
-			IVS.EnableTraceCollision(false);
-			IVS.CreateNavigation(false);
-			IVS.CreateAISystem(false);
-			IVS.AllowAudioPlayback(false);
-			IVS.CreatePhysicsScene(true);
-
-			World->InitWorld(IVS);
-			World->PersistentLevel->UpdateModelComponents();
-			World->UpdateWorldComponents(true, false);
-
-			World->FlushLevelStreaming(EFlushLevelStreamingType::Full);
-		}
+	return UWorld::FindWorldInPackage(MapPackage);
+}
+
+ULevel* UWorldPartitionConvertCommandlet::InitWorld(UWorld* World)
+{
+	SET_WARN_COLOR(COLOR_WHITE);
+	UE_LOG(LogWorldPartitionConvertCommandlet, Log, TEXT("Initializing level %s."), *World->GetName());
+	CLEAR_WARN_COLOR();
 
-		return World->PersistentLevel;
+	// Setup the world.
+	World->WorldType = EWorldType::Editor;
+	World->AddToRoot();
+	if (!World->bIsWorldInitialized)
+	{
+		UWorld::InitializationValues IVS;
+		IVS.RequiresHitProxies(false);
+		IVS.ShouldSimulatePhysics(false);
+		IVS.EnableTraceCollision(false);
+		IVS.CreateNavigation(false);
+		IVS.CreateAISystem(false);
+		IVS.AllowAudioPlayback(false);
+		IVS.CreatePhysicsScene(true);
+
+		World->InitWorld(IVS);
+		World->PersistentLevel->UpdateModelComponents();
+		World->UpdateWorldComponents(true, false);
+
+		World->FlushLevelStreaming(EFlushLevelStreamingType::Full);
 	}
 
-	return nullptr;
+	return World->PersistentLevel;
 }
 
 UWorldPartition* UWorldPartitionConvertCommandlet::CreateWorldPartition(AWorldSettings* MainWorldSettings, UWorldComposition* WorldComposition) const
 {
 	UWorldPartition* WorldPartition = NewObject<UWorldPartition>(MainWorldSettings);
+	MainWorldSettings->SetWorldPartition(WorldPartition);
 
 	WorldPartition->EditorHash = NewObject<UWorldPartitionEditorHash>(WorldPartition, EditorHashClass);
 	WorldPartition->RuntimeHash = NewObject<UWorldPartitionRuntimeHash>(WorldPartition, RuntimeHashClass);
@@ -732,26 +736,22 @@ int32 UWorldPartitionConvertCommandlet::Main(const FString& Params)
 		}
 	}
 
-	// Load persistent level
-	ULevel* MainLevel = LoadLevel(Tokens[0]);
-	if (!MainLevel)
+	// Load world
+	UWorld* MainWorld = LoadWorld(Tokens[0]);
+	if (!MainWorld)
 	{
-		UE_LOG(LogWorldPartitionConvertCommandlet, Error, TEXT("Unknown level '%s'"), *Tokens[0]);
+		UE_LOG(LogWorldPartitionConvertCommandlet, Error, TEXT("Unknown world '%s'"), *Tokens[0]);
 		return 1;
 	}
 
-	if (MainLevel->GetWorldPartition())
+	// Make sure the world isn't already partitionned
+	AWorldSettings* MainWorldSettings = MainWorld->GetWorldSettings();
+	if (MainWorldSettings->GetWorldPartition())
 	{
 		UE_LOG(LogWorldPartitionConvertCommandlet, Error, TEXT("Level '%s' is already partitionned"), *Tokens[0]);
 		return 1;
 	}
 
-	UWorld* MainWorld = MainLevel->GetWorld();
-	UPackage* MainPackage = MainLevel->GetPackage();
-	AWorldSettings* MainWorldSettings = MainWorld->GetWorldSettings();
-
-	OnWorldLoaded(MainWorld);
-
 	// Setup the world partition object
 	UWorldPartition* WorldPartition = CreateWorldPartition(MainWorldSettings, MainWorld->WorldComposition);
 	if (!WorldPartition)
@@ -759,10 +759,19 @@ int32 UWorldPartitionConvertCommandlet::Main(const FString& Params)
 		return 1;
 	}
 
-	// Make sure ActorPartitionSubsystem uses WorldPartition
+	// Initialize the world, create subsystems, etc.
+	ULevel* MainLevel = InitWorld(MainWorld);
+	if (!MainLevel)
+	{
+		UE_LOG(LogWorldPartitionConvertCommandlet, Error, TEXT("Unknown level '%s'"), *Tokens[0]);
+		return 1;
+	}
+
+	UPackage* MainPackage = MainLevel->GetPackage();
+
+	OnWorldLoaded(MainWorld);
+
 	UActorPartitionSubsystem* ActorPartitionSubsystem = MainWorld->GetSubsystem<UActorPartitionSubsystem>();
-	ActorPartitionSubsystem->InitializeForWorldPartitionConversion();
-	
 	auto PartitionFoliage = [this, MainWorld, ActorPartitionSubsystem](AInstancedFoliageActor* IFA)
 	{
 		for (auto& Pair : IFA->FoliageInfos)
@@ -1047,7 +1056,6 @@ int32 UWorldPartitionConvertCommandlet::Main(const FString& Params)
 		CollectGarbage(RF_Standalone);
 
 		MainLevel->bUseExternalActors = true;
-		MainWorldSettings->SetWorldPartition(WorldPartition);
 		MainWorld->WorldComposition = nullptr;
 		MainLevel->bIsPartitioned = true;
 
diff --git a/Engine/Source/Runtime/Engine/Private/ActorPartition/ActorPartitionSubsystem.cpp b/Engine/Source/Runtime/Engine/Private/ActorPartition/ActorPartitionSubsystem.cpp
--- a/Engine/Source/Runtime/Engine/Private/ActorPartition/ActorPartitionSubsystem.cpp
+++ b/Engine/Source/Runtime/Engine/Private/ActorPartition/ActorPartitionSubsystem.cpp
@@ -132,7 +132,7 @@ public:
 		: FBaseActorPartition(InWorld)
 	{
 		WorldPartition = InWorld->GetSubsystem<UWorldPartitionSubsystem>();
-		check(WorldPartition || IsRunningCommandlet());
+		check(WorldPartition);
 	}
 
 	UActorPartitionSubsystem::FCellCoord GetActorPartitionHash(const FActorPartitionGetParams& GetParams) const override 
@@ -207,12 +207,6 @@ bool UActorPartitionSubsystem::IsLevelPartition() const
 }
 
 #if WITH_EDITOR
-void UActorPartitionSubsystem::InitializeForWorldPartitionConversion()
-{
-	check(IsRunningCommandlet());
-	ActorPartition.Reset(new FActorPartitionWorldPartition(GetWorld()));
-}
-
 void UActorPartitionSubsystem::Initialize(FSubsystemCollectionBase& Collection)
 {
 	UWorldPartitionSubsystem* WorldPartitionSubsystem = Collection.InitializeDependency<UWorldPartitionSubsystem>();
diff --git a/Engine/Source/Runtime/Engine/Public/ActorPartition/ActorPartitionSubsystem.h b/Engine/Source/Runtime/Engine/Public/ActorPartition/ActorPartitionSubsystem.h
--- a/Engine/Source/Runtime/Engine/Public/ActorPartition/ActorPartitionSubsystem.h
+++ b/Engine/Source/Runtime/Engine/Public/ActorPartition/ActorPartitionSubsystem.h
@@ -106,8 +106,6 @@ public:
 	APartitionActor* GetActor(const FActorPartitionGetParams& GetParam);
 	APartitionActor* GetActor(const TSubclassOf<APartitionActor>& InActorClass, const FCellCoord& InCellCoords, bool bInCreate);
 
-	virtual void InitializeForWorldPartitionConversion();
-
 	virtual void Initialize(FSubsystemCollectionBase& Collection) override;
 	virtual void Deinitialize() override;
 #endif
