Fixed crash when importing FBX static mesh with invalid UV indices.
This is a band aid to unblock, but needs to be investigated in greater detail and a better fix made accordingly.

diff --git a/Engine/Source/Editor/UnrealEd/Private/Fbx/FbxStaticMeshImport.cpp b/Engine/Source/Editor/UnrealEd/Private/Fbx/FbxStaticMeshImport.cpp
--- a/Engine/Source/Editor/UnrealEd/Private/Fbx/FbxStaticMeshImport.cpp
+++ b/Engine/Source/Editor/UnrealEd/Private/Fbx/FbxStaticMeshImport.cpp
@@ -1110,10 +1110,18 @@ bool UnFbx::FFbxImporter::BuildStaticMeshFromGeometry(FbxNode* Node, UStaticMesh
 								? UVMapIndex
 								: FBXUVs.LayerElementUV[UVLayerIndex]->GetIndexArray().GetAt(UVMapIndex);
 
-							UVIDs[VertexIndex] = UVIndex + UVOffsets[UVLayerIndex];
+							if (UVIndex != -1)
+							{
+								UVIDs[VertexIndex] = UVIndex + UVOffsets[UVLayerIndex];
 
-							check(MeshDescription->VertexInstanceAttributes().GetAttribute<FVector2D>(CornerInstanceIDs[VertexIndex], MeshAttribute::VertexInstance::TextureCoordinate, UVLayerIndex) ==
-								MeshDescription->UVAttributes(UVLayerIndex).GetAttribute<FVector2D>(UVIndex + UVOffsets[UVLayerIndex], MeshAttribute::UV::UVCoordinate));
+								check(MeshDescription->VertexInstanceAttributes().GetAttribute<FVector2D>(CornerInstanceIDs[VertexIndex], MeshAttribute::VertexInstance::TextureCoordinate, UVLayerIndex) ==
+									MeshDescription->UVAttributes(UVLayerIndex).GetAttribute<FVector2D>(UVIndex + UVOffsets[UVLayerIndex], MeshAttribute::UV::UVCoordinate));
+							}
+							else
+							{
+								// TODO: what does it mean to have a UV index of -1?
+								// Investigate this case more carefully and handle it properly.
+							}
 						}
 					}
 
