Fixing a crash with shutdown ordering in standalone toolkit hosts with their mode managers.

#rnx

diff --git a/Engine/Source/Editor/UnrealEd/Private/Toolkits/SStandaloneAssetEditorToolkitHost.cpp b/Engine/Source/Editor/UnrealEd/Private/Toolkits/SStandaloneAssetEditorToolkitHost.cpp
--- a/Engine/Source/Editor/UnrealEd/Private/Toolkits/SStandaloneAssetEditorToolkitHost.cpp
+++ b/Engine/Source/Editor/UnrealEd/Private/Toolkits/SStandaloneAssetEditorToolkitHost.cpp
@@ -15,6 +15,8 @@
 #include "StatusBarSubsystem.h"
 #include "ToolMenus.h"
 #include "IIntroTutorials.h"
+#include "EditorModeManager.h"
+#include "EditorModes.h"
 
 #define LOCTEXT_NAMESPACE "StandaloneAssetEditorToolkit"
 
@@ -249,14 +251,7 @@ FEditorModeTools& SStandaloneAssetEditorToolkitHost::GetEditorModeManager() cons
 
 SStandaloneAssetEditorToolkitHost::~SStandaloneAssetEditorToolkitHost()
 {
-	// Reset the mode manager so that any modes from the HostedAssetEditorToolkit are exited
-	GetEditorModeManager().SetDefaultMode(FBuiltinEditorModes::EM_Default);
-	GetEditorModeManager().DeactivateAllModes();
-
-	// Let the toolkit manager know that we're going away now
-	FToolkitManager::Get().OnToolkitHostDestroyed( this );
-	HostedToolkits.Reset();
-	HostedAssetEditorToolkit.Reset();
+	ShutdownToolkitHost();
 }
 
 
@@ -289,6 +284,26 @@ void SStandaloneAssetEditorToolkitHost::OnToolkitHostingStarted( const TSharedRe
 	}
 }
 
+void SStandaloneAssetEditorToolkitHost::ShutdownToolkitHost()
+{
+	const TSharedPtr<SDockTab> HostTab = HostTabPtr.Pin();
+	if (HostTab.IsValid())
+	{
+		HostTab->RequestCloseTab();
+	}
+
+	if (HostedAssetEditorToolkit.IsValid())
+	{
+		// Reset the mode manager so that any modes from the HostedAssetEditorToolkit are exited
+		GetEditorModeManager().SetDefaultMode(FBuiltinEditorModes::EM_Default);
+		GetEditorModeManager().DeactivateAllModes();
+	}
+
+	// Let the toolkit manager know that we're going away now
+	FToolkitManager::Get().OnToolkitHostDestroyed(this);
+	HostedToolkits.Reset();
+	HostedAssetEditorToolkit.Reset();
+}
 
 void SStandaloneAssetEditorToolkitHost::OnToolkitHostingFinished( const TSharedRef< class IToolkit >& Toolkit )
 {
@@ -300,13 +315,7 @@ void SStandaloneAssetEditorToolkitHost::OnToolkitHostingFinished( const TSharedR
 	// Standalone Asset Editors close by shutting down their major tab.
 	if (Toolkit == HostedAssetEditorToolkit)
 	{
-		HostedAssetEditorToolkit.Reset();
-
-		const TSharedPtr<SDockTab> HostTab = HostTabPtr.Pin();
-		if (HostTab.IsValid())
-		{
-			HostTab->RequestCloseTab();
-		}
+		ShutdownToolkitHost();
 	}
 	else if (HostedAssetEditorToolkit.IsValid())
 	{
diff --git a/Engine/Source/Editor/UnrealEd/Private/Toolkits/SStandaloneAssetEditorToolkitHost.h b/Engine/Source/Editor/UnrealEd/Private/Toolkits/SStandaloneAssetEditorToolkitHost.h
--- a/Engine/Source/Editor/UnrealEd/Private/Toolkits/SStandaloneAssetEditorToolkitHost.h
+++ b/Engine/Source/Editor/UnrealEd/Private/Toolkits/SStandaloneAssetEditorToolkitHost.h
@@ -91,6 +91,7 @@ public:
 
 private:
 	void OnTabClosed(TSharedRef<SDockTab> TabClosed) const;
+	void ShutdownToolkitHost();
 
 	FName GetMenuName() const;
 
