Fixing a crash on shutdown when a modeling tool is active.

#rnx

diff --git a/Engine/Source/Editor/UnrealEd/Private/Subsystems/AssetEditorSubsystem.cpp b/Engine/Source/Editor/UnrealEd/Private/Subsystems/AssetEditorSubsystem.cpp
--- a/Engine/Source/Editor/UnrealEd/Private/Subsystems/AssetEditorSubsystem.cpp
+++ b/Engine/Source/Editor/UnrealEd/Private/Subsystems/AssetEditorSubsystem.cpp
@@ -51,7 +51,7 @@ void UAssetEditorSubsystem::Initialize(FSubsystemCollectionBase& Collection)
 	FCoreUObjectDelegates::OnPackageReloaded.AddUObject(this, &UAssetEditorSubsystem::HandlePackageReloaded);
 
 	GEditor->OnEditorClose().AddUObject(this, &UAssetEditorSubsystem::OnEditorClose);
-	FCoreDelegates::OnPreExit.AddUObject(this, &UAssetEditorSubsystem::UnregisterEditorModes);
+	FCoreDelegates::OnEnginePreExit.AddUObject(this, &UAssetEditorSubsystem::UnregisterEditorModes);
 	FCoreDelegates::OnPostEngineInit.AddUObject(this, &UAssetEditorSubsystem::RegisterEditorModes);
 }
 
