Editgrated 14337207 from 4.26 stream (Fix FBX import crash with degenerate meshes).

diff --git a/Engine/Source/Editor/UnrealEd/Private/Fbx/FbxSceneImportFactory.cpp b/Engine/Source/Editor/UnrealEd/Private/Fbx/FbxSceneImportFactory.cpp
--- a/Engine/Source/Editor/UnrealEd/Private/Fbx/FbxSceneImportFactory.cpp
+++ b/Engine/Source/Editor/UnrealEd/Private/Fbx/FbxSceneImportFactory.cpp
@@ -2362,6 +2362,7 @@ UObject* UFbxSceneImportFactory::ImportANode(void* VoidFbxImporter, TArray<void*
 	}
 	else
 	{
+		Pkg->SetDirtyFlag(false);
 		Pkg->RemoveFromRoot();
 		Pkg->ConditionalBeginDestroy();
 	}
diff --git a/Engine/Source/Editor/UnrealEd/Private/Fbx/FbxStaticMeshImport.cpp b/Engine/Source/Editor/UnrealEd/Private/Fbx/FbxStaticMeshImport.cpp
--- a/Engine/Source/Editor/UnrealEd/Private/Fbx/FbxStaticMeshImport.cpp
+++ b/Engine/Source/Editor/UnrealEd/Private/Fbx/FbxStaticMeshImport.cpp
@@ -1999,6 +1999,7 @@ UStaticMesh* UnFbx::FFbxImporter::ImportStaticMeshAsSingle(UObject* InParent, TA
 		// If we couldn't build the static mesh, its package is invalid. We should reject it entirely to prevent issues from arising from trying to use it in the editor.
 		if (!NewPackageName.IsEmpty())
 		{
+			Package->SetDirtyFlag(false);
 			Package->RemoveFromRoot();
 			Package->ConditionalBeginDestroy();
 		}
diff --git a/Engine/Source/Editor/UnrealEd/Private/Fbx/ReimportFbxSceneFactory.cpp b/Engine/Source/Editor/UnrealEd/Private/Fbx/ReimportFbxSceneFactory.cpp
--- a/Engine/Source/Editor/UnrealEd/Private/Fbx/ReimportFbxSceneFactory.cpp
+++ b/Engine/Source/Editor/UnrealEd/Private/Fbx/ReimportFbxSceneFactory.cpp
@@ -1220,6 +1220,7 @@ EReimportResult::Type UReimportFbxSceneFactory::ImportStaticMesh(void* VoidFbxIm
 	{
 		if (Pkg != nullptr)
 		{
+			Pkg->SetDirtyFlag(false);
 			Pkg->RemoveFromRoot();
 			Pkg->ConditionalBeginDestroy();
 		}
