Fix crash when reimporting UDIM texture with new file(s)

diff --git a/Engine/Plugins/Editor/ContentBrowser/ContentBrowserAssetDataSource/Source/ContentBrowserAssetDataSource/Private/AssetFileContextMenu.cpp b/Engine/Plugins/Editor/ContentBrowser/ContentBrowserAssetDataSource/Source/ContentBrowserAssetDataSource/Private/AssetFileContextMenu.cpp
--- a/Engine/Plugins/Editor/ContentBrowser/ContentBrowserAssetDataSource/Source/ContentBrowserAssetDataSource/Private/AssetFileContextMenu.cpp
+++ b/Engine/Plugins/Editor/ContentBrowser/ContentBrowserAssetDataSource/Source/ContentBrowserAssetDataSource/Private/AssetFileContextMenu.cpp
@@ -1429,16 +1429,16 @@ void FAssetFileContextMenu::ExecuteReimportWithNewFile(int32 SourceFileIndex /*=
 	//Check if the data is valid
 	if (SourceFileIndex == INDEX_NONE)
 	{
-		check(AssetSourcePaths.Num() <= 1);
-		//Ask for a new file for the index 0
-		SourceFileIndexToReplace = 0;
+		if (AssetSourcePaths.Num() > 1)
+		{
+			//Ask for a new file for the index 0
+			SourceFileIndexToReplace = 0;
+		}
 	}
 	else
 	{
 		check(AssetSourcePaths.IsValidIndex(SourceFileIndex));
 	}
-	check(SourceFileIndexToReplace >= 0);
-	
 
 	FReimportManager::Instance()->ValidateAllSourceFileAndReimport(CopyOfSelectedAssets, true, SourceFileIndexToReplace, true);
 }
diff --git a/Engine/Source/Editor/UnrealEd/Private/Factories/EditorFactories.cpp b/Engine/Source/Editor/UnrealEd/Private/Factories/EditorFactories.cpp
--- a/Engine/Source/Editor/UnrealEd/Private/Factories/EditorFactories.cpp
+++ b/Engine/Source/Editor/UnrealEd/Private/Factories/EditorFactories.cpp
@@ -3950,13 +3950,10 @@ UObject* UTextureFactory::FactoryCreateBinary
 
 					const FString ShortPackageName = ObjectTools::SanitizeInvalidChars(BaseUDIMName, INVALID_LONGPACKAGE_CHARACTERS);
 
-					if (PackageUDIMIndex == -1)
-					{
-						// If we're re-importing UDIM texture, the package will already be correctly named after the UDIM base name
-						// In this case we'll fail to parse the UDIM name, but the package should already have the proper name
-						check(PackageName.EndsWith(ShortPackageName, ESearchCase::CaseSensitive));
-					}
-					else
+					// If we're re-importing UDIM texture, the package will already be correctly named after the UDIM base name
+					// In this case we'll fail to parse the UDIM name, but the package should already have the proper name
+					// Note that the package name may not match the asset name in this case, if we're reimporting new files
+					if (PackageUDIMIndex != -1)
 					{
 						check(PackageUDIMIndex == BaseUDIMIndex);
 						check(PackageUDIMName.EndsWith(ShortPackageName, ESearchCase::CaseSensitive));
