Strip UDIM suffix from imported texture names, even if there's only a single source image

diff --git a/Engine/Source/Editor/UnrealEd/Private/Factories/EditorFactories.cpp b/Engine/Source/Editor/UnrealEd/Private/Factories/EditorFactories.cpp
--- a/Engine/Source/Editor/UnrealEd/Private/Factories/EditorFactories.cpp
+++ b/Engine/Source/Editor/UnrealEd/Private/Factories/EditorFactories.cpp
@@ -3923,53 +3923,55 @@ UObject* UTextureFactory::FactoryCreateBinary
 				if (!CurrentFilename.EndsWith(UDIMFile) && FactoryCanImport(UDIMFile))
 				{
 					const int32 UDIMIndex = ParseUDIMName(FPaths::GetBaseFilename(UDIMFile), UdimRegexPattern, PreUDIMName, PostUDIMName);
-					const FString UDIMName = PreUDIMName + PostUDIMName;
-					if (!UDIMIndexToFile.Contains(UDIMIndex) && UDIMName == BaseUDIMName)
+					if (UDIMIndex != INDEX_NONE)
 					{
-						UDIMIndexToFile.Add(UDIMIndex, Path / UDIMFile);
+						const FString UDIMName = PreUDIMName + PostUDIMName;
+						if (!UDIMIndexToFile.Contains(UDIMIndex) && UDIMName == BaseUDIMName)
+						{
+							UDIMIndexToFile.Add(UDIMIndex, Path / UDIMFile);
+						}
 					}
 				}
 			}
-			if (UDIMIndexToFile.Num() > 1)
-			{
-				// Found multiple UDIM pages, so import as UDIM texture
-				// Exclude UDIM number from the name of the UE4 texture asset we create
-				TextureName = *BaseUDIMName;
 
-				// Don't try to rename the package if its the transient package
-				if ( InParent != GetTransientPackage() )
-				{
-					// Need to rename the package to match the new texture name, since package was already created
-					// Package name will be the same as the object name, except will contain additional path information,
-					// so we take the existing package name, then extract the UDIM index in order to preserve the path
-					FString PackageName;
-					InParent->GetName(PackageName);
+			// Exclude UDIM number from the name of the UE4 texture asset we create
+			// We do this even in the case where we're only importing a single 1001 image, which won't technically be a UDIM
+			// We still want to strip the UDIM suffix in this case, as industry standard still considers this a UDIM imageset
+			TextureName = *BaseUDIMName;
 
-					const int32 PackageUDIMIndex = ParseUDIMName(PackageName, UdimRegexPattern, PreUDIMName, PostUDIMName);
-					const FString PackageUDIMName = PreUDIMName + PostUDIMName;
+			// Don't try to rename the package if its the transient package
+			if (InParent != GetTransientPackage())
+			{
+				// Need to rename the package to match the new texture name, since package was already created
+				// Package name will be the same as the object name, except will contain additional path information,
+				// so we take the existing package name, then extract the UDIM index in order to preserve the path
+				FString PackageName;
+				InParent->GetName(PackageName);
+
+				const int32 PackageUDIMIndex = ParseUDIMName(PackageName, UdimRegexPattern, PreUDIMName, PostUDIMName);
+				const FString PackageUDIMName = PreUDIMName + PostUDIMName;
 
-					const FString ShortPackageName = ObjectTools::SanitizeInvalidChars(BaseUDIMName, INVALID_LONGPACKAGE_CHARACTERS);
+				const FString ShortPackageName = ObjectTools::SanitizeInvalidChars(BaseUDIMName, INVALID_LONGPACKAGE_CHARACTERS);
 
-					// If we're re-importing UDIM texture, the package will already be correctly named after the UDIM base name
-					// In this case we'll fail to parse the UDIM name, but the package should already have the proper name
-					// Note that the package name may not match the asset name in this case, if we're reimporting new files
-					if (PackageUDIMIndex != -1)
+				// If we're re-importing UDIM texture, the package will already be correctly named after the UDIM base name
+				// In this case we'll fail to parse the UDIM name, but the package should already have the proper name
+				// Note that the package name may not match the asset name in this case, if we're reimporting new files
+				if (PackageUDIMIndex != -1)
+				{
+					check(PackageUDIMIndex == BaseUDIMIndex);
+					check(PackageUDIMName.EndsWith(ShortPackageName, ESearchCase::CaseSensitive));
+
+					// In normal case, higher level code would have already checked for duplicate package name
+					// But since we're changing package name here, check to see if package with the new name already exists...
+					// If it does, code later in this method will prompt user to overwrite the existing asset
+					UPackage* ExistingPackage = FindPackage(InParent->GetOuter(), *PackageUDIMName);
+					if (ExistingPackage)
 					{
-						check(PackageUDIMIndex == BaseUDIMIndex);
-						check(PackageUDIMName.EndsWith(ShortPackageName, ESearchCase::CaseSensitive));
-
-						// In normal case, higher level code would have already checked for duplicate package name
-						// But since we're changing package name here, check to see if package with the new name already exists...
-						// If it does, code later in this method will prompt user to overwrite the existing asset
-						UPackage* ExistingPackage = FindPackage(InParent->GetOuter(), *PackageUDIMName);
-						if (ExistingPackage)
-						{
-							InParent = ExistingPackage;
-						}
-						else
-						{
-							verify(InParent->Rename(*PackageUDIMName, nullptr, REN_DontCreateRedirectors));
-						}
+						InParent = ExistingPackage;
+					}
+					else
+					{
+						verify(InParent->Rename(*PackageUDIMName, nullptr, REN_DontCreateRedirectors));
 					}
 				}
 			}
