Fix crash in world partition convert commandlet when the user delete the external folders manually.

diff --git a/Engine/Source/Editor/UnrealEd/Private/Commandlets/WorldPartitionConvertCommandlet.cpp b/Engine/Source/Editor/UnrealEd/Private/Commandlets/WorldPartitionConvertCommandlet.cpp
--- a/Engine/Source/Editor/UnrealEd/Private/Commandlets/WorldPartitionConvertCommandlet.cpp
+++ b/Engine/Source/Editor/UnrealEd/Private/Commandlets/WorldPartitionConvertCommandlet.cpp
@@ -710,24 +710,28 @@ int32 UWorldPartitionConvertCommandlet::Main(const FString& Params)
 
 		FString OldLevelName = Tokens[0] + ConversionSuffix;
 		FString ExternalActorsPath = ULevel::GetExternalActorsPath(OldLevelName);
+		FString ExternalActorsFilePath = FPackageName::LongPackageNameToFilename(ExternalActorsPath);
 
-		bool bResult = IFileManager::Get().IterateDirectoryRecursively(*FPackageName::LongPackageNameToFilename(ExternalActorsPath), [this](const TCHAR* FilenameOrDirectory, bool bIsDirectory)
+		if (IFileManager::Get().DirectoryExists(*ExternalActorsFilePath))
 		{
-			if (!bIsDirectory)
+			bool bResult = IFileManager::Get().IterateDirectoryRecursively(*ExternalActorsFilePath, [this](const TCHAR* FilenameOrDirectory, bool bIsDirectory)
 			{
-				FString Filename(FilenameOrDirectory);
-				if (Filename.EndsWith(FPackageName::GetAssetPackageExtension()))
+				if (!bIsDirectory)
 				{
-					return PackageHelper.Delete(Filename);
+					FString Filename(FilenameOrDirectory);
+					if (Filename.EndsWith(FPackageName::GetAssetPackageExtension()))
+					{
+						return PackageHelper.Delete(Filename);
+					}
 				}
-			}
-			return true;
-		});
+				return true;
+			});
 
-		if (!bResult)
-		{
-			UE_LOG(LogWorldPartitionConvertCommandlet, Error, TEXT("Failed to delete external actor package(s)"));
-			return 1;
+			if (!bResult)
+			{
+				UE_LOG(LogWorldPartitionConvertCommandlet, Error, TEXT("Failed to delete external actor package(s)"));
+				return 1;
+			}
 		}
 
 		if (FPackageName::SearchForPackageOnDisk(OldLevelName, &OldLevelName))
