Integrate CL 14366888 "validate the CookServer in the call to IsCookByTheBookInEditorFinished"

diff --git a/Engine/Source/Editor/UnrealEd/Private/UnrealEdEngine.cpp b/Engine/Source/Editor/UnrealEd/Private/UnrealEdEngine.cpp
--- a/Engine/Source/Editor/UnrealEd/Private/UnrealEdEngine.cpp
+++ b/Engine/Source/Editor/UnrealEd/Private/UnrealEdEngine.cpp
@@ -297,7 +297,7 @@ void UUnrealEdEngine::StartCookByTheBookInEditor( const TArray<ITargetPlatform*>
 
 bool UUnrealEdEngine::IsCookByTheBookInEditorFinished() const 
 { 
-	return !CookServer->IsCookByTheBookRunning();
+	return CookServer ? !CookServer->IsCookByTheBookRunning() : true;
 }
 
 void UUnrealEdEngine::CancelCookByTheBookInEditor()
