Integrate CL 14363187 [Navigation] Fix missing navmesh tiles when cooking maps from the editor 'Launch'
That fixes: "Multiple AI tests failed when locally running in cooked build"

diff --git a/Engine/Source/Editor/UnrealEd/Private/EditorEngine.cpp b/Engine/Source/Editor/UnrealEd/Private/EditorEngine.cpp
--- a/Engine/Source/Editor/UnrealEd/Private/EditorEngine.cpp
+++ b/Engine/Source/Editor/UnrealEd/Private/EditorEngine.cpp
@@ -7313,6 +7313,10 @@ void UEditorEngine::InitializeNewlyCreatedInactiveWorld(UWorld* World)
 	{
 		const bool bOldDirtyState = World->GetOutermost()->IsDirty();
 
+		// Make sure we have a navigation system if we are cooking the asset.
+		// Typically nav bounds are added when AddNavigationSystemToWorld() is called from UEditorEngine::Map_Load().
+		const bool bCooking = (IsCookByTheBookInEditorFinished() == false);
+
 		// Create the world without a physics scene because creating too many physics scenes causes deadlock issues in PhysX. The scene will be created when it is opened in the level editor.
 		// Also, don't create an FXSystem because it consumes too much video memory. This is also created when the level editor opens this world.
 		// Do not create AISystem/Navigation for inactive world. These ones will also be created when the level editor opens this world. if required.
@@ -7320,12 +7324,21 @@ void UEditorEngine::InitializeNewlyCreatedInactiveWorld(UWorld* World)
 			.CreatePhysicsScene(false)
 			.CreateFXSystem(false)
 			.CreateAISystem(false)
-			.CreateNavigation(false)
+			.CreateNavigation(bCooking)
 			);
 
 		// Update components so the scene is populated
 		World->UpdateWorldComponents(true, true);
 
+		if (bCooking)
+		{
+			// When calling World->InitWorld() with bCreateNavigation=true (just above), 
+			// it calls internally FNavigationSystem::AddNavigationSystemToWorld() with bInitializeForWorld=false.
+			// That does not gather nav bounds. When cooking, the nav system and nav bounds are needed on the navmesh serialize-save for tiles to be added to the archive.
+			// Also this call needs to occur after World->UpdateWorldComponents() else no bounds are found.
+			FNavigationSystem::AddNavigationSystemToWorld(*World, FNavigationSystemRunMode::EditorMode);
+		}
+
 		// Need to restore the dirty state as registering components dirties the world
 		if (!bOldDirtyState)
 		{
