[Navigation WP] First pass to generate per cell navigation data and allow navmesh streaming using world partitioning
Addition of ANavigationDataChunkActor, external actors used to hold navdata chunks
Addition of UWorldPartitionBuildNavigationDataCommandlet, commandlet to generate world partition navigation data
At the moment steps to generate a partitioned world with navmesh per cell are:
1) convert world to WP (WorldPartitionConvertCommandlet)
2) generate navigation data chunk actors (WorldPartitionBuildNavigationDataCommandlet)
3) pre cook (WorldPartitionPreCookCommandlet)
4) cook

diff --git a/Engine/Source/Editor/UnrealEd/Private/Commandlets/WorldPartitionBuildNavigationDataCommandlet.cpp b/Engine/Source/Editor/UnrealEd/Private/Commandlets/WorldPartitionBuildNavigationDataCommandlet.cpp
new file mode 100644
--- /dev/null
+++ b/Engine/Source/Editor/UnrealEd/Private/Commandlets/WorldPartitionBuildNavigationDataCommandlet.cpp
@@ -0,0 +1,185 @@
+// Copyright Epic Games, Inc. All Rights Reserved.
+
+#include "Commandlets/WorldPartitionBuildNavigationDataCommandlet.h"
+
+#include "CoreMinimal.h"
+#include "EngineUtils.h"
+#include "Logging/LogMacros.h"
+
+#include "WorldPartition/WorldPartitionEditorHash.h"
+#include "WorldPartition/WorldPartition.h"
+#include "WorldPartition/NavigationData/NavigationDataChunkActor.h"
+#include "WorldPartition/WorldPartitionSubsystem.h"
+
+#if WITH_EDITOR
+#include "Editor/EditorEngine.h"
+extern UNREALED_API UEditorEngine* GEditor;
+#endif
+
+DEFINE_LOG_CATEGORY_STATIC(LogWorldPartitionBuildNavigationDataCommandlet, Log, All);
+
+UWorldPartitionBuildNavigationDataCommandlet::UWorldPartitionBuildNavigationDataCommandlet(const FObjectInitializer& ObjectInitializer)
+	: Super(ObjectInitializer)
+{
+}
+
+int32 UWorldPartitionBuildNavigationDataCommandlet::Main(const FString& Params)
+{
+	TArray<FString> Tokens, Switches;
+	ParseCommandLine(*Params, Tokens, Switches);
+
+	// Retrieve map name from the command line
+	FString MapName;
+	const bool bMapProvided = FParse::Value(*Params, TEXT("Map="), MapName);
+	if (!bMapProvided)
+	{
+		UE_LOG(LogWorldPartitionBuildNavigationDataCommandlet, Error, TEXT("No map specified."));
+		return 1;
+	}
+
+	// Load the map package
+	UPackage* MapPackage = LoadPackage(NULL, *MapName, LOAD_None); 
+	if (!MapPackage)
+	{
+		UE_LOG(LogWorldPartitionBuildNavigationDataCommandlet, Error, TEXT("Couldn't load package %s."), *MapName);
+		return 1;
+	}
+
+	// Find the world in the given package
+	UWorld* World = UWorld::FindWorldInPackage(MapPackage);
+	if (!World)
+	{
+		UE_LOG(LogWorldPartitionBuildNavigationDataCommandlet, Error, TEXT("No world in package %s."), *MapPackage->GetName());
+		return 1;
+	}
+
+	// Setup the world.
+	World->WorldType = EWorldType::Editor;
+	World->AddToRoot();
+	if (!World->bIsWorldInitialized)
+	{
+		UWorld::InitializationValues IVS;
+		IVS.RequiresHitProxies(false);
+		IVS.ShouldSimulatePhysics(false);
+		IVS.EnableTraceCollision(false);
+		IVS.CreateNavigation(true);
+		IVS.CreateAISystem(false);
+		IVS.AllowAudioPlayback(false);
+		IVS.CreatePhysicsScene(true);
+
+		World->InitWorld(UWorld::InitializationValues(IVS));
+		World->PersistentLevel->UpdateModelComponents();
+		World->UpdateWorldComponents(true /*bRerunConstructionScripts*/, false /*bCurrentLevelOnly*/);
+	}
+
+	// Retrieve the world partition.
+	const UWorldPartitionSubsystem* WorldPartitionSubsystem = World->GetSubsystem<UWorldPartitionSubsystem>();
+	UWorldPartition* WorldPartition = nullptr;
+	if (WorldPartitionSubsystem)
+	{
+		WorldPartition = World->GetWorldPartition();
+		check(WorldPartition);
+	}
+
+	FWorldContext& WorldContext = GEditor->GetEditorWorldContext(true /*bEnsureIsGWorld*/);
+	WorldContext.SetCurrentWorld(World);
+	GWorld = World;
+	
+	// For now, load all cells
+	// In the future, we'll want the commandlet to be able to perform partial updates of the map
+	// to allow rebuild to be distributed on multiple machines.
+	const FBox LoadBox(FVector(-WORLD_MAX, -WORLD_MAX, -WORLD_MAX), FVector(WORLD_MAX, WORLD_MAX, WORLD_MAX));
+	WorldPartition->LoadEditorCells(LoadBox);
+
+	TSet<UPackage*> NavigationDataChunkActorPackages;
+
+	// Gather all packages before any navigation data chunk actors are deleted
+	for (TActorIterator<ANavigationDataChunkActor> ItActor(World); ItActor; ++ItActor)
+	{
+		NavigationDataChunkActorPackages.Add(ItActor->GetPackage());
+	}
+
+	// Rebuild ANavigationDataChunkActor for the whole world
+	WorldPartition->GenerateNavigationData();
+
+	// Gather all packages again to include newly created ANavigationDataChunkActor actors
+	for (TActorIterator<ANavigationDataChunkActor> ItActor(World); ItActor; ++ItActor)
+	{
+		NavigationDataChunkActorPackages.Add(ItActor->GetPackage());
+	}
+
+	TArray<UPackage*> PackagesToDelete;
+	TArray<UPackage*> PackagesToSave;
+
+	for (UPackage* Package : NavigationDataChunkActorPackages)
+	{
+		if (Package && Package->IsDirty())
+		{
+			if (UPackage::IsEmptyPackage(Package))
+			{
+				PackagesToDelete.Add(Package);
+			}
+			else
+			{
+				PackagesToSave.Add(Package);
+			}
+		}
+	}
+
+	// Delete packages
+	if (!PackagesToDelete.IsEmpty())
+	{
+		UE_LOG(LogWorldPartitionBuildNavigationDataCommandlet, Log, TEXT("Deleting %d packages."), PackagesToDelete.Num());
+		if (!PackageHelper.Delete(PackagesToDelete))
+		{
+			UE_LOG(LogWorldPartitionBuildNavigationDataCommandlet, Error, TEXT("Error deleting packages."));
+			return 1;
+		}
+	}
+
+	// Save packages
+	if (!PackagesToSave.IsEmpty())
+	{
+		// Checkout packages to save
+		UE_LOG(LogWorldPartitionBuildNavigationDataCommandlet, Log, TEXT("Checking out %d actor packages."), PackagesToSave.Num());
+		for (UPackage* Package : PackagesToSave)
+		{
+			if (!PackageHelper.Checkout(Package))
+			{
+				UE_LOG(LogWorldPartitionBuildNavigationDataCommandlet, Error, TEXT("Error checking out package %s."), *Package->GetName());
+				return 1;
+			}
+		}
+
+		// Save packages
+		UE_LOG(LogWorldPartitionBuildNavigationDataCommandlet, Log, TEXT("Saving %d packages."), PackagesToSave.Num());
+		for (UPackage* Package : PackagesToSave)
+		{
+			UE_LOG(LogWorldPartitionBuildNavigationDataCommandlet, Verbose, TEXT("Saving package  %s."), *Package->GetName());
+			FString PackageFileName = SourceControlHelpers::PackageFilename(Package);
+			if (!UPackage::SavePackage(Package, nullptr, RF_Standalone, *PackageFileName, GError, nullptr, false, true, SAVE_Async))
+			{
+				UE_LOG(LogWorldPartitionBuildNavigationDataCommandlet, Error, TEXT("Error saving package %s."), *Package->GetName());
+				return 1;
+			}
+		}
+
+		// Add new packages to source control
+		UE_LOG(LogWorldPartitionBuildNavigationDataCommandlet, Log, TEXT("Adding packages to source control."));
+		for (UPackage* Package : PackagesToSave)
+		{
+			if (!PackageHelper.AddToSourceControl(Package))
+			{
+				UE_LOG(LogWorldPartitionBuildNavigationDataCommandlet, Error, TEXT("Error adding package %s to source control."), *Package->GetName());
+				return 1;
+			}
+		}
+	}
+
+	// Cleanup
+	World->RemoveFromRoot();
+	WorldContext.SetCurrentWorld(nullptr);
+	GWorld = nullptr;
+
+	return 0;
+}
diff --git a/Engine/Source/Editor/UnrealEd/Private/Commandlets/WorldPartitionBuildNavigationDataCommandlet.h b/Engine/Source/Editor/UnrealEd/Private/Commandlets/WorldPartitionBuildNavigationDataCommandlet.h
new file mode 100644
--- /dev/null
+++ b/Engine/Source/Editor/UnrealEd/Private/Commandlets/WorldPartitionBuildNavigationDataCommandlet.h
@@ -0,0 +1,23 @@
+// Copyright 1998-2019 Epic Games, Inc. All Rights Reserved.
+
+#pragma once
+
+#include "Commandlets/Commandlet.h"
+#include "PackageSourceControlHelper.h"
+#include "WorldPartitionBuildNavigationDataCommandlet.generated.h"
+
+/**
+* Commandlet to build navigation data for a partioned level 
+*/
+UCLASS()
+class UWorldPartitionBuildNavigationDataCommandlet : public UCommandlet
+{
+	GENERATED_UCLASS_BODY()
+
+	//~ Begin UCommandlet Interface
+	virtual int32 Main(const FString& Params) override;
+	//~ End UCommandlet Interface
+
+private:
+	FPackageSourceControlHelper PackageHelper;
+};
diff --git a/Engine/Source/Runtime/Engine/Classes/AI/NavigationSystemBase.h b/Engine/Source/Runtime/Engine/Classes/AI/NavigationSystemBase.h
--- a/Engine/Source/Runtime/Engine/Classes/AI/NavigationSystemBase.h
+++ b/Engine/Source/Runtime/Engine/Classes/AI/NavigationSystemBase.h
@@ -238,6 +238,11 @@ public:
 	UE_DEPRECATED(4.20, "GetMainNavData is deprecated. Use FNavigationSystem::GetCurrent<UNavigationSystemV1>()->GetDefaultNavDataInstance instead")
 	INavigationDataInterface* GetMainNavData(int) { return nullptr; }
 
+	virtual bool ContainsNavData(const FBox& Bounds) const { return false; }
+	virtual void AddNavigationDataChunk(class ANavigationDataChunkActor& DataChunkActor) {}
+	virtual void RemoveNavigationDataChunk(class ANavigationDataChunkActor& DataChunkActor) {}
+	virtual void FillNavigationDataChunkActor(const FBox& Bounds, class ANavigationDataChunkActor& DataChunkActor) {}
+
 protected:
 	/**	Sets the Transform the Navigation System will use when converting from FromCoordType
 	 *	to ToCoordType
diff --git a/Engine/Source/Runtime/Engine/Private/WorldPartition/NavigationData/NavigationDataChunkActor.cpp b/Engine/Source/Runtime/Engine/Private/WorldPartition/NavigationData/NavigationDataChunkActor.cpp
new file mode 100644
--- /dev/null
+++ b/Engine/Source/Runtime/Engine/Private/WorldPartition/NavigationData/NavigationDataChunkActor.cpp
@@ -0,0 +1,86 @@
+// Copyright Epic Games, Inc. All Rights Reserved.
+
+#include "WorldPartition/NavigationData/NavigationDataChunkActor.h"
+#include "AI/NavigationSystemBase.h"
+
+ANavigationDataChunkActor::ANavigationDataChunkActor(const FObjectInitializer& ObjectInitializer)
+	: Super(ObjectInitializer)
+{
+	if (!IsTemplate())
+	{
+		USceneComponent* Root = CreateDefaultSubobject<USceneComponent>(USceneComponent::GetDefaultSceneRootVariableName());
+		Root->SetMobility(EComponentMobility::Static);
+		SetRootComponent(Root);
+	}
+
+	SetCanBeDamaged(false);
+	SetActorEnableCollision(false);
+}
+
+void ANavigationDataChunkActor::Serialize(FArchive& Ar)
+{
+	UE_LOG(LogNavigation, VeryVerbose, TEXT("%s"), ANSI_TO_TCHAR(__FUNCTION__));
+
+	Super::Serialize(Ar);
+}
+
+void ANavigationDataChunkActor::PostLoad()
+{
+	UE_LOG(LogNavigation, VeryVerbose, TEXT("%s"), ANSI_TO_TCHAR(__FUNCTION__));
+
+	Super::PostLoad();
+}
+
+void ANavigationDataChunkActor::CollectNavData(const FBox& Bounds)
+{
+	UE_LOG(LogNavigation, VeryVerbose, TEXT("%s"), ANSI_TO_TCHAR(__FUNCTION__));
+
+	UWorld* World = GetWorld();
+	if (World)
+	{
+		if (UNavigationSystemBase* NavSys = World->GetNavigationSystem())
+		{
+			NavSys->FillNavigationDataChunkActor(Bounds, *this);
+		}
+	}
+}
+
+void ANavigationDataChunkActor::BeginPlay()
+{
+	UE_LOG(LogNavigation, VeryVerbose, TEXT("%s"), ANSI_TO_TCHAR(__FUNCTION__));
+
+	Super::BeginPlay();
+
+	UWorld* World = GetWorld();
+	if (World)
+	{
+		if (UNavigationSystemBase* NavSys = World->GetNavigationSystem())
+		{
+			NavSys->AddNavigationDataChunk(*this);
+		}
+	}
+}
+
+void ANavigationDataChunkActor::EndPlay(const EEndPlayReason::Type EndPlayReason)
+{
+	UE_LOG(LogNavigation, VeryVerbose, TEXT("%s"), ANSI_TO_TCHAR(__FUNCTION__));
+
+	UWorld* World = GetWorld();
+	if (World)
+	{
+		if (UNavigationSystemBase* NavSys = World->GetNavigationSystem())
+		{
+			NavSys->RemoveNavigationDataChunk(*this);
+		}
+	}
+
+	Super::EndPlay(EndPlayReason);
+}
+
+
+#if WITH_EDITOR
+EActorGridPlacement ANavigationDataChunkActor::GetDefaultGridPlacement() const
+{
+	return EActorGridPlacement::Location;
+}
+#endif // WITH_EDITOR
diff --git a/Engine/Source/Runtime/Engine/Private/WorldPartition/WorldPartition.cpp b/Engine/Source/Runtime/Engine/Private/WorldPartition/WorldPartition.cpp
--- a/Engine/Source/Runtime/Engine/Private/WorldPartition/WorldPartition.cpp
+++ b/Engine/Source/Runtime/Engine/Private/WorldPartition/WorldPartition.cpp
@@ -38,7 +38,8 @@
 #include "WorldPartition/WorldPartitionEditorHash.h"
 #include "WorldPartition/WorldPartitionEditorSpatialHash.h"
 #include "WorldPartition/WorldPartitionRuntimeHash.h"
-#endif
+#include "Misc/Base64.h"
+#endif //WITH_EDITOR
 
 DEFINE_LOG_CATEGORY(LogWorldPartition);
 
@@ -310,7 +311,7 @@ void UWorldPartition::Initialize(UWorld* InWorld, const FTransform& InTransform)
 			}
 		}
 	}
-#endif
+#endif //WITH_EDITOR
 
 	InitState = EWorldPartitionInitState::Initialized;
 
@@ -1268,6 +1269,11 @@ void UWorldPartition::GenerateHLOD()
 	RuntimeHash->GenerateHLOD();
 }
 
+void UWorldPartition::GenerateNavigationData()
+{
+	RuntimeHash->GenerateNavigationData();
+}
+
 void UWorldPartition::OnPreFixupForPIE(int32 InPIEInstanceID, FSoftObjectPath& ObjectPath)
 {
 	if (GetWorld()->IsPlayInEditor())
diff --git a/Engine/Source/Runtime/Engine/Private/WorldPartition/WorldPartitionRuntimeSpatialHash.cpp b/Engine/Source/Runtime/Engine/Private/WorldPartition/WorldPartitionRuntimeSpatialHash.cpp
--- a/Engine/Source/Runtime/Engine/Private/WorldPartition/WorldPartitionRuntimeSpatialHash.cpp
+++ b/Engine/Source/Runtime/Engine/Private/WorldPartition/WorldPartitionRuntimeSpatialHash.cpp
@@ -26,6 +26,8 @@
 #include "WorldPartition/HLOD/HLODLayer.h"
 #include "WorldPartition/HLOD/HLODActor.h"
 #include "WorldPartition/HLOD/HLODActorDesc.h"
+#include "WorldPartition/NavigationData/NavigationDataChunkActor.h"
+#include "AI/NavigationSystemBase.h"
 
 #include "AssetRegistryModule.h"
 #include "AssetData.h"
@@ -34,7 +36,7 @@
 #include "LevelUtils.h"
 
 extern UNREALED_API class UEditorEngine* GEditor;
-#endif
+#endif //WITH_EDITOR
 
 DEFINE_LOG_CATEGORY_STATIC(LogWorldPartitionRuntimeSpatialHash, Log, All);
 
@@ -74,7 +76,7 @@ private:
 	UWorldPartition* WorldPartition;
 	TArray<AActor*> LoadedActors;
 };
-#endif
+#endif //WITH_EDITOR
 
 static int32 GShowRuntimeSpatialHashGridLevel = 0;
 static FAutoConsoleVariableRef CVarShowRuntimeSpatialHashGridLevel(
@@ -492,7 +494,7 @@ struct FSquare2DGridHelper
 			}
 		}
 	}
-#endif
+#endif //WITH_EDITOR
 
 public:
 	FVector Origin;
@@ -868,7 +870,7 @@ static FSquare2DGridHelper GetPartitionedActors(const UWorldPartition* WorldPart
 	
 	return PartitionedActors;
 }
-#endif
+#endif //WITH_EDITOR
 
 UWorldPartitionRuntimeSpatialHash::UWorldPartitionRuntimeSpatialHash(const FObjectInitializer& ObjectInitializer)
 	: Super(ObjectInitializer)
@@ -1003,7 +1005,6 @@ FName UWorldPartitionRuntimeSpatialHash::GetCellName(FName InGridName, int32 InL
 	return FName(*FString::Printf(TEXT("WPRT_%s_%s_Cell_L%d_X%02d_Y%02d"), *PackageNameNoPIEPrefix, *InGridName.ToString(), InLevel, InCellX, InCellY));
 }
 
-#if WITH_EDITOR
 void UWorldPartitionRuntimeSpatialHash::CacheHLODParents()
 {
 	CachedHLODParents.Reset();
@@ -1125,7 +1126,6 @@ bool UWorldPartitionRuntimeSpatialHash::CreateStreamingGrid(const FSpatialHashRu
 
 	return true;
 }
-#endif
 
 void UWorldPartitionRuntimeSpatialHash::FlushStreaming()
 {
@@ -1271,7 +1271,89 @@ bool UWorldPartitionRuntimeSpatialHash::GenerateHLOD()
 
 	return true;
 }
-#endif
+
+bool UWorldPartitionRuntimeSpatialHash::GenerateNavigationData()
+{
+	TRACE_CPUPROFILER_EVENT_SCOPE(UWorldPartitionRuntimeSpatialHash::GenerateNavigationData);
+
+	UE_LOG(LogWorldPartitionRuntimeSpatialHash, Log, TEXT("%s"), ANSI_TO_TCHAR(__FUNCTION__));
+
+	UWorldPartition* WorldPartition = GetOuterUWorldPartition();
+	UWorld* World = WorldPartition->World;
+
+	TArray<TArray<const UWorldPartition::FActorCluster*>> GridActors;
+	GridActors.InsertDefaulted(0, Grids.Num());
+
+	// Generate navmesh
+	// Make sure navigation is added and initialized in EditorMode
+	FNavigationSystem::AddNavigationSystemToWorld(*World, FNavigationSystemRunMode::EditorMode);
+
+	// Invoke navigation data generator
+	FNavigationSystem::Build(*World);
+
+	// For each cell, gather navmesh and generate a datachunk actor
+	const FBox WorldBounds = WorldPartition->GetWorldBounds();
+
+	const int32 GridIndex = 0; //Todo At: only work with grid 0 for now.
+	UE_LOG(LogWorldPartitionRuntimeSpatialHash, Verbose, TEXT("Generate NavDataChunk actor for GridIndex %i."), GridIndex);
+
+	const FSpatialHashRuntimeGrid& RuntimeGrid = Grids[GridIndex];
+	const FSquare2DGridHelper GridHelper = GetPartitionedActors(WorldPartition, WorldBounds, RuntimeGrid, GridActors[GridIndex]);
+	int32 ActorCount = 0;
+
+	const UNavigationSystemBase* NavSystem = World->GetNavigationSystem();
+	if (NavSystem == nullptr)
+	{
+		UE_LOG(LogWorldPartitionRuntimeSpatialHash, Verbose, TEXT("No navigation system to generate navigation data."));
+		return false;
+	}
+
+	GridHelper.ForEachCells([GridHelper, RuntimeGrid, WorldPartition, &ActorCount, World, NavSystem, this](const FIntVector& CellCoord)
+	{
+		const int32 Level = CellCoord.Z;
+		if (Level != 3) //Todo AT, only generate for level 3 for now
+		{
+			return;
+		}
+
+		FBox2D CellBounds;
+		GridHelper.GetCellBounds(CellCoord, CellBounds);
+		const float HalfHeight = 100000.f; //Todo AT: try use WORLD_MAX instead
+		const FBox Bounds(FVector(CellBounds.Min.X, CellBounds.Min.Y, -HalfHeight), FVector(CellBounds.Max.X, CellBounds.Max.Y, HalfHeight));
+
+		if (NavSystem->ContainsNavData(Bounds) == false)
+		{
+			// Skip if there is no navdata for this cell
+			return;
+		}
+
+		const FSquare2DGridHelper::FGridLevel::FGridCell& GridCell = GridHelper.GetCell(CellCoord);
+
+		FActorSpawnParameters SpawnParams;
+		SpawnParams.bDeferConstruction = true;
+		SpawnParams.bCreateActorPackage = true;
+		ANavigationDataChunkActor* DataChunkActor = World->SpawnActor<ANavigationDataChunkActor>(SpawnParams);
+		ActorCount++;
+
+		const FName CellName = GetCellName(RuntimeGrid.GridName, CellCoord.Z, CellCoord.X, CellCoord.Y);
+		DataChunkActor->SetActorLabel(FString::Printf(TEXT("NavDataChunkActor_%s_%s"), *GetName(), *CellName.ToString()));
+			
+		const FVector2D CellCenter = CellBounds.GetCenter();
+		DataChunkActor->SetActorLocation(FVector(CellCenter.X, CellCenter.Y, 0.f));
+			
+		DataChunkActor->CollectNavData(Bounds);
+
+		// Set target grid
+		DataChunkActor->RuntimeGrid = RuntimeGrid.GridName;
+
+		WorldPartition->UpdateActorDesc(DataChunkActor);
+	});
+
+	UE_LOG(LogWorldPartitionRuntimeSpatialHash, Verbose, TEXT("   %i ANavigationDataChunkActor actors added."), ActorCount);
+
+	return true;
+}
+#endif //WITH_EDITOR
 
 FAutoConsoleCommand UWorldPartitionRuntimeSpatialHash::OverrideLoadingRangeCommand(
 	TEXT("wp.Runtime.OverrideRuntimeSpatialHashLoadingRange"),
@@ -1482,4 +1564,4 @@ void UWorldPartitionRuntimeSpatialHash::Draw3D(const TArray<FWorldPartitionStrea
 			StreamingGrid.Draw3D(World, Sources);
 		}
 	}
-}
\ No newline at end of file
+}
diff --git a/Engine/Source/Runtime/Engine/Public/WorldPartition/NavigationData/NavigationDataChunkActor.h b/Engine/Source/Runtime/Engine/Public/WorldPartition/NavigationData/NavigationDataChunkActor.h
new file mode 100644
--- /dev/null
+++ b/Engine/Source/Runtime/Engine/Public/WorldPartition/NavigationData/NavigationDataChunkActor.h
@@ -0,0 +1,41 @@
+// Copyright Epic Games, Inc. All Rights Reserved.
+
+#pragma once
+
+#include "GameFramework/Actor.h"
+#include "Containers/Array.h"
+#include "NavigationDataChunkActor.generated.h"
+
+class UNavigationDataChunk;
+
+UCLASS(NotPlaceable)
+class ENGINE_API ANavigationDataChunkActor : public AActor
+{
+	GENERATED_UCLASS_BODY()
+
+public:
+	//~ Begin UObject Interface
+	virtual void Serialize(FArchive& Ar) override;
+	virtual void PostLoad() override;
+	//~ End UObject Interface
+
+	const TArray<UNavigationDataChunk*>& GetNavDataChunk() const { return NavDataChunks; }
+	TArray<UNavigationDataChunk*>& GetMutableNavDataChunk() { return NavDataChunks; }
+
+	void CollectNavData(const FBox& Bounds);
+
+protected:
+	//~ Begin AActor Interface.
+	virtual void BeginPlay() override;
+	virtual void EndPlay(const EEndPlayReason::Type EndPlayReason) override;
+	//~ End AActor Interface.
+
+#if WITH_EDITOR
+	//~ Begin AActor Interface.
+	virtual EActorGridPlacement GetDefaultGridPlacement() const override;
+	//~ End AActor Interface.
+#endif
+
+	UPROPERTY()
+	TArray<UNavigationDataChunk*> NavDataChunks;
+};
diff --git a/Engine/Source/Runtime/Engine/Public/WorldPartition/WorldPartition.h b/Engine/Source/Runtime/Engine/Public/WorldPartition/WorldPartition.h
--- a/Engine/Source/Runtime/Engine/Public/WorldPartition/WorldPartition.h
+++ b/Engine/Source/Runtime/Engine/Public/WorldPartition/WorldPartition.h
@@ -130,6 +130,7 @@ public:
 	FBox GetWorldBounds() const;
 
 	void GenerateHLOD();
+	void GenerateNavigationData();
 
 	// Clustering
 	struct FActorCluster
diff --git a/Engine/Source/Runtime/Engine/Public/WorldPartition/WorldPartitionRuntimeHash.h b/Engine/Source/Runtime/Engine/Public/WorldPartition/WorldPartitionRuntimeHash.h
--- a/Engine/Source/Runtime/Engine/Public/WorldPartition/WorldPartitionRuntimeHash.h
+++ b/Engine/Source/Runtime/Engine/Public/WorldPartition/WorldPartitionRuntimeHash.h
@@ -21,6 +21,7 @@ class ENGINE_API UWorldPartitionRuntimeHash : public UObject
 	virtual bool GenerateStreaming(EWorldPartitionStreamingMode Mode, class UWorldPartitionStreamingPolicy* Policy) { return false; }
 	virtual void FlushStreaming() {}
 	virtual bool GenerateHLOD() { return false; }
+	virtual bool GenerateNavigationData() { return false; }
 	virtual FName GetActorRuntimeGrid(const AActor* Actor) const { return NAME_None; }
 #endif
 
diff --git a/Engine/Source/Runtime/Engine/Public/WorldPartition/WorldPartitionRuntimeSpatialHash.h b/Engine/Source/Runtime/Engine/Public/WorldPartition/WorldPartitionRuntimeSpatialHash.h
--- a/Engine/Source/Runtime/Engine/Public/WorldPartition/WorldPartitionRuntimeSpatialHash.h
+++ b/Engine/Source/Runtime/Engine/Public/WorldPartition/WorldPartitionRuntimeSpatialHash.h
@@ -148,6 +148,7 @@ public:
 	virtual bool GenerateStreaming(EWorldPartitionStreamingMode Mode, class UWorldPartitionStreamingPolicy* StreamingPolicy) override;
 	virtual void FlushStreaming() override;
 	virtual bool GenerateHLOD() override;
+	virtual bool GenerateNavigationData() override;
 	virtual FName GetActorRuntimeGrid(const AActor* Actor) const override;
 #endif
 
diff --git a/Engine/Source/Runtime/NavigationSystem/Private/NavMesh/RecastNavMesh.cpp b/Engine/Source/Runtime/NavigationSystem/Private/NavMesh/RecastNavMesh.cpp
--- a/Engine/Source/Runtime/NavigationSystem/Private/NavMesh/RecastNavMesh.cpp
+++ b/Engine/Source/Runtime/NavigationSystem/Private/NavMesh/RecastNavMesh.cpp
@@ -19,6 +19,7 @@
 #include "NavMesh/RecastNavMeshDataChunk.h"
 #include "NavMesh/RecastQueryFilter.h"
 #include "VisualLogger/VisualLogger.h"
+#include "WorldPartition/NavigationData/NavigationDataChunkActor.h"
 
 #if WITH_EDITOR
 #include "ObjectEditorUtils.h"
@@ -1878,9 +1879,19 @@ void ARecastNavMesh::InvalidateAffectedPaths(const TArray<uint32>& ChangedTiles)
 }
 
 URecastNavMeshDataChunk* ARecastNavMesh::GetNavigationDataChunk(ULevel* InLevel) const
+{
+	return GetNavigationDataChunk(InLevel->NavDataChunks);
+}
+
+URecastNavMeshDataChunk* ARecastNavMesh::GetNavigationDataChunk(const ANavigationDataChunkActor& InActor) const
+{
+	return GetNavigationDataChunk(InActor.GetNavDataChunk());
+}
+
+URecastNavMeshDataChunk* ARecastNavMesh::GetNavigationDataChunk(const TArray<UNavigationDataChunk*>& InChunks) const
 {
 	FName ThisName = GetFName();
-	int32 ChunkIndex = InLevel->NavDataChunks.IndexOfByPredicate([&](UNavigationDataChunk* Chunk) 
+	int32 ChunkIndex = InChunks.IndexOfByPredicate([&](UNavigationDataChunk* Chunk)
 	{
 		return Chunk->NavigationDataName == ThisName;
 	});
@@ -1888,7 +1899,7 @@ URecastNavMeshDataChunk* ARecastNavMesh::GetNavigationDataChunk(ULevel* InLevel)
 	URecastNavMeshDataChunk* RcNavDataChunk = nullptr;
 	if (ChunkIndex != INDEX_NONE)
 	{
-		RcNavDataChunk = Cast<URecastNavMeshDataChunk>(InLevel->NavDataChunks[ChunkIndex]);
+		RcNavDataChunk = Cast<URecastNavMeshDataChunk>(InChunks[ChunkIndex]);
 	}
 		
 	return RcNavDataChunk;
@@ -2076,6 +2087,55 @@ void ARecastNavMesh::ApplyWorldOffset(const FVector& InOffset, bool bWorldShift)
 	RequestDrawingUpdate();
 }
 
+void ARecastNavMesh::FillNavigationDataChunkActor(const FBox& Bounds, ANavigationDataChunkActor& DataChunkActor) const
+{
+	if (RecastNavMeshImpl)
+	{
+		UE_LOG(LogNavigation, Verbose, TEXT("%s Bounds pos: (%s)  size: (%s)."), ANSI_TO_TCHAR(__FUNCTION__), *Bounds.GetCenter().ToString(), *Bounds.GetSize().ToString());
+
+		const TArray<FBox> Boxes({ Bounds });
+		TArray<int32> TileIndices;
+		RecastNavMeshImpl->GetNavMeshTilesIn(Boxes, TileIndices);
+		if (!TileIndices.IsEmpty())
+		{
+			// Add a data chunk for this navmesh
+			URecastNavMeshDataChunk* DataChunk = NewObject<URecastNavMeshDataChunk>(&DataChunkActor);
+			DataChunk->NavigationDataName = GetFName();
+			DataChunkActor.GetMutableNavDataChunk().Add(DataChunk);
+
+			DataChunk->GetTiles(RecastNavMeshImpl, TileIndices, EGatherTilesCopyMode::CopyData);
+		}
+	}
+}
+
+void ARecastNavMesh::OnStreamingNavDataAdded(ANavigationDataChunkActor& InActor)
+{
+	QUICK_SCOPE_CYCLE_COUNTER(STAT_RecastNavMesh_OnStreamingNavDataAdded);
+
+	if (SupportsStreaming() && RecastNavMeshImpl)
+	{
+		URecastNavMeshDataChunk* NavDataChunk = GetNavigationDataChunk(InActor);
+		if (NavDataChunk)
+		{
+			AttachNavMeshDataChunk(*NavDataChunk);
+		}
+	}
+}
+
+void ARecastNavMesh::OnStreamingNavDataRemoved(ANavigationDataChunkActor& InActor)
+{
+	QUICK_SCOPE_CYCLE_COUNTER(STAT_RecastNavMesh_OnStreamingNavDataRemoved);
+
+	if (SupportsStreaming() && RecastNavMeshImpl)
+	{
+		URecastNavMeshDataChunk* NavDataChunk = GetNavigationDataChunk(InActor);
+		if (NavDataChunk)
+		{
+			DetachNavMeshDataChunk(*NavDataChunk);
+		}
+	}
+}
+
 void ARecastNavMesh::OnStreamingLevelAdded(ULevel* InLevel, UWorld* InWorld)
 {
 	QUICK_SCOPE_CYCLE_COUNTER(STAT_RecastNavMesh_OnStreamingLevelAdded);
diff --git a/Engine/Source/Runtime/NavigationSystem/Private/NavigationSystem.cpp b/Engine/Source/Runtime/NavigationSystem/Private/NavigationSystem.cpp
--- a/Engine/Source/Runtime/NavigationSystem/Private/NavigationSystem.cpp
+++ b/Engine/Source/Runtime/NavigationSystem/Private/NavigationSystem.cpp
@@ -1834,6 +1834,51 @@ ANavigationData* UNavigationSystemV1::GetNavDataForAgentName(const FName AgentNa
 	return Result;
 }
 
+bool UNavigationSystemV1::ContainsNavData(const FBox& Bounds) const
+{
+	for (const ANavigationData* NavData : NavDataSet)
+	{
+		if (NavData && Bounds.Intersect(NavData->GetBounds()))
+		{
+			return true;
+		}
+	}
+	return false;
+}
+
+void UNavigationSystemV1::AddNavigationDataChunk(ANavigationDataChunkActor& DataChunkActor)
+{
+	for (ANavigationData* NavData : NavDataSet)
+	{
+		if (NavData)
+		{
+			NavData->OnStreamingNavDataAdded(DataChunkActor);
+		}
+	}
+}
+
+void UNavigationSystemV1::RemoveNavigationDataChunk(ANavigationDataChunkActor& DataChunkActor)
+{
+	for (ANavigationData* NavData : NavDataSet)
+	{
+		if (NavData)
+		{
+			NavData->OnStreamingNavDataRemoved(DataChunkActor);
+		}
+	}
+}
+
+void UNavigationSystemV1::FillNavigationDataChunkActor(const FBox& Bounds, ANavigationDataChunkActor& DataChunkActor)
+{
+	for (const ANavigationData* NavData : NavDataSet)
+	{
+		if (NavData)
+		{
+			NavData->FillNavigationDataChunkActor(Bounds, DataChunkActor);
+		}
+	}
+}
+
 ANavigationData* UNavigationSystemV1::GetDefaultNavDataInstance(FNavigationSystem::ECreateIfMissing CreateNewIfNoneFound)
 {
 	checkSlow(IsInGameThread() == true);
diff --git a/Engine/Source/Runtime/NavigationSystem/Public/NavMesh/RecastNavMesh.h b/Engine/Source/Runtime/NavigationSystem/Public/NavMesh/RecastNavMesh.h
--- a/Engine/Source/Runtime/NavigationSystem/Public/NavMesh/RecastNavMesh.h
+++ b/Engine/Source/Runtime/NavigationSystem/Public/NavMesh/RecastNavMesh.h
@@ -759,6 +759,11 @@ public:
 	/** Called on world origin changes **/
 	virtual void ApplyWorldOffset(const FVector& InOffset, bool bWorldShift) override;
 
+	virtual void FillNavigationDataChunkActor(const FBox& InBounds, class ANavigationDataChunkActor& DataChunkActor) const override;
+
+	virtual void OnStreamingNavDataAdded(class ANavigationDataChunkActor& InActor) override;
+	virtual void OnStreamingNavDataRemoved(class ANavigationDataChunkActor& InActor) override;
+	
 	virtual void OnStreamingLevelAdded(ULevel* InLevel, UWorld* InWorld) override;
 	virtual void OnStreamingLevelRemoved(ULevel* InLevel, UWorld* InWorld) override;
 	//~ End ANavigationData Interface
@@ -1080,12 +1085,18 @@ private:
 	/** @return Navmesh data chunk that belongs to this actor */
 	URecastNavMeshDataChunk* GetNavigationDataChunk(ULevel* InLevel) const;
 
+	/** @return Navmesh data chunk that belongs to this actor */
+	URecastNavMeshDataChunk* GetNavigationDataChunk(const ANavigationDataChunkActor& InActor) const;
+
 protected:
 	// retrieves RecastNavMeshImpl
 	FPImplRecastNavMesh* GetRecastNavMeshImpl() { return RecastNavMeshImpl; }
 	const FPImplRecastNavMesh* GetRecastNavMeshImpl() const { return RecastNavMeshImpl; }
 
 private:
+	/** @return Navmesh data chunk that belongs to this actor */
+	URecastNavMeshDataChunk* GetNavigationDataChunk(const TArray<UNavigationDataChunk*>& InChunks) const;
+
 	/** NavMesh versioning. */
 	uint32 NavMeshVersion;
 	
diff --git a/Engine/Source/Runtime/NavigationSystem/Public/NavigationData.h b/Engine/Source/Runtime/NavigationSystem/Public/NavigationData.h
--- a/Engine/Source/Runtime/NavigationSystem/Public/NavigationData.h
+++ b/Engine/Source/Runtime/NavigationSystem/Public/NavigationData.h
@@ -586,6 +586,11 @@ public:
 	virtual bool SupportsRuntimeGeneration() const;
 	virtual bool SupportsStreaming() const;
 	virtual void OnNavigationBoundsChanged();
+
+	virtual void FillNavigationDataChunkActor(const FBox& Bounds, class ANavigationDataChunkActor& DataChunkActor) const {};
+	virtual void OnStreamingNavDataAdded(class ANavigationDataChunkActor& InActor) {};
+	virtual void OnStreamingNavDataRemoved(class ANavigationDataChunkActor& InActor) {};
+	
 	virtual void OnStreamingLevelAdded(ULevel* InLevel, UWorld* InWorld) {};
 	virtual void OnStreamingLevelRemoved(ULevel* InLevel, UWorld* InWorld) {};
 	
diff --git a/Engine/Source/Runtime/NavigationSystem/Public/NavigationSystem.h b/Engine/Source/Runtime/NavigationSystem/Public/NavigationSystem.h
--- a/Engine/Source/Runtime/NavigationSystem/Public/NavigationSystem.h
+++ b/Engine/Source/Runtime/NavigationSystem/Public/NavigationSystem.h
@@ -584,6 +584,11 @@ public:
 	virtual INavigationDataInterface* GetMainNavData() const override { return Cast<INavigationDataInterface>(GetDefaultNavDataInstance()); }
 	ANavigationData& GetMainNavDataChecked() const { check(MainNavData); return *MainNavData; }
 
+	virtual bool ContainsNavData(const FBox& Bounds) const override;
+	virtual void AddNavigationDataChunk(class ANavigationDataChunkActor& DataChunkActor) override;
+	virtual void RemoveNavigationDataChunk(class ANavigationDataChunkActor& DataChunkActor) override;
+	virtual void FillNavigationDataChunkActor(const FBox& Bounds, class ANavigationDataChunkActor& DataChunkActor) override;
+
 	ANavigationData* GetDefaultNavDataInstance() const { return MainNavData; }
 
 	ANavigationData* GetAbstractNavData() const { return AbstractNavData; }
