Move workflowcentric application to UnrealEd so it can be shared.

diff --git a/Engine/Source/Editor/Kismet/Private/ApplicationMode.cpp b/Engine/Source/Editor/Kismet/Private/ApplicationMode.cpp
deleted file mode 100644
diff --git a/Engine/Source/Editor/Kismet/Private/WorkflowCentricApplication.cpp b/Engine/Source/Editor/Kismet/Private/WorkflowCentricApplication.cpp
deleted file mode 100644
diff --git a/Engine/Source/Editor/Kismet/Private/WorkflowOrientedApp/SContentReference.cpp b/Engine/Source/Editor/Kismet/Private/WorkflowOrientedApp/SContentReference.cpp
deleted file mode 100644
diff --git a/Engine/Source/Editor/Kismet/Private/WorkflowOrientedApp/SModeWidget.cpp b/Engine/Source/Editor/Kismet/Private/WorkflowOrientedApp/SModeWidget.cpp
deleted file mode 100644
diff --git a/Engine/Source/Editor/Kismet/Private/WorkflowOrientedApp/WorkflowTabFactory.cpp b/Engine/Source/Editor/Kismet/Private/WorkflowOrientedApp/WorkflowTabFactory.cpp
deleted file mode 100644
diff --git a/Engine/Source/Editor/Kismet/Private/WorkflowOrientedApp/WorkflowTabManager.cpp b/Engine/Source/Editor/Kismet/Private/WorkflowOrientedApp/WorkflowTabManager.cpp
deleted file mode 100644
diff --git a/Engine/Source/Editor/Kismet/Public/WorkflowOrientedApp/ApplicationMode.h b/Engine/Source/Editor/Kismet/Public/WorkflowOrientedApp/ApplicationMode.h
deleted file mode 100644
diff --git a/Engine/Source/Editor/Kismet/Public/WorkflowOrientedApp/SContentReference.h b/Engine/Source/Editor/Kismet/Public/WorkflowOrientedApp/SContentReference.h
deleted file mode 100644
diff --git a/Engine/Source/Editor/Kismet/Public/WorkflowOrientedApp/SModeWidget.h b/Engine/Source/Editor/Kismet/Public/WorkflowOrientedApp/SModeWidget.h
deleted file mode 100644
diff --git a/Engine/Source/Editor/Kismet/Public/WorkflowOrientedApp/WorkflowCentricApplication.h b/Engine/Source/Editor/Kismet/Public/WorkflowOrientedApp/WorkflowCentricApplication.h
deleted file mode 100644
diff --git a/Engine/Source/Editor/Kismet/Public/WorkflowOrientedApp/WorkflowTabFactory.h b/Engine/Source/Editor/Kismet/Public/WorkflowOrientedApp/WorkflowTabFactory.h
deleted file mode 100644
diff --git a/Engine/Source/Editor/Kismet/Public/WorkflowOrientedApp/WorkflowTabManager.h b/Engine/Source/Editor/Kismet/Public/WorkflowOrientedApp/WorkflowTabManager.h
deleted file mode 100644
diff --git a/Engine/Source/Editor/Kismet/Public/WorkflowOrientedApp/WorkflowUObjectDocuments.h b/Engine/Source/Editor/Kismet/Public/WorkflowOrientedApp/WorkflowUObjectDocuments.h
deleted file mode 100644
diff --git a/Engine/Source/Editor/UnrealEd/Private/ApplicationMode.cpp b/Engine/Source/Editor/UnrealEd/Private/ApplicationMode.cpp
new file mode 100644
--- /dev/null
+++ b/Engine/Source/Editor/UnrealEd/Private/ApplicationMode.cpp
@@ -0,0 +1,65 @@
+// Copyright Epic Games, Inc. All Rights Reserved.
+
+#include "WorkflowOrientedApp/ApplicationMode.h"
+#include "Framework/MultiBox/MultiBoxExtender.h"
+#include "Framework/Docking/LayoutService.h"
+#include "UnrealEdMisc.h"
+
+#define LOCTEXT_NAMESPACE "ApplicationMode"
+
+/////////////////////////////////////////////////////
+// FApplicationMode
+
+FApplicationMode::FApplicationMode(FName InModeName)
+	: ModeName(InModeName)
+{
+	ToolbarExtender = MakeShareable(new FExtender);
+	WorkspaceMenuCategory = FWorkspaceItem::NewGroup(LOCTEXT("WorkspaceMenu_AssetEditor", "Asset Editor"));
+}
+
+FApplicationMode::FApplicationMode(FName InModeName, FText(*GetLocalizedMode)(const FName))
+	: ModeName(InModeName)
+{
+	ToolbarExtender = MakeShareable(new FExtender);
+	WorkspaceMenuCategory = FWorkspaceItem::NewGroup(FText::Format(LOCTEXT("WorkspaceMenu_ApplicationMode", "{0} Editor"), GetLocalizedMode(InModeName)));
+}
+
+void FApplicationMode::DeactivateMode(TSharedPtr<FTabManager> InTabManager)
+{
+	// Save the layout to INI
+	check(InTabManager.IsValid());
+	if (FUnrealEdMisc::Get().IsSavingLayoutOnClosedAllowed())
+	{
+		FLayoutSaveRestore::SaveToConfig(GEditorLayoutIni, InTabManager->PersistLayout());
+	}
+	// Unregister the tabs
+	/*
+	for (int32 Index = 0; Index < AllowableTabs.Num(); ++Index)
+	{
+		const FName TabName = AllowableTabs[Index];
+
+		PinnedTabManager->UnregisterTabSpawner(TabName);
+	}
+	*/
+}
+
+TSharedRef<FTabManager::FLayout> FApplicationMode::ActivateMode(TSharedPtr<FTabManager> InTabManager)
+{
+	check(InTabManager.IsValid());
+	RegisterTabFactories(InTabManager);
+
+	/*
+	for (int32 Index = 0; Index < AllowableTabs.Num(); ++Index)
+	{
+		const FName TabName = AllowableTabs[Index];
+
+		PinnedTabManager->RegisterTabSpawner(TabName, FOnSpawnTab::CreateSP(this, &FApplicationMode::SpawnTab, TabName));
+	}
+	*/
+
+	// Try loading the layout from INI
+	check(TabLayout.IsValid());
+	return FLayoutSaveRestore::LoadFromConfig(GEditorLayoutIni, TabLayout.ToSharedRef());
+}
+
+#undef LOCTEXT_NAMESPACE
diff --git a/Engine/Source/Editor/UnrealEd/Private/WorkflowCentricApplication.cpp b/Engine/Source/Editor/UnrealEd/Private/WorkflowCentricApplication.cpp
new file mode 100644
--- /dev/null
+++ b/Engine/Source/Editor/UnrealEd/Private/WorkflowCentricApplication.cpp
@@ -0,0 +1,161 @@
+// Copyright Epic Games, Inc. All Rights Reserved.
+
+#include "WorkflowOrientedApp/WorkflowCentricApplication.h"
+#include "Framework/Application/SlateApplication.h"
+#include "WorkflowOrientedApp/WorkflowTabFactory.h"
+#include "WorkflowOrientedApp/WorkflowTabManager.h"
+#include "WorkflowOrientedApp/ApplicationMode.h"
+#include "ToolMenus.h"
+
+#define LOCTEXT_NAMESPACE "WorkflowCentricApplication"
+
+/////////////////////////////////////////////////////
+// FWorkflowCentricApplication
+
+TArray<FWorkflowApplicationModeExtender> FWorkflowCentricApplication::ModeExtenderList;
+
+void FWorkflowCentricApplication::RegisterTabSpawners(const TSharedRef<class FTabManager>& InTabManager)
+{
+	if (CurrentAppModePtr.IsValid())
+	{
+		CurrentAppModePtr->RegisterTabFactories(InTabManager);
+	}
+}
+
+void FWorkflowCentricApplication::UnregisterTabSpawners(const TSharedRef<class FTabManager>& InTabManager)
+{
+	InTabManager->UnregisterAllTabSpawners();
+}
+
+FName FWorkflowCentricApplication::GetToolMenuToolbarName(FName& OutParentName) const
+{
+	return GetToolMenuToolbarNameForMode(GetCurrentMode(), OutParentName);
+}
+
+FName FWorkflowCentricApplication::GetToolMenuToolbarNameForMode(const FName InModeName, FName& OutParentName) const
+{
+	const FName BaseMenuName = FAssetEditorToolkit::GetToolMenuToolbarName(OutParentName);
+	if (InModeName != NAME_None)
+	{
+		OutParentName = BaseMenuName;
+		return *(BaseMenuName.ToString() + TEXT(".") + InModeName.ToString());
+	}
+
+	return BaseMenuName;
+}
+
+UToolMenu* FWorkflowCentricApplication::RegisterModeToolbarIfUnregistered(const FName InModeName)
+{
+	FName ParentToolbarName;
+	const FName ModeSpecificToolbarName = GetToolMenuToolbarNameForMode(InModeName, ParentToolbarName);
+	if (!UToolMenus::Get()->IsMenuRegistered(ModeSpecificToolbarName))
+	{
+		return UToolMenus::Get()->RegisterMenu(ModeSpecificToolbarName, ParentToolbarName, EMultiBoxType::ToolBar);
+	}
+
+	return nullptr;
+}
+
+FName FWorkflowCentricApplication::GetCurrentMode() const
+{
+	return CurrentAppModePtr.IsValid() ? CurrentAppModePtr->GetModeName() : NAME_None;
+}
+
+void FWorkflowCentricApplication::SetCurrentMode(FName NewMode)
+{
+	const bool bModeAlreadyActive = CurrentAppModePtr.IsValid() && (NewMode == CurrentAppModePtr->GetModeName());
+
+	if (!bModeAlreadyActive)
+	{
+		check(TabManager.IsValid());
+
+		TSharedPtr<FApplicationMode> NewModePtr = ApplicationModeList.FindRef(NewMode);
+
+		LayoutExtenders.Reset();
+
+		if (NewModePtr.IsValid())
+		{
+			if (NewModePtr->LayoutExtender.IsValid())
+			{
+				LayoutExtenders.Add(NewModePtr->LayoutExtender);
+			}
+			
+			// Deactivate the old mode
+			if (CurrentAppModePtr.IsValid())
+			{
+				check(TabManager.IsValid());
+				CurrentAppModePtr->PreDeactivateMode();
+				CurrentAppModePtr->DeactivateMode(TabManager);
+				RemoveToolbarExtender(CurrentAppModePtr->GetToolbarExtender());
+				RemoveAllToolbarWidgets();
+			}
+
+			// Unregister tab spawners
+			TabManager->UnregisterAllTabSpawners();
+
+			//@TODO: Should do some validation here
+			CurrentAppModePtr = NewModePtr;
+
+			// Establish the workspace menu category for the new mode
+			TabManager->ClearLocalWorkspaceMenuCategories();
+			TabManager->AddLocalWorkspaceMenuItem( CurrentAppModePtr->GetWorkspaceMenuCategory() );
+
+			// Activate the new layout
+			const TSharedRef<FTabManager::FLayout> NewLayout = CurrentAppModePtr->ActivateMode(TabManager);
+			RestoreFromLayout(NewLayout);
+
+			// Give the new mode a chance to do init
+			CurrentAppModePtr->PostActivateMode();
+
+			AddToolbarExtender(NewModePtr->GetToolbarExtender());
+			RegenerateMenusAndToolbars();
+		}
+	}
+}
+
+void FWorkflowCentricApplication::PushTabFactories(FWorkflowAllowedTabSet& FactorySetToPush)
+{
+	check(TabManager.IsValid());
+	for (auto FactoryIt = FactorySetToPush.CreateIterator(); FactoryIt; ++FactoryIt)
+	{
+		FactoryIt.Value()->RegisterTabSpawner(TabManager.ToSharedRef(), CurrentAppModePtr.Get());
+	}
+}
+
+bool FWorkflowCentricApplication::OnRequestClose()
+{
+	if (!FSlateApplication::Get().IsNormalExecution())
+	{
+		return false;
+	}
+
+	if (CurrentAppModePtr.IsValid())
+	{
+		check(TabManager.IsValid());
+
+		// Deactivate the old mode
+		CurrentAppModePtr->PreDeactivateMode();
+		CurrentAppModePtr->DeactivateMode(TabManager);
+		RemoveToolbarExtender(CurrentAppModePtr->GetToolbarExtender());
+		RemoveAllToolbarWidgets();
+
+		// Unregister tab spawners
+		TabManager->UnregisterAllTabSpawners();
+	}
+
+	return true;
+}
+
+void FWorkflowCentricApplication::AddApplicationMode(FName ModeName, TSharedRef<FApplicationMode> Mode)
+{
+	for (int32 Index = 0; Index < ModeExtenderList.Num(); ++Index)
+	{
+		Mode = ModeExtenderList[Index].Execute(ModeName, Mode);
+	}
+
+	ApplicationModeList.Add(ModeName, Mode);
+}
+
+/////////////////////////////////////////////////////
+
+#undef LOCTEXT_NAMESPACE
diff --git a/Engine/Source/Editor/UnrealEd/Private/WorkflowOrientedApp/SContentReference.cpp b/Engine/Source/Editor/UnrealEd/Private/WorkflowOrientedApp/SContentReference.cpp
new file mode 100644
--- /dev/null
+++ b/Engine/Source/Editor/UnrealEd/Private/WorkflowOrientedApp/SContentReference.cpp
@@ -0,0 +1,296 @@
+// Copyright Epic Games, Inc. All Rights Reserved.
+
+#include "WorkflowOrientedApp/SContentReference.h"
+#include "Widgets/SBoxPanel.h"
+#include "Widgets/Layout/SBorder.h"
+#include "Widgets/Images/SImage.h"
+#include "Widgets/Text/STextBlock.h"
+#include "Widgets/Layout/SBox.h"
+#include "Widgets/Input/SButton.h"
+#include "Widgets/Input/SComboButton.h"
+#include "EditorStyleSet.h"
+#include "Editor.h"
+#include "ContentBrowserModule.h"
+
+#define LOCTEXT_NAMESPACE "SContentReference"
+
+//////////////////////////////////////////////////////////////////////////
+// SContentReference
+
+void SContentReference::Construct(const FArguments& InArgs)
+{
+	// Save off the attributes
+	AssetReference = InArgs._AssetReference;
+	ShowFindInBrowserButton = InArgs._ShowFindInBrowserButton;
+	ShowToolsButton = InArgs._ShowToolsButton;
+	AllowSelectingNewAsset = InArgs._AllowSelectingNewAsset;
+	AllowClearingReference = InArgs._AllowClearingReference;
+	AllowedClass = InArgs._AllowedClass;
+	AssetPickerSizeOverride = InArgs._AssetPickerSizeOverride;
+	InitialAssetViewType = InArgs._InitialAssetViewType;
+
+	// Save off delegates
+	OnShouldFilterAsset = InArgs._OnShouldFilterAsset;
+	OnSetReference = InArgs._OnSetReference;
+
+	// Cache resources
+	BorderImageNormal = FEditorStyle::GetBrush(InArgs._Style, ".Background.Normal");
+	BorderImageHovered = FEditorStyle::GetBrush(InArgs._Style, ".Background.Hovered");
+
+	static const FName InvertedForegroundName("InvertedForeground");
+
+	PickerComboButton = SNew(SComboButton)
+		.ButtonStyle(FEditorStyle::Get(), "NoBorder")
+		.ContentPadding(1.f)
+		.Visibility(this, &SContentReference::GetPickButtonVisibility)
+		.OnGetMenuContent(this, &SContentReference::MakeAssetPickerMenu)
+		.HasDownArrow(false)
+		.ToolTipText(LOCTEXT("PickAsset", "Pick an asset from a popup menu"))
+		.ButtonContent()
+		[
+			SNew(SImage)
+			.Image(FEditorStyle::GetBrush(InArgs._Style, ".PickAsset"))
+		];
+
+	// Create the widgets
+	ChildSlot
+	[
+		SNew(SHorizontalBox)
+
+		// Text box containing the name of the asset
+		+SHorizontalBox::Slot()
+		.AutoWidth()
+		.Padding(1,0)
+		[
+			SAssignNew(AssetReferenceNameBorderWidget, SBorder)
+			.BorderImage(this, &SContentReference::GetBorderImage)
+			.Padding(FEditorStyle::GetMargin(InArgs._Style, ".BorderPadding"))
+			.BorderBackgroundColor( FLinearColor::White )
+			.ForegroundColor(FEditorStyle::GetSlateColor(InvertedForegroundName))
+			.ToolTipText(this, &SContentReference::GetAssetFullName)
+			.OnMouseDoubleClick(this, &SContentReference::OnDoubleClickedOnAssetName)
+			[
+				SNew(SBox)
+				.WidthOverride(InArgs._WidthOverride)
+				[
+					SNew(STextBlock)
+					.Text(this, &SContentReference::GetAssetShortName)
+				]
+			]
+		]
+
+		// Pick an asset
+		+SHorizontalBox::Slot()
+		.AutoWidth()
+		.HAlign(HAlign_Center)
+		.VAlign(VAlign_Center)
+		.Padding(1,0)
+		[
+			PickerComboButton.ToSharedRef()
+		]
+			
+		// Find in content browser button
+		+SHorizontalBox::Slot()
+		.AutoWidth()
+		.HAlign(HAlign_Center)
+		.VAlign(VAlign_Center)
+		.Padding(1,0)
+		[
+			SNew(SButton)
+			.ButtonStyle( FEditorStyle::Get(), "NoBorder" )
+			.OnClicked(this, &SContentReference::OnClickFindButton)
+			.ContentPadding(0)
+			.Visibility(this, &SContentReference::GetFindButtonVisibility)
+			.ToolTipText(LOCTEXT("Find", "Find in content browser"))
+			[
+				SNew(SImage)
+				.Image( FEditorStyle::GetBrush(InArgs._Style, ".FindInContentBrowser") )
+			]
+		]
+			
+		// Clear button
+		+SHorizontalBox::Slot()
+		.AutoWidth()
+		.HAlign(HAlign_Center)
+		.VAlign(VAlign_Center)
+		.Padding(1,0)
+		[
+			SNew(SButton)
+			.ButtonStyle( FEditorStyle::Get(), "NoBorder" )
+			.OnClicked(this, &SContentReference::OnClickClearButton)
+			.ContentPadding(1.f)
+			.Visibility( this, &SContentReference::GetClearButtonVisibility)
+			.ToolTipText(LOCTEXT("Clear", "Clear"))
+			[
+				SNew(SImage)
+				.Image( FEditorStyle::GetBrush(InArgs._Style, ".Clear") )
+			]
+		]
+		
+		// Tools button
+		+SHorizontalBox::Slot()
+		.AutoWidth()
+		.HAlign(HAlign_Center)
+		.VAlign(VAlign_Center)
+		.Padding(1,0)
+		[
+			SNew(SButton)
+			.ButtonStyle( FEditorStyle::Get(), "NoBorder" )
+			.OnClicked(InArgs._OnClickedTools)
+			.ContentPadding(1.f)
+			.Visibility(this, &SContentReference::GetToolsButtonVisibility)
+			.ToolTipText(LOCTEXT("Tools", "Tools"))
+			[
+				SNew(SImage)
+				.Image( FEditorStyle::GetBrush(InArgs._Style, ".Tools") )
+			]
+		]
+	];
+}
+
+void SContentReference::OpenAssetPickerMenu()
+{
+	PickerComboButton->SetIsOpen(true);
+}
+
+EVisibility SContentReference::GetUseButtonVisibility() const
+{
+	return AllowSelectingNewAsset.Get() ? EVisibility::Visible : EVisibility::Collapsed;
+}
+
+EVisibility SContentReference::GetPickButtonVisibility() const
+{
+	return AllowSelectingNewAsset.Get() ? EVisibility::Visible : EVisibility::Collapsed;
+}
+
+EVisibility SContentReference::GetFindButtonVisibility() const
+{
+	return ShowFindInBrowserButton.Get() ? EVisibility::Visible : EVisibility::Collapsed;
+}
+
+EVisibility SContentReference::GetClearButtonVisibility() const
+{
+	return AllowClearingReference.Get() ? EVisibility::Visible : EVisibility::Collapsed;
+}
+
+EVisibility SContentReference::GetToolsButtonVisibility() const
+{
+	return ShowToolsButton.Get() ? EVisibility::Visible : EVisibility::Collapsed;
+}
+
+FReply SContentReference::OnClickUseButton()
+{
+	TArray<FAssetData> Objects;
+
+	FContentBrowserModule& ContentBrowserModule = FModuleManager::Get().LoadModuleChecked<FContentBrowserModule>(TEXT("ContentBrowser"));
+	ContentBrowserModule.Get().GetSelectedAssets(Objects);
+
+	for(auto It(Objects.CreateIterator());It;++It)
+	{
+		OnSetReference.ExecuteIfBound(It->GetAsset());
+	}
+	return FReply::Handled();
+}
+
+FReply SContentReference::OnClickFindButton()
+{
+	FindObjectInContentBrowser(AssetReference.Get());
+	return FReply::Handled();
+}
+
+FReply SContentReference::OnClickClearButton()
+{
+	OnSetReference.ExecuteIfBound(NULL);
+	return FReply::Handled();
+}
+
+TSharedRef<SWidget> SContentReference::MakeAssetPickerMenu()
+{
+	FContentBrowserModule& ContentBrowserModule = FModuleManager::Get().LoadModuleChecked<FContentBrowserModule>(TEXT("ContentBrowser"));
+
+	FAssetPickerConfig AssetPickerConfig;
+
+	UClass* FilterClass = AllowedClass.Get();
+	if (FilterClass != NULL)
+	{
+		AssetPickerConfig.Filter.ClassNames.Add(FilterClass->GetFName());
+		AssetPickerConfig.Filter.bRecursiveClasses = true;
+	}
+
+	AssetPickerConfig.OnAssetSelected = FOnAssetSelected::CreateSP(this, &SContentReference::OnAssetSelectedFromPicker);
+	AssetPickerConfig.OnShouldFilterAsset = OnShouldFilterAsset;
+	AssetPickerConfig.bAllowNullSelection = true;
+	AssetPickerConfig.ThumbnailLabel = EThumbnailLabel::ClassName;
+	AssetPickerConfig.InitialAssetViewType = InitialAssetViewType;
+
+	return SNew(SBox)
+		.WidthOverride(AssetPickerSizeOverride.Get().X)
+		.HeightOverride(AssetPickerSizeOverride.Get().Y)
+		[
+			ContentBrowserModule.Get().CreateAssetPicker(AssetPickerConfig)
+		];
+}
+
+void SContentReference::OnAssetSelectedFromPicker(const FAssetData& AssetData)
+{
+	PickerComboButton->SetIsOpen(false);
+	OnSetReference.ExecuteIfBound(AssetData.GetAsset());
+}
+
+FText SContentReference::GetAssetShortName() const
+{
+	if (UObject* Asset = AssetReference.Get())
+	{
+		return FText::FromString(Asset->GetName());
+	}
+	else
+	{
+		return LOCTEXT("NullReference", "(None)");
+	}
+}
+
+FText SContentReference::GetAssetFullName() const
+{
+	if (UObject* Asset = AssetReference.Get())
+	{
+		return FText::FromString(Asset->GetFullName());
+	}
+	else
+	{
+		return LOCTEXT("NullReferenceTooltip", "(None)");
+	}
+}
+
+FReply SContentReference::OnDoubleClickedOnAssetName(const FGeometry& InMyGeometry, const FPointerEvent& InMouseEvent)
+{
+	OpenAssetToEdit();
+	return FReply::Handled();
+}
+
+void SContentReference::OpenAssetToEdit()
+{
+	if (UObject* Asset = AssetReference.Get())
+	{
+		GEditor->EditObject(Asset);
+	}
+}
+
+const FSlateBrush* SContentReference::GetBorderImage() const
+{
+	return AssetReferenceNameBorderWidget->IsHovered() ? BorderImageHovered : BorderImageNormal;
+}
+
+void SContentReference::FindObjectInContentBrowser(UObject* Object)
+{
+	if ((Object != NULL) && (Object->IsAsset()))
+	{
+		TArray<UObject*> ObjectsToSyncTo;
+		ObjectsToSyncTo.Add(Object);
+
+		GEditor->SyncBrowserToObjects(ObjectsToSyncTo);
+	}
+}
+
+//////////////////////////////////////////////////////////////////////////
+
+#undef LOCTEXT_NAMESPACE
diff --git a/Engine/Source/Editor/UnrealEd/Private/WorkflowOrientedApp/SModeWidget.cpp b/Engine/Source/Editor/UnrealEd/Private/WorkflowOrientedApp/SModeWidget.cpp
new file mode 100644
--- /dev/null
+++ b/Engine/Source/Editor/UnrealEd/Private/WorkflowOrientedApp/SModeWidget.cpp
@@ -0,0 +1,179 @@
+// Copyright Epic Games, Inc. All Rights Reserved.
+
+#include "WorkflowOrientedApp/SModeWidget.h"
+#include "Misc/Paths.h"
+#include "Widgets/SBoxPanel.h"
+#include "Widgets/Layout/SBorder.h"
+#include "Widgets/Images/SImage.h"
+#include "Widgets/Text/STextBlock.h"
+#include "Framework/MultiBox/MultiBoxDefs.h"
+#include "Styling/CoreStyle.h"
+#include "EditorStyleSet.h"
+
+//////////////////////////////////////////////////////////////////////////
+// SModeWidget
+
+void SModeWidget::Construct(const FArguments& InArgs, const FText& InText, const FName InMode)
+{
+	// Copy arguments
+	ModeText = InText;
+	ThisMode = InMode;
+	OnGetActiveMode = InArgs._OnGetActiveMode;
+	CanBeSelected = InArgs._CanBeSelected;
+	OnSetActiveMode = InArgs._OnSetActiveMode;
+
+	// Load resources
+	InactiveModeBorderImage = FEditorStyle::GetBrush("ModeSelector.ToggleButton.Normal");
+	ActiveModeBorderImage = FEditorStyle::GetBrush("ModeSelector.ToggleButton.Pressed");
+	HoverBorderImage = FEditorStyle::GetBrush("ModeSelector.ToggleButton.Hovered");
+	
+	TSharedRef<SHorizontalBox> InnerRow = SNew(SHorizontalBox);
+
+	FMargin IconPadding(4.0f, 0.0f, 4.0f, 0.0f);
+	FMargin BodyPadding(0.0f, 0.0f, 0.0f, 0.0f);
+	
+	// Large Icon
+	if (InArgs._IconImage.IsSet())
+	{
+		InnerRow->AddSlot()
+			.AutoWidth()
+			.VAlign(VAlign_Center)
+			.Padding(IconPadding)
+			[
+				SNew(SImage)
+				.Image(InArgs._IconImage)
+				.Visibility(this, &SModeWidget::GetLargeIconVisibility)
+			];
+	}
+
+	// Small Icon
+	if (InArgs._SmallIconImage.IsSet())
+	{
+		InnerRow->AddSlot()
+			.AutoWidth()
+			.VAlign(VAlign_Center)
+			.Padding(IconPadding)
+			[
+				SNew(SImage)
+				.Image(InArgs._SmallIconImage)
+				.Visibility(this, &SModeWidget::GetSmallIconVisibility)
+			];
+	}
+
+	// Label + content
+	InnerRow->AddSlot()
+		.AutoWidth()
+		.VAlign(VAlign_Center)
+		.Padding(BodyPadding)
+		[
+			SNew(SVerticalBox)
+
+			// Mode 'tab'
+			+SVerticalBox::Slot()
+			.AutoHeight()
+			[
+				SNew(SHorizontalBox)
+
+				//Mode Name
+				+SHorizontalBox::Slot()
+				.AutoWidth()
+				[
+					SNew(STextBlock)
+					.Text(ModeText)
+					.Font(this, &SModeWidget::GetDesiredTitleFont)
+				]
+
+				//Dirty flag
+				+SHorizontalBox::Slot()
+				.AutoWidth()
+				.Padding(3)
+				[
+					SNew(SImage)
+					.Image( InArgs._DirtyMarkerBrush )
+				]
+			]
+
+			// Body of 'ribbon'
+			+SVerticalBox::Slot()
+			.AutoHeight()
+			[
+
+				InArgs._ShortContents.Widget
+			]
+		];
+
+
+	// Create the widgets
+	ChildSlot
+	.VAlign(VAlign_Fill)
+	[
+		SNew(SBorder)
+		.BorderImage(this, &SModeWidget::GetModeNameBorderImage)
+		.OnMouseButtonDown(this, &SModeWidget::OnModeTabClicked)
+		[
+			InnerRow
+		]
+	];
+
+	SetEnabled(CanBeSelected);
+}
+
+EVisibility SModeWidget::GetLargeIconVisibility() const
+{
+	return FMultiBoxSettings::UseSmallToolBarIcons.Get() ? EVisibility::Collapsed : EVisibility::Visible;
+}
+
+EVisibility SModeWidget::GetSmallIconVisibility() const
+{
+	return FMultiBoxSettings::UseSmallToolBarIcons.Get() ? EVisibility::Visible : EVisibility::Collapsed;
+}
+
+const FSlateBrush* SModeWidget::GetModeNameBorderImage() const
+{
+	if (IsActiveMode())
+	{
+		return ActiveModeBorderImage;
+	}
+	else if (IsHovered())
+	{
+		return HoverBorderImage;
+	}
+	else
+	{
+		return InactiveModeBorderImage;
+	}
+}
+
+bool SModeWidget::IsActiveMode() const
+{
+	return OnGetActiveMode.Get() == ThisMode;
+}
+
+FReply SModeWidget::OnModeTabClicked(const FGeometry& MyGeometry, const FPointerEvent& MouseEvent)
+{
+	if (MouseEvent.GetEffectingButton() == EKeys::LeftMouseButton)
+	{
+		// Try to change the mode
+		if (!IsActiveMode() && (CanBeSelected.Get() == true))
+		{
+			OnSetActiveMode.ExecuteIfBound(ThisMode);
+		}
+	}
+
+	return FReply::Handled();
+}
+
+FSlateFontInfo SModeWidget::GetDesiredTitleFont() const
+{
+	const bool bSmallIcons = FMultiBoxSettings::UseSmallToolBarIcons.Get();	
+	const int16 FontSize = bSmallIcons ? 10 : 14;
+
+	if (IsActiveMode())
+	{
+		return FCoreStyle::GetDefaultFontStyle(bSmallIcons ? TEXT("Bold") : TEXT("Regular"), FontSize);
+	}
+	else
+	{
+		return FCoreStyle::GetDefaultFontStyle(bSmallIcons ? TEXT("Regular") : TEXT("Light"), FontSize);
+	}
+}
diff --git a/Engine/Source/Editor/UnrealEd/Private/WorkflowOrientedApp/WorkflowTabFactory.cpp b/Engine/Source/Editor/UnrealEd/Private/WorkflowOrientedApp/WorkflowTabFactory.cpp
new file mode 100644
--- /dev/null
+++ b/Engine/Source/Editor/UnrealEd/Private/WorkflowOrientedApp/WorkflowTabFactory.cpp
@@ -0,0 +1,189 @@
+// Copyright Epic Games, Inc. All Rights Reserved.
+
+#include "WorkflowOrientedApp/WorkflowTabFactory.h"
+#include "Widgets/Layout/SBorder.h"
+#include "Widgets/Text/STextBlock.h"
+#include "Framework/MultiBox/MultiBoxBuilder.h"
+#include "EditorStyleSet.h"
+#include "WorkflowOrientedApp/ApplicationMode.h"
+#include "WorkflowOrientedApp/WorkflowTabManager.h"
+#include "IDocumentation.h"
+#include "Widgets/Docking/SDockTab.h"
+#include "Framework/Docking/TabManager.h"
+#include "Misc/BlacklistNames.h"
+
+/////////////////////////////////////////////////////
+// FWorkflowTabFactory
+
+FWorkflowTabFactory::FWorkflowTabFactory(FName InIdentifier, TSharedPtr<class FAssetEditorToolkit> InHostingApp)
+	: TabIdentifier(InIdentifier)
+	, TabRole(ETabRole::PanelTab)
+	, TabIcon(FSlateIcon())
+	, InsideTabPadding(0.0f)
+	, bIsSingleton(false)
+	, bShouldAutosize(false)
+	, HostingApp(InHostingApp)
+{
+}
+
+void FWorkflowTabFactory::EnableTabPadding()
+{
+	InsideTabPadding = 4.0f;
+}
+
+TSharedRef<SDockTab> FWorkflowTabFactory::SpawnTab(const FWorkflowTabSpawnInfo& Info) const
+{
+	// Get the tab body content
+	TSharedRef<SWidget> TabBody = CreateTabBody(Info);
+
+	// Pad the content if requested
+	if (InsideTabPadding > 0.0f)
+	{
+		// propagate tag from original content, or use TabId
+		TSharedPtr<FTagMetaData> MetaData = TabBody->GetMetaData<FTagMetaData>();
+
+		TabBody = SNew(SBorder)
+			.Padding(InsideTabPadding)
+			.BorderImage(FEditorStyle::GetBrush("ToolPanel.GroupBorder"))
+			.AddMetaData<FTagMetaData>(MetaData.IsValid() ? MetaData->Tag : this->TabIdentifier)
+			[
+				TabBody
+			];
+	}
+
+	// Spawn the tab
+	TSharedRef<SDockTab> NewTab = SNew(SDockTab)
+		.TabRole(this->TabRole)
+		.Icon(GetTabIcon(Info))
+		.Label(ConstructTabName(Info))
+		.ShouldAutosize(bShouldAutosize)
+		[
+			TabBody
+		];
+
+	NewTab->SetTabToolTipWidget(CreateTabToolTipWidget(Info));
+
+	return NewTab;
+}
+
+TSharedRef<SDockTab> FWorkflowTabFactory::SpawnBlankTab() const
+{
+	// Spawn the tab
+	TSharedRef<SDockTab> NewTab = SNew(SDockTab)
+		.TabRole(this->TabRole)
+		.ShouldAutosize(bShouldAutosize);
+
+	return NewTab;
+}
+
+
+void FWorkflowTabFactory::UpdateTab(TSharedPtr<SDockTab> InDockTab, const FWorkflowTabSpawnInfo& Info, TSharedPtr< SWidget > InContent) const
+{
+	// Get the tab body content
+	TSharedRef<SWidget> TabBody = InContent.ToSharedRef();
+
+	// Pad the content if requested
+	if (InsideTabPadding > 0.0f)
+	{
+		TabBody = SNew(SBorder)
+			.Padding(InsideTabPadding)
+			.BorderImage(FEditorStyle::GetBrush("ToolPanel.GroupBorder"))
+			[
+				TabBody
+			];
+	}
+
+	InDockTab->SetContent(TabBody);
+	InDockTab->SetLabel(ConstructTabName(Info));
+	InDockTab->SetTabIcon(GetTabIcon(Info));
+	InDockTab->SetTabToolTipWidget(CreateTabToolTipWidget(Info));
+}
+
+TSharedPtr<SToolTip> FWorkflowTabFactory::CreateTabToolTipWidget(const FWorkflowTabSpawnInfo& Info) const
+{
+	const FString DocLink = TEXT("Shared/Editors/Common/EditorTabs");
+	const FString DocExcerptName = TabIdentifier.ToString();
+	return IDocumentation::Get()->CreateToolTip(GetTabToolTipText(Info), NULL, DocLink, DocExcerptName);
+}
+
+TSharedRef<SDockTab> FWorkflowTabFactory::OnSpawnTab(const FSpawnTabArgs& SpawnArgs, TWeakPtr<FTabManager> WeakTabManager) const
+{
+	FWorkflowTabSpawnInfo SpawnInfo;
+	SpawnInfo.TabManager = WeakTabManager.Pin();
+
+	return SpawnInfo.TabManager.IsValid() ? SpawnTab(SpawnInfo) : SNew(SDockTab);
+}
+
+bool FWorkflowTabFactory::CanSpawnTab(const FSpawnTabArgs& SpawnArgs, TWeakPtr<FTabManager> WeakTabManager) const
+{
+	TSharedPtr<FTabManager> TabManager = WeakTabManager.Pin();
+	if (TabManager.IsValid())
+	{
+		return TabManager->GetTabBlacklist()->PassesFilter(SpawnArgs.GetTabId().TabType);
+	}
+
+	return true;
+}
+
+FTabSpawnerEntry& FWorkflowTabFactory::RegisterTabSpawner(TSharedRef<FTabManager> InTabManager, const FApplicationMode* CurrentApplicationMode) const
+{
+	FWorkflowTabSpawnInfo SpawnInfo;
+	SpawnInfo.TabManager = InTabManager;
+
+	TWeakPtr<FTabManager> WeakTabManager(InTabManager);
+	FTabSpawnerEntry& SpawnerEntry = InTabManager->RegisterTabSpawner(GetIdentifier(), FOnSpawnTab::CreateSP(this, &FWorkflowTabFactory::OnSpawnTab, WeakTabManager), FCanSpawnTab::CreateSP(this, &FWorkflowTabFactory::CanSpawnTab, WeakTabManager))
+		.SetDisplayName(ConstructTabName(SpawnInfo).Get())
+		.SetTooltipText(GetTabToolTipText(SpawnInfo));
+
+	if (CurrentApplicationMode)
+	{
+		SpawnerEntry.SetGroup(CurrentApplicationMode->GetWorkspaceMenuCategory());
+	}
+	
+	// Add the tab icon to the menu entry if one was provided
+	const FSlateIcon& TabSpawnerIcon = GetTabSpawnerIcon(SpawnInfo);
+	if (TabSpawnerIcon.IsSet())
+	{
+		SpawnerEntry.SetIcon(TabSpawnerIcon);
+	}
+
+	return SpawnerEntry;
+}
+
+TAttribute<FText> FWorkflowTabFactory::ConstructTabName(const FWorkflowTabSpawnInfo& Info) const
+{
+	return TabLabel;
+}
+
+TSharedRef<SWidget> FWorkflowTabFactory::CreateTabBody(const FWorkflowTabSpawnInfo& Info) const
+{
+	return SNew(STextBlock).Text(NSLOCTEXT("WorkflowTabFactory", "EmptyTabMessage", "EmptyBody"));
+}
+
+const FSlateBrush* FWorkflowTabFactory::GetTabIcon(const FWorkflowTabSpawnInfo& Info) const
+{
+	return TabIcon.GetIcon();
+}
+
+const FSlateIcon& FWorkflowTabFactory::GetTabSpawnerIcon(const FWorkflowTabSpawnInfo& Info) const
+{
+	return TabIcon;
+}
+
+void FWorkflowTabFactory::CreateViewMenuEntry(FMenuBuilder& MenuBuilder, const FUIAction& Action) const
+{
+	MenuBuilder.AddMenuEntry(ViewMenuDescription, ViewMenuTooltip, FSlateIcon(), Action);
+}
+
+/////////////////////////////////////////////////////
+// FDocumentTabFactory
+
+FDocumentTabFactory::FDocumentTabFactory(FName InIdentifier, TSharedPtr<class FAssetEditorToolkit> InHostingApp) : FWorkflowTabFactory(InIdentifier, InHostingApp)
+{
+	TabRole = ETabRole::DocumentTab;
+}
+
+TSharedRef<FGenericTabHistory> FDocumentTabFactory::CreateTabHistoryNode(TSharedPtr<FTabPayload> Payload)
+{
+	return MakeShareable(new FGenericTabHistory(SharedThis(this), Payload));
+}
diff --git a/Engine/Source/Editor/UnrealEd/Private/WorkflowOrientedApp/WorkflowTabManager.cpp b/Engine/Source/Editor/UnrealEd/Private/WorkflowOrientedApp/WorkflowTabManager.cpp
new file mode 100644
--- /dev/null
+++ b/Engine/Source/Editor/UnrealEd/Private/WorkflowOrientedApp/WorkflowTabManager.cpp
@@ -0,0 +1,950 @@
+// Copyright Epic Games, Inc. All Rights Reserved.
+
+#include "WorkflowOrientedApp/WorkflowTabManager.h"
+#include "Widgets/Layout/SBorder.h"
+#include "Widgets/Input/SMenuAnchor.h"
+#include "Framework/MultiBox/MultiBoxBuilder.h"
+#include "Framework/Application/SlateApplication.h"
+#include "Widgets/Images/SImage.h"
+#include "Widgets/Input/SButton.h"
+#include "EditorStyleSet.h"
+#include "WorkflowOrientedApp/WorkflowUObjectDocuments.h"
+#include "Widgets/Docking/SDockTab.h"
+#include "Misc/BlacklistNames.h"
+
+/////////////////////////////////////////////////////
+// FTabInfo
+
+namespace WorkflowTabManagerHelpers
+{
+	/** Max number of history items that can be stored.  Once the max is reached, the oldest history item is removed */
+	const int32 MaxHistoryEntries = 300;
+}
+
+FTabInfo::FTabInfo(const TSharedRef<SDockTab>& InTab, const TSharedPtr<FDocumentTabFactory>& InSpawner, const TSharedPtr<class FDocumentTracker>& InTracker)
+	: Tab(InTab)
+	, WeakTracker(InTracker)
+{
+}
+
+bool FTabInfo::PayloadMatches(const TSharedPtr<FTabPayload> TestPayload) const
+{
+	return PayloadMatches(GetPayload(), TestPayload);
+}
+
+bool FTabInfo::PayloadMatches(TSharedPtr<FTabPayload> A, TSharedPtr<FTabPayload> B)
+{
+	if (A.IsValid() && B.IsValid())
+	{
+		return A->IsEqual(B.ToSharedRef());
+	}
+	else if (!B.IsValid() && !A.IsValid())
+	{
+		return true;
+	}
+	else
+	{
+		return false;
+	}
+}
+
+void FTabInfo::AddTabHistory(TSharedPtr< struct FGenericTabHistory > InHistoryNode, bool bInSaveHistory/* = true*/)
+{
+	TSharedPtr<FDocumentTracker> Tracker = WeakTracker.Pin();
+
+	if (Tracker.IsValid())
+	{
+		if (Tracker->CurrentHistoryIndex == Tracker->History.Num() - 1)
+		{
+			// History added to the end
+			if (Tracker->History.Num() == WorkflowTabManagerHelpers::MaxHistoryEntries)
+			{
+				// If max history entries has been reached
+				// remove the oldest history
+				Tracker->History.RemoveAt(0);
+			}
+		}
+		else
+		{
+			// Clear out any history that is in front of the current location in the history list
+			Tracker->History.RemoveAt(Tracker->CurrentHistoryIndex + 1, Tracker->History.Num() - (Tracker->CurrentHistoryIndex + 1), true);
+		}
+
+		Tracker->History.Add(InHistoryNode);
+		Tracker->CurrentHistoryIndex = Tracker->History.Num() - 1;
+
+		SetCurrentHistory(InHistoryNode, bInSaveHistory, false);
+	}
+}
+
+TSharedPtr<struct FGenericTabHistory> FTabInfo::GetCurrentHistory() const
+{
+	return CurrentHistory;
+}
+
+void FTabInfo::SetCurrentHistory(TSharedPtr<FGenericTabHistory> NewHistory, bool bInSaveHistory, bool bShouldRestore)
+{
+	bool bPayloadMatches = false;
+
+	if (CurrentHistory.IsValid() && NewHistory.IsValid())
+	{
+		if (bInSaveHistory)
+		{
+			CurrentHistory->SaveHistory();
+		}
+
+		bPayloadMatches = PayloadMatches(CurrentHistory->GetPayload(), NewHistory->GetPayload());
+	}
+	CurrentHistory = NewHistory;
+
+	if (CurrentHistory.IsValid())
+	{
+		CurrentHistory->BindToTab(AsShared());
+
+		// This creates the tab widget but does not foreground it
+		CurrentHistory->EvokeHistory(AsShared(), bPayloadMatches);
+		if (bShouldRestore)
+		{
+			CurrentHistory->RestoreHistory();
+		}
+		TSharedPtr<FDocumentTabFactory> Factory = CurrentHistory->GetFactory().Pin();
+		if (Factory.IsValid())
+		{
+			// Notify listeners that tab contents have changed
+			Factory->OnTabActivated(Tab.Pin());
+		}
+		
+	}
+}
+
+FReply FTabInfo::OnGoForwardInHistory()
+{
+	TSharedPtr<FDocumentTracker> Tracker = WeakTracker.Pin();
+
+	if (Tracker.IsValid())
+	{
+		int32 NextHistoryIndex = Tracker->CurrentHistoryIndex;
+		while (NextHistoryIndex < Tracker->History.Num() - 1)
+		{
+			++NextHistoryIndex;
+				
+			if (Tracker->NavigateToTabHistory(NextHistoryIndex))
+			{
+				break;
+			}
+		}
+	}
+	return FReply::Handled();
+}
+
+FReply FTabInfo::OnGoBackInHistory()
+{
+	TSharedPtr<FDocumentTracker> Tracker = WeakTracker.Pin();
+
+	if (Tracker.IsValid())
+	{
+		int32 PreviousHistoryIndex = Tracker->CurrentHistoryIndex;
+		while (PreviousHistoryIndex > 0)
+		{
+			--PreviousHistoryIndex;
+				
+			if (Tracker->NavigateToTabHistory(PreviousHistoryIndex))
+			{
+				break;
+			}
+		}
+	}
+	return FReply::Handled();
+}
+
+void FTabInfo::JumpToNearestValidHistoryData()
+{
+	// Each tab only knows current history, close if it's invalid
+	if (!CurrentHistory.IsValid() || !CurrentHistory->IsHistoryValid())
+	{
+		if (Tab.IsValid())
+		{
+			Tab.Pin()->RequestCloseTab();
+		}
+	}
+}
+
+TWeakPtr<FDocumentTabFactory> FTabInfo::GetFactory() const
+{
+	if (CurrentHistory.IsValid())
+	{
+		return CurrentHistory->GetFactory();
+	}
+	return nullptr;
+}
+
+TSharedPtr<FTabPayload> FTabInfo::GetPayload() const
+{
+	if (CurrentHistory.IsValid())
+	{
+		return CurrentHistory->GetPayload();
+	}
+	return nullptr;
+}
+
+void FTabInfo::GoToHistoryIndex(int32 InHistoryIdx)
+{
+	TSharedPtr<FDocumentTracker> Tracker = WeakTracker.Pin();
+	if (Tracker.IsValid())
+	{
+		Tracker->NavigateToTabHistory(InHistoryIdx);
+	}
+}
+
+TSharedRef<SWidget> FTabInfo::CreateHistoryMenu(bool bInBackHistory) const
+{
+	TSharedPtr<FDocumentTracker> Tracker = WeakTracker.Pin();
+	if (!Tracker.IsValid())
+	{
+		return SNullWidget::NullWidget;
+	}
+
+	FMenuBuilder MenuBuilder(true, nullptr);
+	if(bInBackHistory)
+	{
+		int32 HistoryIdx = Tracker->CurrentHistoryIndex - 1;
+		while( HistoryIdx >= 0 )
+		{
+			TSharedPtr<FGenericTabHistory> HistoryItem = Tracker->History[HistoryIdx];
+			if(HistoryItem->IsHistoryValid())
+			{
+				MenuBuilder.AddMenuEntry(HistoryItem->GetHistoryTitle().Get(), FText(), FSlateIcon(), 
+					FUIAction(
+					FExecuteAction::CreateRaw(const_cast<FTabInfo*>(this), &FTabInfo::GoToHistoryIndex, HistoryIdx)
+					), 
+					NAME_None, EUserInterfaceActionType::Button);
+			}
+
+			--HistoryIdx;
+		}
+	}
+	else
+	{
+		int32 HistoryIdx = Tracker->CurrentHistoryIndex + 1;
+		while( HistoryIdx < Tracker->History.Num() )
+		{
+			TSharedPtr<FGenericTabHistory> HistoryItem = Tracker->History[HistoryIdx];
+			if(HistoryItem->IsHistoryValid())
+			{
+				MenuBuilder.AddMenuEntry(HistoryItem->GetHistoryTitle().Get(), FText(), FSlateIcon(), 
+					FUIAction(
+					FExecuteAction::CreateRaw(const_cast<FTabInfo*>(this), &FTabInfo::GoToHistoryIndex, HistoryIdx)
+					), 
+					NAME_None, EUserInterfaceActionType::Button);
+			}
+
+			++HistoryIdx;
+		}
+	}
+
+	return MenuBuilder.MakeWidget();
+}
+
+bool FTabInfo::CanStepBackwardInHistory() const
+{
+	TSharedPtr<FDocumentTracker> Tracker = WeakTracker.Pin();
+	if (!Tracker.IsValid())
+	{
+		return false;
+	}
+
+	int32 HistoryIdx = Tracker->CurrentHistoryIndex - 1;
+	while( HistoryIdx >= 0 )
+	{
+		if(Tracker->History[HistoryIdx]->IsHistoryValid())
+		{
+			return true;
+		}
+
+		--HistoryIdx;
+	}
+	return false;
+}
+
+bool FTabInfo::CanStepForwardInHistory() const
+{
+	TSharedPtr<FDocumentTracker> Tracker = WeakTracker.Pin();
+	if (!Tracker.IsValid())
+	{
+		return false;
+	}
+
+	int32 HistoryIdx = Tracker->CurrentHistoryIndex + 1;
+	while( HistoryIdx < Tracker->History.Num() )
+	{
+		if(Tracker->History[HistoryIdx]->IsHistoryValid())
+		{
+			return true;
+		}
+
+		++HistoryIdx;
+	}
+	return false;
+}
+
+FReply FTabInfo::OnMouseDownHistory( const FGeometry& MyGeometry, const FPointerEvent& MouseEvent, TWeakPtr< SMenuAnchor > InMenuAnchor )
+{
+	if(MouseEvent.GetEffectingButton() == EKeys::RightMouseButton && InMenuAnchor.IsValid())
+	{
+		InMenuAnchor.Pin()->SetIsOpen(true);
+		return FReply::Handled();
+	}
+
+	return FReply::Unhandled();
+}
+
+TSharedRef< SWidget > FTabInfo::CreateHistoryNavigationWidget()
+{
+	if(!HistoryNavigationWidget.IsValid())
+	{
+		TWeakPtr< SMenuAnchor > BackMenuAnchorPtr;
+		TWeakPtr< SMenuAnchor > FwdMenuAnchorPtr;
+		HistoryNavigationWidget = 
+			SNew(SHorizontalBox)
+			+SHorizontalBox::Slot()
+			.AutoWidth()
+			[
+				SNew(SBorder)
+				.OnMouseButtonDown(this, &FTabInfo::OnMouseDownHistory, BackMenuAnchorPtr)
+				.BorderImage( FEditorStyle::GetBrush("NoBorder") )
+				[
+					SAssignNew(BackMenuAnchorPtr, SMenuAnchor)
+					.Placement( MenuPlacement_BelowAnchor )
+					.OnGetMenuContent( this, &FTabInfo::CreateHistoryMenu, true )
+					[
+						SNew(SButton)
+						.OnClicked( this, &FTabInfo::OnGoBackInHistory )
+						.ButtonStyle( FEditorStyle::Get(), "GraphBreadcrumbButton" )
+						.IsEnabled(this, &FTabInfo::CanStepBackwardInHistory)
+						.ToolTipText(NSLOCTEXT("WorkflowNavigationBrowser", "Backward_Tooltip", "Step backward in the tab history. Right click to see full history."))
+						[
+							SNew(SImage)
+							.Image( FEditorStyle::GetBrush("GraphBreadcrumb.BrowseBack") )
+						]
+					]
+				]
+			]
+
+			+SHorizontalBox::Slot()
+			.AutoWidth()
+			[
+				SNew(SBorder)
+				.OnMouseButtonDown(this, &FTabInfo::OnMouseDownHistory, FwdMenuAnchorPtr)
+				.BorderImage( FEditorStyle::GetBrush("NoBorder") )
+				[
+					SAssignNew(FwdMenuAnchorPtr, SMenuAnchor)
+					.Placement( MenuPlacement_BelowAnchor )
+					.OnGetMenuContent( this, &FTabInfo::CreateHistoryMenu, false )
+					[
+						SNew(SButton)
+						.OnClicked( this, &FTabInfo::OnGoForwardInHistory )
+						.ButtonStyle( FEditorStyle::Get(), "GraphBreadcrumbButton" )
+						.IsEnabled(this, &FTabInfo::CanStepForwardInHistory)
+						.ToolTipText(NSLOCTEXT("WorkflowNavigationBrowser", "Forward_Tooltip", "Step forward in the tab history. Right click to see full history."))
+						[
+							SNew(SImage)
+							.Image( FEditorStyle::GetBrush("GraphBreadcrumb.BrowseForward") )
+						]
+					]
+				]
+			];
+	}
+
+	return HistoryNavigationWidget.ToSharedRef();
+}
+
+/////////////////////////////////////////////////////
+// FWorkflowAllowedTabSet
+
+// Searches this set for a factory with the specified ID, or returns NULL.
+TSharedPtr<class FWorkflowTabFactory> FWorkflowAllowedTabSet::GetFactory(FName FactoryID) const
+{
+	return Factories.FindRef(FactoryID);
+}
+
+// Registers a factory with this set
+void FWorkflowAllowedTabSet::RegisterFactory(TSharedPtr<class FWorkflowTabFactory> Factory)
+{
+	FName NewIdentifier = Factory->GetIdentifier();
+	check(!Factories.Contains(NewIdentifier));
+	Factories.Add(NewIdentifier, Factory);
+}
+
+void FWorkflowAllowedTabSet::UnregisterFactory(FName FactoryID)
+{
+	int32 Removed = Factories.Remove(FactoryID);
+	check(Removed != 0);
+}
+
+// Merges in a set of factories into this set
+void FWorkflowAllowedTabSet::MergeInSet(const FWorkflowAllowedTabSet& OtherSet)
+{
+	Factories.Append(OtherSet.Factories);
+}
+
+// Clears the set
+void FWorkflowAllowedTabSet::Clear()
+{
+	Factories.Empty();
+}
+
+TMap< FName, TSharedPtr<class FWorkflowTabFactory> >::TIterator FWorkflowAllowedTabSet::CreateIterator()
+{
+	return Factories.CreateIterator();
+}
+
+/////////////////////////////////////////////////////
+// FDocumentTracker
+
+void FDocumentTracker::ClearDocumentFactories()
+{
+	PotentialTabFactories.Empty();
+}
+
+void FDocumentTracker::RegisterDocumentFactory(TSharedPtr<class FDocumentTabFactory> Factory)
+{
+	FName NewIdentifier = Factory->GetIdentifier();
+	check(!PotentialTabFactories.Contains(NewIdentifier));
+	PotentialTabFactories.Add(NewIdentifier, Factory);
+}
+
+FDocumentTracker::FDocumentTracker()
+{
+	// Make sure we know when tabs become active
+	OnActiveTabChangedDelegateHandle = FGlobalTabmanager::Get()->OnActiveTabChanged_Subscribe( FOnActiveTabChanged::FDelegate::CreateRaw( this, &FDocumentTracker::OnActiveTabChanged ) );
+	TabForegroundedDelegateHandle = FGlobalTabmanager::Get()->OnTabForegrounded_Subscribe(FOnActiveTabChanged::FDelegate::CreateRaw(this, &FDocumentTracker::OnTabForegrounded));
+
+	CurrentHistoryIndex = INDEX_NONE;
+}
+
+FDocumentTracker::~FDocumentTracker()
+{
+	FGlobalTabmanager::Get()->OnActiveTabChanged_Unsubscribe(OnActiveTabChangedDelegateHandle);
+	FGlobalTabmanager::Get()->OnTabForegrounded_Unsubscribe(TabForegroundedDelegateHandle);
+}
+
+// Called by the global active tab changed callback; dispatches to individually registered callbacks
+void FDocumentTracker::OnActiveTabChanged(TSharedPtr<SDockTab> PreviouslyActive, TSharedPtr<SDockTab> NewlyActivated)
+{
+	FTabList& List = GetSpawnedList();
+	for (auto ListIt = List.CreateIterator(); ListIt; ++ListIt)
+	{
+		// Get the factory (can't fail; the tabs had to come from somewhere; failing means a tab survived a mode transition to a mode where it is not allowed!)
+		TSharedPtr<FDocumentTabFactory> Factory = (*ListIt)->GetFactory().Pin();
+		if (ensure(Factory.IsValid()))
+		{
+			TSharedPtr<SDockTab> Tab = (*ListIt)->GetTab().Pin();
+			if (Tab == NewlyActivated)
+			{
+				WeakLastEditedTabInfo = *ListIt;
+				Factory->OnTabActivated(Tab);
+			}
+		}
+	}
+}
+
+void FDocumentTracker::OnTabForegrounded(TSharedPtr<SDockTab> ForegroundedTab, TSharedPtr<SDockTab> BackgroundedTab)
+{
+	TSharedPtr<FTabInfo> NewTabInfo;
+	TSharedPtr<SDockTab> OwnedForeground, OwnedBackground;
+	TSharedPtr<FDocumentTabFactory> ForegroundFactory, BackgroundFactory;
+
+	FTabList& List = GetSpawnedList();
+	for ( const TSharedPtr<FTabInfo>& TabInfo : List )
+	{
+		// Get the factory (can't fail; the tabs had to come from somewhere; failing means a tab survived a mode transition to a mode where it is not allowed!)
+		TSharedPtr<FDocumentTabFactory> Factory = TabInfo->GetFactory().Pin();
+		if ( ensure(Factory.IsValid()) )
+		{
+			TSharedPtr<SDockTab> Tab = TabInfo->GetTab().Pin();
+			if ( Tab == ForegroundedTab )
+			{
+				NewTabInfo = TabInfo;
+				OwnedForeground = Tab;
+				ForegroundFactory = Factory;
+			}
+			else if ( Tab == BackgroundedTab )
+			{
+				OwnedBackground = Tab;
+				BackgroundFactory = Factory;
+			}
+		}
+	}
+
+	if ( OwnedBackground.IsValid() )
+	{
+		BackgroundFactory->OnTabBackgrounded(OwnedBackground);
+	}
+
+	if ( OwnedForeground.IsValid() )
+	{
+		ForegroundFactory->OnTabForegrounded(OwnedForeground);
+	}
+
+	TSharedPtr<FGenericTabHistory> CurrentTabHistory = GetCurrentTabHistory();
+	TSharedPtr<FGenericTabHistory> NewTabHistory;
+	if (NewTabInfo.IsValid())
+	{
+		NewTabHistory = NewTabInfo->GetCurrentHistory();
+	}
+
+	if ( ForegroundFactory.IsValid() && NewTabHistory.IsValid() && NewTabHistory->IsHistoryValid() && NewTabHistory != CurrentTabHistory )
+	{
+		// If a tab was manually foregrounded, need to add tab history
+		NewTabInfo->AddTabHistory(ForegroundFactory->CreateTabHistoryNode(NewTabInfo->GetPayload()), true);
+	}
+}
+
+FDocumentTracker::FTabList& FDocumentTracker::GetSpawnedList()
+{
+	// Remove any closed tabs
+	{
+		struct
+		{ 
+			bool operator()(const TSharedPtr<FTabInfo>& RemovalCandidate) const
+			{
+				return !RemovalCandidate->GetTab().IsValid();
+			}
+		} CullClosedTabs;
+	
+		SpawnedTabs.RemoveAll( CullClosedTabs );
+	}
+
+	return SpawnedTabs;
+}
+
+TSharedPtr<FTabInfo> FDocumentTracker::GetLastEditedTabInfo()
+{
+	TSharedPtr<FTabInfo> LastEditedTabInfo = WeakLastEditedTabInfo.Pin();
+	if (!LastEditedTabInfo.IsValid() || !LastEditedTabInfo->GetTab().IsValid())
+	{
+		// Clear if either the info or actual tab are gone
+		WeakLastEditedTabInfo = nullptr;
+		return nullptr;
+	}
+	return LastEditedTabInfo;
+}
+
+void FDocumentTracker::Initialize(TSharedPtr<FAssetEditorToolkit> InHostingApp )
+{
+	check(!HostingAppPtr.IsValid());
+	HostingAppPtr = InHostingApp;
+}
+
+void FDocumentTracker::SetTabManager( const TSharedRef<FTabManager>& InTabManager )
+{
+	TabManager = InTabManager;
+}
+
+TSharedPtr<SDockTab> FDocumentTracker::OpenDocument(TSharedPtr<FTabPayload> InPayload, EOpenDocumentCause InOpenCause)
+{
+	// Managed document spawn can occur at times like 'Ctrl + Z' to undo, don't force them to spawn a new tab.
+	if(FSlateApplication::Get().GetModifierKeys().IsControlDown() && InOpenCause != SpawnManagedDocument)
+	{
+		InOpenCause = FDocumentTracker::ForceOpenNewDocument;
+	}
+
+	if(InOpenCause == CreateHistoryEvent)
+	{
+		// Deprecated, all opens now save history
+		InOpenCause = FDocumentTracker::OpenNewDocument;
+	}
+
+	if(InOpenCause == NavigatingCurrentDocument || InOpenCause == QuickNavigateCurrentDocument || InOpenCause == NavigateBackwards || InOpenCause == NavigateForwards)
+	{
+		return NavigateCurrentTab(InPayload, InOpenCause);
+	}
+
+	// Spawning or restoring a tab, so the factory has to be valid
+	TSharedPtr<FDocumentTabFactory> Factory = FindSupportingFactory(InPayload.ToSharedRef());
+	if(!Factory.IsValid())
+	{
+		return nullptr;
+	}
+
+	if(InOpenCause == OpenNewDocument || InOpenCause == RestorePreviousDocument || InOpenCause == SpawnManagedDocument)
+	{
+		// If the current tab matches we'll re-use it.
+		TSharedPtr<FTabInfo> LastEditedTabInfo = GetLastEditedTabInfo();
+		if(LastEditedTabInfo.IsValid() && LastEditedTabInfo->PayloadMatches(InPayload))
+		{
+			TSharedPtr<SDockTab> Tab = LastEditedTabInfo->GetTab().Pin();
+			if (Tab.IsValid())
+			{
+				LastEditedTabInfo->AddTabHistory(Factory->CreateTabHistoryNode(InPayload), true);
+				// Ensure that the tab appears if the tab isn't currently in the foreground.
+				Tab->ActivateInParent(ETabActivationCause::SetDirectly);
+			}
+			return Tab;
+		}
+		else
+		{
+			// Check if the payload is currently open in any tab
+			FTabList& List = GetSpawnedList();
+			for (const TSharedPtr<FTabInfo>& TabInfo : List)
+			{
+				TSharedPtr<SDockTab> Tab = TabInfo->GetTab().Pin();
+				if (TabInfo->PayloadMatches(InPayload))
+				{
+					// Manually opening an existing tab, add to history
+					TabInfo->AddTabHistory(Factory->CreateTabHistoryNode(InPayload), true);
+					TabManager->DrawAttention(Tab.ToSharedRef());
+					return Tab;
+				}
+			}
+		}
+	}
+
+	// Occurs when forcing open a new tab, or if it failed to find existing one
+	return OpenNewTab(Factory->CreateTabHistoryNode(InPayload), InOpenCause);
+}
+
+TWeakPtr< FTabInfo > FDocumentTracker::FindTabInForeground()
+{
+	// Find a tab that is in the foreground
+	FTabList& List = GetSpawnedList();
+	for (auto ListIt = List.CreateIterator(); ListIt; ++ListIt)
+	{
+		if((*ListIt)->GetTab().Pin()->IsForeground())
+		{
+			return *ListIt;
+		}
+	}
+
+	return nullptr;
+}
+
+bool FDocumentTracker::NavigateToTabHistory(int32 InHistoryIdx)
+{
+	if (History.IsValidIndex(InHistoryIdx))
+	{
+		TSharedPtr<FGenericTabHistory> NewHistory = History[InHistoryIdx];
+		if (NewHistory->IsHistoryValid())
+		{
+			CurrentHistoryIndex = InHistoryIdx;
+
+			TSharedPtr<FTabInfo> FoundPayloadTab, FoundBoundTab;
+			FTabList& List = GetSpawnedList();
+			for (TSharedPtr<FTabInfo>& TabInfo : List)
+			{
+				if (TabInfo->GetCurrentHistory() == NewHistory)
+				{
+					// If it's literally the same history item that is already active, just focus the tab
+					TabManager->DrawAttention(TabInfo->GetTab().Pin().ToSharedRef());
+					return true;
+				}
+				else if (TabInfo->GetCurrentHistory()->GetBoundTab() == NewHistory->GetBoundTab())
+				{
+					FoundBoundTab = TabInfo;
+				}
+				else if (TabInfo->PayloadMatches(NewHistory->GetPayload()))
+				{
+					FoundPayloadTab = TabInfo;
+				}
+			}
+
+			// If no bound tab but found matching payload, use that
+			if (!FoundBoundTab.IsValid() && FoundPayloadTab.IsValid())
+			{
+				FoundBoundTab = FoundPayloadTab;
+			}
+
+			if (FoundBoundTab.IsValid())
+			{
+				// Found a mostly matching tab we want to restore in
+				FoundBoundTab->SetCurrentHistory(NewHistory, true, true);
+				TabManager->DrawAttention(FoundBoundTab->GetTab().Pin().ToSharedRef());
+				return true;
+			}
+
+			// The tab was closed, so spawn a new one
+			OpenNewTab(NewHistory, EOpenDocumentCause::NavigatingHistory);
+			return true;
+		}
+	}
+
+	return false;
+}
+
+TSharedPtr<FGenericTabHistory> FDocumentTracker::GetCurrentTabHistory()
+{
+	if (History.IsValidIndex(CurrentHistoryIndex))
+	{
+		return History[CurrentHistoryIndex];
+	}
+	return nullptr;
+}
+
+TSharedPtr<SDockTab> FDocumentTracker::NavigateCurrentTab(TSharedPtr<FTabPayload> InPayload, EOpenDocumentCause InNavigateCause)
+{
+	ensure(InNavigateCause == NavigatingCurrentDocument || InNavigateCause == QuickNavigateCurrentDocument || InNavigateCause == NavigateBackwards || InNavigateCause == NavigateForwards);
+
+	FTabList& List = GetSpawnedList();
+	if(List.Num())
+	{
+		// Make sure we find an available tab to navigate, there are ones available.
+		TSharedPtr<FTabInfo> LastEditedTabInfo = GetLastEditedTabInfo();
+		if(!LastEditedTabInfo.IsValid())
+		{
+			WeakLastEditedTabInfo = FindTabInForeground();
+
+			// Check if we still do not have a valid tab, if we do not, activate the first tab
+			// Any invalid tabs would have been deleted inside GetSpawnedList so we know 0 is valid
+			if(!WeakLastEditedTabInfo.IsValid())
+			{
+				WeakLastEditedTabInfo = List[0];
+				List[0]->GetTab().Pin()->ActivateInParent(ETabActivationCause::SetDirectly);
+			}
+
+			LastEditedTabInfo = GetLastEditedTabInfo();
+		}
+
+		check(LastEditedTabInfo.IsValid());
+		TSharedPtr<SDockTab> LastEditedTab = LastEditedTabInfo->GetTab().Pin();
+
+		if(InNavigateCause == NavigatingCurrentDocument || InNavigateCause == QuickNavigateCurrentDocument)
+		{
+			TSharedPtr<FDocumentTabFactory> Factory = FindSupportingFactory(InPayload.ToSharedRef());
+			if (Factory.IsValid())
+			{
+				// If doing a Quick navigate of the document, do not save history data as it's likely still at the default values. The object is always saved
+				LastEditedTabInfo->AddTabHistory(Factory->CreateTabHistoryNode(InPayload), InNavigateCause != QuickNavigateCurrentDocument);
+				// Ensure that the tab appears if the tab isn't currently in the foreground.
+				LastEditedTab->ActivateInParent(ETabActivationCause::SetDirectly);
+			}
+		}
+		else if(InNavigateCause == NavigateBackwards)
+		{
+			LastEditedTabInfo->OnGoBackInHistory();
+		}
+		else if(InNavigateCause == NavigateForwards)
+		{
+			LastEditedTabInfo->OnGoForwardInHistory();
+		}
+
+		return LastEditedTab;
+	}
+	
+	// Open in new tab
+	TSharedPtr<FDocumentTabFactory> Factory = FindSupportingFactory(InPayload.ToSharedRef());
+	if (Factory.IsValid())
+	{
+		return OpenNewTab(Factory->CreateTabHistoryNode(InPayload), OpenNewDocument);
+	}
+	
+	return nullptr;
+}
+
+TSharedPtr<SDockTab> FDocumentTracker::OpenNewTab(TSharedPtr<FGenericTabHistory> InTabHistory, EOpenDocumentCause InOpenCause)
+{
+	ensure(InOpenCause == ForceOpenNewDocument || InOpenCause == OpenNewDocument || InOpenCause == RestorePreviousDocument || InOpenCause == NavigatingHistory || InOpenCause == SpawnManagedDocument);
+
+	TSharedPtr<FDocumentTabFactory> Factory = InTabHistory->GetFactory().Pin();
+	TSharedPtr<SDockTab> NewTab;
+
+	if(Factory.IsValid())
+	{
+		TSharedPtr<FAssetEditorToolkit> HostingApp = HostingAppPtr.Pin();
+		FWorkflowTabSpawnInfo SpawnInfo;
+		SpawnInfo.Payload = InTabHistory->GetPayload();
+
+		NewTab = Factory->SpawnBlankTab();
+
+		TSharedPtr<FTabInfo> NewTabInfo = MakeShareable( new FTabInfo(NewTab.ToSharedRef(), Factory, AsShared()) );
+		SpawnedTabs.Add( NewTabInfo );
+
+		if (InOpenCause == NavigatingHistory)
+		{
+			// Don't want to create a new history entry when going back/forward
+			NewTabInfo->SetCurrentHistory(InTabHistory, true, true);
+		}
+		else
+		{
+			NewTabInfo->AddTabHistory(InTabHistory);
+		}
+
+		if (InOpenCause == ForceOpenNewDocument  || InOpenCause == OpenNewDocument)
+		{
+			TabManager->InsertNewDocumentTab( "Document", FTabManager::ESearchPreference::RequireClosedTab, NewTab.ToSharedRef() );
+		}
+		else if (InOpenCause == RestorePreviousDocument)
+		{
+			TabManager->RestoreDocumentTab( "Document", FTabManager::ESearchPreference::RequireClosedTab, NewTab.ToSharedRef() );
+
+			// Clear tab history before this so previous restores don't show up
+			History.RemoveAt(0, History.Num() - 1, true);
+			CurrentHistoryIndex = History.Num() - 1;
+		}
+	}
+
+	return NewTab;
+}
+
+void FDocumentTracker::CloseTab(TSharedPtr<FTabPayload> Payload)
+{
+	FTabList List = GetSpawnedList();
+	for (auto ListIt = List.CreateIterator(); ListIt; ++ListIt)
+	{
+		if ((*ListIt)->PayloadMatches(Payload))
+		{
+			TSharedPtr<SDockTab> Tab = (*ListIt)->GetTab().Pin();
+			Tab->RequestCloseTab();
+		}
+	}
+}
+
+void FDocumentTracker::CleanInvalidTabs()
+{
+	FTabList List = GetSpawnedList();
+	for (auto ListIt = List.CreateIterator(); ListIt; ++ListIt)
+	{
+		(*ListIt)->JumpToNearestValidHistoryData();
+	}
+}
+
+// Finds a factory that can deal with the supplied payload
+TSharedPtr<FDocumentTabFactory> FDocumentTracker::FindSupportingFactory(TSharedRef<FTabPayload> Payload) const
+{
+	for (auto FactoryIt = PotentialTabFactories.CreateConstIterator(); FactoryIt; ++FactoryIt)
+	{
+		TSharedPtr<FDocumentTabFactory> Factory = FactoryIt.Value();
+
+		if (Factory->IsPayloadSupported(Payload) && TabManager->GetTabBlacklist()->PassesFilter(Factory->GetIdentifier()))
+		{
+			return Factory;
+		}
+	}
+
+	return TSharedPtr<FDocumentTabFactory>();
+}
+
+// Finds all tabs that match the specified identifier and payload, placing them in the specified array
+void FDocumentTracker::FindMatchingTabs( TSharedPtr<FTabPayload> Payload, /*inout*/ TArray< TSharedPtr<SDockTab> >& Results)
+{
+	FTabList& List = GetSpawnedList();
+	for (auto ListIt = List.CreateIterator(); ListIt; ++ListIt)
+	{
+		if ((*ListIt)->PayloadMatches(Payload))
+		{
+			Results.Add((*ListIt)->GetTab().Pin());
+		}
+	}
+}
+
+// Finds all tabs that match the specified identifier, placing them in the specified array
+void FDocumentTracker::FindAllTabsForFactory(const TWeakPtr<FDocumentTabFactory>& Factory, /*inout*/ TArray< TSharedPtr<SDockTab> >& Results)
+{
+	FTabList& List = GetSpawnedList();
+	for (auto ListIt = List.CreateIterator(); ListIt; ++ListIt)
+	{
+		if ((*ListIt)->GetFactory() == Factory)
+		{
+			Results.Add((*ListIt)->GetTab().Pin());
+		}
+		
+	}
+}
+
+// Saves the state of all tabs
+//@TODO: DOCMANAGEMENT: Basically a total copy of RefreshAllTabs; need to make this sort of pattern be a delegate for-all-tabs
+void FDocumentTracker::SaveAllState()
+{
+	for (auto ToolPanelIt = SpawnedTabs.CreateIterator(); ToolPanelIt; ++ToolPanelIt)
+	{
+		// Get the factory (can't fail; the tabs had to come from somewhere; failing means a tab survived a mode transition to a mode where it is not allowed!)
+		TSharedPtr<FDocumentTabFactory> Factory = (*ToolPanelIt)->GetFactory().Pin();
+		if (ensure(Factory.IsValid()) && (*ToolPanelIt)->GetTab().IsValid() && !(*ToolPanelIt)->GetTab().Pin()->IsForeground())
+		{
+			Factory->SaveState((*ToolPanelIt)->GetTab().Pin(), (*ToolPanelIt)->GetPayload());
+		}
+	}
+	// Save the foreground tabs after saving the background tabs
+	// This ensures foreground tabs always are restored after the background tabs.
+	for (auto ToolPanelIt = SpawnedTabs.CreateIterator(); ToolPanelIt; ++ToolPanelIt)
+	{
+		TSharedPtr<FDocumentTabFactory> Factory = (*ToolPanelIt)->GetFactory().Pin();
+		if (ensure(Factory.IsValid()) && (*ToolPanelIt)->GetTab().IsValid() && (*ToolPanelIt)->GetTab().Pin()->IsForeground())
+		{
+			Factory->SaveState((*ToolPanelIt)->GetTab().Pin(), (*ToolPanelIt)->GetPayload());
+		}
+	}
+}
+
+// Calls OnTabRefreshed for each open tab (on the factories that created them)
+void FDocumentTracker::RefreshAllTabs() const
+{
+	for (auto ToolPanelIt = SpawnedTabs.CreateConstIterator(); ToolPanelIt; ++ToolPanelIt)
+	{
+		// Get the factory (can't fail; the tabs had to come from somewhere; failing means a tab survived a mode transition to a mode where it is not allowed!)
+		TSharedPtr<FDocumentTabFactory> Factory = (*ToolPanelIt)->GetFactory().Pin();
+		// Run thru the open tabs for this one
+		if (ensure(Factory.IsValid()))
+		{
+			TSharedPtr<SDockTab> Tab = (*ToolPanelIt)->GetTab().Pin();
+			if (Tab.IsValid())
+			{
+				Factory->OnTabRefreshed(Tab);
+			}			
+		}
+	}
+}
+
+// Replaces the open payload in the specified tab with a new one; recreating the contents
+void FDocumentTracker::ReplacePayloadInTab(TSharedPtr<SDockTab> TargetTab, TSharedPtr<FTabPayload> NewPayload)
+{
+	// Find the existing tab
+	for (auto ToolPanelIt = SpawnedTabs.CreateIterator(); ToolPanelIt; ++ToolPanelIt)
+	{
+		TSharedPtr<SDockTab> Tab = (*ToolPanelIt)->GetTab().Pin();
+
+		if (Tab == TargetTab)
+		{
+			// Get the factory (can't fail; the tabs had to come from somewhere; failing means a tab survived a mode transition to a mode where it is not allowed!)
+			TSharedPtr<FDocumentTabFactory> Factory = (*ToolPanelIt)->GetFactory().Pin();
+			if (ensure(Factory.IsValid()) && ensure(Factory->IsPayloadSupported(NewPayload.ToSharedRef())))
+			{
+				FWorkflowTabSpawnInfo Info;
+				Info.Payload = NewPayload;
+				Tab->SetContent(Factory->CreateTabBody(Info));
+
+				return;
+			}
+		}
+	}
+
+}
+
+TArray< TSharedPtr<SDockTab> > FDocumentTracker::GetAllDocumentTabs() const
+{
+	TArray< TSharedPtr<SDockTab> > AllSpawnedDocuments;
+	for (int32 DocIndex=0; DocIndex < SpawnedTabs.Num(); ++DocIndex)
+	{
+		if (SpawnedTabs[DocIndex]->GetTab().IsValid())
+		{
+			AllSpawnedDocuments.Add(SpawnedTabs[DocIndex]->GetTab().Pin());
+		}
+	}
+
+	return AllSpawnedDocuments;
+}
+
+TSharedPtr<SDockTab> FDocumentTracker::GetActiveTab() const
+{
+	if(WeakLastEditedTabInfo.IsValid())
+	{
+		return WeakLastEditedTabInfo.Pin()->GetTab().Pin();
+	}
+
+	return nullptr;
+}
+
+FReply FDocumentTracker::OnNavigateTab(FDocumentTracker::EOpenDocumentCause InCause)
+{
+	NavigateCurrentTab(FTabPayload_UObject::Make(nullptr), InCause);
+	return FReply::Handled();
+}
diff --git a/Engine/Source/Editor/UnrealEd/Public/WorkflowOrientedApp/ApplicationMode.h b/Engine/Source/Editor/UnrealEd/Public/WorkflowOrientedApp/ApplicationMode.h
new file mode 100644
--- /dev/null
+++ b/Engine/Source/Editor/UnrealEd/Public/WorkflowOrientedApp/ApplicationMode.h
@@ -0,0 +1,69 @@
+// Copyright Epic Games, Inc. All Rights Reserved.
+
+#pragma once
+
+#include "CoreMinimal.h"
+#include "Framework/Docking/WorkspaceItem.h"
+#include "Framework/Docking/TabManager.h"
+#include "Framework/Docking/LayoutExtender.h"
+
+DECLARE_DELEGATE_RetVal_OneParam(TSharedPtr<class FWorkflowTabFactory>, FCreateWorkflowTabFactory, TSharedPtr<class FAssetEditorToolkit>)
+
+class FExtender;
+
+
+/////////////////////////////////////////////////////
+// FApplicationMode
+
+class UNREALED_API FApplicationMode : public TSharedFromThis<FApplicationMode>
+{
+protected:
+	// The layout to use in this mode
+	TSharedPtr<FTabManager::FLayout> TabLayout;
+
+	// The list of tabs that can be spawned in this game mode.
+	// Note: This list should be treated as readonly outside of the constructor.
+	//FWorkflowAllowedTabSet AllowableTabs;
+
+	// The internal name of this mode
+	FName ModeName;
+
+	//@TODO: For test suite use only
+	FString UserLayoutString;
+	
+	/** The toolbar extension for this mode */
+	TSharedPtr<FExtender> ToolbarExtender;
+
+	/** The workspace menu category for this mode */
+	TSharedPtr<FWorkspaceItem> WorkspaceMenuCategory;
+
+public:
+	FApplicationMode(FName InModeName);
+
+	FApplicationMode(FName InModeName, FText(*GetLocalizedMode)(const FName));
+
+	virtual ~FApplicationMode() {}
+
+	void DeactivateMode(TSharedPtr<FTabManager> InTabManager);
+	TSharedRef<FTabManager::FLayout> ActivateMode(TSharedPtr<FTabManager> InTabManager);
+	virtual void RegisterTabFactories(TSharedPtr<FTabManager> InTabManager) {}
+
+	virtual void AddTabFactory(FCreateWorkflowTabFactory FactoryCreator) {}
+	virtual void RemoveTabFactory(FName TabFactoryID) {}
+	
+	virtual void PreDeactivateMode() {}
+	virtual void PostActivateMode() {}
+
+	FName GetModeName() const
+	{
+		return ModeName;
+	}
+
+	TSharedPtr<FExtender> GetToolbarExtender() { return ToolbarExtender; }
+
+	/** @return The the workspace category for this asset editor */
+	TSharedRef<FWorkspaceItem> GetWorkspaceMenuCategory() const { return WorkspaceMenuCategory.ToSharedRef(); }
+
+	/** Extender for adding to the default layout for this mode */
+	TSharedPtr<FLayoutExtender> LayoutExtender;
+};
diff --git a/Engine/Source/Editor/UnrealEd/Public/WorkflowOrientedApp/SContentReference.h b/Engine/Source/Editor/UnrealEd/Public/WorkflowOrientedApp/SContentReference.h
new file mode 100644
--- /dev/null
+++ b/Engine/Source/Editor/UnrealEd/Public/WorkflowOrientedApp/SContentReference.h
@@ -0,0 +1,135 @@
+// Copyright Epic Games, Inc. All Rights Reserved.
+
+#pragma once
+
+#include "CoreMinimal.h"
+#include "Misc/Attribute.h"
+#include "Types/SlateStructs.h"
+#include "Layout/Visibility.h"
+#include "Input/Reply.h"
+#include "Widgets/DeclarativeSyntaxSupport.h"
+#include "Widgets/SWidget.h"
+#include "Widgets/SCompoundWidget.h"
+#include "Framework/SlateDelegates.h"
+#include "Editor/ContentBrowser/Public/IContentBrowserSingleton.h"
+
+class SBorder;
+class SComboButton;
+
+// This widget is a content reference, which can optionally show buttons for interacting with the reference
+class UNREALED_API SContentReference : public SCompoundWidget
+{
+public:
+
+	/** Allows for informing user when someone tries to set a new reference*/
+	DECLARE_DELEGATE_OneParam(FOnSetReference, UObject*)
+
+	SLATE_BEGIN_ARGS(SContentReference)
+		: _Style(TEXT("ContentReference"))
+		, _AssetReference((UObject*)NULL)
+		, _ShowFindInBrowserButton(true)
+		, _ShowToolsButton(false)
+		, _AllowSelectingNewAsset(false)
+		, _AllowClearingReference(false)
+		, _AllowedClass((UClass*)NULL)
+		, _WidthOverride(FOptionalSize())
+		, _AssetPickerSizeOverride(FVector2D(384, 768))
+		, _InitialAssetViewType(EAssetViewType::Tile)
+	{}
+		// The style of the content reference widget (optional)
+		SLATE_ARGUMENT(FName, Style)
+
+		// The asset being displayed by this widget
+		SLATE_ATTRIBUTE(UObject*, AssetReference)
+
+		// Should the find in content browser button be shown?
+		SLATE_ATTRIBUTE(bool, ShowFindInBrowserButton)
+
+		// Should the tools button be shown?
+		SLATE_ATTRIBUTE(bool, ShowToolsButton)
+
+		// Can the user pick a new asset to reference?
+		SLATE_ATTRIBUTE(bool, AllowSelectingNewAsset)
+
+		// Can the user clear the existing asset reference?
+		SLATE_ATTRIBUTE(bool, AllowClearingReference)
+
+		// If specified, the asset references are restricted to this class
+		SLATE_ATTRIBUTE(UClass*, AllowedClass)
+
+		// The event to call when the tools button is clicked
+		SLATE_EVENT(FOnClicked, OnClickedTools)
+
+		// If bound, this delegate will be executed for each asset in the asset picker to see if it should be filtered
+		SLATE_EVENT(FOnShouldFilterAsset, OnShouldFilterAsset)
+
+		//when user changes the referenced object
+		SLATE_EVENT(FOnSetReference, OnSetReference)
+
+		/** When specified, the path box will request this fixed size. */
+		SLATE_ATTRIBUTE(FOptionalSize, WidthOverride)
+
+		// when you'd like to modify size of the asset picker window
+		SLATE_ATTRIBUTE(FVector2D, AssetPickerSizeOverride)
+
+		// initial view type
+		SLATE_ARGUMENT(EAssetViewType::Type, InitialAssetViewType)
+		// The event to call when the user selects a new asset or clears the existing reference
+		//@TODO: Support that stuff!
+	SLATE_END_ARGS()
+
+public:
+	// Construct a SContentReference
+	void Construct(const FArguments& InArgs);
+
+	// Causes the asset picker to open, allowing the user to select a new asset
+	void OpenAssetPickerMenu();
+
+protected:
+	EVisibility GetUseButtonVisibility() const;
+	EVisibility GetPickButtonVisibility() const;
+	EVisibility GetFindButtonVisibility() const;
+	EVisibility GetClearButtonVisibility() const;
+	EVisibility GetToolsButtonVisibility() const;
+
+	FReply OnClickUseButton();
+	FReply OnClickFindButton();
+	FReply OnClickClearButton();
+
+	TSharedRef<SWidget> MakeAssetPickerMenu();
+	void OnAssetSelectedFromPicker(const FAssetData& AssetData);
+
+	FText GetAssetShortName() const;
+	FText GetAssetFullName() const;
+
+	FReply OnDoubleClickedOnAssetName(const FGeometry& InMyGeometry, const FPointerEvent& InMouseEvent);
+	void OpenAssetToEdit();
+
+	const FSlateBrush* GetBorderImage() const;
+
+	static void FindObjectInContentBrowser(UObject* Object);
+protected:
+	// Attributes
+	TAttribute<UObject*> AssetReference;
+	TAttribute<bool> ShowFindInBrowserButton;
+	TAttribute<bool> ShowToolsButton;
+	TAttribute<bool> AllowSelectingNewAsset;
+	TAttribute<bool> AllowClearingReference;
+	TAttribute<UClass*> AllowedClass;
+
+	// Asset picker configuration
+	TAttribute<FVector2D>	AssetPickerSizeOverride;
+	EAssetViewType::Type	InitialAssetViewType;
+
+	// Delegates
+	FOnShouldFilterAsset OnShouldFilterAsset;
+	FOnSetReference OnSetReference;
+
+	// Widgets
+	TSharedPtr<SBorder> AssetReferenceNameBorderWidget;
+	TSharedPtr<SComboButton> PickerComboButton;
+
+	// Resources
+	const FSlateBrush* BorderImageNormal;
+	const FSlateBrush* BorderImageHovered;
+};
diff --git a/Engine/Source/Editor/UnrealEd/Public/WorkflowOrientedApp/SModeWidget.h b/Engine/Source/Editor/UnrealEd/Public/WorkflowOrientedApp/SModeWidget.h
new file mode 100644
--- /dev/null
+++ b/Engine/Source/Editor/UnrealEd/Public/WorkflowOrientedApp/SModeWidget.h
@@ -0,0 +1,81 @@
+// Copyright Epic Games, Inc. All Rights Reserved.
+
+#pragma once
+
+#include "CoreMinimal.h"
+#include "Misc/Attribute.h"
+#include "Layout/Visibility.h"
+#include "Input/Reply.h"
+#include "Fonts/SlateFontInfo.h"
+#include "Widgets/DeclarativeSyntaxSupport.h"
+#include "Widgets/SCompoundWidget.h"
+
+struct FSlateBrush;
+
+// Event sent when a mode change is requested
+DECLARE_DELEGATE_OneParam( FOnModeChangeRequested, FName );
+
+
+// This is the base class for a persona mode widget
+class UNREALED_API SModeWidget : public SCompoundWidget
+{
+public:
+	SLATE_BEGIN_ARGS(SModeWidget)
+		: _CanBeSelected(true)
+	{}
+		// The currently active mode, used to determine which mode is highlighted
+		SLATE_ATTRIBUTE( FName, OnGetActiveMode )
+
+		// The large image for the icon
+		SLATE_ATTRIBUTE( const FSlateBrush*, IconImage )
+
+		// The small image for the icon
+		SLATE_ATTRIBUTE( const FSlateBrush*, SmallIconImage )
+
+		// The delegate that will be called when this widget wants to change the active mode
+		SLATE_EVENT( FOnModeChangeRequested, OnSetActiveMode )
+
+		// Can this mode ever be selected?
+		SLATE_ATTRIBUTE( bool, CanBeSelected )
+
+		// Slot for the always displayed contents
+		SLATE_NAMED_SLOT( FArguments, ShortContents )
+
+		SLATE_ATTRIBUTE( const FSlateBrush*, DirtyMarkerBrush)
+		
+	SLATE_END_ARGS()
+
+public:
+	// Construct a SModeWidget
+	void Construct(const FArguments& InArgs, const FText& InText, const FName InMode);
+
+private:
+	EVisibility GetLargeIconVisibility() const;
+	EVisibility GetSmallIconVisibility() const;
+
+	bool IsActiveMode() const;
+	const FSlateBrush* GetModeNameBorderImage() const;
+	FReply OnModeTabClicked(const FGeometry& MyGeometry, const FPointerEvent& MouseEvent);
+
+	FSlateFontInfo GetDesiredTitleFont() const;
+private:
+	// The active mode of this group
+	TAttribute<FName> OnGetActiveMode;
+
+	// The delegate to call when this mode is selected
+	FOnModeChangeRequested OnSetActiveMode;
+
+	// Can this mode be selected?
+	TAttribute<bool> CanBeSelected;
+
+	// The text representation of this mode
+	FText ModeText;
+
+	// The mode this widget is representing
+	FName ThisMode;
+
+	// Border images
+	const FSlateBrush* ActiveModeBorderImage;
+	const FSlateBrush* InactiveModeBorderImage;
+	const FSlateBrush* HoverBorderImage;
+};
diff --git a/Engine/Source/Editor/UnrealEd/Public/WorkflowOrientedApp/WorkflowCentricApplication.h b/Engine/Source/Editor/UnrealEd/Public/WorkflowOrientedApp/WorkflowCentricApplication.h
new file mode 100644
--- /dev/null
+++ b/Engine/Source/Editor/UnrealEd/Public/WorkflowOrientedApp/WorkflowCentricApplication.h
@@ -0,0 +1,68 @@
+// Copyright Epic Games, Inc. All Rights Reserved.
+
+#pragma once
+
+#include "CoreMinimal.h"
+#include "Toolkits/AssetEditorToolkit.h"
+
+class FApplicationMode;
+class FWorkflowAllowedTabSet;
+class FWorkflowTabFactory;
+class UToolMenu;
+
+// Delegate for mutating a mode
+DECLARE_DELEGATE_RetVal_TwoParams(TSharedRef<FApplicationMode>, FWorkflowApplicationModeExtender, const FName /*ModeName*/, TSharedRef<FApplicationMode> /*InMode*/);
+
+/////////////////////////////////////////////////////
+// FWorkflowCentricApplication
+
+class UNREALED_API FWorkflowCentricApplication : public FAssetEditorToolkit
+{
+public:
+	// IToolkit interface
+	virtual void RegisterTabSpawners(const TSharedRef<class FTabManager>& TabManager) override;
+	virtual void UnregisterTabSpawners(const TSharedRef<class FTabManager>& TabManager) override;
+	// End of IToolkit interface
+	
+	// FAssetEditorToolkit interface
+	virtual bool OnRequestClose() override;
+	virtual FName GetToolMenuToolbarName(FName& OutParentName) const override;
+	// End of FAssetEditorToolkit interface
+
+	// Returns the name of the toolbar for the given mode along with the name of the toolbar's parent.
+	FName GetToolMenuToolbarNameForMode(const FName InModeName, FName& OutParentName) const;
+
+	// Returns the current mode of this application
+	FName GetCurrentMode() const;
+	TSharedPtr<FApplicationMode> GetCurrentModePtr() const { return CurrentAppModePtr; }
+	bool IsModeCurrent(FName ModeToCheck) const {return GetCurrentMode() == ModeToCheck;}
+
+	// Attempt to set the current mode.  If this mode is illegal or unknown, the mode will remain unchanged.
+	virtual void SetCurrentMode(FName NewMode);
+
+	void PushTabFactories(FWorkflowAllowedTabSet& FactorySetToPush);
+
+	// Gets the mode extender list for all workflow applications (append to customize a specific mode)
+	static TArray<FWorkflowApplicationModeExtender>& GetModeExtenderList() { return ModeExtenderList; }
+
+	/**
+	* Registers toolbar for a mode if not already registered
+	*
+	* @param InModeName - Name of the mode to register toolbar for
+	* @return nullptr if toolbar already registered.
+	*/
+	virtual UToolMenu* RegisterModeToolbarIfUnregistered(const FName InModeName);
+protected:
+
+	virtual void AddApplicationMode(FName ModeName, TSharedRef<FApplicationMode> Mode);
+protected:
+	TSharedPtr<FApplicationMode> CurrentAppModePtr;
+
+private:
+	// List of modes; do not access directly, use AddApplicationMode and SetCurrentMode
+	TMap< FName, TSharedPtr<FApplicationMode> > ApplicationModeList;
+
+	static TArray<FWorkflowApplicationModeExtender> ModeExtenderList;
+};
+
+
diff --git a/Engine/Source/Editor/UnrealEd/Public/WorkflowOrientedApp/WorkflowTabFactory.h b/Engine/Source/Editor/UnrealEd/Public/WorkflowOrientedApp/WorkflowTabFactory.h
new file mode 100644
--- /dev/null
+++ b/Engine/Source/Editor/UnrealEd/Public/WorkflowOrientedApp/WorkflowTabFactory.h
@@ -0,0 +1,196 @@
+// Copyright Epic Games, Inc. All Rights Reserved.
+
+#pragma once
+
+#include "CoreMinimal.h"
+#include "Misc/Attribute.h"
+#include "Widgets/SWidget.h"
+#include "Textures/SlateIcon.h"
+#include "Toolkits/AssetEditorToolkit.h"
+
+class FApplicationMode;
+class FMenuBuilder;
+class FTabInfo;
+class SToolTip;
+
+/////////////////////////////////////////////////////
+// FTabPayload
+
+struct FTabPayload
+{
+public:
+	// The type of this payload
+	const FName PayloadType;
+public:
+	virtual ~FTabPayload() {}
+
+	// Is this payload equal to the other one?
+	virtual bool IsEqual(const TSharedRef<FTabPayload>& OtherPayload) const
+	{
+		return this->PayloadType == OtherPayload->PayloadType;
+	}
+
+	virtual bool IsValid() const
+	{
+		return true;
+	}
+
+protected:
+	// Buried constructor.  This is a base class meant to be inherited from
+	FTabPayload(FName InPayloadType)
+		: PayloadType(InPayloadType)
+	{
+	}
+};
+
+/////////////////////////////////////////////////////
+// FWorkflowTabSpawnInfo
+
+class FWorkflowTabSpawnInfo
+{
+public:
+	// Payload to spawn the tab with
+	TSharedPtr<FTabPayload> Payload;
+
+	// The TabInfo being used to spawn this
+	TSharedPtr<class FTabInfo> TabInfo;
+
+	/** The parent tab manager */
+	TSharedPtr<FTabManager> TabManager;
+
+public:
+	FWorkflowTabSpawnInfo()
+		: TabInfo(NULL)
+	{
+	}
+};
+
+/////////////////////////////////////////////////////
+// FWorkflowTabFactory
+
+class UNREALED_API FWorkflowTabFactory : public TSharedFromThis<FWorkflowTabFactory>
+{
+protected:
+	FName TabIdentifier;
+	ETabRole TabRole;
+
+	FText TabLabel;
+	FSlateIcon TabIcon;
+
+	float InsideTabPadding;
+
+	bool bIsSingleton;
+	bool bShouldAutosize;
+
+	FText ViewMenuDescription;
+	FText ViewMenuTooltip;
+
+	TWeakPtr<class FAssetEditorToolkit> HostingApp;
+
+public:
+	virtual ~FWorkflowTabFactory() {}
+
+	// Creates a toolkit tab given the spawning information
+	virtual TSharedRef<SDockTab> SpawnTab(const FWorkflowTabSpawnInfo& Info) const;
+	// Creates a toolkit tab without any content, given the spawning information
+	virtual TSharedRef<SDockTab> SpawnBlankTab() const;
+
+	/**
+	 * Updates a given tab with new content
+	 *
+	 * @param InDockTab		The tab to update
+	 * @param InInfo		The spawning info, used to generate required tab information
+	 * @param InContent		The content to update the tab with
+	 */
+	virtual void UpdateTab(TSharedPtr<SDockTab> InDockTab, const FWorkflowTabSpawnInfo& Info, TSharedPtr< SWidget > InContent) const;
+
+	// Creates a view menu entry for this factory
+	virtual void CreateViewMenuEntry(FMenuBuilder& MenuBuilder, const FUIAction& Action) const;
+
+	/** Register a tab spawner for this factory */
+	virtual FTabSpawnerEntry& RegisterTabSpawner(TSharedRef<FTabManager> TabManager, const FApplicationMode* CurrentApplicationMode) const;
+
+public:
+	// Constructor
+	FWorkflowTabFactory(FName InIdentifier, TSharedPtr<class FAssetEditorToolkit> InHostingApp);
+
+	// Get the identifier of this factory
+	FName GetIdentifier() const { return TabIdentifier; }
+
+	// Are tabs produced by this factory limited to only having one open at a time?
+	bool IsSingleton() const { return bIsSingleton; }
+
+	// Creates the label for the tab
+	virtual TAttribute<FText> ConstructTabName(const FWorkflowTabSpawnInfo& Info) const;
+
+	/** Returns TRUE if the payload is valid */
+	virtual bool IsPayloadValid(TSharedRef<FTabPayload> Payload) const { return false; }
+
+protected:
+	// Call this to enable the default padding surrounding the tab contents
+	void EnableTabPadding();
+
+	/** Callback function for spawning the tab */
+	virtual TSharedRef<SDockTab> OnSpawnTab(const FSpawnTabArgs& SpawnArgs, TWeakPtr<FTabManager> WeakTabManager) const;
+
+	/** Callback function for checking to see if a tab can/should be spawned */
+	bool CanSpawnTab(const FSpawnTabArgs& SpawnArgs, TWeakPtr<FTabManager> WeakTabManager) const;
+
+public:
+	// Creates the widget that will be the body of the new tab
+	// Users will most commonly override just this.
+	virtual TSharedRef<SWidget> CreateTabBody(const FWorkflowTabSpawnInfo& Info) const;
+
+	// Virtual to create an icon for the tab
+	virtual const FSlateBrush* GetTabIcon(const FWorkflowTabSpawnInfo& Info) const;
+	
+	// Virtual to create an icon for the menu entry
+	virtual const FSlateIcon& GetTabSpawnerIcon(const FWorkflowTabSpawnInfo& Info) const;
+
+	// Create a tooltip widget for the tab
+	virtual TSharedPtr<SToolTip> CreateTabToolTipWidget(const FWorkflowTabSpawnInfo& Info) const;
+
+	// Virtual to get optional extra tooltip text for tab
+	virtual FText GetTabToolTipText(const FWorkflowTabSpawnInfo& Info) const
+	{
+		return TabLabel;
+	}
+
+private:
+	friend class FDocumentTracker;
+};
+
+/////////////////////////////////////////////////////
+// FDocumentTabFactory
+
+class UNREALED_API FDocumentTabFactory : public FWorkflowTabFactory
+{
+public:
+	// Saves the state for this tab
+	//@TODO: DOCMANAGEMENT: This is very special-cased and isn't a general purpose solution
+	virtual void SaveState(TSharedPtr<SDockTab> Tab, TSharedPtr<FTabPayload> Payload) const {}
+
+	// Called when a tab created from this factory is activated
+	virtual void OnTabActivated(TSharedPtr<SDockTab> Tab) const {}
+
+	// Called when a tab created from this factory is brought to the foreground
+	virtual void OnTabForegrounded(TSharedPtr<SDockTab> Tab) const {}
+
+	// Called when a tab created from this factory is sent to the background
+	virtual void OnTabBackgrounded(TSharedPtr<SDockTab> Tab) const {}
+
+	// Called when a tab created from this factory should be refreshed
+	virtual void OnTabRefreshed(TSharedPtr<SDockTab> Tab) const {}
+
+	// Called when a tab is being closed
+	//virtual void OnTabClosed(TSharedRef<SDockTab> Tab, TSharedPtr<FTabPayload> Payload) const {}
+
+	// Does this tab support the specified payload
+	virtual bool IsPayloadSupported(TSharedRef<FTabPayload> Payload) const { return false; }
+
+	/** Creates a tab history node from the passed in payload */
+	virtual TSharedRef<struct FGenericTabHistory> CreateTabHistoryNode(TSharedPtr<FTabPayload> Payload);
+
+protected:
+	FDocumentTabFactory(FName InIdentifier, TSharedPtr<class FAssetEditorToolkit> InHostingApp);
+};
diff --git a/Engine/Source/Editor/UnrealEd/Public/WorkflowOrientedApp/WorkflowTabManager.h b/Engine/Source/Editor/UnrealEd/Public/WorkflowOrientedApp/WorkflowTabManager.h
new file mode 100644
--- /dev/null
+++ b/Engine/Source/Editor/UnrealEd/Public/WorkflowOrientedApp/WorkflowTabManager.h
@@ -0,0 +1,404 @@
+// Copyright Epic Games, Inc. All Rights Reserved.
+
+#pragma once
+
+#include "CoreMinimal.h"
+#include "Misc/Attribute.h"
+#include "Input/Reply.h"
+#include "Widgets/SWidget.h"
+#include "Framework/Docking/TabManager.h"
+#include "WorkflowOrientedApp/WorkflowTabFactory.h"
+
+class SMenuAnchor;
+struct FGenericTabHistory;
+
+/////////////////////////////////////////////////////
+
+namespace EDocumentOpenBehavior
+{
+	enum Type
+	{
+		OnlyReturnExisting,
+		AutomaticallyOpen
+	};
+}
+
+/////////////////////////////////////////////////////
+// FWorkflowAllowedTabSet
+
+class UNREALED_API FWorkflowAllowedTabSet
+{
+public:
+	// Searches this set for a factory with the specified ID, or returns NULL.
+	TSharedPtr<class FWorkflowTabFactory> GetFactory(FName FactoryID) const;
+
+	// Registers a factory with this set - must be unique.
+	void RegisterFactory(TSharedPtr<class FWorkflowTabFactory> Factory);
+
+	// Unregisters a factory with this specified ID.
+	void UnregisterFactory(FName FactoryID);
+
+	// Merges in a set of factories into this set - duplicates are OK.
+	void MergeInSet(const FWorkflowAllowedTabSet& OtherSet);
+
+	// Clears the set
+	void Clear();
+
+	// Allows iteration over Factories
+	TMap< FName, TSharedPtr<class FWorkflowTabFactory> >::TIterator CreateIterator();
+
+private:
+	// Map of available factories in this mode
+	TMap< FName, TSharedPtr<class FWorkflowTabFactory> > Factories;
+
+	friend class FDocumentTracker;
+};
+
+/////////////////////////////////////////////////////
+// FTabInfo
+
+class UNREALED_API FTabInfo : public TSharedFromThis<FTabInfo>
+{
+public:
+	FTabInfo(const TSharedRef<SDockTab>& InTab, const TSharedPtr<FDocumentTabFactory>& InSpawner, const TSharedPtr<class FDocumentTracker>& InTracker);
+
+	/** Returns TRUE if the payload used by this TabInfo is the same as passed in */
+	bool PayloadMatches(const TSharedPtr<FTabPayload> TestPayload) const;
+
+	static bool PayloadMatches(TSharedPtr<FTabPayload> A, TSharedPtr<FTabPayload> B);
+
+	/**
+	 * Adds history data to immediately after the current location in the history list, wiping out any history after in the process
+	 *
+	 * @param InHistoryNode		The history node to add
+	 * @param bInSaveHistory	TRUE if history should be saved
+	 */
+	void AddTabHistory(TSharedPtr< struct FGenericTabHistory > InHistoryNode, bool bInSaveHistory = true);
+
+	/** Retrieves the history currently visible in this tab */
+	TSharedPtr<struct FGenericTabHistory> GetCurrentHistory() const;
+
+	/** Sets the current history and evoke it, this can either be an old or new history node */
+	void SetCurrentHistory(TSharedPtr<struct FGenericTabHistory> NewHistory, bool bInSaveHistory = true, bool bShouldRestore = true);
+
+	/** Single step forward in history */
+	FReply OnGoForwardInHistory();
+
+	/** Single step back in history */
+	FReply OnGoBackInHistory();
+
+	/** Jumps backwards to the nearest valid history. If no history backwards is valid, jumps forward. In the case that there is no valid history, the tab is closed */
+	void JumpToNearestValidHistoryData();
+
+	/** Retrieves the factory from the current history data, which is being displayed by the tab */
+	TWeakPtr<FDocumentTabFactory> GetFactory() const;
+
+	/** Retrieves the payload from the current history data, which is being displayed by the tab */
+	TSharedPtr<FTabPayload> GetPayload() const;
+
+	/** Retrieves the tab used by this TabInfo */
+	TWeakPtr<class SDockTab> GetTab() const
+	{
+		return Tab;
+	}
+
+	/** Jumps immediately to an index in the history if valid */
+	void GoToHistoryIndex(int32 InHistoryIdx);
+
+	/** Returns TRUE if stepping backward in history is allowed */
+	bool CanStepBackwardInHistory() const;
+
+	/** Returns TRUE if stepping forward in history is allowed */
+	bool CanStepForwardInHistory() const;
+
+	/** Creates the history navigation widget used for navigating back and forward in history */
+	TSharedRef< SWidget > CreateHistoryNavigationWidget();
+
+private:
+	/**
+	 * Mouse down callback to display a history menu
+	 *
+	 * @param InMenuAnchor		This is the anchor the menu will use for positioning
+	 */
+	FReply OnMouseDownHistory( const FGeometry& MyGeometry, const FPointerEvent& MouseEvent, TWeakPtr< SMenuAnchor > InMenuAnchor );
+
+	/** 
+	 * Callback to create the history menu.
+	 *
+	 * @param bInBackHistory		TRUE if the back history is requested, FALSE if the forward history is
+	 *
+	 * @return						The menu widget displaying all available history
+	 */
+	TSharedRef<SWidget> CreateHistoryMenu(bool bInBackHistory) const;
+
+private:	
+	/** Cached history navigation widget */
+	TSharedPtr<SWidget> HistoryNavigationWidget;
+	/** Tab this info represents */
+	TWeakPtr<class SDockTab> Tab;
+	/** Current tab history element in this tab, may not necessarily be in the global history */
+	TSharedPtr<struct FGenericTabHistory> CurrentHistory;
+	/** Pointer to document tracker that spawned us */
+	TWeakPtr<class FDocumentTracker> WeakTracker;
+};
+
+/////////////////////////////////////////////////////
+// FGenericTabHistory
+
+struct FGenericTabHistory
+{
+public:
+	/**
+	 * @param InFactory		The factory used to regenerate the content
+	 * @param InPayload		The payload object used to regenerate the content
+	 */
+	FGenericTabHistory(TSharedPtr<class FDocumentTabFactory> InFactory, TSharedPtr<FTabPayload> InPayload)
+		: FactoryPtr(InFactory)
+		, Payload(InPayload)
+	{
+
+	}
+
+	virtual ~FGenericTabHistory() {}
+
+	/**
+	 * Evokes the history data this struct represents, reinstating the payload and other data it had at the time the history was saved
+	 *
+	 * @param InTab				The tab to update with the history data
+	 */
+	virtual void EvokeHistory(TSharedPtr<FTabInfo> InTabInfo, bool bPrevTabMatches)
+	{
+		FWorkflowTabSpawnInfo SpawnInfo;
+		SpawnInfo.Payload = Payload;
+		SpawnInfo.TabInfo = InTabInfo;
+
+		if(!bPrevTabMatches)
+		{	
+			FactoryPtr.Pin()->UpdateTab(InTabInfo->GetTab().Pin(), SpawnInfo, FactoryPtr.Pin()->CreateTabBody(SpawnInfo));
+		}
+	}
+
+	/** Saves any important information from the payload into the history node */
+	virtual void SaveHistory() {}
+
+	/** Retrieves any important information from the history node and sets it in the payload */
+	virtual void RestoreHistory() {}
+
+	/** Returns the Payload this history node represents */
+	TSharedPtr<FTabPayload> GetPayload() const
+	{
+		return Payload;
+	}
+
+	/** Returns the Factory this history node represents */
+	TWeakPtr<FDocumentTabFactory> GetFactory() const
+	{
+		return FactoryPtr;
+	}
+
+	/**
+	 * Checks if the payload for this history node is currently valid
+	 *
+	 * @return		TRUE if the history is valid
+	 */
+	bool IsHistoryValid() const
+	{
+		return FactoryPtr.Pin()->IsPayloadValid(Payload.ToSharedRef());
+	}
+
+	/** Returns the title used to display this history node */
+	TAttribute<FText> GetHistoryTitle() const
+	{
+		FWorkflowTabSpawnInfo SpawnInfo;
+		SpawnInfo.Payload = Payload;
+
+		return FactoryPtr.Pin()->ConstructTabName(SpawnInfo);
+	}
+
+	/** Returns the tab this is bound to */
+	TSharedPtr<FTabInfo> GetBoundTab() const
+	{
+		return BoundTab.Pin();
+	}
+
+	/** Binds this to a specific tab */
+	void BindToTab(const TSharedPtr<FTabInfo>& InTab)
+	{
+		BoundTab = InTab;
+	}
+
+protected:
+	/** The factory used to generate widgets for the payload */
+	TWeakPtr<class FDocumentTabFactory> FactoryPtr;
+
+	/** The payload this history node represents */
+	TSharedPtr<FTabPayload> Payload;
+
+	/** Specific tab this is bound to, may be null */
+	TWeakPtr<FTabInfo> BoundTab;
+};
+
+/////////////////////////////////////////////////////
+// FDocumentTracker
+
+DECLARE_DELEGATE_OneParam(FOnTabActivated, TSharedPtr<SDockTab>);
+
+class UNREALED_API FDocumentTracker : public TSharedFromThis<FDocumentTracker>
+{
+public:
+	FDocumentTracker();
+	~FDocumentTracker();
+
+	void ClearDocumentFactories();
+	void RegisterDocumentFactory(TSharedPtr<class FDocumentTabFactory> Factory);
+
+	// Creates a new tab manager
+	void Initialize(TSharedPtr<FAssetEditorToolkit> InHostingApp);
+
+	void SetTabManager(const TSharedRef<FTabManager>& InTabManager);
+
+	/** Are we opening a new document, or restoring a previously opened document */
+	enum EOpenDocumentCause
+	{
+		// Opens a new document, will open in existing tab if one is found
+		OpenNewDocument,
+		// Always open in a new tab
+		ForceOpenNewDocument,
+		// Restoring tab state from a previous session
+		RestorePreviousDocument,
+		// Navigating within an open tab
+		NavigatingCurrentDocument,
+		// Navigate without saving current document position
+		QuickNavigateCurrentDocument,
+		// Key/Mouse binding for back, turns into NavigatingHistory
+		NavigateBackwards,
+		// Key/Mouse binding for forward
+		NavigateForwards,
+		// Going forward or back in the history
+		NavigatingHistory,
+		// Same as OpenNewDocument, but does not create a "Document" tab
+		SpawnManagedDocument,
+
+		// Deprecated, treated the same as OpenNewDocument
+		CreateHistoryEvent,
+	};
+
+	/**
+	 * Handles opening a document and handling tab navigation or forcing a new tab to be opened
+	 *
+	 * @param InPayload					Payload to use for the tab
+	 * @param InOpenCause				The cause of the document being opened
+	 *
+	 * @return							The tab modified
+	 */
+	TSharedPtr<SDockTab> OpenDocument(TSharedPtr<FTabPayload> InPayload, EOpenDocumentCause InOpenCause);
+
+	// Closes any tabs of the specified type that match the payload (or all of them if the payload is NULL)
+	void CloseTab(TSharedPtr<FTabPayload> Payload);
+
+	/** Cleans any invalid tabs, either moving them to a valid place in history or closing the tab altogether */
+	void CleanInvalidTabs();
+
+	// Calls OnTabRefreshed for each open tab (on the factories that created them)
+	void RefreshAllTabs() const;
+
+	// Finds a factory that can deal with the supplied payload
+	TSharedPtr<FDocumentTabFactory> FindSupportingFactory(TSharedRef<FTabPayload> Payload) const;
+
+	// Finds all tabs that match the specified identifier and optional payload, placing them in the specified array
+	void FindMatchingTabs(TSharedPtr<FTabPayload> Payload, /*inout*/ TArray< TSharedPtr<SDockTab> >& Results);
+
+	// Finds all tabs that match the specified identifier, placing them in the specified array
+	void FindAllTabsForFactory(const TWeakPtr<FDocumentTabFactory>& Factory, /*inout*/ TArray< TSharedPtr<SDockTab> >& Results);
+
+	// Saves the state of all tabs
+	void SaveAllState();
+
+	// Replaces the open payload in the specified tab with a new one; recreating the contents
+	void ReplacePayloadInTab(TSharedPtr<SDockTab> TargetTab, TSharedPtr<FTabPayload> NewPayload);
+
+	/** Get all the live documents tracked by this tracker */
+	TArray< TSharedPtr<SDockTab> > GetAllDocumentTabs() const;
+
+	/** Returns the current active tab */
+	TSharedPtr<SDockTab> GetActiveTab() const;
+private:
+
+	// List of open tabs; used to bring one to front and focus it
+	typedef TArray< TSharedPtr< FTabInfo > > FTabList;
+private:
+	// List of currently spawned tool panels
+	FTabList SpawnedTabs;
+	
+	// List of allowable tab spawning mechanisms
+	TMap< FName, TSharedPtr<class FDocumentTabFactory> > PotentialTabFactories;
+
+	// The hosting application we're associated with
+	TWeakPtr<FAssetEditorToolkit> HostingAppPtr;
+
+	// The Slate tab manager for the application we are managing documents for
+	TSharedPtr<FTabManager> TabManager;
+
+	// The last active tab info to be selected, used for opening new documents in the most recently used tab
+	TWeakPtr<FTabInfo> WeakLastEditedTabInfo;
+
+	/** Handle to the registered OnActiveTabChanged delegate */
+	FDelegateHandle OnActiveTabChangedDelegateHandle;
+
+	/** Handle to the registered OnTabForegrounded delegate */
+	FDelegateHandle TabForegroundedDelegateHandle;
+
+	/** List of history data nodes */
+	TArray< TSharedPtr<FGenericTabHistory> > History;
+
+	/** Current history index */
+	int32 CurrentHistoryIndex;
+
+private:
+	// Clean the spawned list 
+	FTabList& GetSpawnedList();
+
+	// Cleans the LastEditedTabInfo and returns pinned
+	TSharedPtr<FTabInfo> GetLastEditedTabInfo();
+
+	// Called by the global active tab changed callback; dispatches to individually registered callbacks
+	void OnActiveTabChanged(TSharedPtr<SDockTab> PreviouslyActive, TSharedPtr<SDockTab> NewlyActivated);
+
+	// Called by the global tab manager foregrounds a tab; dispatches to individually registered callbacks
+	void OnTabForegrounded(TSharedPtr<SDockTab> PreviouslyActive, TSharedPtr<SDockTab> NewlyActivated);
+
+	/** Callback for navigating tabs with buttons */
+	FReply OnNavigateTab(FDocumentTracker::EOpenDocumentCause InCause);
+
+	/** Helper function to find the tab in the foreground */
+	TWeakPtr< FTabInfo > FindTabInForeground();
+
+	/** Navigates to specific index in tab history, returns success if it was found */
+	bool NavigateToTabHistory(int32 InHistoryIdx);
+
+	/** Returns the currently active tab history, may be null */
+	TSharedPtr<FGenericTabHistory> GetCurrentTabHistory();
+
+	/**
+	 * Navigates current tab
+	 *
+	 * @param InPayload					Payload to use for the tab
+	 * @param InNavigateCause			The cause of the navigation
+	 *
+	 * @return							The tab navigated and modified
+	 */
+	TSharedPtr<SDockTab> NavigateCurrentTab(TSharedPtr<FTabPayload> InPayload, EOpenDocumentCause InNavigateCause);
+
+	/**
+	 * Opens a new tab
+	 *
+	 * @param InTabHistory				New or existing tab history to open tab to
+	 * @param InOpenCause				The cause of the tab being opened
+	 *
+	 * @return							The tab opened and modified
+	 */
+	TSharedPtr<SDockTab> OpenNewTab(TSharedPtr<FGenericTabHistory> InTabHistory, EOpenDocumentCause InOpenCause);
+
+	friend class FTabInfo;
+};
+
diff --git a/Engine/Source/Editor/UnrealEd/Public/WorkflowOrientedApp/WorkflowUObjectDocuments.h b/Engine/Source/Editor/UnrealEd/Public/WorkflowOrientedApp/WorkflowUObjectDocuments.h
new file mode 100644
--- /dev/null
+++ b/Engine/Source/Editor/UnrealEd/Public/WorkflowOrientedApp/WorkflowUObjectDocuments.h
@@ -0,0 +1,146 @@
+// Copyright Epic Games, Inc. All Rights Reserved.
+
+#pragma once
+
+#include "CoreMinimal.h"
+#include "UObject/Object.h"
+#include "Misc/Attribute.h"
+#include "Templates/Casts.h"
+#include "Widgets/SWidget.h"
+#include "Toolkits/AssetEditorToolkit.h"
+#include "WorkflowOrientedApp/WorkflowTabFactory.h"
+#include "WorkflowOrientedApp/WorkflowTabFactory.h"
+
+
+/////////////////////////////////////////////////////
+// FTabPayload_UObject
+
+struct FTabPayload_UObject : public FTabPayload
+{
+public:
+	// Create a new payload wrapper for a UObject
+	static TSharedRef<FTabPayload_UObject> Make(const UObject* DocumentID)
+	{
+		return MakeShareable(new FTabPayload_UObject(const_cast<UObject*>(DocumentID)));
+	}
+
+	// Helper method to get the payload object as a specific type
+	// (checks both that the payload is a FTabPayload_UObject, and that the type is correct)
+	template <typename CastType>
+	static CastType* CastChecked(TSharedPtr<FTabPayload> Payload)
+	{
+		check((Payload.IsValid()) && (Payload->PayloadType == NAME_Object));
+
+		UObject* UntypedObject = StaticCastSharedPtr<FTabPayload_UObject>(Payload)->DocumentID.Get(true);
+		return ::CastChecked<CastType>(UntypedObject);
+	}
+
+	// Determine if another payload is the same as this one
+	virtual bool IsEqual(const TSharedRef<FTabPayload>& OtherPayload) const override
+	{
+		if (OtherPayload->PayloadType == PayloadType)
+		{
+			return this->DocumentID.HasSameIndexAndSerialNumber(StaticCastSharedRef<FTabPayload_UObject>(OtherPayload)->DocumentID);
+		}
+
+		return false;
+	}
+
+	virtual bool IsValid() const override
+	{
+		return DocumentID.IsValid();
+	}
+
+	virtual ~FTabPayload_UObject() {};
+
+private:
+	// Buried constructor. Use Make instead!
+	FTabPayload_UObject(UObject* InDocumentID)
+		: FTabPayload(NAME_Object)
+		, DocumentID(InDocumentID)
+	{
+	}
+
+private:
+	// The object that is the real payload
+	TWeakObjectPtr<UObject> DocumentID;
+};
+
+/////////////////////////////////////////////////////
+// FDocumentTabFactoryForObjects
+
+template<typename BaseClass>
+struct FDocumentTabFactoryForObjects : public FDocumentTabFactory
+{
+public:
+	// Does this factory support this type of objects?
+	virtual bool SupportsObjectType(UObject* DocumentID) const
+	{
+		return DocumentID->IsA(BaseClass::StaticClass());
+	}
+public:
+
+	// FWorkflowTabFactory public interface
+	virtual TSharedRef<SWidget> CreateTabBody(const FWorkflowTabSpawnInfo& Info) const override
+	{
+		BaseClass* TypedDocumentID = FTabPayload_UObject::CastChecked<BaseClass>(Info.Payload);
+		return CreateTabBodyForObject(Info, TypedDocumentID);
+	}
+
+	virtual const FSlateBrush* GetTabIcon(const FWorkflowTabSpawnInfo& Info) const override
+	{
+		BaseClass* TypedDocumentID = FTabPayload_UObject::CastChecked<BaseClass>(Info.Payload);
+		return GetTabIconForObject(Info, TypedDocumentID);
+	}
+
+	virtual bool IsPayloadSupported(TSharedRef<FTabPayload> Payload) const override
+	{
+		if (Payload->PayloadType == NAME_Object && Payload->IsValid())
+		{
+			UObject* DocumentID = FTabPayload_UObject::CastChecked<UObject>(Payload);
+			return SupportsObjectType(DocumentID);
+		}
+		return false;
+	}
+
+	virtual bool IsPayloadValid(TSharedRef<FTabPayload> Payload) const override
+	{
+		if (Payload->PayloadType == NAME_Object)
+		{
+			return Payload->IsValid();
+		}
+		return false;
+	}
+	// End of FWorkflowTabFactory public interface
+
+protected:
+	// Protected constructor
+	FDocumentTabFactoryForObjects(FName InIdentifier, TSharedPtr<class FAssetEditorToolkit> InHostingApp)
+		: FDocumentTabFactory(InIdentifier, InHostingApp)
+	{
+	}
+
+	// Creates the label for the tab
+	virtual TAttribute<FText> ConstructTabName(const FWorkflowTabSpawnInfo& Info) const override
+	{
+		check(Info.Payload.IsValid());
+		BaseClass* TypedDocumentID = FTabPayload_UObject::CastChecked<BaseClass>(Info.Payload);
+		return ConstructTabNameForObject(TypedDocumentID);
+		//@TODO: DOCMANAGEMENT: Need to incorporate Info.TabPrefix into these guys!
+	}
+
+	// Methods to override in subclasses that make life easier for UObject-based documents
+	virtual TAttribute<FText> ConstructTabNameForObject(BaseClass* DocumentID) const=0;
+	virtual TSharedRef<SWidget> CreateTabBodyForObject(const FWorkflowTabSpawnInfo& Info, BaseClass* DocumentID) const=0;
+	virtual const FSlateBrush* GetTabIconForObject(const FWorkflowTabSpawnInfo& Info, BaseClass* DocumentID) const=0;
+};
+
+
+
+
+
+
+
+
+
+
