PR #7354: Make Alt+Drag identical to Ctrl+W (Contributed by [REDACTED_PERSON])

diff --git a/Engine/Source/Editor/UnrealEd/Private/EditorActor.cpp b/Engine/Source/Editor/UnrealEd/Private/EditorActor.cpp
--- a/Engine/Source/Editor/UnrealEd/Private/EditorActor.cpp
+++ b/Engine/Source/Editor/UnrealEd/Private/EditorActor.cpp
@@ -364,7 +364,7 @@ void UUnrealEdEngine::edactPasteSelected(UWorld* InWorld, bool bDuplicate, bool
 	}
 }
 
-namespace DuplicateSelectedActors {
+namespace  {
 /**
  * A collection of actors to duplicate and prefabs to instance that all belong to the same level.
  */
@@ -488,8 +488,6 @@ void UUnrealEdEngine::edactDuplicateSelected( ULevel* InLevel, bool bOffsetLocat
 	}
 	else
 	{
-		using namespace DuplicateSelectedActors;
-
 		const FScopedBusyCursor BusyCursor;
 		GetSelectedActors()->Modify();
 
diff --git a/Engine/Source/Editor/UnrealEd/Private/LevelEditorViewport.cpp b/Engine/Source/Editor/UnrealEd/Private/LevelEditorViewport.cpp
--- a/Engine/Source/Editor/UnrealEd/Private/LevelEditorViewport.cpp
+++ b/Engine/Source/Editor/UnrealEd/Private/LevelEditorViewport.cpp
@@ -2563,9 +2563,11 @@ bool FLevelEditorViewportClient::InputWidgetDelta(FViewport* InViewport, EAxisLi
 						{
 							// Widget hasn't been dragged since ALT+LMB went down.
 							bDuplicateOnNextDrag = false;
-							ABrush::SetSuppressBSPRegeneration(true);
-							GEditor->edactDuplicateSelected(GetWorld()->GetCurrentLevel(), false);
-							ABrush::SetSuppressBSPRegeneration(false);
+							UUnrealEdEngine* UnrealEdEngine = Cast<UUnrealEdEngine>(GEditor);
+							if (UnrealEdEngine)
+							{
+								UnrealEdEngine->DuplicateSelectedActors(GetWorld());
+							}
 						}
 					}
 				}
diff --git a/Engine/Source/Editor/UnrealEd/Private/UnrealEdEngine.cpp b/Engine/Source/Editor/UnrealEd/Private/UnrealEdEngine.cpp
--- a/Engine/Source/Editor/UnrealEd/Private/UnrealEdEngine.cpp
+++ b/Engine/Source/Editor/UnrealEd/Private/UnrealEdEngine.cpp
@@ -69,6 +69,7 @@
 #include "Logging/MessageLog.h"
 #include "Misc/MessageDialog.h"
 #include "Logging/MessageLog.h"
+#include "Subsystems/EditorActorSubsystem.h"
 
 DEFINE_LOG_CATEGORY_STATIC(LogUnrealEdEngine, Log, All);
 
@@ -990,6 +991,27 @@ void UUnrealEdEngine::Serialize(FArchive& Ar)
 	Ar << MatineeCopyPasteBuffer;
 }
 
+void UUnrealEdEngine::DuplicateSelectedActors(UWorld* InWorld)
+{
+	bool bHandled = false;
+	bHandled |= GLevelEditorModeTools().ProcessEditDuplicate();
+
+	// if not specially handled by the current editing mode,
+	if (!bHandled)
+	{
+		UEditorActorSubsystem* EditorActorSubsystem = GEditor->GetEditorSubsystem<UEditorActorSubsystem>();
+
+		if (EditorActorSubsystem)
+		{
+			EditorActorSubsystem->DuplicateSelectedActors(InWorld);
+		}
+	}
+	// DuplicateSelectedActors also calls RedrawLevelEditingViewports
+	else
+	{
+		RedrawLevelEditingViewports();
+	}
+}
 
 void UUnrealEdEngine::MakeSelectedActorsLevelCurrent()
 {
diff --git a/Engine/Source/Editor/UnrealEd/Private/UnrealEdSrv.cpp b/Engine/Source/Editor/UnrealEd/Private/UnrealEdSrv.cpp
--- a/Engine/Source/Editor/UnrealEd/Private/UnrealEdSrv.cpp
+++ b/Engine/Source/Editor/UnrealEd/Private/UnrealEdSrv.cpp
@@ -479,7 +479,7 @@
 				if( !bSilent )
 				{
 					GWarn->UpdateProgress( 1, 1 );
-					GWarn->EndSlowTask();
+					GWarn->EndSlowTask();
 				}
 			}
 		}
@@ -2692,24 +2692,7 @@
 	}
 	else if( FParse::Command(&Str,TEXT("DUPLICATE")) )
 	{
-		bool bHandled = false;
-		bHandled |= GLevelEditorModeTools().ProcessEditDuplicate();
-
-		// if not specially handled by the current editing mode,
-		if (!bHandled)
-		{
-			UEditorActorSubsystem* EditorActorSubsystem = GEditor->GetEditorSubsystem<UEditorActorSubsystem>();
-
-			if (EditorActorSubsystem)
-			{
-				EditorActorSubsystem->DuplicateSelectedActors(InWorld);
-			}
-		}
-		// DuplicateSelectedActors also calls RedrawLevelEditingViewports
-		else
-		{
-			RedrawLevelEditingViewports(); 
-		}
+		DuplicateSelectedActors(InWorld);
 		return true;
 	}
 	else if( FParse::Command(&Str, TEXT("ALIGN")) )
