Add proto for UnrealEdEngine::DuplicateSelectedActors

#rnx

diff --git a/Engine/Source/Editor/UnrealEd/Classes/Editor/UnrealEdEngine.h b/Engine/Source/Editor/UnrealEd/Classes/Editor/UnrealEdEngine.h
--- a/Engine/Source/Editor/UnrealEd/Classes/Editor/UnrealEdEngine.h
+++ b/Engine/Source/Editor/UnrealEd/Classes/Editor/UnrealEdEngine.h
@@ -419,7 +419,7 @@
 	 * @return								true unless the delete operation was aborted.
 	 */
 	virtual bool edactDeleteSelected( UWorld* InWorld, bool bVerifyDeletionCanHappen=true, bool bWarnAboutReferences = true, bool bWarnAboutSoftReferences = true) override;
-
+
 	/**
 	 * Creates a new group from the current selection removing any existing groups.
 	 */
@@ -680,8 +680,15 @@
 	 * @return true if level streaming should prefer to stream levels from disk instead of duplicating them from editor world
 	 */
 	virtual bool PreferToStreamLevelsInPIE() const override;
-	
+
 	/**
+	 * Duplicate the currently selected actors.
+	 *
+	 * This is a high level routine which may ultimately call edactDuplicateSelected
+	 */
+	void DuplicateSelectedActors(UWorld* InWorld);
+
+	/**
 	 * If all selected actors belong to the same level, that level is made the current level.
 	 */
 	void MakeSelectedActorsLevelCurrent();
