Fix typo in comment.

diff --git a/Engine/Source/Editor/UnrealEd/Private/PackageSourceControlHelper.cpp b/Engine/Source/Editor/UnrealEd/Private/PackageSourceControlHelper.cpp
--- a/Engine/Source/Editor/UnrealEd/Private/PackageSourceControlHelper.cpp
+++ b/Engine/Source/Editor/UnrealEd/Private/PackageSourceControlHelper.cpp
@@ -122,7 +122,7 @@ bool FPackageSourceControlHelper::Delete(const TArray<UPackage*>& Packages) cons
 	{
 		PackagesNames.Add(Package->GetName());
 
-		// Must clear dirty flag befor unloading
+		// Must clear dirty flag before unloading
 		Package->SetDirtyFlag(false);
 	}
 
