The mode manager will now internally ensure only one visible mode is active at a time.

#rnx

diff --git a/Engine/Plugins/MeshPainting/Source/MeshPaintEditorMode/Private/MeshPaintEditorMode.cpp b/Engine/Plugins/MeshPainting/Source/MeshPaintEditorMode/Private/MeshPaintEditorMode.cpp
--- a/Engine/Plugins/MeshPainting/Source/MeshPaintEditorMode/Private/MeshPaintEditorMode.cpp
+++ b/Engine/Plugins/MeshPainting/Source/MeshPaintEditorMode/Private/MeshPaintEditorMode.cpp
@@ -93,9 +93,6 @@ void FMeshPaintEditorModeModule::OnMeshPaintModeButtonClicked()
 	// activate the default mode when the stack becomes empty, resulting in multiple active visible modes.
 	GLevelEditorModeTools().ActivateMode(GetDefault<UMeshPaintMode>()->GetID());
 
-	// Find and disable any other 'visible' modes since we only ever allow one of those active at a time.
-	GLevelEditorModeTools().DeactivateOtherVisibleModes(GetDefault<UMeshPaintMode>()->GetID());
-
 }
 
 
diff --git a/Engine/Source/Editor/LevelEditor/Private/LevelEditor.cpp b/Engine/Source/Editor/LevelEditor/Private/LevelEditor.cpp
--- a/Engine/Source/Editor/LevelEditor/Private/LevelEditor.cpp
+++ b/Engine/Source/Editor/LevelEditor/Private/LevelEditor.cpp
@@ -208,7 +208,6 @@ TSharedRef<SDockTab> FLevelEditorModule::SpawnLevelEditor( const FSpawnTabArgs&
 		LevelEditorTmp->Initialize(LevelEditorTab, OwnerWindow.ToSharedRef());
 
 		GLevelEditorModeTools().DeactivateAllModes();
-		GLevelEditorModeTools().ActivateDefaultMode();
 
 		LevelEditorCreatedEvent.Broadcast(LevelEditorTmp);
 
diff --git a/Engine/Source/Editor/LevelEditor/Private/SLevelEditor.cpp b/Engine/Source/Editor/LevelEditor/Private/SLevelEditor.cpp
--- a/Engine/Source/Editor/LevelEditor/Private/SLevelEditor.cpp
+++ b/Engine/Source/Editor/LevelEditor/Private/SLevelEditor.cpp
@@ -1477,9 +1477,6 @@ void SLevelEditor::ToggleEditorMode( FEditorModeID ModeID )
 	// *Important* - activate the mode first since FEditorModeTools::DeactivateMode will
 	// activate the default mode when the stack becomes empty, resulting in multiple active visible modes.
 	GLevelEditorModeTools().ActivateMode( ModeID );
-
-	// Find and disable any other 'visible' modes since we only ever allow one of those active at a time.
-	GLevelEditorModeTools().DeactivateOtherVisibleModes(ModeID);
 }
 
 bool SLevelEditor::IsModeActive( FEditorModeID ModeID )
diff --git a/Engine/Source/Editor/UnrealEd/Private/EditorModeManager.cpp b/Engine/Source/Editor/UnrealEd/Private/EditorModeManager.cpp
--- a/Engine/Source/Editor/UnrealEd/Private/EditorModeManager.cpp
+++ b/Engine/Source/Editor/UnrealEd/Private/EditorModeManager.cpp
@@ -843,10 +843,11 @@ void FEditorModeTools::ActivateMode(FEditorModeID InID, bool bToggle)
 	}
 
 	// Remove anything that isn't compatible with this mode
+	const bool bIsVisibleMode = ScriptableMode->GetModeInfo().bVisible;
 	for (int32 ModeIndex = ActiveScriptableModes.Num() - 1; ModeIndex >= 0; ModeIndex--)
 	{
 		const bool bModesAreCompatible = ScriptableMode->IsCompatibleWith(ActiveScriptableModes[ModeIndex]->GetID()) || ActiveScriptableModes[ModeIndex]->IsCompatibleWith(ScriptableMode->GetID());
-		if (!bModesAreCompatible)
+		if (!bModesAreCompatible || (bIsVisibleMode && ActiveScriptableModes[ModeIndex]->GetModeInfo().bVisible))
 		{
 			DeactivateScriptableModeAtIndex(ModeIndex);
 		}
diff --git a/Engine/Source/Editor/UnrealEd/Private/EditorServer.cpp b/Engine/Source/Editor/UnrealEd/Private/EditorServer.cpp
--- a/Engine/Source/Editor/UnrealEd/Private/EditorServer.cpp
+++ b/Engine/Source/Editor/UnrealEd/Private/EditorServer.cpp
@@ -2525,12 +2525,6 @@ bool UEditorEngine::Map_Load(const TCHAR* Str, FOutputDevice& Ar)
 					FStatsViewerModule& StatsViewerModule = FModuleManager::Get().LoadModuleChecked<FStatsViewerModule>(TEXT("StatsViewer"));
 					StatsViewerModule.GetPage(EStatsPage::LightingBuildInfo)->Clear();
 
-					if (GUnrealEd)
-					{
-						// The default mode is only available when the full editor has been initialized
-						GLevelEditorModeTools().ActivateDefaultMode();
-					}
-
 					OldOuter = Context.World()->GetOuter();
 
 					ResetTransaction( LocalizedLoadingMap );
diff --git a/Engine/Source/Editor/UnrealEd/Public/EditorModeManager.h b/Engine/Source/Editor/UnrealEd/Public/EditorModeManager.h
--- a/Engine/Source/Editor/UnrealEd/Public/EditorModeManager.h
+++ b/Engine/Source/Editor/UnrealEd/Public/EditorModeManager.h
@@ -479,6 +479,7 @@ public:
 	EEditAction::Type  GetActionEditCopy();
 	EEditAction::Type GetActionEditPaste();
 
+	UE_DEPRECATED(5.0, "This function is redundant, and is handled as part of a call to ActivateMode.")
 	void DeactivateOtherVisibleModes(FEditorModeID InMode);
 	bool IsSnapRotationEnabled() const;
 	bool SnapRotatorToGridOverride(FRotator& InRotation) const;
diff --git a/Engine/Source/Editor/VREditor/VREditorActions.cpp b/Engine/Source/Editor/VREditor/VREditorActions.cpp
--- a/Engine/Source/Editor/VREditor/VREditorActions.cpp
+++ b/Engine/Source/Editor/VREditor/VREditorActions.cpp
@@ -521,9 +521,6 @@ void FVREditorActionCallbacks::ChangeEditorModes(FEditorModeID InMode)
 	// *Important* - activate the mode first since FEditorModeTools::DeactivateMode will
 	// activate the default mode when the stack becomes empty, resulting in multiple active visible modes.
 	GLevelEditorModeTools().ActivateMode(InMode);
-
-	// Find and disable any other 'visible' modes since we only ever allow one of those active at a time.
-	GLevelEditorModeTools().DeactivateOtherVisibleModes(InMode);
 }
 
 ECheckBoxState FVREditorActionCallbacks::EditorModeActive(FEditorModeID InMode)
