First phase of converting Level Editor viewport selection to use typed elements

This phase focuses on the minimum set of changes needed to port the UnrealEd logic for handling the selection of actors and components within the Level Editor viewport to use typed element interfaces. There is still future work to be done to clean this up further.

This change adds a new framework type, UTypedElementSelectionSet, which manages the concept of "selection" for typed elements. Internally this owns its own UTypedElementList, and ensures that mutation of that list goes via the UTypedElementSelectionInterface implementations.

To allow specific asset editors to customize their selection behavior, UTypedElementSelectionSet may have "selection proxies" that implement UTypedElementAssetEditorSelectionProxy registered to it. This is what's used to allow the level editor to implement its type specific rules.

The core Level Editor selection implementation for actors and components now lives inside UActorElementLevelEditorSelectionProxy and UComponentElementLevelEditorSelectionProxy, and the existing UUnrealEdEngine functions that used to deal with this logic now proxy through to those implementations via UTypedElementSelectionSet. This means that the implementation has moved into the LevelEditor module without UnrealEd even knowing.

Future work will focus on:
 - Cleaning up the legacy USelection bridge, so that all USelection instances are just a proxy around a UTypedElementSelectionSet.
   - This will involve making a basic "Object" element type to handle the generic UObject based selection (for assets) that USelection also handles.
   - This should allow GetMutableElementList to be removed from UTypedElementSelectionSet, to avoid potential abuse or misuse.
 - Investigating the best way to integrate element based selection into the editor modes.

#rnx

diff --git a/Engine/Source/Editor/EditorFramework/EditorFramework.Build.cs b/Engine/Source/Editor/EditorFramework/EditorFramework.Build.cs
--- a/Engine/Source/Editor/EditorFramework/EditorFramework.Build.cs
+++ b/Engine/Source/Editor/EditorFramework/EditorFramework.Build.cs
@@ -31,6 +31,8 @@ public class EditorFramework : ModuleRules
 
 		PublicDependencyModuleNames.AddRange(
 			new string[] {
+				"TypedElementFramework",
+				"TypedElementInterfaces",
 			}
 		);
 
diff --git a/Engine/Source/Editor/EditorFramework/Private/Elements/TypedElementSelectionSet.cpp b/Engine/Source/Editor/EditorFramework/Private/Elements/TypedElementSelectionSet.cpp
new file mode 100644
--- /dev/null
+++ b/Engine/Source/Editor/EditorFramework/Private/Elements/TypedElementSelectionSet.cpp
@@ -0,0 +1,144 @@
+// Copyright Epic Games, Inc. All Rights Reserved.
+
+#include "Elements/TypedElementSelectionSet.h"
+#include "TypedElementList.h"
+#include "TypedElementRegistry.h"
+
+UTypedElementSelectionSet::UTypedElementSelectionSet()
+{
+	if (!HasAnyFlags(RF_ClassDefaultObject))
+	{
+		ElementList = UTypedElementRegistry::GetInstance()->CreateElementList();
+	}
+}
+
+bool UTypedElementSelectionSet::IsElementSelected(const FTypedElementHandle& InElementHandle, const FTypedElementIsSelectedOptions InSelectionOptions) const
+{
+	FTypedElementSelectionSetElement SelectionSetElement = ResolveSelectionSetElement(InElementHandle);
+	return SelectionSetElement && SelectionSetElement.IsElementSelected(InSelectionOptions);
+}
+
+bool UTypedElementSelectionSet::CanSelectElement(const FTypedElementHandle& InElementHandle, const FTypedElementSelectionOptions InSelectionOptions) const
+{
+	FTypedElementSelectionSetElement SelectionSetElement = ResolveSelectionSetElement(InElementHandle);
+	return SelectionSetElement && SelectionSetElement.CanSelectElement(InSelectionOptions);
+}
+
+bool UTypedElementSelectionSet::CanDeselectElement(const FTypedElementHandle& InElementHandle, const FTypedElementSelectionOptions InSelectionOptions) const
+{
+	FTypedElementSelectionSetElement SelectionSetElement = ResolveSelectionSetElement(InElementHandle);
+	return SelectionSetElement && SelectionSetElement.CanDeselectElement(InSelectionOptions);
+}
+
+bool UTypedElementSelectionSet::SelectElement(const FTypedElementHandle& InElementHandle, const FTypedElementSelectionOptions InSelectionOptions)
+{
+	FTypedElementSelectionSetElement SelectionSetElement = ResolveSelectionSetElement(InElementHandle);
+	return SelectionSetElement && SelectionSetElement.CanSelectElement(InSelectionOptions) && SelectionSetElement.SelectElement(InSelectionOptions);
+}
+
+bool UTypedElementSelectionSet::SelectElements(const TArray<FTypedElementHandle>& InElementHandles, const FTypedElementSelectionOptions InSelectionOptions)
+{
+	return SelectElements(MakeArrayView(InElementHandles), InSelectionOptions);
+}
+
+bool UTypedElementSelectionSet::SelectElements(TArrayView<const FTypedElementHandle> InElementHandles, const FTypedElementSelectionOptions InSelectionOptions)
+{
+	FTypedElementListLegacySyncScopedBatch LegacySyncBatch(ElementList, InSelectionOptions.AllowLegacyNotifications());
+
+	bool bSelectionChanged = false;
+
+	for (const FTypedElementHandle& ElementHandle : InElementHandles)
+	{
+		bSelectionChanged |= SelectElement(ElementHandle, InSelectionOptions);
+	}
+
+	return bSelectionChanged;
+}
+
+bool UTypedElementSelectionSet::DeselectElement(const FTypedElementHandle& InElementHandle, const FTypedElementSelectionOptions InSelectionOptions)
+{
+	FTypedElementSelectionSetElement SelectionSetElement = ResolveSelectionSetElement(InElementHandle);
+	return SelectionSetElement && SelectionSetElement.CanDeselectElement(InSelectionOptions) && SelectionSetElement.DeselectElement(InSelectionOptions);
+}
+
+bool UTypedElementSelectionSet::DeselectElements(const TArray<FTypedElementHandle>& InElementHandles, const FTypedElementSelectionOptions InSelectionOptions)
+{
+	return DeselectElements(MakeArrayView(InElementHandles), InSelectionOptions);
+}
+
+bool UTypedElementSelectionSet::DeselectElements(TArrayView<const FTypedElementHandle> InElementHandles, const FTypedElementSelectionOptions InSelectionOptions)
+{
+	FTypedElementListLegacySyncScopedBatch LegacySyncBatch(ElementList, InSelectionOptions.AllowLegacyNotifications());
+
+	bool bSelectionChanged = false;
+
+	for (const FTypedElementHandle& ElementHandle : InElementHandles)
+	{
+		bSelectionChanged |= DeselectElement(ElementHandle, InSelectionOptions);
+	}
+
+	return bSelectionChanged;
+}
+
+bool UTypedElementSelectionSet::ClearSelection(const FTypedElementSelectionOptions InSelectionOptions)
+{
+	FTypedElementListLegacySyncScopedBatch LegacySyncBatch(ElementList, InSelectionOptions.AllowLegacyNotifications());
+
+	bool bSelectionChanged = false;
+
+	// Run deselection via the selection interface where possible
+	{
+		// Take a copy of the currently selected elements to avoid mutating the selection set while iterating
+		TArray<FTypedElementHandle, TInlineAllocator<8>> ElementsCopy;
+		ElementsCopy.Reserve(ElementList->Num());
+		for (const FTypedElementHandle& ElementHandle : *ElementList)
+		{
+			ElementsCopy.Emplace(ElementHandle);
+		}
+
+		for (const FTypedElementHandle& ElementHandle : ElementsCopy)
+		{
+			bSelectionChanged |= DeselectElement(ElementHandle, InSelectionOptions);
+		}
+	}
+
+	// TODO: BSP surfaces?
+
+	// If anything remains in the selection set after processing elements that implement this interface, just clear it
+	if (ElementList->Num() > 0)
+	{
+		bSelectionChanged = true;
+		ElementList->Reset();
+	}
+
+	return bSelectionChanged;
+}
+
+bool UTypedElementSelectionSet::AllowSelectionModifiers(const FTypedElementHandle& InElementHandle) const
+{
+	FTypedElementSelectionSetElement SelectionSetElement = ResolveSelectionSetElement(InElementHandle);
+	return SelectionSetElement && SelectionSetElement.AllowSelectionModifiers();
+}
+
+FTypedElementHandle UTypedElementSelectionSet::GetSelectionElement(const FTypedElementHandle& InElementHandle, const ETypedElementSelectionMethod InSelectionMethod) const
+{
+	FTypedElementSelectionSetElement SelectionSetElement = ResolveSelectionSetElement(InElementHandle);
+	return SelectionSetElement ? SelectionSetElement.GetSelectionElement(InSelectionMethod) : FTypedElementHandle();
+}
+
+void UTypedElementSelectionSet::RegisterAssetEditorSelectionProxy(const FName InElementTypeName, UTypedElementAssetEditorSelectionProxy* InAssetEditorSelectionProxy)
+{
+	const FTypedHandleTypeId ElementTypeId = ElementList->GetRegisteredElementTypeId(InElementTypeName);
+	checkf(ElementTypeId > 0, TEXT("Element type '%s' has not been registered!"), *InElementTypeName.ToString());
+	RegisteredAssetEditorSelectionProxies[ElementTypeId - 1] = InAssetEditorSelectionProxy;
+}
+
+FTypedElementSelectionSetElement UTypedElementSelectionSet::ResolveSelectionSetElement(const FTypedElementHandle& InElementHandle) const
+{
+	if (InElementHandle)
+	{
+		UTypedElementAssetEditorSelectionProxy* AssetEditorSelectionProxy = RegisteredAssetEditorSelectionProxies[InElementHandle.GetId().GetTypeId() - 1];
+		return FTypedElementSelectionSetElement(ElementList->GetElement<UTypedElementSelectionInterface>(InElementHandle), ElementList, AssetEditorSelectionProxy ? AssetEditorSelectionProxy : GetMutableDefault<UTypedElementAssetEditorSelectionProxy>());
+	}
+	return FTypedElementSelectionSetElement();
+}
diff --git a/Engine/Source/Editor/EditorFramework/Public/Elements/TypedElementSelectionSet.h b/Engine/Source/Editor/EditorFramework/Public/Elements/TypedElementSelectionSet.h
new file mode 100644
--- /dev/null
+++ b/Engine/Source/Editor/EditorFramework/Public/Elements/TypedElementSelectionSet.h
@@ -0,0 +1,192 @@
+// Copyright Epic Games, Inc. All Rights Reserved.
+
+#pragma once
+
+#include "TypedElementSelectionInterface.h"
+#include "TypedElementSelectionSet.generated.h"
+
+class UTypedElementList;
+
+/**
+ * Proxy type used to allow asset editors (such as the level editor) to override the base behavior of element selection,
+ * by injecting extra pre/post selection logic around the call into the selection interface for an element type.
+ */
+UCLASS()
+class EDITORFRAMEWORK_API UTypedElementAssetEditorSelectionProxy : public UObject
+{
+	GENERATED_BODY()
+
+public:
+	//~ See UTypedElementSelectionInterface for API docs
+	virtual bool IsElementSelected(const TTypedElement<UTypedElementSelectionInterface>& InElementSelectionHandle, const UTypedElementList* InSelectionSet, const FTypedElementIsSelectedOptions& InSelectionOptions) { return InElementSelectionHandle.IsElementSelected(InSelectionSet, InSelectionOptions); }
+	virtual bool CanSelectElement(const TTypedElement<UTypedElementSelectionInterface>& InElementSelectionHandle, const FTypedElementSelectionOptions& InSelectionOptions) { return InElementSelectionHandle.CanSelectElement(InSelectionOptions); }
+	virtual bool CanDeselectElement(const TTypedElement<UTypedElementSelectionInterface>& InElementSelectionHandle, const FTypedElementSelectionOptions& InSelectionOptions) { return InElementSelectionHandle.CanDeselectElement(InSelectionOptions); }
+	virtual bool SelectElement(const TTypedElement<UTypedElementSelectionInterface>& InElementSelectionHandle, UTypedElementList* InSelectionSet, const FTypedElementSelectionOptions& InSelectionOptions) { return InElementSelectionHandle.SelectElement(InSelectionSet, InSelectionOptions); }
+	virtual bool DeselectElement(const TTypedElement<UTypedElementSelectionInterface>& InElementSelectionHandle, UTypedElementList* InSelectionSet, const FTypedElementSelectionOptions& InSelectionOptions) { return InElementSelectionHandle.DeselectElement(InSelectionSet, InSelectionOptions); }
+	virtual bool AllowSelectionModifiers(const TTypedElement<UTypedElementSelectionInterface>& InElementSelectionHandle, const UTypedElementList* InSelectionSet) { return InElementSelectionHandle.AllowSelectionModifiers(InSelectionSet); }
+	virtual FTypedElementHandle GetSelectionElement(const TTypedElement<UTypedElementSelectionInterface>& InElementSelectionHandle, const UTypedElementList* InCurrentSelection, const ETypedElementSelectionMethod InSelectionMethod) { return InElementSelectionHandle.GetSelectionElement(InCurrentSelection, InSelectionMethod); }
+};
+
+/**
+ * Utility to hold a typed element handle and its associated selection interface and asset editor selection proxy.
+ */
+struct EDITORFRAMEWORK_API FTypedElementSelectionSetElement
+{
+public:
+	FTypedElementSelectionSetElement() = default;
+
+	FTypedElementSelectionSetElement(TTypedElement<UTypedElementSelectionInterface> InElementSelectionHandle, UTypedElementList* InElementList, UTypedElementAssetEditorSelectionProxy* InAssetEditorSelectionProxy)
+		: ElementSelectionHandle(MoveTemp(InElementSelectionHandle))
+		, ElementList(InElementList)
+		, AssetEditorSelectionProxy(InAssetEditorSelectionProxy)
+	{
+	}
+
+	FTypedElementSelectionSetElement(const FTypedElementSelectionSetElement&) = default;
+	FTypedElementSelectionSetElement& operator=(const FTypedElementSelectionSetElement&) = default;
+
+	FTypedElementSelectionSetElement(FTypedElementSelectionSetElement&&) = default;
+	FTypedElementSelectionSetElement& operator=(FTypedElementSelectionSetElement&&) = default;
+
+	FORCEINLINE explicit operator bool() const
+	{
+		return IsSet();
+	}
+
+	FORCEINLINE bool IsSet() const
+	{
+		return ElementSelectionHandle.IsSet()
+			&& ElementList
+			&& AssetEditorSelectionProxy;
+	}
+
+	//~ See UTypedElementSelectionInterface for API docs
+	bool IsElementSelected(const FTypedElementIsSelectedOptions& InSelectionOptions) const { return AssetEditorSelectionProxy->IsElementSelected(ElementSelectionHandle, ElementList, InSelectionOptions); }
+	bool CanSelectElement(const FTypedElementSelectionOptions& InSelectionOptions) const { return AssetEditorSelectionProxy->CanSelectElement(ElementSelectionHandle, InSelectionOptions); }
+	bool CanDeselectElement(const FTypedElementSelectionOptions& InSelectionOptions) const { return AssetEditorSelectionProxy->CanDeselectElement(ElementSelectionHandle, InSelectionOptions); }
+	bool SelectElement(const FTypedElementSelectionOptions& InSelectionOptions) const { return AssetEditorSelectionProxy->SelectElement(ElementSelectionHandle, ElementList, InSelectionOptions); }
+	bool DeselectElement(const FTypedElementSelectionOptions& InSelectionOptions) const { return AssetEditorSelectionProxy->DeselectElement(ElementSelectionHandle, ElementList, InSelectionOptions); }
+	bool AllowSelectionModifiers() const { return AssetEditorSelectionProxy->AllowSelectionModifiers(ElementSelectionHandle, ElementList); }
+	FTypedElementHandle GetSelectionElement(const ETypedElementSelectionMethod InSelectionMethod) const { return AssetEditorSelectionProxy->GetSelectionElement(ElementSelectionHandle, ElementList, InSelectionMethod); }
+
+private:
+	TTypedElement<UTypedElementSelectionInterface> ElementSelectionHandle;
+	UTypedElementList* ElementList = nullptr;
+	UTypedElementAssetEditorSelectionProxy* AssetEditorSelectionProxy = nullptr;
+};
+
+/**
+ * A wrapper around an element list that ensures mutation goes via the selection 
+ * interfaces, as well as providing some utilities for batching operations.
+ */
+UCLASS(Transient)
+class EDITORFRAMEWORK_API UTypedElementSelectionSet : public UObject
+{
+	GENERATED_BODY()
+
+public:
+	UTypedElementSelectionSet();
+
+	/**
+	 * Test to see whether the given element is currently considered selected.
+	 */
+	UFUNCTION(BlueprintPure, Category="TypedElementFramework|Selection")
+	bool IsElementSelected(const FTypedElementHandle& InElementHandle, const FTypedElementIsSelectedOptions InSelectionOptions) const;
+
+	/**
+	 * Test to see whether the given element can be selected.
+	 */
+	UFUNCTION(BlueprintPure, Category="TypedElementFramework|Selection")
+	bool CanSelectElement(const FTypedElementHandle& InElementHandle, const FTypedElementSelectionOptions InSelectionOptions) const;
+
+	/**
+	 * Test to see whether the given element can be deselected.
+	 */
+	UFUNCTION(BlueprintPure, Category="TypedElementFramework|Selection")
+	bool CanDeselectElement(const FTypedElementHandle& InElementHandle, const FTypedElementSelectionOptions InSelectionOptions) const;
+
+	/**
+	 * Attempt to select the given element.
+	 * @return True if the selection was changed, false otherwise.
+	 */
+	UFUNCTION(BlueprintCallable, Category="TypedElementFramework|Selection")
+	bool SelectElement(const FTypedElementHandle& InElementHandle, const FTypedElementSelectionOptions InSelectionOptions);
+
+	/**
+	 * Attempt to select the given elements.
+	 * @return True if the selection was changed, false otherwise.
+	 */
+	UFUNCTION(BlueprintCallable, Category="TypedElementFramework|Selection")
+	bool SelectElements(const TArray<FTypedElementHandle>& InElementHandles, const FTypedElementSelectionOptions InSelectionOptions);
+	bool SelectElements(TArrayView<const FTypedElementHandle> InElementHandles, const FTypedElementSelectionOptions InSelectionOptions);
+
+	/**
+	 * Attempt to deselect the given element.
+	 * @return True if the selection was changed, false otherwise.
+	 */
+	UFUNCTION(BlueprintCallable, Category="TypedElementFramework|Selection")
+	bool DeselectElement(const FTypedElementHandle& InElementHandle, const FTypedElementSelectionOptions InSelectionOptions);
+
+	/**
+	 * Attempt to deselect the given elements.
+	 * @return True if the selection was changed, false otherwise.
+	 */
+	UFUNCTION(BlueprintCallable, Category="TypedElementFramework|Selection")
+	bool DeselectElements(const TArray<FTypedElementHandle>& InElementHandles, const FTypedElementSelectionOptions InSelectionOptions);
+	bool DeselectElements(TArrayView<const FTypedElementHandle> InElementHandles, const FTypedElementSelectionOptions InSelectionOptions);
+
+	/**
+	 * Clear the current selection.
+	 * @return True if the selection was changed, false otherwise.
+	 */
+	UFUNCTION(BlueprintCallable, Category="TypedElementFramework|Selection")
+	bool ClearSelection(const FTypedElementSelectionOptions InSelectionOptions);
+
+	/**
+	 * Test to see whether selection modifiers (Ctrl or Shift) are allowed while selecting this element.
+	 */
+	UFUNCTION(BlueprintPure, Category="TypedElementFramework|Selection")
+	bool AllowSelectionModifiers(const FTypedElementHandle& InElementHandle) const;
+
+	/**
+	 * Given an element, return the element that should actually perform a selection operation.
+	 */
+	UFUNCTION(BlueprintPure, Category="TypedElementFramework|Selection")
+	FTypedElementHandle GetSelectionElement(const FTypedElementHandle& InElementHandle, const ETypedElementSelectionMethod InSelectionMethod) const;
+
+	/**
+	 * Get the underlying element list holding the selection state.
+	 */
+	const UTypedElementList* GetElementList() const
+	{
+		return ElementList;
+	}
+
+	/**
+	 * Get the underlying element list holding the selection state.
+	 * @note Mutating the internal element list directly is a bad idea.
+	 */
+	UTypedElementList* GetMutableElementList()
+	{
+		return ElementList;
+	}
+
+	/**
+	 * Register an asset editor selection proxy for the given named element type.
+	 */
+	void RegisterAssetEditorSelectionProxy(const FName InElementTypeName, UTypedElementAssetEditorSelectionProxy* InAssetEditorSelectionProxy);
+
+private:
+	/**
+	 * Attempt to resolve the selection interface and asset editor selection proxy for the given element, if any.
+	 */
+	FTypedElementSelectionSetElement ResolveSelectionSetElement(const FTypedElementHandle& InElementHandle) const;
+
+	/** Underlying element list holding the selection state. */
+	UPROPERTY()
+	UTypedElementList* ElementList = nullptr;
+
+	/** Array of registered asset editor selection proxies, indexed by ElementTypeId-1. */
+	UPROPERTY()
+	UTypedElementAssetEditorSelectionProxy* RegisteredAssetEditorSelectionProxies[TypedHandleMaxTypeId - 1];
+};
diff --git a/Engine/Source/Editor/LevelEditor/LevelEditor.Build.cs b/Engine/Source/Editor/LevelEditor/LevelEditor.Build.cs
--- a/Engine/Source/Editor/LevelEditor/LevelEditor.Build.cs
+++ b/Engine/Source/Editor/LevelEditor/LevelEditor.Build.cs
@@ -60,6 +60,7 @@ public class LevelEditor : ModuleRules
 				"Projects",
 				"TargetPlatform",
 				"TypedElementFramework",
+				"TypedElementInterfaces",
 				"EngineSettings",
 				"PropertyEditor",
 				"Kismet",
diff --git a/Engine/Source/Editor/LevelEditor/Private/Elements/Actor/ActorElementLevelEditorSelectionProxy.cpp b/Engine/Source/Editor/LevelEditor/Private/Elements/Actor/ActorElementLevelEditorSelectionProxy.cpp
new file mode 100644
--- /dev/null
+++ b/Engine/Source/Editor/LevelEditor/Private/Elements/Actor/ActorElementLevelEditorSelectionProxy.cpp
@@ -0,0 +1,300 @@
+// Copyright Epic Games, Inc. All Rights Reserved.
+
+#include "Elements/Actor/ActorElementLevelEditorSelectionProxy.h"
+#include "Elements/Actor/ActorElementData.h"
+#include "GameFramework/Actor.h"
+#include "Components/ActorComponent.h"
+#include "TypedElementList.h"
+
+#include "Elements/Actor/ActorElementSelectionInterface.h"
+#include "Elements/Component/ComponentElementSelectionInterface.h"
+
+#include "Editor.h"
+#include "LevelUtils.h"
+#include "UnrealEdGlobals.h"
+#include "EditorModeManager.h"
+#include "Editor/GroupActor.h"
+#include "ActorGroupingUtils.h"
+#include "Editor/UnrealEdEngine.h"
+#include "Kismet2/ComponentEditorUtils.h"
+
+DEFINE_LOG_CATEGORY_STATIC(LogActorLevelEditorSelection, Log, All);
+
+bool UActorElementLevelEditorSelectionProxy::CanSelectElement(const TTypedElement<UTypedElementSelectionInterface>& InElementSelectionHandle, const FTypedElementSelectionOptions& InSelectionOptions)
+{
+	return InElementSelectionHandle.IsDataOfType<FActorElementData>() && CanSelectActorElement(InElementSelectionHandle, InSelectionOptions);
+}
+
+bool UActorElementLevelEditorSelectionProxy::CanDeselectElement(const TTypedElement<UTypedElementSelectionInterface>& InElementSelectionHandle, const FTypedElementSelectionOptions& InSelectionOptions)
+{
+	return InElementSelectionHandle.IsDataOfType<FActorElementData>() && CanDeselectActorElement(InElementSelectionHandle, InSelectionOptions);
+}
+
+bool UActorElementLevelEditorSelectionProxy::SelectElement(const TTypedElement<UTypedElementSelectionInterface>& InElementSelectionHandle, UTypedElementList* InSelectionSet, const FTypedElementSelectionOptions& InSelectionOptions)
+{
+	return InElementSelectionHandle.IsDataOfType<FActorElementData>() && SelectActorElement(InElementSelectionHandle, InSelectionSet, InSelectionOptions);
+}
+
+bool UActorElementLevelEditorSelectionProxy::DeselectElement(const TTypedElement<UTypedElementSelectionInterface>& InElementSelectionHandle, UTypedElementList* InSelectionSet, const FTypedElementSelectionOptions& InSelectionOptions)
+{
+	return InElementSelectionHandle.IsDataOfType<FActorElementData>() && DeselectActorElement(InElementSelectionHandle, InSelectionSet, InSelectionOptions);
+}
+
+bool UActorElementLevelEditorSelectionProxy::AllowSelectionModifiers(const TTypedElement<UTypedElementSelectionInterface>& InElementSelectionHandle, const UTypedElementList* InSelectionSet)
+{
+	// Ctrl or Shift clicking an actor is the same as regular clicking when components are selected
+	return InElementSelectionHandle.IsDataOfType<FActorElementData>() && !UComponentElementSelectionInterface::HasSelectedComponents(InSelectionSet);
+}
+
+FTypedElementHandle UActorElementLevelEditorSelectionProxy::GetSelectionElement(const TTypedElement<UTypedElementSelectionInterface>& InElementSelectionHandle, const UTypedElementList* InCurrentSelection, const ETypedElementSelectionMethod InSelectionMethod)
+{
+	if (const FActorElementData* ActorData = InElementSelectionHandle.GetData<FActorElementData>())
+	{
+		AActor* ConsideredActor = ActorData->Actor;
+		while (ConsideredActor->IsChildActor())
+		{
+			ConsideredActor = ConsideredActor->GetParentActor();
+		}
+		return ConsideredActor->AcquireEditorElementHandle();
+	}
+	return InElementSelectionHandle;
+}
+
+bool UActorElementLevelEditorSelectionProxy::CanSelectActorElement(const TTypedElement<UTypedElementSelectionInterface>& InActorSelectionHandle, const FTypedElementSelectionOptions& InSelectionOptions)
+{
+	AActor* Actor = InActorSelectionHandle.GetDataChecked<FActorElementData>().Actor;
+
+	// Bail if global selection is locked, or this actor cannot be edited or selected
+	if (GEdSelectionLock || !Actor->IsEditable() || !Actor->IsSelectable())
+	{
+		return false;
+	}
+
+	// Bail if the actor is hidden, and we're not allowed to select hidden elements
+	if (!InSelectionOptions.AllowHidden() && (Actor->IsHiddenEd() || !FLevelUtils::IsLevelVisible(Actor->GetLevel())))
+	{
+		return false;
+	}
+
+	// Ensure that neither the level nor the actor is being destroyed or is unreachable
+	const EObjectFlags InvalidSelectableFlags = RF_BeginDestroyed;
+	if (Actor->GetLevel()->HasAnyFlags(InvalidSelectableFlags) || Actor->GetLevel()->IsPendingKillOrUnreachable())
+	{
+		UE_LOG(LogActorLevelEditorSelection, Warning, TEXT("SelectActor: %s (%s)"), TEXT("The requested operation could not be completed because the level has invalid flags."), *Actor->GetActorLabel());
+		return false;
+	}
+	if (Actor->HasAnyFlags(InvalidSelectableFlags) || Actor->IsPendingKillOrUnreachable())
+	{
+		UE_LOG(LogActorLevelEditorSelection, Warning, TEXT("SelectActor: %s (%s)"), TEXT("The requested operation could not be completed because the actor has invalid flags."), *Actor->GetActorLabel());
+		return false;
+	}
+
+	if (!Actor->IsTemplate() && FLevelUtils::IsLevelLocked(Actor->GetLevel()))
+	{
+		UE_CLOG(InSelectionOptions.WarnIfLocked(), LogActorLevelEditorSelection, Warning, TEXT("SelectActor: %s (%s)"), TEXT("The requested operation could not be completed because the level is locked."), *Actor->GetActorLabel());
+		return false;
+	}
+
+	// If grouping operations are not currently allowed, don't select groups
+	AGroupActor* SelectedGroupActor = Cast<AGroupActor>(Actor);
+	if (SelectedGroupActor && (!UActorGroupingUtils::IsGroupingActive() || !InSelectionOptions.AllowGroups()))
+	{
+		return false;
+	}
+
+	// Allow active modes to determine whether the selection is allowed
+	return GLevelEditorModeTools().IsSelectionAllowed(Actor, /*bInSelected*/true);
+}
+
+bool UActorElementLevelEditorSelectionProxy::CanDeselectActorElement(const TTypedElement<UTypedElementSelectionInterface>& InActorSelectionHandle, const FTypedElementSelectionOptions& InSelectionOptions)
+{
+	AActor* Actor = InActorSelectionHandle.GetDataChecked<FActorElementData>().Actor;
+
+	// Bail if global selection is locked
+	if (GEdSelectionLock)
+	{
+		return false;
+	}
+
+	// Allow active modes to determine whether the deselection is allowed
+	return GLevelEditorModeTools().IsSelectionAllowed(Actor, /*bInSelected*/false);
+}
+
+bool UActorElementLevelEditorSelectionProxy::SelectActorElement(const TTypedElement<UTypedElementSelectionInterface>& InActorSelectionHandle, UTypedElementList* InSelectionSet, const FTypedElementSelectionOptions& InSelectionOptions)
+{
+	AActor* Actor = InActorSelectionHandle.GetDataChecked<FActorElementData>().Actor;
+
+	// Allow active modes to potentially handle the selection
+	// TODO: Should this pass through the selection set?
+	if (GLevelEditorModeTools().IsSelectionHandled(Actor, /*bInSelected*/true))
+	{
+		return true;
+	}
+
+	// If trying to select an actor, use this actors root selection actor instead (if it has one)
+	if (AActor* RootSelection = Actor->GetRootSelectionParent())
+	{
+		Actor = RootSelection;
+	}
+
+	bool bSelectionChanged = false;
+
+	if (UActorGroupingUtils::IsGroupingActive() && InSelectionOptions.AllowGroups())
+	{
+		// If this actor is a group, do a group select
+		if (AGroupActor* SelectedGroupActor = Cast<AGroupActor>(Actor))
+		{
+			bSelectionChanged |= SelectActorGroup(SelectedGroupActor, InSelectionSet, InSelectionOptions, /*bForce*/true);
+		}
+		// Select/Deselect this actor's entire group, starting from the top locked group
+		else if (AGroupActor* ActorLockedRootGroup = AGroupActor::GetRootForActor(Actor, true))
+		{
+			bSelectionChanged |= SelectActorGroup(ActorLockedRootGroup, InSelectionSet, InSelectionOptions, /*bForce*/false);
+		}
+	}
+
+	// Select the desired actor
+	{
+		TTypedElement<UTypedElementSelectionInterface> ActorSelectionHandle = InSelectionSet->GetElement<UTypedElementSelectionInterface>(Actor->AcquireEditorElementHandle());
+		if (!ActorSelectionHandle.SelectElement(InSelectionSet, InSelectionOptions))
+		{
+			return bSelectionChanged;
+		}
+	}
+
+	UE_LOG(LogActorLevelEditorSelection, Verbose, TEXT("Selected Actor: %s"), *Actor->GetClass()->GetName());
+
+	// Bind the override delegates for the components on the selected actor
+	for (UActorComponent* Component : Actor->GetComponents())
+	{
+		if (USceneComponent* SceneComponent = Cast<USceneComponent>(Component))
+		{
+			FComponentEditorUtils::BindComponentSelectionOverride(SceneComponent, /*bBind*/true);
+		}
+	}
+		
+	// Flush some cached data
+	GUnrealEd->PostActorSelectionChanged();
+
+	// A fast path to mark selection rather than reconnecting ALL components for ALL actors that have changed state
+	Actor->PushSelectionToProxies();
+
+	return true;
+}
+
+bool UActorElementLevelEditorSelectionProxy::DeselectActorElement(const TTypedElement<UTypedElementSelectionInterface>& InActorSelectionHandle, UTypedElementList* InSelectionSet, const FTypedElementSelectionOptions& InSelectionOptions)
+{
+	AActor* Actor = InActorSelectionHandle.GetDataChecked<FActorElementData>().Actor;
+
+	// Allow active modes to potentially handle the deselection
+	// TODO: Should this pass through the selection set?
+	if (GLevelEditorModeTools().IsSelectionHandled(Actor, /*bInSelected*/false))
+	{
+		return true;
+	}
+
+	bool bSelectionChanged = false;
+
+	if (UActorGroupingUtils::IsGroupingActive() && InSelectionOptions.AllowGroups())
+	{
+		// If this actor is a group, do a group select
+		if (AGroupActor* SelectedGroupActor = Cast<AGroupActor>(Actor))
+		{
+			bSelectionChanged |= DeselectActorGroup(SelectedGroupActor, InSelectionSet, InSelectionOptions, /*bForce*/true);
+		}
+		// Select/Deselect this actor's entire group, starting from the top locked group
+		else if (AGroupActor* ActorLockedRootGroup = AGroupActor::GetRootForActor(Actor, true))
+		{
+			bSelectionChanged |= DeselectActorGroup(ActorLockedRootGroup, InSelectionSet, InSelectionOptions, /*bForce*/false);
+		}
+	}
+
+	// Deselect the desired actor
+	{
+		TTypedElement<UTypedElementSelectionInterface> ActorSelectionHandle = InSelectionSet->GetElement<UTypedElementSelectionInterface>(Actor->AcquireEditorElementHandle());
+		if (!ActorSelectionHandle.DeselectElement(InSelectionSet, InSelectionOptions))
+		{
+			return bSelectionChanged;
+		}
+	}
+	
+	UE_LOG(LogActorLevelEditorSelection, Verbose, TEXT("Deselected Actor: %s"), *Actor->GetClass()->GetName());
+	
+	// Deselect and unbind the override delegates for the components on the selected actor
+	{
+		FTypedElementListLegacySyncScopedBatch LegacySyncBatch(InSelectionSet, InSelectionOptions.AllowLegacyNotifications());
+
+		for (UActorComponent* Component : Actor->GetComponents())
+		{
+			TTypedElement<UTypedElementSelectionInterface> ComponentSelectionHandle = InSelectionSet->GetElement<UTypedElementSelectionInterface>(Component->AcquireEditorElementHandle());
+			ComponentSelectionHandle.DeselectElement(InSelectionSet, InSelectionOptions);
+
+			if (USceneComponent* SceneComponent = Cast<USceneComponent>(Component))
+			{
+				FComponentEditorUtils::BindComponentSelectionOverride(SceneComponent, /*bBind*/false);
+			}
+		}
+	}
+
+	// Flush some cached data
+	GUnrealEd->PostActorSelectionChanged();
+
+	// A fast path to mark selection rather than reconnecting ALL components for ALL actors that have changed state
+	Actor->PushSelectionToProxies();
+
+	return true;
+}
+
+bool UActorElementLevelEditorSelectionProxy::SelectActorGroup(AGroupActor* InGroupActor, UTypedElementList* InSelectionSet, const FTypedElementSelectionOptions& InSelectionOptions, const bool bForce)
+{
+	bool bSelectionChanged = false;
+
+	// Select all actors within the group (if locked or forced)
+	if (bForce || InGroupActor->IsLocked())
+	{
+		FTypedElementListLegacySyncScopedBatch LegacySyncBatch(InSelectionSet, InSelectionOptions.AllowLegacyNotifications());
+
+		const FTypedElementSelectionOptions GroupSelectionOptions = FTypedElementSelectionOptions(InSelectionOptions)
+			.SetAllowGroups(false);
+
+		TArray<AActor*> GroupActors;
+		InGroupActor->GetGroupActors(GroupActors);
+		for (AActor* Actor : GroupActors)
+		{
+			TTypedElement<UTypedElementSelectionInterface> ActorSelectionHandle = InSelectionSet->GetElement<UTypedElementSelectionInterface>(Actor->AcquireEditorElementHandle());
+			if (CanSelectActorElement(ActorSelectionHandle, GroupSelectionOptions))
+			{
+				bSelectionChanged |= SelectActorElement(ActorSelectionHandle, InSelectionSet, GroupSelectionOptions);
+			}
+		}
+	}
+
+	return bSelectionChanged;
+}
+
+bool UActorElementLevelEditorSelectionProxy::DeselectActorGroup(AGroupActor* InGroupActor, UTypedElementList* InSelectionSet, const FTypedElementSelectionOptions& InSelectionOptions, const bool bForce)
+{
+	bool bSelectionChanged = false;
+
+	// Deselect all actors within the group (if locked or forced)
+	if (bForce || InGroupActor->IsLocked())
+	{
+		FTypedElementListLegacySyncScopedBatch LegacySyncBatch(InSelectionSet, InSelectionOptions.AllowLegacyNotifications());
+
+		const FTypedElementSelectionOptions GroupSelectionOptions = FTypedElementSelectionOptions(InSelectionOptions)
+			.SetAllowGroups(false);
+
+		TArray<AActor*> GroupActors;
+		InGroupActor->GetGroupActors(GroupActors);
+		for (AActor* Actor : GroupActors)
+		{
+			TTypedElement<UTypedElementSelectionInterface> ActorSelectionHandle = InSelectionSet->GetElement<UTypedElementSelectionInterface>(Actor->AcquireEditorElementHandle());
+			if (CanDeselectActorElement(ActorSelectionHandle, GroupSelectionOptions))
+			{
+				bSelectionChanged |= DeselectActorElement(ActorSelectionHandle, InSelectionSet, GroupSelectionOptions);
+			}
+		}
+	}
+
+	return bSelectionChanged;
+}
diff --git a/Engine/Source/Editor/LevelEditor/Private/Elements/Actor/ActorElementLevelEditorSelectionProxy.h b/Engine/Source/Editor/LevelEditor/Private/Elements/Actor/ActorElementLevelEditorSelectionProxy.h
new file mode 100644
--- /dev/null
+++ b/Engine/Source/Editor/LevelEditor/Private/Elements/Actor/ActorElementLevelEditorSelectionProxy.h
@@ -0,0 +1,31 @@
+// Copyright Epic Games, Inc. All Rights Reserved.
+
+#pragma once
+
+#include "Elements/TypedElementSelectionSet.h"
+#include "ActorElementLevelEditorSelectionProxy.generated.h"
+
+class AGroupActor;
+
+UCLASS()
+class UActorElementLevelEditorSelectionProxy : public UTypedElementAssetEditorSelectionProxy
+{
+	GENERATED_BODY()
+
+public:
+	virtual bool CanSelectElement(const TTypedElement<UTypedElementSelectionInterface>& InElementSelectionHandle, const FTypedElementSelectionOptions& InSelectionOptions) override;
+	virtual bool CanDeselectElement(const TTypedElement<UTypedElementSelectionInterface>& InElementSelectionHandle, const FTypedElementSelectionOptions& InSelectionOptions) override;
+	virtual bool SelectElement(const TTypedElement<UTypedElementSelectionInterface>& InElementSelectionHandle, UTypedElementList* InSelectionSet, const FTypedElementSelectionOptions& InSelectionOptions) override;
+	virtual bool DeselectElement(const TTypedElement<UTypedElementSelectionInterface>& InElementSelectionHandle, UTypedElementList* InSelectionSet, const FTypedElementSelectionOptions& InSelectionOptions) override;
+	virtual bool AllowSelectionModifiers(const TTypedElement<UTypedElementSelectionInterface>& InElementSelectionHandle, const UTypedElementList* InSelectionSet) override;
+	virtual FTypedElementHandle GetSelectionElement(const TTypedElement<UTypedElementSelectionInterface>& InElementSelectionHandle, const UTypedElementList* InCurrentSelection, const ETypedElementSelectionMethod InSelectionMethod) override;
+
+	static bool CanSelectActorElement(const TTypedElement<UTypedElementSelectionInterface>& InActorSelectionHandle, const FTypedElementSelectionOptions& InSelectionOptions);
+	static bool CanDeselectActorElement(const TTypedElement<UTypedElementSelectionInterface>& InActorSelectionHandle, const FTypedElementSelectionOptions& InSelectionOptions);
+
+	static bool SelectActorElement(const TTypedElement<UTypedElementSelectionInterface>& InActorSelectionHandle, UTypedElementList* InSelectionSet, const FTypedElementSelectionOptions& InSelectionOptions);
+	static bool DeselectActorElement(const TTypedElement<UTypedElementSelectionInterface>& InActorSelectionHandle, UTypedElementList* InSelectionSet, const FTypedElementSelectionOptions& InSelectionOptions);
+
+	static bool SelectActorGroup(AGroupActor* InGroupActor, UTypedElementList* InSelectionSet, const FTypedElementSelectionOptions& InSelectionOptions, const bool bForce);
+	static bool DeselectActorGroup(AGroupActor* InGroupActor, UTypedElementList* InSelectionSet, const FTypedElementSelectionOptions& InSelectionOptions, const bool bForce);
+};
diff --git a/Engine/Source/Editor/LevelEditor/Private/Elements/Component/ComponentElementLevelEditorSelectionProxy.cpp b/Engine/Source/Editor/LevelEditor/Private/Elements/Component/ComponentElementLevelEditorSelectionProxy.cpp
new file mode 100644
--- /dev/null
+++ b/Engine/Source/Editor/LevelEditor/Private/Elements/Component/ComponentElementLevelEditorSelectionProxy.cpp
@@ -0,0 +1,171 @@
+// Copyright Epic Games, Inc. All Rights Reserved.
+
+#include "Elements/Component/ComponentElementLevelEditorSelectionProxy.h"
+#include "Elements/Component/ComponentElementData.h"
+#include "GameFramework/Actor.h"
+#include "Components/ActorComponent.h"
+#include "TypedElementList.h"
+
+#include "Elements/Actor/ActorElementSelectionInterface.h"
+#include "Elements/Component/ComponentElementSelectionInterface.h"
+
+#include "EditorViewportClient.h"
+#include "Kismet2/ComponentEditorUtils.h"
+#include "Kismet2/KismetEditorUtilities.h"
+
+DEFINE_LOG_CATEGORY_STATIC(LogComponentLevelEditorSelection, Log, All);
+
+bool UComponentElementLevelEditorSelectionProxy::CanSelectElement(const TTypedElement<UTypedElementSelectionInterface>& InElementSelectionHandle, const FTypedElementSelectionOptions& InSelectionOptions)
+{
+	return InElementSelectionHandle.IsDataOfType<FComponentElementData>() && CanSelectComponentElement(InElementSelectionHandle, InSelectionOptions);
+}
+
+bool UComponentElementLevelEditorSelectionProxy::CanDeselectElement(const TTypedElement<UTypedElementSelectionInterface>& InElementSelectionHandle, const FTypedElementSelectionOptions& InSelectionOptions)
+{
+	return InElementSelectionHandle.IsDataOfType<FComponentElementData>() && CanDeselectComponentElement(InElementSelectionHandle, InSelectionOptions);
+}
+
+bool UComponentElementLevelEditorSelectionProxy::SelectElement(const TTypedElement<UTypedElementSelectionInterface>& InElementSelectionHandle, UTypedElementList* InSelectionSet, const FTypedElementSelectionOptions& InSelectionOptions)
+{
+	return InElementSelectionHandle.IsDataOfType<FComponentElementData>() && SelectComponentElement(InElementSelectionHandle, InSelectionSet, InSelectionOptions);
+}
+
+bool UComponentElementLevelEditorSelectionProxy::DeselectElement(const TTypedElement<UTypedElementSelectionInterface>& InElementSelectionHandle, UTypedElementList* InSelectionSet, const FTypedElementSelectionOptions& InSelectionOptions)
+{
+	return InElementSelectionHandle.IsDataOfType<FComponentElementData>() && DeselectComponentElement(InElementSelectionHandle, InSelectionSet, InSelectionOptions);
+}
+
+FTypedElementHandle UComponentElementLevelEditorSelectionProxy::GetSelectionElement(const TTypedElement<UTypedElementSelectionInterface>& InElementSelectionHandle, const UTypedElementList* InCurrentSelection, const ETypedElementSelectionMethod InSelectionMethod)
+{
+	if (const FComponentElementData* ComponentData = InElementSelectionHandle.GetData<FComponentElementData>())
+	{
+		const AActor* ConsideredActor = ComponentData->Component->GetOwner();
+		const USceneComponent* ConsideredComponent = Cast<USceneComponent>(ComponentData->Component);
+		if (ConsideredActor && ConsideredComponent)
+		{
+			while (ConsideredActor->IsChildActor())
+			{
+				ConsideredActor = ConsideredActor->GetParentActor();
+				ConsideredComponent = ConsideredActor->GetParentComponent();
+			}
+
+			// If the component selected is a visualization component, we want to select the non-visualization component it's attached to
+			while (ConsideredComponent->IsVisualizationComponent())
+			{
+				ConsideredComponent = ConsideredComponent->GetAttachParent();
+			}
+
+			const FTypedElementHandle ConsideredActorHandle = ConsideredActor->AcquireEditorElementHandle();
+
+			// We want to process the click on the component only if:
+			// 1. The actor clicked is already selected and is the only actor selected
+			// 2. The actor selected is blueprintable
+			// 3. No components are already selected and the click was a double click
+			// 4. OR, a component is already selected and the click was NOT a double click
+			const bool bActorAlreadySelectedExclusively = InCurrentSelection->Contains(ConsideredActorHandle) && UActorElementSelectionInterface::GetNumSelectedActors(InCurrentSelection) == 1;
+			const bool bActorIsBlueprintable = FKismetEditorUtilities::CanCreateBlueprintOfClass(ConsideredActor->GetClass());
+			const bool bComponentAlreadySelected = UComponentElementSelectionInterface::HasSelectedComponents(InCurrentSelection);
+			const bool bWasDoubleClick = InSelectionMethod == ETypedElementSelectionMethod::Secondary;
+
+			const bool bSelectComponent = bActorAlreadySelectedExclusively && bActorIsBlueprintable && (bComponentAlreadySelected != bWasDoubleClick);
+
+			if (bSelectComponent)
+			{
+				return ConsideredComponent->AcquireEditorElementHandle();
+			}
+			else
+			{
+				return ConsideredActorHandle;
+			}
+		}
+	}
+	return InElementSelectionHandle;
+}
+
+bool UComponentElementLevelEditorSelectionProxy::CanSelectComponentElement(const TTypedElement<UTypedElementSelectionInterface>& InComponentSelectionHandle, const FTypedElementSelectionOptions& InSelectionOptions)
+{
+	UActorComponent* Component = InComponentSelectionHandle.GetDataChecked<FComponentElementData>().Component;
+
+	// Bail if global selection is locked
+	return !GEdSelectionLock;
+}
+
+bool UComponentElementLevelEditorSelectionProxy::CanDeselectComponentElement(const TTypedElement<UTypedElementSelectionInterface>& InComponentSelectionHandle, const FTypedElementSelectionOptions& InSelectionOptions)
+{
+	UActorComponent* Component = InComponentSelectionHandle.GetDataChecked<FComponentElementData>().Component;
+
+	// Bail if global selection is locked
+	return !GEdSelectionLock;
+}
+
+bool UComponentElementLevelEditorSelectionProxy::SelectComponentElement(const TTypedElement<UTypedElementSelectionInterface>& InComponentSelectionHandle, UTypedElementList* InSelectionSet, const FTypedElementSelectionOptions& InSelectionOptions)
+{
+	UActorComponent* Component = InComponentSelectionHandle.GetDataChecked<FComponentElementData>().Component;
+
+	if (!InComponentSelectionHandle.SelectElement(InSelectionSet, InSelectionOptions))
+	{
+		return false;
+	}
+
+	UE_LOG(LogComponentLevelEditorSelection, Verbose, TEXT("Selected Component: %s"), *Component->GetClass()->GetName());
+	
+	// Make sure the override delegate is bound properly
+	if (USceneComponent* SceneComponent = Cast<USceneComponent>(Component))
+	{
+		FComponentEditorUtils::BindComponentSelectionOverride(SceneComponent, /*bBind*/true);
+	}
+
+	if (AActor* ComponentOwner = Component->GetOwner())
+	{
+		// Selecting a component requires that its owner actor be selected too
+		TTypedElement<UTypedElementSelectionInterface> ActorSelectionHandle = InSelectionSet->GetElement<UTypedElementSelectionInterface>(ComponentOwner->AcquireEditorElementHandle());
+		ActorSelectionHandle.SelectElement(InSelectionSet, InSelectionOptions);
+
+		// Update the selection visualization
+		TInlineComponentArray<UPrimitiveComponent*> PrimitiveComponents;
+		ComponentOwner->GetComponents(PrimitiveComponents);
+	
+		for (UPrimitiveComponent* PrimitiveComponent : PrimitiveComponents)
+		{
+			PrimitiveComponent->PushSelectionToProxy();
+		}
+	}
+
+	return true;
+}
+
+bool UComponentElementLevelEditorSelectionProxy::DeselectComponentElement(const TTypedElement<UTypedElementSelectionInterface>& InComponentSelectionHandle, UTypedElementList* InSelectionSet, const FTypedElementSelectionOptions& InSelectionOptions)
+{
+	UActorComponent* Component = InComponentSelectionHandle.GetDataChecked<FComponentElementData>().Component;
+
+	if (!InComponentSelectionHandle.DeselectElement(InSelectionSet, InSelectionOptions))
+	{
+		return false;
+	}
+
+	UE_LOG(LogComponentLevelEditorSelection, Verbose, TEXT("Deselected Component: %s"), *Component->GetClass()->GetName());
+
+	// Update the selection visualization
+	if (AActor* ComponentOwner = Component->GetOwner())
+	{
+		TTypedElement<UTypedElementSelectionInterface> ActorSelectionHandle = InSelectionSet->GetElement<UTypedElementSelectionInterface>(ComponentOwner->AcquireEditorElementHandle());
+		if (!ActorSelectionHandle.IsElementSelected(InSelectionSet, FTypedElementIsSelectedOptions().SetAllowIndirect(true)))
+		{
+			// Make sure the override delegate is unbound properly
+			if (USceneComponent* SceneComponent = Cast<USceneComponent>(Component))
+			{
+				FComponentEditorUtils::BindComponentSelectionOverride(SceneComponent, /*bBind*/false);
+			}
+		}
+
+		TInlineComponentArray<UPrimitiveComponent*> PrimitiveComponents;
+		ComponentOwner->GetComponents(PrimitiveComponents);
+
+		for (UPrimitiveComponent* PrimitiveComponent : PrimitiveComponents)
+		{
+			PrimitiveComponent->PushSelectionToProxy();
+		}
+	}
+
+	return true;
+}
diff --git a/Engine/Source/Editor/LevelEditor/Private/Elements/Component/ComponentElementLevelEditorSelectionProxy.h b/Engine/Source/Editor/LevelEditor/Private/Elements/Component/ComponentElementLevelEditorSelectionProxy.h
new file mode 100644
--- /dev/null
+++ b/Engine/Source/Editor/LevelEditor/Private/Elements/Component/ComponentElementLevelEditorSelectionProxy.h
@@ -0,0 +1,25 @@
+// Copyright Epic Games, Inc. All Rights Reserved.
+
+#pragma once
+
+#include "Elements/TypedElementSelectionSet.h"
+#include "ComponentElementLevelEditorSelectionProxy.generated.h"
+
+UCLASS()
+class UComponentElementLevelEditorSelectionProxy : public UTypedElementAssetEditorSelectionProxy
+{
+	GENERATED_BODY()
+
+public:
+	virtual bool CanSelectElement(const TTypedElement<UTypedElementSelectionInterface>& InElementSelectionHandle, const FTypedElementSelectionOptions& InSelectionOptions) override;
+	virtual bool CanDeselectElement(const TTypedElement<UTypedElementSelectionInterface>& InElementSelectionHandle, const FTypedElementSelectionOptions& InSelectionOptions) override;
+	virtual bool SelectElement(const TTypedElement<UTypedElementSelectionInterface>& InElementSelectionHandle, UTypedElementList* InSelectionSet, const FTypedElementSelectionOptions& InSelectionOptions) override;
+	virtual bool DeselectElement(const TTypedElement<UTypedElementSelectionInterface>& InElementSelectionHandle, UTypedElementList* InSelectionSet, const FTypedElementSelectionOptions& InSelectionOptions) override;
+	virtual FTypedElementHandle GetSelectionElement(const TTypedElement<UTypedElementSelectionInterface>& InElementSelectionHandle, const UTypedElementList* InCurrentSelection, const ETypedElementSelectionMethod InSelectionMethod) override;
+
+	static bool CanSelectComponentElement(const TTypedElement<UTypedElementSelectionInterface>& InComponentSelectionHandle, const FTypedElementSelectionOptions& InSelectionOptions);
+	static bool CanDeselectComponentElement(const TTypedElement<UTypedElementSelectionInterface>& InComponentSelectionHandle, const FTypedElementSelectionOptions& InSelectionOptions);
+
+	static bool SelectComponentElement(const TTypedElement<UTypedElementSelectionInterface>& InComponentSelectionHandle, UTypedElementList* InSelectionSet, const FTypedElementSelectionOptions& InSelectionOptions);
+	static bool DeselectComponentElement(const TTypedElement<UTypedElementSelectionInterface>& InComponentSelectionHandle, UTypedElementList* InSelectionSet, const FTypedElementSelectionOptions& InSelectionOptions);
+};
diff --git a/Engine/Source/Editor/LevelEditor/Private/SActorDetails.cpp b/Engine/Source/Editor/LevelEditor/Private/SActorDetails.cpp
--- a/Engine/Source/Editor/LevelEditor/Private/SActorDetails.cpp
+++ b/Engine/Source/Editor/LevelEditor/Private/SActorDetails.cpp
@@ -293,7 +293,7 @@ void SActorDetails::PostUndo(bool bSuccess)
 	AActor* SelectedActor = GetSelectedActorInEditor();
 	if (SelectedActor)
 	{
-		GUnrealEd->SetActorSelectionFlags(SelectedActor);
+		SelectedActor->PushSelectionToProxies();
 
 		// Update the pivot (widget) as the current selection may be a component within the Actor instance
 		GUnrealEd->UpdatePivotLocationForSelection();
@@ -351,7 +351,7 @@ void SActorDetails::OnEditorSelectionChanged(UObject* Object)
 				for(FSelectionIterator It(GEditor->GetSelectedActorIterator()); It; ++It)
 				{
 					AActor* Actor = CastChecked<AActor>(*It);
-					GUnrealEd->SetActorSelectionFlags(Actor);
+					Actor->PushSelectionToProxies();
 				}
 			}
 		}
@@ -555,7 +555,7 @@ void SActorDetails::OnSCSEditorTreeViewSelectionChanged(const TArray<FSCSEditorT
 
 					DetailsView->SetObjects(DetailsObjects);
 
-					GUnrealEd->SetActorSelectionFlags(Actor);
+					Actor->PushSelectionToProxies();
 					GUnrealEd->UpdatePivotLocationForSelection(true);
 					GEditor->RedrawLevelEditingViewports();
 				}
diff --git a/Engine/Source/Editor/LevelEditor/Private/SLevelEditor.cpp b/Engine/Source/Editor/LevelEditor/Private/SLevelEditor.cpp
--- a/Engine/Source/Editor/LevelEditor/Private/SLevelEditor.cpp
+++ b/Engine/Source/Editor/LevelEditor/Private/SLevelEditor.cpp
@@ -56,9 +56,11 @@
 #include "StatusBarSubsystem.h"
 #include "Widgets/Colors/SColorPicker.h"
 #include "SourceCodeNavigation.h"
-#include "TypedElementRegistry.h"
 #include "Editor/EnvironmentLightingViewer/Public/EnvironmentLightingModule.h"
 #include "Misc/MessageDialog.h"
+#include "Elements/TypedElementSelectionSet.h"
+#include "Elements/Actor/ActorElementLevelEditorSelectionProxy.h"
+#include "Elements/Component/ComponentElementLevelEditorSelectionProxy.h"
 
 #define LOCTEXT_NAMESPACE "SLevelEditor"
 
@@ -215,12 +217,16 @@ void SLevelEditor::Construct( const SLevelEditor::FArguments& InArgs)
 
 void SLevelEditor::Initialize( const TSharedRef<SDockTab>& OwnerTab, const TSharedRef<SWindow>& OwnerWindow )
 {
-	SelectedElements = UTypedElementRegistry::GetInstance()->CreateElementList();
+	SelectedElements = NewObject<UTypedElementSelectionSet>();
 	SelectedElements->AddToRoot();
 
+	// Register the level editor specific selection behavior
+	SelectedElements->RegisterAssetEditorSelectionProxy(NAME_Actor, NewObject<UActorElementLevelEditorSelectionProxy>());
+	SelectedElements->RegisterAssetEditorSelectionProxy(NAME_Components, NewObject<UComponentElementLevelEditorSelectionProxy>());
+
 	// Allow USelection to bridge to our selected element list
-	GUnrealEd->GetSelectedActors()->SetElementList(SelectedElements);
-	GUnrealEd->GetSelectedComponents()->SetElementList(SelectedElements);
+	GUnrealEd->GetSelectedActors()->SetElementSelectionSet(SelectedElements);
+	GUnrealEd->GetSelectedComponents()->SetElementSelectionSet(SelectedElements);
 
 	// Bind the level editor tab's label to the currently loaded level name string in the main frame
 	OwnerTab->SetLabel( TAttribute<FText>( this, &SLevelEditor::GetTabTitle) );
@@ -329,12 +335,12 @@ SLevelEditor::~SLevelEditor()
 	// Clear USelection from using our selected element list
 	if (GUnrealEd)
 	{
-		GUnrealEd->GetSelectedActors()->SetElementList(nullptr);
-		GUnrealEd->GetSelectedComponents()->SetElementList(nullptr);
+		GUnrealEd->GetSelectedActors()->SetElementSelectionSet(nullptr);
+		GUnrealEd->GetSelectedComponents()->SetElementSelectionSet(nullptr);
 	}
 	if (UObjectInitialized())
 	{
-		SelectedElements->Empty();
+		SelectedElements->ClearSelection(FTypedElementSelectionOptions());
 		SelectedElements->RemoveFromRoot();
 		SelectedElements = nullptr;
 	}
@@ -1612,6 +1618,16 @@ void SLevelEditor::OnLayoutHasChanged()
 	// ...
 }
 
+const UTypedElementSelectionSet* SLevelEditor::GetElementSelectionSet() const
+{
+	return SelectedElements;
+}
+
+UTypedElementSelectionSet* SLevelEditor::GetMutableElementSelectionSet()
+{
+	return SelectedElements;
+}
+
 void SLevelEditor::SummonLevelViewportContextMenu(AActor* HitProxyActor)
 {
 	FLevelEditorContextMenu::SummonMenu( SharedThis( this ), ELevelEditorMenuContext::Viewport, HitProxyActor);
diff --git a/Engine/Source/Editor/LevelEditor/Private/SLevelEditor.h b/Engine/Source/Editor/LevelEditor/Private/SLevelEditor.h
--- a/Engine/Source/Editor/LevelEditor/Private/SLevelEditor.h
+++ b/Engine/Source/Editor/LevelEditor/Private/SLevelEditor.h
@@ -22,7 +22,7 @@ class SActorDetails;
 class SBorder;
 class SLevelEditorModeContent;
 class SLevelEditorToolBox;
-class UTypedElementList;
+class UTypedElementSelectionSet;
 
 /**
  * Unreal editor level editor Slate widget
@@ -65,6 +65,8 @@ public:
 	TSharedPtr<class SLevelViewport> GetActiveViewport();
 
 	/** ILevelEditor interface */
+	virtual const UTypedElementSelectionSet* GetElementSelectionSet() const override;
+	virtual UTypedElementSelectionSet* GetMutableElementSelectionSet() override;
 	virtual void SummonLevelViewportContextMenu(AActor* HitProxyActor = nullptr) override;
 	virtual void SummonLevelViewportViewOptionMenu(const ELevelViewportType ViewOption) override;
 	virtual const TArray< TSharedPtr< class IToolkit > >& GetHostedToolkits() const override;
@@ -262,7 +264,7 @@ private:
 	UWorld* World;
 
 	// The list of selected elements (also set on the global USelection for actors and components).
-	UTypedElementList* SelectedElements = nullptr;
+	UTypedElementSelectionSet* SelectedElements = nullptr;
 
 	// The box that holds the title bar messages.
 	TSharedPtr<SHorizontalBox> TtileBarMessageBox;
diff --git a/Engine/Source/Editor/LevelEditor/Public/ILevelEditor.h b/Engine/Source/Editor/LevelEditor/Public/ILevelEditor.h
--- a/Engine/Source/Editor/LevelEditor/Public/ILevelEditor.h
+++ b/Engine/Source/Editor/LevelEditor/Public/ILevelEditor.h
@@ -17,6 +17,7 @@ class IDetailRootObjectCustomization;
 class ISceneOutliner;
 class IAssetViewport;
 class SLevelViewport;
+class UTypedElementSelectionSet;
 
 /**
  * Public interface to SLevelEditor
@@ -26,6 +27,10 @@ class ILevelEditor : public SCompoundWidget, public IToolkitHost
 
 public:
 
+	/** Get the element selection set used by this level editor */
+	virtual const UTypedElementSelectionSet* GetElementSelectionSet() const = 0;
+	virtual UTypedElementSelectionSet* GetMutableElementSelectionSet() = 0;
+
 	/** Summons a context menu for this level editor at the mouse cursor's location */
 	virtual void SummonLevelViewportContextMenu(AActor* HitProxyActor = nullptr) = 0;
 
diff --git a/Engine/Source/Editor/UnrealEd/Classes/Editor/UnrealEdEngine.h b/Engine/Source/Editor/UnrealEd/Classes/Editor/UnrealEdEngine.h
--- a/Engine/Source/Editor/UnrealEd/Classes/Editor/UnrealEdEngine.h
+++ b/Engine/Source/Editor/UnrealEd/Classes/Editor/UnrealEdEngine.h
@@ -27,8 +27,11 @@ class ITargetPlatform;
 class UPrimitiveComponent;
 class UTexture2D;
 class UUnrealEdOptions;
+class USelection;
+class UTypedElementList;
 class FName;
 typedef FName FEditorModeID;
+struct FTypedElementSelectionOptions;
 
 UENUM()
 enum EPackageNotifyState
@@ -268,10 +271,9 @@ public:
 	virtual void UpdateFloatingPropertyWindowsFromActorList(const TArray< UObject *>& ActorList, bool bForceRefresh=false);
 
 	/**
-	 * Fast track function to set render thread flags marking selection rather than reconnecting all components
-	 * @param InActor - the actor to toggle view flags for
+	 * Called whenever the actor selection has changed to invalidate any cached state
 	 */
-	virtual void SetActorSelectionFlags (AActor* InActor);
+	void PostActorSelectionChanged();
 
 	/**
 	 * Set whether the pivot has been moved independently or not
@@ -801,6 +803,12 @@ protected:
 	/** Called when global editor selection changes */
 	void OnEditorSelectionChanged(UObject* SelectionThatChanged);
 
+	/** Called when the element list pointer set on the global editor selection changes */
+	void OnEditorSelectionElementListPtrChanged(USelection* Selection, UTypedElementList* OldElementList, UTypedElementList* NewElementList);
+
+	/** Called when the element list associated with the global editor selection changes */
+	void OnEditorSelectionElementListChanged(const UTypedElementList* ElementList);
+
 	/** The package auto-saver instance used by the editor */
 	TUniquePtr<IPackageAutoSaver> PackageAutoSaver;
 
diff --git a/Engine/Source/Editor/UnrealEd/Private/EditorSelectUtils.cpp b/Engine/Source/Editor/UnrealEd/Private/EditorSelectUtils.cpp
--- a/Engine/Source/Editor/UnrealEd/Private/EditorSelectUtils.cpp
+++ b/Engine/Source/Editor/UnrealEd/Private/EditorSelectUtils.cpp
@@ -27,6 +27,8 @@
 #include "Logging/MessageLog.h"
 #include "ActorGroupingUtils.h"
 #include "Subsystems/BrushEditingSubsystem.h"
+#include "TypedElementList.h"
+#include "Elements/TypedElementSelectionSet.h"
 
 #define LOCTEXT_NAMESPACE "EditorSelectUtils"
 
@@ -254,9 +256,36 @@ void UUnrealEdEngine::ResetPivot()
 	Selection.
 -----------------------------------------------------------------------------*/
 
-void UUnrealEdEngine::SetActorSelectionFlags (AActor* InActor)
+void UUnrealEdEngine::OnEditorSelectionElementListPtrChanged(USelection* Selection, UTypedElementList* OldElementList, UTypedElementList* NewElementList)
 {
-	InActor->PushSelectionToProxies();	
+	if (Selection == GetSelectedActors())
+	{
+		if (OldElementList)
+		{
+			OldElementList->OnChanged().RemoveAll(this);
+		}
+
+		if (NewElementList)
+		{
+			NewElementList->OnChanged().AddUObject(this, &UUnrealEdEngine::OnEditorSelectionElementListChanged);
+		}
+	}
+}
+
+
+void UUnrealEdEngine::OnEditorSelectionElementListChanged(const UTypedElementList* ElementList)
+{
+	NoteSelectionChange();
+}
+
+
+void UUnrealEdEngine::PostActorSelectionChanged()
+{
+	// Whenever selection changes, recompute whether the selection contains a locked actor
+	bCheckForLockActors = true;
+
+	// Whenever selection changes, recompute whether the selection contains a world info actor
+	bCheckForWorldSettingsActors = true;
 }
 
 
@@ -377,12 +406,7 @@ void UUnrealEdEngine::NoteSelectionChange(bool bNotify)
 	
 	if (!bComponentSelectionChanged)
 	{
-		//whenever selection changes, recompute whether the selection contains a locked actor
-		bCheckForLockActors = true;
-
-		//whenever selection changes, recompute whether the selection contains a world info actor
-		bCheckForWorldSettingsActors = true;
-
+		PostActorSelectionChanged();
 		UpdateFloatingPropertyWindows();
 	}
 
@@ -391,305 +415,134 @@ void UUnrealEdEngine::NoteSelectionChange(bool bNotify)
 
 void UUnrealEdEngine::SelectGroup(AGroupActor* InGroupActor, bool bForceSelection/*=false*/, bool bInSelected/*=true*/, bool bNotify/*=true*/)
 {
-	USelection* SelectedActors = GetSelectedActors();
-	bool bStartedBatchSelect = false;
-	if(!SelectedActors->IsBatchSelecting())
+	USelection* ActorSelection = GetSelectedActors();
+
+	if (UTypedElementSelectionSet* SelectionSet = ActorSelection->GetElementSelectionSet())
 	{
-		bStartedBatchSelect = true;
-		// These will have already been called when batch selecting
-		SelectedActors->BeginBatchSelectOperation();
-		SelectedActors->Modify();
-	}
+		const FTypedElementSelectionOptions SelectionOptions = FTypedElementSelectionOptions()
+			.SetWarnIfLocked(true)
+			.SetAllowGroups(false)
+			.SetAllowLegacyNotifications(false);
 
-	static bool bIteratingGroups = false;
+		bool bSelectionChanged = false;
+
+		// Select/deselect all actors within the group (if locked or forced)
+		if (bForceSelection || InGroupActor->IsLocked())
+		{
+			FTypedElementListLegacySyncScopedBatch LegacySyncBatch(SelectionSet->GetMutableElementList(), SelectionOptions.AllowLegacyNotifications());
 
-	if( !bIteratingGroups )
-	{
-		bIteratingGroups = true;
-		// Select all actors within the group (if locked or forced)
-		if( bForceSelection || InGroupActor->IsLocked() )
-		{	
 			TArray<AActor*> GroupActors;
 			InGroupActor->GetGroupActors(GroupActors);
-			for( int32 ActorIndex=0; ActorIndex < GroupActors.Num(); ++ActorIndex )
-			{                  
-				SelectActor(GroupActors[ActorIndex], bInSelected, false );
-			}
-			bForceSelection = true;
-
-			// Recursively select any subgroups
-			TArray<AGroupActor*> SubGroups;
-			InGroupActor->GetSubGroups(SubGroups);
-			for( int32 GroupIndex=0; GroupIndex < SubGroups.Num(); ++GroupIndex )
+			for (AActor* Actor : GroupActors)
 			{
-				SelectGroup(SubGroups[GroupIndex], bForceSelection, bInSelected, false);
+				bSelectionChanged |= bInSelected
+					? SelectionSet->SelectElement(Actor->AcquireEditorElementHandle(), SelectionOptions)
+					: SelectionSet->DeselectElement(Actor->AcquireEditorElementHandle(), SelectionOptions);
 			}
 		}
 
-		if(bStartedBatchSelect)
-		{
-			SelectedActors->EndBatchSelectOperation(bNotify);
-		}
-		if (bNotify)
+		if (bSelectionChanged && bNotify)
 		{
-			NoteSelectionChange();
+			SelectionSet->GetMutableElementList()->NotifyPendingChanges();
 		}
-
-		//whenever selection changes, recompute whether the selection contains a locked actor
-		bCheckForLockActors = true;
-
-		//whenever selection changes, recompute whether the selection contains a world info actor
-		bCheckForWorldSettingsActors = true;
-
-		bIteratingGroups = false;
 	}
 }
 
-
 bool UUnrealEdEngine::CanSelectActor(AActor* Actor, bool bInSelected, bool bSelectEvenIfHidden, bool bWarnIfLevelLocked ) const
 {
-	// If selections are globally locked, leave.
-	if( !Actor || GEdSelectionLock || !Actor->IsEditable() )
+	if (!Actor)
 	{
 		return false;
 	}
 
-	// Only abort from hidden actors if we are selecting. You can deselect hidden actors without a problem.
-	if ( bInSelected )
-	{
-		// If the actor is NULL or hidden, leave.
-		if ( !bSelectEvenIfHidden && ( Actor->IsHiddenEd() || !FLevelUtils::IsLevelVisible( Actor->GetLevel() ) ) )
-		{
-			return false;
-		}
+	USelection* ActorSelection = GetSelectedActors();
 
-		// If the actor explicitly makes itself unselectable, leave.
-		if (!Actor->IsSelectable())
-		{
-			return false;
-		}
-
-		// Ensure that neither the level nor the actor is being destroyed or is unreachable
-		const EObjectFlags InvalidSelectableFlags = RF_BeginDestroyed;
-		if (Actor->GetLevel()->HasAnyFlags(InvalidSelectableFlags) || Actor->GetLevel()->IsPendingKillOrUnreachable())
-		{
-			UE_LOG(LogEditorSelectUtils, Warning, TEXT("SelectActor: %s (%s)"), TEXT("The requested operation could not be completed because the level has invalid flags."),*Actor->GetActorLabel());
-			return false;
-		}
-		if (Actor->HasAnyFlags(InvalidSelectableFlags) || Actor->IsPendingKillOrUnreachable())
-		{
-			UE_LOG(LogEditorSelectUtils, Warning, TEXT("SelectActor: %s (%s)"), TEXT("The requested operation could not be completed because the actor has invalid flags."),*Actor->GetActorLabel());
-			return false;
-		}
-
-		if ( !Actor->IsTemplate() && FLevelUtils::IsLevelLocked(Actor->GetLevel()) )
-		{
-			if( bWarnIfLevelLocked )
-			{
-				UE_LOG(LogEditorSelectUtils, Warning, TEXT("SelectActor: %s (%s)"), TEXT("The requested operation could not be completed because the level is locked."),*Actor->GetActorLabel());
-			}
-			return false;
-		}
-	}
-
-	// If grouping operations are not currently allowed, don't select groups.
-	AGroupActor* SelectedGroupActor = Cast<AGroupActor>(Actor);
-	if( SelectedGroupActor && !UActorGroupingUtils::IsGroupingActive())
+	if (UTypedElementSelectionSet* SelectionSet = ActorSelection->GetElementSelectionSet())
 	{
-		return false;
+		const FTypedElementSelectionOptions SelectionOptions = FTypedElementSelectionOptions()
+			.SetAllowHidden(bSelectEvenIfHidden)
+			.SetWarnIfLocked(bWarnIfLevelLocked);
+
+		return bInSelected
+			? SelectionSet->CanSelectElement(Actor->AcquireEditorElementHandle(), SelectionOptions)
+			: SelectionSet->CanDeselectElement(Actor->AcquireEditorElementHandle(), SelectionOptions);
 	}
 
-	// Allow active modes to determine whether the selection is allowed. 
-	return GLevelEditorModeTools().IsSelectionAllowed(Actor, bInSelected);
+	return false;
 }
 
 void UUnrealEdEngine::SelectActor(AActor* Actor, bool bInSelected, bool bNotify, bool bSelectEvenIfHidden, bool bForceRefresh)
 {
-	const bool bWarnIfLevelLocked = true;
-	if( !CanSelectActor( Actor, bInSelected, bSelectEvenIfHidden, bWarnIfLevelLocked ) )
+	if (!Actor)
 	{
 		return;
 	}
 
-	bool bSelectionHandled = GLevelEditorModeTools().IsSelectionHandled(Actor, bInSelected);
-
+	USelection* ActorSelection = GetSelectedActors();
 
-	// Select the actor and update its internals.
-	if( !bSelectionHandled )
+	if (UTypedElementSelectionSet* SelectionSet = ActorSelection->GetElementSelectionSet())
 	{
-		if (bInSelected)
-		{
-			// If trying to select an actor, return this actors root selection actor instead (if it has one)
-			if (AActor* RootSelection = Actor->GetRootSelectionParent())
-			{
-				Actor = RootSelection;
-			}
-		}
+		const FTypedElementSelectionOptions SelectionOptions = FTypedElementSelectionOptions()
+			.SetAllowHidden(bSelectEvenIfHidden)
+			.SetWarnIfLocked(true)
+			.SetAllowLegacyNotifications(false);
 
-		if (UActorGroupingUtils::IsGroupingActive())
-		{
-			// if this actor is a group, do a group select/deselect
-			AGroupActor* SelectedGroupActor = Cast<AGroupActor>(Actor);
-			if (SelectedGroupActor)
-			{
-				SelectGroup(SelectedGroupActor, true, bInSelected, bNotify);
-			}
-			else
-			{
-				// Select/Deselect this actor's entire group, starting from the top locked group.
-				// If none is found, just use the actor.
-				AGroupActor* ActorLockedRootGroup = AGroupActor::GetRootForActor(Actor, true);
-				if (ActorLockedRootGroup)
-				{
-					SelectGroup(ActorLockedRootGroup, false, bInSelected, bNotify);
-				}
-			}
-		}
+		const bool bSelectionChanged = bInSelected
+			? SelectionSet->SelectElement(Actor->AcquireEditorElementHandle(), SelectionOptions)
+			: SelectionSet->DeselectElement(Actor->AcquireEditorElementHandle(), SelectionOptions);
 
-		// Don't do any work if the actor's selection state is already the selected state.
-		const bool bActorSelected = Actor->IsSelected();
-		if ( bActorSelected != bInSelected )
+		if (bSelectionChanged)
 		{
-			if(bInSelected)
-			{
-				UE_LOG(LogEditorSelectUtils, Verbose,  TEXT("Selected Actor: %s"), *Actor->GetClass()->GetName());
-			}
-			else
-			{
-				UE_LOG(LogEditorSelectUtils, Verbose,  TEXT("Deselected Actor: %s"), *Actor->GetClass()->GetName() );
-			}
-
-			GetSelectedActors()->Select( Actor, bInSelected );
-			if (!bInSelected)
-			{
-				if (GetSelectedComponentCount() > 0)
-				{
-					GetSelectedComponents()->Modify();
-				}
-
-				GetSelectedComponents()->BeginBatchSelectOperation();
-				for (UActorComponent* Component : Actor->GetComponents())
-				{
-					if (Component)
-					{
-						GetSelectedComponents()->Deselect(Component);
-
-						// Remove the selection override delegates from the deselected components
-						if (USceneComponent* SceneComponent = Cast<USceneComponent>(Component))
-						{
-							FComponentEditorUtils::BindComponentSelectionOverride(SceneComponent, false);
-						}
-					}
-				}
-				GetSelectedComponents()->EndBatchSelectOperation(false);
-			}
-			else
+			if (bNotify)
 			{
-				// Bind the override delegates for the components in the selected actor
-				for (UActorComponent* Component : Actor->GetComponents())
-				{
-					if (USceneComponent* SceneComponent = Cast<USceneComponent>(Component))
-					{
-						FComponentEditorUtils::BindComponentSelectionOverride(SceneComponent, true);
-					}
-				}
+				SelectionSet->GetMutableElementList()->NotifyPendingChanges();
 			}
-
-
-			if( bNotify )
-			{
-				NoteSelectionChange();
-			}
-
-			//whenever selection changes, recompute whether the selection contains a locked actor
-			bCheckForLockActors = true;
-
-			//whenever selection changes, recompute whether the selection contains a world info actor
-			bCheckForWorldSettingsActors = true;
 		}
-		else
+		else if (bNotify || bForceRefresh)
 		{
-			if (bNotify || bForceRefresh)
-			{
-				//reset the property windows.  In case something has changed since previous selection
-				UpdateFloatingPropertyWindows(bForceRefresh);
-			}
+			// Reset the property windows, in case something has changed since previous selection
+			UpdateFloatingPropertyWindows(bForceRefresh);
 		}
-
-		//A fast path to mark selection rather than reconnecting ALL components for ALL actors that have changed state
-		SetActorSelectionFlags(Actor);
 	}
 }
 
 void UUnrealEdEngine::SelectComponent(UActorComponent* Component, bool bInSelected, bool bNotify, bool bSelectEvenIfHidden)
 {
-	// Don't do any work if the component's selection state matches the target selection state
-	const bool bComponentSelected = GetSelectedComponents()->IsSelected(Component);
-	if (( bComponentSelected && !bInSelected ) || ( !bComponentSelected && bInSelected ))
+	if (!Component)
 	{
-		if (bInSelected)
-		{
-			UE_LOG(LogEditorSelectUtils, Verbose, TEXT("Selected Component: %s"), *Component->GetClass()->GetName());
-		}
-		else
-		{
-			UE_LOG(LogEditorSelectUtils, Verbose, TEXT("Deselected Component: %s"), *Component->GetClass()->GetName());
-		}
-
-		GetSelectedComponents()->Select(Component, bInSelected);
+		return;
+	}
 
-		// Make sure the override delegate is bound properly
-		USceneComponent* SceneComponent = Cast<USceneComponent>(Component);
-		if (SceneComponent)
-		{
-			FComponentEditorUtils::BindComponentSelectionOverride(SceneComponent, true);
-		}
+	USelection* ComponentSelection = GetSelectedComponents();
 
-		// Update the selection visualization
-		AActor* ComponentOwner = Component->GetOwner();
-		if (ComponentOwner != nullptr)
-		{
-			TInlineComponentArray<UPrimitiveComponent*> PrimitiveComponents;
-			ComponentOwner->GetComponents(PrimitiveComponents);
+	if (UTypedElementSelectionSet* SelectionSet = ComponentSelection->GetElementSelectionSet())
+	{
+		const FTypedElementSelectionOptions SelectionOptions = FTypedElementSelectionOptions()
+			.SetAllowHidden(bSelectEvenIfHidden)
+			.SetWarnIfLocked(true)
+			.SetAllowLegacyNotifications(false);
 
-			for (int32 Idx = 0; Idx < PrimitiveComponents.Num(); ++Idx)
-			{
-				PrimitiveComponents[Idx]->PushSelectionToProxy();
-			}
-		}
+		const bool bSelectionChanged = bInSelected
+			? SelectionSet->SelectElement(Component->AcquireEditorElementHandle(), SelectionOptions)
+			: SelectionSet->DeselectElement(Component->AcquireEditorElementHandle(), SelectionOptions);
 
-		if (bNotify)
+		if (bSelectionChanged && bNotify)
 		{
-			NoteSelectionChange();
+			SelectionSet->GetMutableElementList()->NotifyPendingChanges();
 		}
 	}
 }
 
 bool UUnrealEdEngine::IsComponentSelected(const UPrimitiveComponent* PrimComponent)
 {
-	bool bIsSelected = false;
-	if (GetSelectedComponentCount() > 0)
-	{
-		const UActorComponent* PotentiallySelectedComponent = nullptr;
+	USelection* ComponentSelection = GetSelectedComponents();
 
-		AActor* ComponentOwner = PrimComponent->GetOwner();
-		if (ComponentOwner->IsChildActor())
-		{
-			do 
-			{
-				PotentiallySelectedComponent = ComponentOwner->GetParentComponent();
-				ComponentOwner = ComponentOwner->GetParentActor();
-			} while (ComponentOwner->IsChildActor());
-		}
-		else
-		{
-			PotentiallySelectedComponent = (PrimComponent->IsVisualizationComponent() ? PrimComponent->GetAttachParent() : PrimComponent);
-		}
-
-		bIsSelected = GetSelectedComponents()->IsSelected(PotentiallySelectedComponent);
+	if (UTypedElementSelectionSet* SelectionSet = ComponentSelection->GetElementSelectionSet())
+	{
+		return SelectionSet->IsElementSelected(PrimComponent->AcquireEditorElementHandle(), FTypedElementIsSelectedOptions().SetAllowIndirect(true));
 	}
 
-	return bIsSelected;
+	return false;
 }
 
 void UUnrealEdEngine::SelectBSPSurf(UModel* InModel, int32 iSurf, bool bSelected, bool bNoteSelectionChange)
@@ -716,11 +569,7 @@ void UUnrealEdEngine::SelectBSPSurf(UModel* InModel, int32 iSurf, bool bSelected
 		NoteSelectionChange();
 	}
 
-	//whenever selection changes, recompute whether the selection contains a locked actor
-	bCheckForLockActors = true;
-
-	//whenever selection changes, recompute whether the selection contains a world info actor
-	bCheckForWorldSettingsActors = true;
+	PostActorSelectionChanged();
 }
 
 /**
@@ -782,87 +631,34 @@ void UUnrealEdEngine::DeselectAllSurfaces()
 
 void UUnrealEdEngine::SelectNone(bool bNoteSelectionChange, bool bDeselectBSPSurfs, bool WarnAboutManyActors)
 {
-	if( GEdSelectionLock )
+	if (GEdSelectionLock)
 	{
 		return;
 	}
 
-	bool bShowProgress = false;
-
-	// If there are a lot of actors to process, pop up a warning "are you sure?" box
-	if( WarnAboutManyActors )
-	{
-		int32 NumSelectedActors = GEditor->GetSelectedActorCount();
-		if( NumSelectedActors >= EditorActorSelectionDefs::MaxActorsToSelectBeforeWarning )
-		{
-			bShowProgress = true;
-
-			const FText ConfirmText = FText::Format( NSLOCTEXT("UnrealEd", "Warning_ManyActorsForDeselect", "There are {0} selected actors. Are you sure you want to deselect them all?" ), FText::AsNumber(NumSelectedActors) );
-
-			FSuppressableWarningDialog::FSetupInfo Info( ConfirmText, NSLOCTEXT( "UnrealEd", "Warning_ManyActors", "Warning: Many Actors" ), "Warning_ManyActors" );
-			Info.ConfirmText = NSLOCTEXT("ModalDialogs", "ManyActorsForDeselectConfirm", "Continue Deselection");
-			Info.CancelText = NSLOCTEXT("ModalDialogs", "ManyActorsForDeselectCancel", "Keep Current Selection");
-
-			FSuppressableWarningDialog ManyActorsWarning( Info );
-			if( ManyActorsWarning.ShowModal() == FSuppressableWarningDialog::Cancel )
-			{
-				return;
-			}
-		}
-	}
+	bool bSelectionChanged = false;
 
-	if( bShowProgress )
-	{
-		GWarn->BeginSlowTask( LOCTEXT("BeginDeselectingActorsTaskMessage", "Deselecting Actors"), true );
-	}
+	USelection* ActorSelection = GetSelectedActors();
+	UTypedElementSelectionSet* SelectionSet = ActorSelection->GetElementSelectionSet();
 
-	// Make a list of selected actors . . .
-	TArray<AActor*> ActorsToDeselect;
-	for ( FSelectionIterator It( GetSelectedActorIterator() ) ; It ; ++It )
+	if (SelectionSet)
 	{
-		AActor* Actor = static_cast<AActor*>( *It );
-		checkSlow( Actor->IsA(AActor::StaticClass()) );
-
-		ActorsToDeselect.Add( Actor );
+		bSelectionChanged |= SelectionSet->ClearSelection(FTypedElementSelectionOptions().SetAllowLegacyNotifications(false));
 	}
 
-	USelection* SelectedActors = GetSelectedActors();
-	SelectedActors->BeginBatchSelectOperation();
-	SelectedActors->Modify();
-
-	// . . . and deselect them.
-	for ( int32 ActorIndex = 0 ; ActorIndex < ActorsToDeselect.Num() ; ++ActorIndex )
+	if (bDeselectBSPSurfs)
 	{
-		AActor* Actor = ActorsToDeselect[ ActorIndex ];
-		SelectActor( Actor, false, false );
+		bSelectionChanged |= DeselectAllSurfacesForWorld(GWorld) > 0;
 	}
 
-	uint32 NumDeselectSurfaces = 0;
-	if( bDeselectBSPSurfs )
+	if (bSelectionChanged)
 	{
-		NumDeselectSurfaces = DeselectAllSurfacesForWorld(GWorld);
-	}
-
-	SelectedActors->EndBatchSelectOperation(bNoteSelectionChange);
+		PostActorSelectionChanged();
 
-	//prevents clicking on background multiple times spamming selection changes
-	if (ActorsToDeselect.Num() || NumDeselectSurfaces)
-	{
-		if( bNoteSelectionChange )
+		if (bNoteSelectionChange && SelectionSet)
 		{
-			NoteSelectionChange();
+			SelectionSet->GetMutableElementList()->NotifyPendingChanges();
 		}
-
-		//whenever selection changes, recompute whether the selection contains a locked actor
-		bCheckForLockActors = true;
-
-		//whenever selection changes, recompute whether the selection contains a world info actor
-		bCheckForWorldSettingsActors = true;
-	}
-
-	if( bShowProgress )
-	{
-		GWarn->EndSlowTask();
 	}
 }
 
diff --git a/Engine/Source/Editor/UnrealEd/Private/Elements/Actor/ActorElementEditorSelectionInterface.cpp b/Engine/Source/Editor/UnrealEd/Private/Elements/Actor/ActorElementEditorSelectionInterface.cpp
--- a/Engine/Source/Editor/UnrealEd/Private/Elements/Actor/ActorElementEditorSelectionInterface.cpp
+++ b/Engine/Source/Editor/UnrealEd/Private/Elements/Actor/ActorElementEditorSelectionInterface.cpp
@@ -3,3 +3,33 @@
 #include "Elements/Actor/ActorElementEditorSelectionInterface.h"
 #include "Elements/Actor/ActorElementData.h"
 #include "GameFramework/Actor.h"
+#include "TypedElementList.h"
+
+bool UActorElementEditorSelectionInterface::IsElementSelected(const FTypedElementHandle& InElementHandle, const UTypedElementList* InSelectionSet, const FTypedElementIsSelectedOptions& InSelectionOptions)
+{
+	const FActorElementData* ActorData = InElementHandle.GetData<FActorElementData>();
+	return ActorData && IsActorSelected(ActorData->Actor, InSelectionSet, InSelectionOptions);
+}
+
+bool UActorElementEditorSelectionInterface::IsActorSelected(const AActor* InActor, const UTypedElementList* InSelectionSet, const FTypedElementIsSelectedOptions& InSelectionOptions)
+{
+	if (InSelectionSet->Num() == 0)
+	{
+		return false;
+	}
+
+	if (InSelectionSet->Contains(InActor->AcquireEditorElementHandle()))
+	{
+		return true;
+	}
+
+	if (InSelectionOptions.AllowIndirect())
+	{
+		if (const AActor* RootSelectionActor = InActor->GetRootSelectionParent())
+		{
+			return InSelectionSet->Contains(RootSelectionActor->AcquireEditorElementHandle());
+		}
+	}
+
+	return false;
+}
diff --git a/Engine/Source/Editor/UnrealEd/Private/Elements/Actor/ActorElementEditorSelectionInterface.h b/Engine/Source/Editor/UnrealEd/Private/Elements/Actor/ActorElementEditorSelectionInterface.h
--- a/Engine/Source/Editor/UnrealEd/Private/Elements/Actor/ActorElementEditorSelectionInterface.h
+++ b/Engine/Source/Editor/UnrealEd/Private/Elements/Actor/ActorElementEditorSelectionInterface.h
@@ -5,10 +5,15 @@
 #include "Elements/Actor/ActorElementSelectionInterface.h"
 #include "ActorElementEditorSelectionInterface.generated.h"
 
+class AActor;
+
 UCLASS()
 class UActorElementEditorSelectionInterface : public UActorElementSelectionInterface
 {
 	GENERATED_BODY()
 
 public:
+	virtual bool IsElementSelected(const FTypedElementHandle& InElementHandle, const UTypedElementList* InSelectionSet, const FTypedElementIsSelectedOptions& InSelectionOptions) override;
+
+	static bool IsActorSelected(const AActor* InActor, const UTypedElementList* InSelectionSet, const FTypedElementIsSelectedOptions& InSelectionOptions);
 };
diff --git a/Engine/Source/Editor/UnrealEd/Private/Elements/Component/ComponentElementEditorSelectionInterface.cpp b/Engine/Source/Editor/UnrealEd/Private/Elements/Component/ComponentElementEditorSelectionInterface.cpp
--- a/Engine/Source/Editor/UnrealEd/Private/Elements/Component/ComponentElementEditorSelectionInterface.cpp
+++ b/Engine/Source/Editor/UnrealEd/Private/Elements/Component/ComponentElementEditorSelectionInterface.cpp
@@ -2,4 +2,48 @@
 
 #include "Elements/Component/ComponentElementEditorSelectionInterface.h"
 #include "Elements/Component/ComponentElementData.h"
-#include "Components/ActorComponent.h"
+#include "Components/SceneComponent.h"
+#include "GameFramework/Actor.h"
+#include "TypedElementList.h"
+
+bool UComponentElementEditorSelectionInterface::IsElementSelected(const FTypedElementHandle& InElementHandle, const UTypedElementList* InSelectionSet, const FTypedElementIsSelectedOptions& InSelectionOptions)
+{
+	const FComponentElementData* ComponentData = InElementHandle.GetData<FComponentElementData>();
+	return ComponentData && IsComponentSelected(ComponentData->Component, InSelectionSet, InSelectionOptions);
+}
+
+bool UComponentElementEditorSelectionInterface::IsComponentSelected(const UActorComponent* InComponent, const UTypedElementList* InSelectionSet, const FTypedElementIsSelectedOptions& InSelectionOptions)
+{
+	if (InSelectionSet->Num() == 0)
+	{
+		return false;
+	}
+
+	if (InSelectionSet->Contains(InComponent->AcquireEditorElementHandle()))
+	{
+		return true;
+	}
+
+	if (InSelectionOptions.AllowIndirect())
+	{
+		const AActor* ConsideredActor = InComponent->GetOwner();
+		const USceneComponent* ConsideredComponent = Cast<USceneComponent>(InComponent);
+		if (ConsideredActor && ConsideredComponent)
+		{
+			while (ConsideredActor->IsChildActor())
+			{
+				ConsideredActor = ConsideredActor->GetParentActor();
+				ConsideredComponent = ConsideredActor->GetParentComponent();
+			}
+
+			while (ConsideredComponent->IsVisualizationComponent())
+			{
+				ConsideredComponent = ConsideredComponent->GetAttachParent();
+			}
+		}
+
+		return InSelectionSet->Contains(ConsideredComponent->AcquireEditorElementHandle());
+	}
+
+	return false;
+}
diff --git a/Engine/Source/Editor/UnrealEd/Private/Elements/Component/ComponentElementEditorSelectionInterface.h b/Engine/Source/Editor/UnrealEd/Private/Elements/Component/ComponentElementEditorSelectionInterface.h
--- a/Engine/Source/Editor/UnrealEd/Private/Elements/Component/ComponentElementEditorSelectionInterface.h
+++ b/Engine/Source/Editor/UnrealEd/Private/Elements/Component/ComponentElementEditorSelectionInterface.h
@@ -5,10 +5,15 @@
 #include "Elements/Component/ComponentElementSelectionInterface.h"
 #include "ComponentElementEditorSelectionInterface.generated.h"
 
+class UActorComponent;
+
 UCLASS()
 class UComponentElementEditorSelectionInterface : public UComponentElementSelectionInterface
 {
 	GENERATED_BODY()
 
 public:
+	virtual bool IsElementSelected(const FTypedElementHandle& InElementHandle, const UTypedElementList* InSelectionSet, const FTypedElementIsSelectedOptions& InSelectionOptions) override;
+
+	static bool IsComponentSelected(const UActorComponent* InComponent, const UTypedElementList* InSelectionSet, const FTypedElementIsSelectedOptions& InSelectionOptions);
 };
diff --git a/Engine/Source/Editor/UnrealEd/Private/Elements/EditorElements.cpp b/Engine/Source/Editor/UnrealEd/Private/Elements/EditorElements.cpp
--- a/Engine/Source/Editor/UnrealEd/Private/Elements/EditorElements.cpp
+++ b/Engine/Source/Editor/UnrealEd/Private/Elements/EditorElements.cpp
@@ -1,12 +1,14 @@
 // Copyright Epic Games, Inc. All Rights Reserved.
 
-#include "EditorElements.h"
+#include "Elements/EditorElements.h"
 #include "TypedElementRegistry.h"
 
 #include "Elements/Actor/ActorElementEditorSelectionInterface.h"
 
 #include "Elements/Component/ComponentElementEditorSelectionInterface.h"
 
+FSimpleMulticastDelegate OnRegisterEditorElementsDelegate;
+
 void RegisterEditorActorElements()
 {
 	UTypedElementRegistry* Registry = UTypedElementRegistry::GetInstance();
@@ -25,4 +27,6 @@ void RegisterEditorElements()
 {
 	RegisterEditorActorElements();
 	RegisterEditorComponentElements();
+
+	OnRegisterEditorElementsDelegate.Broadcast();
 }
diff --git a/Engine/Source/Editor/UnrealEd/Private/Elements/EditorElements.h b/Engine/Source/Editor/UnrealEd/Private/Elements/EditorElements.h
deleted file mode 100644
diff --git a/Engine/Source/Editor/UnrealEd/Private/LevelEditorViewport.cpp b/Engine/Source/Editor/UnrealEd/Private/LevelEditorViewport.cpp
--- a/Engine/Source/Editor/UnrealEd/Private/LevelEditorViewport.cpp
+++ b/Engine/Source/Editor/UnrealEd/Private/LevelEditorViewport.cpp
@@ -2129,6 +2129,8 @@ void FLevelEditorViewportClient::ProcessClick(FSceneView& View, HHitProxy* HitPr
 	}
 	if (!ModeTools->HandleClick(this, HitProxy,Click))
 	{
+		const FTypedElementHandle HitElement = HitProxy ? HitProxy->GetElementHandle() : FTypedElementHandle();
+
 		if (HitProxy == NULL)
 		{
 			LevelViewportClickHandlers::ClickBackdrop(this,Click);
@@ -2177,6 +2179,10 @@ void FLevelEditorViewportClient::ProcessClick(FSceneView& View, HHitProxy* HitPr
 		{
 			// Component vis manager handled the click
 		}
+		else if (HitElement && LevelViewportClickHandlers::ClickElement(this, HitElement, Click))
+		{
+			// Element handled the click
+		}
 		else if (HitProxy->IsA(HActor::StaticGetType()))
 		{
 			HActor* ActorHitProxy = (HActor*)HitProxy;
diff --git a/Engine/Source/Editor/UnrealEd/Private/LevelViewportClickHandlers.cpp b/Engine/Source/Editor/UnrealEd/Private/LevelViewportClickHandlers.cpp
--- a/Engine/Source/Editor/UnrealEd/Private/LevelViewportClickHandlers.cpp
+++ b/Engine/Source/Editor/UnrealEd/Private/LevelViewportClickHandlers.cpp
@@ -30,11 +30,30 @@
 #include "StaticLightingSystem/StaticLightingPrivate.h"
 #include "LightMap.h"
 #include "Subsystems/BrushEditingSubsystem.h"
+#include "Elements/TypedElementSelectionSet.h"
 
 #define LOCTEXT_NAMESPACE "ClickHandlers"
 
 namespace LevelViewportClickHandlers
 {
+	static const UTypedElementSelectionSet* PrivateGetElementSelectionSet(FLevelEditorViewportClient* ViewportClient)
+	{
+		if (TSharedPtr<ILevelEditor> LevelEditor = ViewportClient->ParentLevelEditor.Pin())
+		{
+			return LevelEditor->GetElementSelectionSet();
+		}
+		return nullptr;
+	}
+
+	static UTypedElementSelectionSet* PrivateGetMutableElementSelectionSet(FLevelEditorViewportClient* ViewportClient)
+	{
+		if (TSharedPtr<ILevelEditor> LevelEditor = ViewportClient->ParentLevelEditor.Pin())
+		{
+			return LevelEditor->GetMutableElementSelectionSet();
+		}
+		return nullptr;
+	}
+
 	static void PrivateSummonContextMenu( FLevelEditorViewportClient* ViewportClient, AActor* HitProxyActor = nullptr)
 	{
 		if( ViewportClient->ParentLevelEditor.IsValid() )
@@ -119,6 +138,90 @@ namespace LevelViewportClickHandlers
 		return false;
 	}
 
+	bool ClickElement(FLevelEditorViewportClient* ViewportClient, const FTypedElementHandle& HitElement, const FViewportClick& Click)
+	{
+		// Pivot snapping
+		if (Click.GetKey() == EKeys::MiddleMouseButton && Click.IsAltDown())
+		{
+			//GEditor->SetPivot(GEditor->ClickLocation, true, false, true); // TODO: This last param is only for actor pivots
+			//return true;
+			return false; // Let actor and component clicks handle pivots for now
+		}
+
+		UTypedElementSelectionSet* LevelEditorElementSelectionSet = PrivateGetMutableElementSelectionSet(ViewportClient);
+		if (!LevelEditorElementSelectionSet)
+		{
+			return false;
+		}
+
+		bool bHandledClick = false;
+
+		const bool bIsLeftClickSelection = Click.GetKey() == EKeys::LeftMouseButton && !(ViewportClient->Viewport->KeyState(EKeys::T) || ViewportClient->Viewport->KeyState(EKeys::L) || ViewportClient->Viewport->KeyState(EKeys::S) || ViewportClient->Viewport->KeyState(EKeys::A));
+		const bool bIsRightClickSelection = Click.GetKey() == EKeys::RightMouseButton && !Click.IsControlDown() && !ViewportClient->Viewport->KeyState(EKeys::LeftMouseButton);
+
+		if (bIsLeftClickSelection || bIsRightClickSelection)
+		{
+			const ETypedElementSelectionMethod SelectionMethod = Click.GetEvent() == IE_DoubleClick ? ETypedElementSelectionMethod::Secondary : ETypedElementSelectionMethod::Primary;
+			if (const FTypedElementHandle ResolvedElement = LevelEditorElementSelectionSet->GetSelectionElement(HitElement, SelectionMethod))
+			{
+				bHandledClick = true;
+
+				const FTypedElementSelectionOptions SelectionOptions = FTypedElementSelectionOptions()
+					.SetAllowHidden(true)
+					.SetWarnIfLocked(true);
+
+				bool bNeedViewportRefresh = false;
+
+				if (LevelEditorElementSelectionSet->CanSelectElement(ResolvedElement, SelectionOptions))
+				{
+					const FScopedTransaction Transaction(NSLOCTEXT("UnrealEd", "ClickingOnElements", "Clicking on Elements"));
+
+					const bool bAllowSelectionModifiers = bIsLeftClickSelection && LevelEditorElementSelectionSet->AllowSelectionModifiers(ResolvedElement);
+					if (Click.IsControlDown() && bAllowSelectionModifiers)
+					{
+						if (LevelEditorElementSelectionSet->IsElementSelected(ResolvedElement, FTypedElementIsSelectedOptions()))
+						{
+							LevelEditorElementSelectionSet->DeselectElement(ResolvedElement, SelectionOptions);
+						}
+						else
+						{
+							LevelEditorElementSelectionSet->SelectElement(ResolvedElement, SelectionOptions);
+						}
+					}
+					else if (Click.IsShiftDown() && bAllowSelectionModifiers)
+					{
+						LevelEditorElementSelectionSet->SelectElement(ResolvedElement, SelectionOptions);
+					}
+					else
+					{
+						// Skip the clear if we're doing a RMB select and this actor is already selected, as we want to summon the menu for the current selection
+						if (bIsLeftClickSelection || !LevelEditorElementSelectionSet->IsElementSelected(ResolvedElement, FTypedElementIsSelectedOptions().SetAllowIndirect(true)))
+						{
+							bNeedViewportRefresh = bIsRightClickSelection; // Refresh the viewport so the user will see what they just clicked while the menu is open
+							GEditor->DeselectAllSurfaces();
+							LevelEditorElementSelectionSet->ClearSelection(SelectionOptions);
+						}
+						LevelEditorElementSelectionSet->SelectElement(ResolvedElement, SelectionOptions);
+					}
+				}
+
+				if (bNeedViewportRefresh)
+				{
+					// Redraw the viewport so the user can see which object was clicked on
+					ViewportClient->Viewport->Draw();
+					FlushRenderingCommands();
+				}
+
+				if (bIsRightClickSelection)
+				{
+					PrivateSummonContextMenu(ViewportClient);
+				}
+			}
+		}
+
+		return bHandledClick;
+	}
+
 	bool ClickActor(FLevelEditorViewportClient* ViewportClient,AActor* Actor,const FViewportClick& Click,bool bAllowSelectionChange)
 	{
 		// Pivot snapping
diff --git a/Engine/Source/Editor/UnrealEd/Private/Selection.cpp b/Engine/Source/Editor/UnrealEd/Private/Selection.cpp
--- a/Engine/Source/Editor/UnrealEd/Private/Selection.cpp
+++ b/Engine/Source/Editor/UnrealEd/Private/Selection.cpp
@@ -4,13 +4,9 @@
 #include "UObject/Package.h"
 #include "UObject/UObjectAnnotation.h"
 #include "GameFramework/Actor.h"
-
-#define UE_USE_ELEMENT_LIST_SELECTION (1)
-
-#if UE_USE_ELEMENT_LIST_SELECTION
 #include "TypedElementList.h"
 #include "TypedElementSelectionInterface.h"
-#endif	// UE_USE_ELEMENT_LIST_SELECTION
+#include "Elements/TypedElementSelectionSet.h"
 
 DEFINE_LOG_CATEGORY_STATIC(LogSelection, Log, All);
 
@@ -57,6 +53,7 @@ public:
 
 	//~ ISelectionStore interface
 	virtual void SetElementList(UTypedElementList* InElementList) override;
+	virtual UTypedElementList* GetElementList() const override;
 	virtual int32 GetNumObjects() const override;
 	virtual UObject* GetObjectAtIndex(const int32 InIndex) const override;
 	virtual bool IsObjectSelected(const UObject* InObject) const override;
@@ -86,7 +83,12 @@ FObjectSelectionStore::FObjectSelectionStore(const UClass* InObjectBaseClass)
 
 void FObjectSelectionStore::SetElementList(UTypedElementList* InElementList)
 {
-	checkf(!UE_USE_ELEMENT_LIST_SELECTION, TEXT("This selection store does not support element lists! Only actor and component selections may currently use element lists!"));
+	checkf(false, TEXT("This selection store does not support element lists! Only actor and component selections may currently use element lists!"));
+}
+
+UTypedElementList* FObjectSelectionStore::GetElementList() const
+{
+	return nullptr;
 }
 
 int32 FObjectSelectionStore::GetNumObjects() const
@@ -108,8 +110,14 @@ bool FObjectSelectionStore::IsObjectSelected(const UObject* InObject) const
 
 void FObjectSelectionStore::SelectObject(UObject* InObject)
 {
+	if (Sink && !bIsBatchDirty)
+	{
+		Sink->OnPreModify();
+	}
+
 	bIsBatchDirty = true;
 	SelectedObjects.Add(InObject);
+	
 	if (Sink)
 	{
 		Sink->OnObjectSelected(InObject, !IsBatchSelecting());
@@ -118,8 +126,14 @@ void FObjectSelectionStore::SelectObject(UObject* InObject)
 
 void FObjectSelectionStore::DeselectObject(UObject* InObject)
 {
+	if (Sink && !bIsBatchDirty)
+	{
+		Sink->OnPreModify();
+	}
+
 	bIsBatchDirty = true;
 	SelectedObjects.RemoveSingle(InObject);
+	
 	if (Sink)
 	{
 		Sink->OnObjectDeselected(InObject, !IsBatchSelecting());
@@ -130,29 +144,37 @@ int32 FObjectSelectionStore::DeselectObjects(TFunctionRef<bool(UObject*)> InPred
 {
 	int32 RemovedCount = 0;
 	
-	for (int32 Index = SelectedObjects.Num() - 1; Index >= 0; --Index)
+	if (SelectedObjects.Num() > 0)
 	{
-		UObject* SelectedObject = SelectedObjects[Index].Get();
-		if (!SelectedObject || InPredicate(SelectedObject))
+		if (Sink && !bIsBatchDirty)
+		{
+			Sink->OnPreModify();
+		}
+
+		for (int32 Index = SelectedObjects.Num() - 1; Index >= 0; --Index)
 		{
-			SelectedObjects.RemoveAt(Index);
-			if (SelectedObject)
+			UObject* SelectedObject = SelectedObjects[Index].Get();
+			if (!SelectedObject || InPredicate(SelectedObject))
 			{
-				if (Sink)
+				SelectedObjects.RemoveAt(Index);
+				if (SelectedObject)
 				{
-					Sink->OnObjectDeselected(SelectedObject, /*bNotify*/false);
+					if (Sink)
+					{
+						Sink->OnObjectDeselected(SelectedObject, /*bNotify*/false);
+					}
+					++RemovedCount;
 				}
-				++RemovedCount;
 			}
 		}
-	}
 
-	if (RemovedCount > 0)
-	{
-		bIsBatchDirty = true;
-		if (Sink)
+		if (RemovedCount > 0)
 		{
-			Sink->OnSelectedChanged(/*bSyncState*/false, !IsBatchSelecting());
+			bIsBatchDirty = true;
+			if (Sink)
+			{
+				Sink->OnSelectedChanged(/*bSyncState*/false, !IsBatchSelecting());
+			}
 		}
 	}
 
@@ -192,8 +214,6 @@ void FObjectSelectionStore::ForceBatchDirty()
 }
 
 
-#if UE_USE_ELEMENT_LIST_SELECTION
-
 class FElementSelectionStore : public FSelectionStoreBase
 {
 public:
@@ -202,6 +222,7 @@ public:
 
 	//~ ISelectionStore interface
 	virtual void SetElementList(UTypedElementList* InElementList) override;
+	virtual UTypedElementList* GetElementList() const override;
 	virtual int32 GetNumObjects() const override;
 	virtual UObject* GetObjectAtIndex(const int32 InIndex) const override;
 	virtual bool IsValidObjectToSelect(const UObject* InObject) const override;
@@ -217,6 +238,7 @@ public:
 	virtual FTypedElementHandle GetElementHandleForObject(const UObject* InObject, const bool bAllowCreate = true) const = 0;
 
 private:
+	void OnElementListPreChangeEvent(const UTypedElementList* InElementList);
 	void OnElementListSyncEvent(const UTypedElementList* InElementList, FTypedElementListLegacySync::ESyncType InSyncType, const FTypedElementHandle& InElementHandle, bool bIsWithinBatchOperation);
 	UObject* GetObjectForElementHandle(const FTypedElementHandle& InElementHandle) const;
 
@@ -233,6 +255,7 @@ FElementSelectionStore::~FElementSelectionStore()
 {
 	if (ElementList)
 	{
+		ElementList->OnPreChange().RemoveAll(this);
 		ElementList->Legacy_GetSync().OnSyncEvent().RemoveAll(this);
 	}
 }
@@ -241,6 +264,7 @@ void FElementSelectionStore::SetElementList(UTypedElementList* InElementList)
 {
 	if (ElementList)
 	{
+		ElementList->OnPreChange().RemoveAll(this);
 		ElementList->Legacy_GetSync().OnSyncEvent().RemoveAll(this);
 	}
 
@@ -248,10 +272,16 @@ void FElementSelectionStore::SetElementList(UTypedElementList* InElementList)
 
 	if (ElementList)
 	{
+		ElementList->OnPreChange().AddRaw(this, &FElementSelectionStore::OnElementListPreChangeEvent);
 		ElementList->Legacy_GetSync().OnSyncEvent().AddRaw(this, &FElementSelectionStore::OnElementListSyncEvent);
 	}
 }
 
+UTypedElementList* FElementSelectionStore::GetElementList() const
+{
+	return ElementList;
+}
+
 int32 FElementSelectionStore::GetNumObjects() const
 {
 	return ElementList
@@ -279,7 +309,7 @@ bool FElementSelectionStore::IsValidObjectToSelect(const UObject* InObject) cons
 		TTypedElement<UTypedElementSelectionInterface> SelectionElement = ElementList->GetElement<UTypedElementSelectionInterface>(GetElementHandleForObject(InObject));
 		if (SelectionElement)
 		{
-			bIsValid &= SelectionElement.IsValidSelection() && SelectionElement.Legacy_GetSelectionObject() == InObject;
+			bIsValid &= SelectionElement.Legacy_GetSelectionObject() == InObject;
 		}
 		else
 		{
@@ -368,6 +398,14 @@ void FElementSelectionStore::ForceBatchDirty()
 	}
 }
 
+void FElementSelectionStore::OnElementListPreChangeEvent(const UTypedElementList* InElementList)
+{
+	if (Sink)
+	{
+		Sink->OnPreModify();
+	}
+}
+
 void FElementSelectionStore::OnElementListSyncEvent(const UTypedElementList* InElementList, FTypedElementListLegacySync::ESyncType InSyncType, const FTypedElementHandle& InElementHandle, bool bIsWithinBatchOperation)
 {
 	if (Sink)
@@ -461,13 +499,12 @@ FTypedElementHandle FComponentElementSelectionStore::GetElementHandleForObject(c
 	return Component->AcquireEditorElementHandle(bAllowCreate);
 }
 
-#endif	// UE_USE_ELEMENT_LIST_SELECTION
-
 } // namespace Selection_Private
 
-USelection::FOnSelectionChanged	USelection::SelectionChangedEvent;
-USelection::FOnSelectionChanged	USelection::SelectObjectEvent;
-FSimpleMulticastDelegate		USelection::SelectNoneEvent;
+USelection::FOnSelectionChanged				USelection::SelectionChangedEvent;
+USelection::FOnSelectionChanged				USelection::SelectObjectEvent;
+FSimpleMulticastDelegate					USelection::SelectNoneEvent;
+USelection::FOnSelectionElementListChanged	USelection::SelectionElementListChanged;
 
 USelection* USelection::CreateObjectSelection(FUObjectAnnotationSparseBool* InSelectionAnnotation, UObject* InOuter, FName InName, EObjectFlags InFlags)
 {
@@ -479,22 +516,14 @@ USelection* USelection::CreateObjectSelection(FUObjectAnnotationSparseBool* InSe
 USelection* USelection::CreateActorSelection(FUObjectAnnotationSparseBool* InSelectionAnnotation, UObject* InOuter, FName InName, EObjectFlags InFlags)
 {
 	USelection* Selection = NewObject<USelection>(InOuter, InName, InFlags);
-#if UE_USE_ELEMENT_LIST_SELECTION
 	Selection->Initialize(InSelectionAnnotation, MakeShared<Selection_Private::FActorElementSelectionStore>());
-#else	// UE_USE_ELEMENT_LIST_SELECTION
-	Selection->Initialize(InSelectionAnnotation, MakeShared<Selection_Private::FObjectSelectionStore>(AActor::StaticClass()));
-#endif	// UE_USE_ELEMENT_LIST_SELECTION
 	return Selection;
 }
 
 USelection* USelection::CreateComponentSelection(FUObjectAnnotationSparseBool* InSelectionAnnotation, UObject* InOuter, FName InName, EObjectFlags InFlags)
 {
 	USelection* Selection = NewObject<USelection>(InOuter, InName, InFlags);
-#if UE_USE_ELEMENT_LIST_SELECTION
 	Selection->Initialize(InSelectionAnnotation, MakeShared<Selection_Private::FComponentElementSelectionStore>());
-#else	// UE_USE_ELEMENT_LIST_SELECTION
-	Selection->Initialize(InSelectionAnnotation, MakeShared<Selection_Private::FObjectSelectionStore>(UActorComponent::StaticClass()));
-#endif	// UE_USE_ELEMENT_LIST_SELECTION
 	return Selection;
 }
 
@@ -517,6 +546,19 @@ void USelection::Initialize(FUObjectAnnotationSparseBool* InSelectionAnnotation,
 	SyncSelectedState();
 }
 
+void USelection::SetElementSelectionSet(UTypedElementSelectionSet* InElementSelectionSet)
+{
+	ElementSelectionSet = InElementSelectionSet;
+	UTypedElementList* OldElementList = SelectionStore->GetElementList();
+	SelectionStore->SetElementList(InElementSelectionSet ? InElementSelectionSet->GetMutableElementList() : nullptr);
+	USelection::SelectionElementListChanged.Broadcast(this, OldElementList, SelectionStore->GetElementList());
+}
+
+void USelection::OnPreModify()
+{
+	Modify();
+}
+
 void USelection::OnObjectSelected(UObject* InObject, const bool bNotify)
 {
 	check(InObject);
@@ -736,5 +778,3 @@ void USelection::BeginDestroy()
 		SelectionAnnotation = nullptr;
 	}
 }
-
-#undef UE_USE_ELEMENT_LIST_SELECTION
diff --git a/Engine/Source/Editor/UnrealEd/Private/UnrealEdEngine.cpp b/Engine/Source/Editor/UnrealEd/Private/UnrealEdEngine.cpp
--- a/Engine/Source/Editor/UnrealEd/Private/UnrealEdEngine.cpp
+++ b/Engine/Source/Editor/UnrealEd/Private/UnrealEdEngine.cpp
@@ -123,6 +123,7 @@ void UUnrealEdEngine::Init(IEngineLoop* InEngineLoop)
 	FEditorSupportDelegates::PostWindowsMessage.AddUObject(this, &UUnrealEdEngine::OnPostWindowsMessage);
 
 	USelection::SelectionChangedEvent.AddUObject(this, &UUnrealEdEngine::OnEditorSelectionChanged);
+	USelection::SelectionElementListChanged.AddUObject(this, &UUnrealEdEngine::OnEditorSelectionElementListPtrChanged);
 
 	// Initialize the snap manager
 	FSnappingUtils::InitEditorSnappingTools();
diff --git a/Engine/Source/Editor/UnrealEd/Public/Elements/EditorElements.h b/Engine/Source/Editor/UnrealEd/Public/Elements/EditorElements.h
new file mode 100644
--- /dev/null
+++ b/Engine/Source/Editor/UnrealEd/Public/Elements/EditorElements.h
@@ -0,0 +1,9 @@
+// Copyright Epic Games, Inc. All Rights Reserved.
+
+#pragma once
+
+#include "Delegates/Delegate.h"
+
+extern UNREALED_API FSimpleMulticastDelegate OnRegisterEditorElementsDelegate;
+
+void RegisterEditorElements();
diff --git a/Engine/Source/Editor/UnrealEd/Public/LevelViewportClickHandlers.h b/Engine/Source/Editor/UnrealEd/Public/LevelViewportClickHandlers.h
--- a/Engine/Source/Editor/UnrealEd/Public/LevelViewportClickHandlers.h
+++ b/Engine/Source/Editor/UnrealEd/Public/LevelViewportClickHandlers.h
@@ -11,6 +11,7 @@ class ABrush;
 class FLevelEditorViewportClient;
 class UModel;
 class USceneComponent;
+struct FTypedElementHandle;
 struct FViewportClick;
 struct HActor;
 struct HGeomEdgeProxy;
@@ -45,6 +46,8 @@ namespace LevelViewportClickHandlers
 {
 	bool ClickViewport(FLevelEditorViewportClient* ViewportClient, const FViewportClick& Click);
 
+	bool ClickElement(FLevelEditorViewportClient* ViewportClient, const FTypedElementHandle& HitElement, const FViewportClick& Click);
+
 	bool UNREALED_API ClickActor(FLevelEditorViewportClient* ViewportClient,AActor* Actor,const FViewportClick& Click,bool bAllowSelectionChange);
 
 	bool ClickComponent(FLevelEditorViewportClient* ViewportClient, HActor* ActorHitProxy, const FViewportClick& Click);
diff --git a/Engine/Source/Editor/UnrealEd/Public/Selection.h b/Engine/Source/Editor/UnrealEd/Public/Selection.h
--- a/Engine/Source/Editor/UnrealEd/Public/Selection.h
+++ b/Engine/Source/Editor/UnrealEd/Public/Selection.h
@@ -10,6 +10,7 @@
 #include "Selection.generated.h"
 
 class UTypedElementList;
+class UTypedElementSelectionSet;
 
 namespace Selection_Private
 {
@@ -19,6 +20,11 @@ class ISelectionStoreSink
 public:
 	virtual ~ISelectionStoreSink() = default;
 
+	/**
+	 * Called prior to potential changes being made within the underlying store.
+	 */
+	virtual void OnPreModify() = 0;
+
 	/**
 	 * Called when the given object is selected within the underlying store.
 	 */
@@ -51,6 +57,11 @@ public:
 	 */
 	virtual void SetElementList(UTypedElementList* InElementList) = 0;
 
+	/**
+	 * Get the element list instance for this selection set, if any.
+	 */
+	virtual UTypedElementList* GetElementList() const = 0;
+
 	/**
 	 * Get the number of objects within the underlying store.
 	 * This is the total number of objects within the store (for use as an upper limit of GetObjectAtIndex),
@@ -164,12 +175,17 @@ public:
 	/** Params: UObject* NewSelection */
 	DECLARE_EVENT_OneParam(USelection, FOnSelectionChanged, UObject*);
 
+	/** Event fired when the element list for a selection is changed */
+	DECLARE_EVENT_ThreeParams(USelection, FOnSelectionElementListChanged, USelection* /*Selection*/, UTypedElementList* /*OldElementList*/, UTypedElementList* /*NewElementList*/);
+
 	/** Called when selection in editor has changed */
 	static FOnSelectionChanged SelectionChangedEvent;
 	/** Called when an object has been selected (generally an actor) */
 	static FOnSelectionChanged SelectObjectEvent;
 	/** Called to deselect everything */
 	static FSimpleMulticastDelegate SelectNoneEvent;
+	/** Called when the element list for a selection has changed */
+	static FOnSelectionElementListChanged SelectionElementListChanged;
 
 	typedef ClassArray::TIterator TClassIterator;
 	typedef ClassArray::TConstIterator TClassConstIterator;
@@ -183,7 +199,32 @@ public:
 	 */
 	void SetElementList(UTypedElementList* InElementList)
 	{
+		ElementSelectionSet.Reset();
+		UTypedElementList* OldElementList = SelectionStore->GetElementList();
 		SelectionStore->SetElementList(InElementList);
+		USelection::SelectionElementListChanged.Broadcast(this, OldElementList, SelectionStore->GetElementList());
+	}
+
+	/**
+	 * Get the element list instance for this selection set, if any.
+	 */
+	UTypedElementList* GetElementList() const
+	{
+		return SelectionStore->GetElementList();
+	}
+
+	/**
+	 * Set the element selection set instance for this selection set.
+	 * @note Also sets the element list instance.
+	 */
+	void SetElementSelectionSet(UTypedElementSelectionSet* InElementSelectionSet);
+
+	/**
+	 * Get the element selection set instance for this selection set, if any.
+	 */
+	UTypedElementSelectionSet* GetElementSelectionSet() const
+	{
+		return ElementSelectionSet.Get();
 	}
 
 	/**
@@ -482,6 +523,7 @@ protected:
 	void Initialize(FUObjectAnnotationSparseBool* InSelectionAnnotation, TSharedRef<Selection_Private::ISelectionStore>&& InSelectionStore);
 
 	//~ ISelectionStoreSink interface
+	virtual void OnPreModify() override;
 	virtual void OnObjectSelected(UObject* InObject, const bool bNotify) override;
 	virtual void OnObjectDeselected(UObject* InObject, const bool bNotify) override;
 	virtual void OnSelectedChanged(const bool bSyncState, const bool bNotify) override;
@@ -492,6 +534,9 @@ protected:
 	/** Store of selected objects. */
 	TSharedPtr<Selection_Private::ISelectionStore> SelectionStore = nullptr;
 
+	/** Underlying element selection set (if any). */
+	TWeakObjectPtr<UTypedElementSelectionSet> ElementSelectionSet;
+
 	/** Tracks the most recently selected actor classes.  Used for UnrealEd menus. */
 	ClassArray SelectedClasses;
 	
diff --git a/Engine/Source/Runtime/Engine/Private/Elements/EngineElements.cpp b/Engine/Source/Runtime/Engine/Private/Elements/EngineElements.cpp
--- a/Engine/Source/Runtime/Engine/Private/Elements/EngineElements.cpp
+++ b/Engine/Source/Runtime/Engine/Private/Elements/EngineElements.cpp
@@ -1,6 +1,6 @@
 // Copyright Epic Games, Inc. All Rights Reserved.
 
-#include "EngineElements.h"
+#include "Elements/EngineElements.h"
 #include "Modules/ModuleManager.h"
 
 #include "TypedElementRegistry.h"
@@ -11,6 +11,8 @@
 #include "Elements/Component/ComponentElementData.h"
 #include "Elements/Component/ComponentElementSelectionInterface.h"
 
+FSimpleMulticastDelegate OnRegisterEngineElementsDelegate;
+
 void RegisterEngineActorElements()
 {
 	UTypedElementRegistry* Registry = UTypedElementRegistry::GetInstance();
@@ -35,4 +37,6 @@ void RegisterEngineElements()
 
 	RegisterEngineActorElements();
 	RegisterEngineComponentElements();
+
+	OnRegisterEngineElementsDelegate.Broadcast();
 }
diff --git a/Engine/Source/Runtime/Engine/Private/Elements/EngineElements.h b/Engine/Source/Runtime/Engine/Private/Elements/EngineElements.h
deleted file mode 100644
diff --git a/Engine/Source/Runtime/Engine/Public/Elements/Actor/ActorElementSelectionInterface.h b/Engine/Source/Runtime/Engine/Public/Elements/Actor/ActorElementSelectionInterface.h
--- a/Engine/Source/Runtime/Engine/Public/Elements/Actor/ActorElementSelectionInterface.h
+++ b/Engine/Source/Runtime/Engine/Public/Elements/Actor/ActorElementSelectionInterface.h
@@ -5,8 +5,6 @@
 #include "TypedElementSelectionInterface.h"
 #include "ActorElementSelectionInterface.generated.h"
 
-class UTypedElementList;
-
 UCLASS()
 class ENGINE_API UActorElementSelectionInterface : public UTypedElementSelectionInterface
 {
diff --git a/Engine/Source/Runtime/Engine/Public/Elements/Component/ComponentElementSelectionInterface.h b/Engine/Source/Runtime/Engine/Public/Elements/Component/ComponentElementSelectionInterface.h
--- a/Engine/Source/Runtime/Engine/Public/Elements/Component/ComponentElementSelectionInterface.h
+++ b/Engine/Source/Runtime/Engine/Public/Elements/Component/ComponentElementSelectionInterface.h
@@ -5,8 +5,6 @@
 #include "TypedElementSelectionInterface.h"
 #include "ComponentElementSelectionInterface.generated.h"
 
-class UTypedElementList;
-
 UCLASS()
 class ENGINE_API UComponentElementSelectionInterface : public UTypedElementSelectionInterface
 {
diff --git a/Engine/Source/Runtime/Engine/Public/Elements/EngineElements.h b/Engine/Source/Runtime/Engine/Public/Elements/EngineElements.h
new file mode 100644
--- /dev/null
+++ b/Engine/Source/Runtime/Engine/Public/Elements/EngineElements.h
@@ -0,0 +1,9 @@
+// Copyright Epic Games, Inc. All Rights Reserved.
+
+#pragma once
+
+#include "Delegates/Delegate.h"
+
+extern ENGINE_API FSimpleMulticastDelegate OnRegisterEngineElementsDelegate;
+
+void RegisterEngineElements();
diff --git a/Engine/Source/Runtime/TypedElementFramework/Private/TypedElementList.cpp b/Engine/Source/Runtime/TypedElementFramework/Private/TypedElementList.cpp
--- a/Engine/Source/Runtime/TypedElementFramework/Private/TypedElementList.cpp
+++ b/Engine/Source/Runtime/TypedElementFramework/Private/TypedElementList.cpp
@@ -73,6 +73,24 @@ void FTypedElementListLegacySync::ForceBatchOperationDirty()
 }
 
 
+FTypedElementListLegacySyncScopedBatch::FTypedElementListLegacySyncScopedBatch(UTypedElementList* InElementList, const bool InNotify)
+	: ElementListLegacySync(InElementList->Legacy_GetSyncPtr())
+	, bNotify(InNotify)
+{
+	if (ElementListLegacySync)
+	{
+		ElementListLegacySync->BeginBatchOperation();
+	}
+}
+
+FTypedElementListLegacySyncScopedBatch::~FTypedElementListLegacySyncScopedBatch()
+{
+	if (ElementListLegacySync)
+	{
+		ElementListLegacySync->EndBatchOperation(bNotify);
+	}
+}
+
 UTypedElementList* UTypedElementList::Private_CreateElementList(UTypedElementRegistry* InRegistry)
 {
 	UTypedElementList* ElementList = NewObject<UTypedElementList>();
@@ -113,6 +131,11 @@ UTypedElementInterface* UTypedElementList::GetElementInterface(const FTypedEleme
 	return Registry->GetElementInterface(InElementHandle, InBaseInterfaceType);
 }
 
+FTypedHandleTypeId UTypedElementList::GetRegisteredElementTypeId(const FName InElementTypeName) const
+{
+	return Registry->GetRegisteredElementTypeId(InElementTypeName);
+}
+
 bool UTypedElementList::AddElementImpl(FTypedElementHandle&& InElementHandle)
 {
 	if (!InElementHandle)
@@ -120,6 +143,8 @@ bool UTypedElementList::AddElementImpl(FTypedElementHandle&& InElementHandle)
 		return false;
 	}
 
+	NoteListMayChange();
+
 	bool bAlreadyAdded = false;
 	ElementCombinedIds.Add(InElementHandle.GetId().GetCombinedId(), &bAlreadyAdded);
 
@@ -139,6 +164,8 @@ bool UTypedElementList::RemoveElementImpl(const FTypedElementId& InElementId)
 		return false;
 	}
 
+	NoteListMayChange();
+
 	const bool bRemoved = ElementCombinedIds.Remove(InElementId.GetCombinedId()) > 0;
 
 	if (bRemoved)
@@ -160,31 +187,29 @@ bool UTypedElementList::RemoveElementImpl(const FTypedElementId& InElementId)
 
 int32 UTypedElementList::RemoveAllElementsImpl(TFunctionRef<bool(const FTypedElementHandle&)> InPredicate)
 {
-	if (LegacySync)
-	{
-		LegacySync->BeginBatchOperation();
-	}
-
 	int32 RemovedCount = 0;
-	for (int32 Index = ElementHandles.Num() - 1; Index >= 0; --Index)
+
+	if (ElementHandles.Num() > 0)
 	{
-		if (InPredicate(ElementHandles[Index]))
+		FTypedElementListLegacySyncScopedBatch LegacySyncBatch(this);
+
+		NoteListMayChange();
+
+		for (int32 Index = ElementHandles.Num() - 1; Index >= 0; --Index)
 		{
-			FTypedElementHandle RemovedElementHandle = MoveTemp(ElementHandles[Index]);
-			ElementCombinedIds.Remove(RemovedElementHandle.GetId().GetCombinedId());
-			ElementHandles.RemoveAt(Index, 1, /*bAllowShrinking*/false);
+			if (InPredicate(ElementHandles[Index]))
+			{
+				FTypedElementHandle RemovedElementHandle = MoveTemp(ElementHandles[Index]);
+				ElementCombinedIds.Remove(RemovedElementHandle.GetId().GetCombinedId());
+				ElementHandles.RemoveAt(Index, 1, /*bAllowShrinking*/false);
 
-			NoteListChanged(EChangeType::Removed, RemovedElementHandle);
+				NoteListChanged(EChangeType::Removed, RemovedElementHandle);
 
-			++RemovedCount;
+				++RemovedCount;
+			}
 		}
 	}
 
-	if (LegacySync)
-	{
-		LegacySync->EndBatchOperation();
-	}
-
 	return RemovedCount;
 }
 
@@ -203,6 +228,11 @@ FTypedElementListLegacySync& UTypedElementList::Legacy_GetSync()
 	return *LegacySync;
 }
 
+FTypedElementListLegacySync* UTypedElementList::Legacy_GetSyncPtr()
+{
+	return LegacySync.Get();
+}
+
 void UTypedElementList::NotifyPendingChanges()
 {
 	if (bHasPendingNotify)
@@ -213,6 +243,14 @@ void UTypedElementList::NotifyPendingChanges()
 	}
 }
 
+void UTypedElementList::NoteListMayChange()
+{
+	if (!bHasPendingNotify)
+	{
+		OnPreChangeDelegate.Broadcast(this);
+	}
+}
+
 void UTypedElementList::NoteListChanged(const EChangeType InChangeType, const FTypedElementHandle& InElementHandle)
 {
 	bHasPendingNotify = true;
diff --git a/Engine/Source/Runtime/TypedElementFramework/Public/TypedElementList.h b/Engine/Source/Runtime/TypedElementFramework/Public/TypedElementList.h
--- a/Engine/Source/Runtime/TypedElementFramework/Public/TypedElementList.h
+++ b/Engine/Source/Runtime/TypedElementFramework/Public/TypedElementList.h
@@ -42,7 +42,7 @@ template <typename BaseInterfaceType>
 struct TTypedElementListInterfaceIterator
 {
 public:
-	explicit TTypedElementListInterfaceIterator(const UTypedElementRegistry* InRegistry, const TArray<FTypedElementHandle>& InElementHandles, const int32 InIndex)
+	TTypedElementListInterfaceIterator(const UTypedElementRegistry* InRegistry, const TArray<FTypedElementHandle>& InElementHandles, const int32 InIndex)
 		: Registry(InRegistry)
 		, ElementHandles(InElementHandles)
 		, InitialNum(InElementHandles.Num())
@@ -105,7 +105,7 @@ template <typename BaseInterfaceType>
 struct TTypedElementListInterfaceIteratorProxy
 {
 public:
-	explicit TTypedElementListInterfaceIteratorProxy(const UTypedElementRegistry* InRegistry, const TArray<FTypedElementHandle>& InElementHandles)
+	TTypedElementListInterfaceIteratorProxy(const UTypedElementRegistry* InRegistry, const TArray<FTypedElementHandle>& InElementHandles)
 		: Registry(InRegistry)
 		, ElementHandles(InElementHandles)
 	{
@@ -196,6 +196,27 @@ private:
 	bool bBatchOperationIsDirty = false;
 };
 
+/**
+ * Helper to batch immediate sync notifications for legacy code.
+ * Does nothing if no legacy sync has been created for the given instance.
+ */
+class TYPEDELEMENTFRAMEWORK_API FTypedElementListLegacySyncScopedBatch
+{
+public:
+	explicit FTypedElementListLegacySyncScopedBatch(UTypedElementList* InElementList, const bool InNotify = true);
+	~FTypedElementListLegacySyncScopedBatch();
+
+	FTypedElementListLegacySyncScopedBatch(const FTypedElementListLegacySyncScopedBatch&) = delete;
+	FTypedElementListLegacySyncScopedBatch& operator=(const FTypedElementListLegacySyncScopedBatch&) = delete;
+
+	FTypedElementListLegacySyncScopedBatch(FTypedElementListLegacySyncScopedBatch&&) = delete;
+	FTypedElementListLegacySyncScopedBatch& operator=(FTypedElementListLegacySyncScopedBatch&&) = delete;
+
+private:
+	FTypedElementListLegacySync* ElementListLegacySync = nullptr;
+	bool bNotify = true;
+};
+
 /**
  * A list of element handles.
  * Provides high-level access to groups of elements, including accessing elements that implement specific interfaces.
@@ -337,6 +358,7 @@ public:
 	UFUNCTION(BlueprintCallable, Category="TypedElementFramework|List")
 	FORCEINLINE void Empty(const int32 InSlack = 0)
 	{
+		NoteListMayChange();
 		ElementCombinedIds.Empty(InSlack);
 		ElementHandles.Empty(InSlack);
 		NoteListChanged(EChangeType::Cleared);
@@ -348,6 +370,7 @@ public:
 	UFUNCTION(BlueprintCallable, Category="TypedElementFramework|List")
 	FORCEINLINE void Reset()
 	{
+		NoteListMayChange();
 		ElementCombinedIds.Reset();
 		ElementHandles.Reset();
 		NoteListChanged(EChangeType::Cleared);
@@ -422,21 +445,13 @@ public:
 	 */
 	void Append(TArrayView<const FTypedElementHandle> InElementHandles)
 	{
-		if (LegacySync)
-		{
-			LegacySync->BeginBatchOperation();
-		}
+		FTypedElementListLegacySyncScopedBatch LegacySyncBatch(this);
 
 		Reserve(Num() + InElementHandles.Num());
 		for (const FTypedElementHandle& ElementHandle : InElementHandles)
 		{
 			AddElementImpl(CopyTemp(ElementHandle));
 		}
-
-		if (LegacySync)
-		{
-			LegacySync->EndBatchOperation();
-		}
 	}
 
 	/**
@@ -454,21 +469,13 @@ public:
 	template <typename ElementDataType>
 	void Append(TArrayView<const TTypedElementOwner<ElementDataType>> InElementOwners)
 	{
-		if (LegacySync)
-		{
-			LegacySync->BeginBatchOperation();
-		}
+		FTypedElementListLegacySyncScopedBatch LegacySyncBatch(this);
 
 		Reserve(Num() + InElementOwners.Num());
 		for (const TTypedElementOwner<ElementDataType>& ElementOwner : InElementOwners)
 		{
 			AddElementImpl(ElementOwner.AcquireHandle());
 		}
-
-		if (LegacySync)
-		{
-			LegacySync->EndBatchOperation();
-		}
 	}
 
 	/**
@@ -524,6 +531,22 @@ public:
 		});
 	}
 
+	/**
+	 * Get the element type ID for the associated element type name, if any.
+	 * @return The element type ID, or 0 if the given name wasn't registered.
+	 */
+	FTypedHandleTypeId GetRegisteredElementTypeId(const FName InElementTypeName) const;
+
+	/**
+	 * Access the delegate that is invoked whenever this element list is potentially about to change.
+	 * @note This may be called even if no actual change happens, though once a change does happen it won't be called again until after the next call to NotifyPendingChanges.
+	 */
+	DECLARE_EVENT_OneParam(UTypedElementList, FOnPreChange, const UTypedElementList* /*InElementList*/);
+	FOnPreChange& OnPreChange()
+	{
+		return OnPreChangeDelegate;
+	}
+
 	/**
 	 * Access the delegate that is invoked whenever this element list has been changed.
 	 * @note This is called automatically at the end of each frame, but can also be manually invoked by NotifyPendingChanges.
@@ -545,6 +568,12 @@ public:
 	 */
 	FTypedElementListLegacySync& Legacy_GetSync();
 
+	/**
+	 * Access the interface to allow external systems (such as USelection) to receive immediate sync notifications as an element list is changed.
+	 * This exists purely as a bridging mechanism and shouldn't be relied on for new code. This will return null if no legacy sync has been created for this instance.
+	 */
+	FTypedElementListLegacySync* Legacy_GetSyncPtr();
+
 	/**
 	 * Get an iterator for elements that implement the given element interface.
 	 */
@@ -590,6 +619,7 @@ private:
 	int32 RemoveAllElementsImpl(TFunctionRef<bool(const FTypedElementHandle&)> InPredicate);
 	bool ContainsElementImpl(const FTypedElementId& InElementId) const;
 
+	void NoteListMayChange();
 	void NoteListChanged(const EChangeType InChangeType, const FTypedElementHandle& InElementHandle = FTypedElementHandle());
 
 	/**
@@ -609,6 +639,11 @@ private:
 	 */
 	TArray<FTypedElementHandle> ElementHandles;
 
+	/**
+	 * Delegate that is invoked whenever this element list is potentially about to change.
+	 */
+	FOnPreChange OnPreChangeDelegate;
+
 	/**
 	 * Delegate that is invoked whenever this element list has been changed.
 	 */
diff --git a/Engine/Source/Runtime/TypedElementFramework/Public/TypedElementRegistry.h b/Engine/Source/Runtime/TypedElementFramework/Public/TypedElementRegistry.h
--- a/Engine/Source/Runtime/TypedElementFramework/Public/TypedElementRegistry.h
+++ b/Engine/Source/Runtime/TypedElementFramework/Public/TypedElementRegistry.h
@@ -41,6 +41,18 @@ public:
 	UFUNCTION(BlueprintPure, Category="TypedElementFramework|Registry")
 	static UTypedElementRegistry* GetInstance();
 
+	/**
+	 * Get the element type ID for the associated element type name, if any.
+	 * @return The element type ID, or 0 if the given name wasn't registered.
+	 */
+	FORCEINLINE FTypedHandleTypeId GetRegisteredElementTypeId(const FName InElementTypeName) const
+	{
+		FReadScopeLock RegisteredElementTypesLock(RegisteredElementTypesRW);
+
+		const FTypedHandleTypeId* TypeId = RegisteredElementTypesNameToId.Find(InElementTypeName);
+		return TypeId ? *TypeId : 0;
+	}
+
 	/**
 	 * Register an element type that doesn't require any additional payload data.
 	 */
diff --git a/Engine/Source/Runtime/TypedElementInterfaces/Private/TypedElementSelectionInterface.cpp b/Engine/Source/Runtime/TypedElementInterfaces/Private/TypedElementSelectionInterface.cpp
new file mode 100644
--- /dev/null
+++ b/Engine/Source/Runtime/TypedElementInterfaces/Private/TypedElementSelectionInterface.cpp
@@ -0,0 +1,19 @@
+// Copyright Epic Games, Inc. All Rights Reserved.
+
+#include "TypedElementSelectionInterface.h"
+#include "TypedElementList.h"
+
+bool UTypedElementSelectionInterface::IsElementSelected(const FTypedElementHandle& InElementHandle, const UTypedElementList* InSelectionSet, const FTypedElementIsSelectedOptions& InSelectionOptions)
+{
+	return InSelectionSet->Contains(InElementHandle);
+}
+
+bool UTypedElementSelectionInterface::SelectElement(const FTypedElementHandle& InElementHandle, UTypedElementList* InSelectionSet, const FTypedElementSelectionOptions& InSelectionOptions)
+{
+	return InSelectionSet->Add(InElementHandle);
+}
+
+bool UTypedElementSelectionInterface::DeselectElement(const FTypedElementHandle& InElementHandle, UTypedElementList* InSelectionSet, const FTypedElementSelectionOptions& InSelectionOptions)
+{
+	return InSelectionSet->Remove(InElementHandle);
+}
diff --git a/Engine/Source/Runtime/TypedElementInterfaces/Public/TypedElementSelectionInterface.h b/Engine/Source/Runtime/TypedElementInterfaces/Public/TypedElementSelectionInterface.h
--- a/Engine/Source/Runtime/TypedElementInterfaces/Public/TypedElementSelectionInterface.h
+++ b/Engine/Source/Runtime/TypedElementInterfaces/Public/TypedElementSelectionInterface.h
@@ -5,29 +5,130 @@
 #include "TypedElementHandle.h"
 #include "TypedElementSelectionInterface.generated.h"
 
-UCLASS()
+class UTypedElementList;
+
+UENUM()
+enum class ETypedElementSelectionMethod : uint8
+{
+	/** Select the "primary" element (eg, a component favor selecting its owner actor) */
+	Primary,
+	/** Select the "secondary" element (eg, a component would favor selecting itself) */
+	Secondary,
+};
+
+USTRUCT(BlueprintType)
+struct FTypedElementIsSelectedOptions
+{
+	GENERATED_BODY()
+
+public:
+	FTypedElementIsSelectedOptions& SetAllowIndirect(const bool InAllowIndirect) { bAllowIndirect = InAllowIndirect; return *this; }
+	bool AllowIndirect() const { return bAllowIndirect; }
+	
+private:
+	UPROPERTY(BlueprintReadWrite, Category="TypedElementInterfaces|Selection|IsSelectedOptions", meta=(AllowPrivateAccess=true))
+	bool bAllowIndirect = false;
+};
+
+USTRUCT(BlueprintType)
+struct FTypedElementSelectionOptions
+{
+	GENERATED_BODY()
+
+public:
+	FTypedElementSelectionOptions& SetAllowHidden(const bool InAllowHidden) { bAllowHidden = InAllowHidden; return *this; }
+	bool AllowHidden() const { return bAllowHidden; }
+
+	FTypedElementSelectionOptions& SetAllowGroups(const bool InAllowGroups) { bAllowGroups = InAllowGroups; return *this; }
+	bool AllowGroups() const { return bAllowGroups; }
+
+	FTypedElementSelectionOptions& SetAllowLegacyNotifications(const bool InAllowLegacyNotifications) { bAllowLegacyNotifications = InAllowLegacyNotifications; return *this; }
+	bool AllowLegacyNotifications() const { return bAllowLegacyNotifications; }
+
+	FTypedElementSelectionOptions& SetWarnIfLocked(const bool InWarnIfLocked) { bWarnIfLocked = InWarnIfLocked; return *this; }
+	bool WarnIfLocked() const { return bWarnIfLocked; }
+	
+private:
+	UPROPERTY(BlueprintReadWrite, Category="TypedElementInterfaces|Selection|SelectionOptions", meta=(AllowPrivateAccess=true))
+	bool bAllowHidden = false;
+
+	UPROPERTY(BlueprintReadWrite, Category="TypedElementInterfaces|Selection|SelectionOptions", meta=(AllowPrivateAccess=true))
+	bool bAllowGroups = true;
+
+	UPROPERTY(BlueprintReadWrite, Category="TypedElementInterfaces|Selection|SelectionOptions", meta=(AllowPrivateAccess=true))
+	bool bAllowLegacyNotifications = true;
+
+	UPROPERTY(BlueprintReadWrite, Category="TypedElementInterfaces|Selection|SelectionOptions", meta=(AllowPrivateAccess=true))
+	bool bWarnIfLocked = false;
+};
+
+UCLASS(Abstract)
 class TYPEDELEMENTINTERFACES_API UTypedElementSelectionInterface : public UTypedElementInterface
 {
 	GENERATED_BODY()
 
 public:
 	/**
-	 * Test to see whether the given handle is in a valid state to be selected.
+	 * Test to see whether the given element is currently considered selected.
+	 */
+	UFUNCTION(BlueprintPure, Category="TypedElementInterfaces|Selection")
+	virtual bool IsElementSelected(const FTypedElementHandle& InElementHandle, const UTypedElementList* InSelectionSet, const FTypedElementIsSelectedOptions& InSelectionOptions);
+
+	/**
+	 * Test to see whether the given element can be selected.
+	 */
+	UFUNCTION(BlueprintPure, Category="TypedElementInterfaces|Selection")
+	virtual bool CanSelectElement(const FTypedElementHandle& InElementHandle, const FTypedElementSelectionOptions& InSelectionOptions) { return true; }
+
+	/**
+	 * Test to see whether the given element can be deselected.
+	 */
+	UFUNCTION(BlueprintPure, Category="TypedElementInterfaces|Selection")
+	virtual bool CanDeselectElement(const FTypedElementHandle& InElementHandle, const FTypedElementSelectionOptions& InSelectionOptions) { return true; }
+
+	/**
+	 * Attempt to select the given element.
+	 * @return True if the selection was changed, false otherwise.
 	 */
 	UFUNCTION(BlueprintCallable, Category="TypedElementInterfaces|Selection")
-	virtual bool IsValidSelection(const FTypedElementHandle& InElementHandle) { return true; }
+	virtual bool SelectElement(const FTypedElementHandle& InElementHandle, UTypedElementList* InSelectionSet, const FTypedElementSelectionOptions& InSelectionOptions);
+
+	/**
+	 * Attempt to deselect the given element.
+	 * @return True if the selection was changed, false otherwise.
+	 */
+	UFUNCTION(BlueprintCallable, Category="TypedElementInterfaces|Selection")
+	virtual bool DeselectElement(const FTypedElementHandle& InElementHandle, UTypedElementList* InSelectionSet, const FTypedElementSelectionOptions& InSelectionOptions);
+
+	/**
+	 * Test to see whether selection modifiers (Ctrl or Shift) are allowed while selecting this element.
+	 */
+	UFUNCTION(BlueprintPure, Category="TypedElementInterfaces|Selection")
+	virtual bool AllowSelectionModifiers(const FTypedElementHandle& InElementHandle, const UTypedElementList* InSelectionSet) { return true; }
+
+	/**
+	 * Given an element, return the element that should actually perform a selection operation.
+	 */
+	UFUNCTION(BlueprintPure, Category="TypedElementInterfaces|Selection")
+	virtual FTypedElementHandle GetSelectionElement(const FTypedElementHandle& InElementHandle, const UTypedElementList* InCurrentSelection, const ETypedElementSelectionMethod InSelectionMethod) { return InElementHandle; }
 
 	/**
 	 * Retreive the object instance that should be selected via this handle.
 	 * @note This exists to allow the legacy USelection to bridge to an UTypedElementList instance. It should not be used in new code!
 	 */
-	UFUNCTION(BlueprintCallable, Category="TypedElementInterfaces|Selection")
+	UFUNCTION(BlueprintPure, Category="TypedElementInterfaces|Selection", DisplayName="GetSelectionObject (Legacy)")
 	virtual UObject* Legacy_GetSelectionObject(const FTypedElementHandle& InElementHandle) { return nullptr; }
 };
 
 template <>
 struct TTypedElement<UTypedElementSelectionInterface> : public TTypedElementBase<UTypedElementSelectionInterface>
 {
-	bool IsValidSelection() const { return InterfacePtr->IsValidSelection(*this); }
+	bool IsElementSelected(const UTypedElementList* InSelectionSet, const FTypedElementIsSelectedOptions& InSelectionOptions) const { return InterfacePtr->IsElementSelected(*this, InSelectionSet, InSelectionOptions); }
+	bool CanSelectElement(const FTypedElementSelectionOptions& InSelectionOptions) const { return InterfacePtr->CanSelectElement(*this, InSelectionOptions); }
+	bool CanDeselectElement(const FTypedElementSelectionOptions& InSelectionOptions) const { return InterfacePtr->CanDeselectElement(*this, InSelectionOptions); }
+	bool SelectElement(UTypedElementList* InSelectionSet, const FTypedElementSelectionOptions& InSelectionOptions) const { return InterfacePtr->SelectElement(*this, InSelectionSet, InSelectionOptions); }
+	bool DeselectElement(UTypedElementList* InSelectionSet, const FTypedElementSelectionOptions& InSelectionOptions) const { return InterfacePtr->DeselectElement(*this, InSelectionSet, InSelectionOptions); }
+	bool AllowSelectionModifiers(const UTypedElementList* InSelectionSet) const { return InterfacePtr->AllowSelectionModifiers(*this, InSelectionSet); }
+	FTypedElementHandle GetSelectionElement(const UTypedElementList* InCurrentSelection, const ETypedElementSelectionMethod InSelectionMethod) const { return InterfacePtr->GetSelectionElement(*this, InCurrentSelection, InSelectionMethod); }
 	UObject* Legacy_GetSelectionObject() const { return InterfacePtr->Legacy_GetSelectionObject(*this); }
 };
