Clear pending selection notifications when the legacy API is not emitting them

#rnx

diff --git a/Engine/Source/Editor/UnrealEd/Private/EditorSelectUtils.cpp b/Engine/Source/Editor/UnrealEd/Private/EditorSelectUtils.cpp
--- a/Engine/Source/Editor/UnrealEd/Private/EditorSelectUtils.cpp
+++ b/Engine/Source/Editor/UnrealEd/Private/EditorSelectUtils.cpp
@@ -441,9 +441,16 @@ void UUnrealEdEngine::SelectGroup(AGroupActor* InGroupActor, bool bForceSelectio
 			}
 		}
 
-		if (bSelectionChanged && bNotify)
+		if (bSelectionChanged)
 		{
-			SelectionSet->GetMutableElementList()->NotifyPendingChanges();
+			if (bNotify)
+			{
+				SelectionSet->GetMutableElementList()->NotifyPendingChanges();
+			}
+			else
+			{
+				SelectionSet->GetMutableElementList()->ClearPendingChanges();
+			}
 		}
 	}
 }
@@ -497,6 +504,10 @@ void UUnrealEdEngine::SelectActor(AActor* Actor, bool bInSelected, bool bNotify,
 			{
 				SelectionSet->GetMutableElementList()->NotifyPendingChanges();
 			}
+			else
+			{
+				SelectionSet->GetMutableElementList()->ClearPendingChanges();
+			}
 		}
 		else if (bNotify || bForceRefresh)
 		{
@@ -526,9 +537,16 @@ void UUnrealEdEngine::SelectComponent(UActorComponent* Component, bool bInSelect
 			? SelectionSet->SelectElement(Component->AcquireEditorElementHandle(), SelectionOptions)
 			: SelectionSet->DeselectElement(Component->AcquireEditorElementHandle(), SelectionOptions);
 
-		if (bSelectionChanged && bNotify)
+		if (bSelectionChanged)
 		{
-			SelectionSet->GetMutableElementList()->NotifyPendingChanges();
+			if (bNotify)
+			{
+				SelectionSet->GetMutableElementList()->NotifyPendingChanges();
+			}
+			else
+			{
+				SelectionSet->GetMutableElementList()->ClearPendingChanges();
+			}
 		}
 	}
 }
@@ -655,9 +673,16 @@ void UUnrealEdEngine::SelectNone(bool bNoteSelectionChange, bool bDeselectBSPSur
 	{
 		PostActorSelectionChanged();
 
-		if (bNoteSelectionChange && SelectionSet)
+		if (SelectionSet)
 		{
-			SelectionSet->GetMutableElementList()->NotifyPendingChanges();
+			if (bNoteSelectionChange)
+			{
+				SelectionSet->GetMutableElementList()->NotifyPendingChanges();
+			}
+			else
+			{
+				SelectionSet->GetMutableElementList()->ClearPendingChanges();
+			}
 		}
 	}
 }
diff --git a/Engine/Source/Runtime/TypedElementFramework/Private/TypedElementList.cpp b/Engine/Source/Runtime/TypedElementFramework/Private/TypedElementList.cpp
--- a/Engine/Source/Runtime/TypedElementFramework/Private/TypedElementList.cpp
+++ b/Engine/Source/Runtime/TypedElementFramework/Private/TypedElementList.cpp
@@ -243,6 +243,11 @@ void UTypedElementList::NotifyPendingChanges()
 	}
 }
 
+void UTypedElementList::ClearPendingChanges()
+{
+	bHasPendingNotify = false;
+}
+
 void UTypedElementList::NoteListMayChange()
 {
 	if (!bHasPendingNotify)
diff --git a/Engine/Source/Runtime/TypedElementFramework/Public/TypedElementList.h b/Engine/Source/Runtime/TypedElementFramework/Public/TypedElementList.h
--- a/Engine/Source/Runtime/TypedElementFramework/Public/TypedElementList.h
+++ b/Engine/Source/Runtime/TypedElementFramework/Public/TypedElementList.h
@@ -562,6 +562,11 @@ public:
 	 */
 	void NotifyPendingChanges();
 
+	/**
+	 * Clear whether there are pending changes for OnChangedDelegate to notify for, without emitting a notification.
+	 */
+	void ClearPendingChanges();
+
 	/**
 	 * Access the interface to allow external systems (such as USelection) to receive immediate sync notifications as an element list is changed.
 	 * This exists purely as a bridging mechanism and shouldn't be relied on for new code. It is lazily created as needed.
