Build Progress widget reskin

diff --git a/Engine/Source/Editor/UnrealEd/Private/Dialogs/SBuildProgress.cpp b/Engine/Source/Editor/UnrealEd/Private/Dialogs/SBuildProgress.cpp
--- a/Engine/Source/Editor/UnrealEd/Private/Dialogs/SBuildProgress.cpp
+++ b/Engine/Source/Editor/UnrealEd/Private/Dialogs/SBuildProgress.cpp
@@ -20,94 +20,68 @@ SBuildProgressWidget::~SBuildProgressWidget()
 
 void SBuildProgressWidget::Construct( const FArguments& InArgs )
 {
-	BorderImage = FInvalidatableBrushAttribute(FEditorStyle::GetBrush("Menu.Background"));
-
-	this->ChildSlot
-	.VAlign(VAlign_Center)
+	ChildSlot
 	[
-		SNew( SVerticalBox )
-		+SVerticalBox::Slot()
-		.AutoHeight()
-		.HAlign(HAlign_Fill)
-		.Padding( 10.0f, 4.0f )
+		SNew(SBorder)
+		.BorderImage(FAppStyle::Get().GetBrush("Brushes.Header"))
+		.Padding(0.0f)
 		[
-			SNew(SBorder)
+			SNew( SVerticalBox )
+			+SVerticalBox::Slot()
+			.VAlign(VAlign_Top)
+			.AutoHeight()
+			.Padding(10.0f, 10.0f, 10.0f, 2.0f )
 			[
-				SNew( SVerticalBox )
-				+SVerticalBox::Slot()
-				.AutoHeight()
-				.HAlign(HAlign_Fill)
-				.Padding( 10.0f, 4.0f )
-				[
-					SNew( STextBlock )
-					.Text( NSLOCTEXT("BuildProgress", "BuildStatusLabel", "Build Status") )
-					.ShadowOffset( FVector2D( 1.0f, 1.0f ) )
-				]
-				+SVerticalBox::Slot()
-				.AutoHeight()
-				.HAlign(HAlign_Fill)
-				.Padding( 10.0f, 4.0f )
-				[
-					SNew( SHorizontalBox)
-					+SHorizontalBox::Slot()
-					.AutoWidth()
-					.Padding( 0, 7.0f )
-					[
-						SNew( STextBlock )
-						.Text( this, &SBuildProgressWidget::OnGetBuildTimeText )
-						.ShadowOffset( FVector2D( 1.0f, 1.0f ) )
-					]
-					+SHorizontalBox::Slot()
-					.AutoWidth()
-					.Padding(10.0f, 7.0f, 10.0f, 7.0f)
-					[
-						SNew( STextBlock )
-						.Text( this, &SBuildProgressWidget::OnGetProgressText )
-						.ShadowOffset( FVector2D( 1.0f, 1.0f ) )
-					]
-				]
+				SNew(STextBlock)
+				.Text(NSLOCTEXT("BuildProgress", "BuildStatusLabel", "Build Status"))
+				.Font(FAppStyle::Get().GetFontStyle("HeadingExtraSmall"))
+				.ColorAndOpacity(FAppStyle::Get().GetSlateColor("Colors.White"))
 			]
-		]
-		+SVerticalBox::Slot()
-		.AutoHeight()
-		.HAlign(HAlign_Fill)
-		.Padding(10.0f, 1.0f)
-		[
-			SNew(SBorder)
+			+ SVerticalBox::Slot()
+			.VAlign(VAlign_Top)
+			.Padding(10.0f, 2.0f)
 			[
-				SNew( SVerticalBox )
-				+SVerticalBox::Slot()
-				.AutoHeight()
-				.HAlign(HAlign_Fill)
-				.Padding( 10.0f, 4.0f )
+				SNew( SHorizontalBox)
+				+SHorizontalBox::Slot()
+				.AutoWidth()
 				[
 					SNew( STextBlock )
-					.Text( NSLOCTEXT("BuildProgress", "BuildProgressLabel", "Build Progress") )
-					.ShadowOffset( FVector2D( 1.0f, 1.0f ) )
+					.Text( this, &SBuildProgressWidget::OnGetBuildTimeText )
 				]
-				+SVerticalBox::Slot()
-				.AutoHeight()
-				.HAlign(HAlign_Fill)
-				.Padding( 10.0f, 7.0f, 10.0f, 7.0f )
+				+SHorizontalBox::Slot()
+				.AutoWidth()
+				.Padding(10.0f, 0, 10.0f, 7.0f)
 				[
-					SNew(SProgressBar)
-					.Percent( this, &SBuildProgressWidget::OnGetProgressFraction )	
+					SNew( STextBlock )
+					.Text( this, &SBuildProgressWidget::OnGetProgressText )
 				]
 			]
-
-		]
-		+SVerticalBox::Slot()
-		.AutoHeight()
-		.Padding(15.0f, 4.0f)
-		.HAlign(HAlign_Center)
-		[
-			SNew(SHorizontalBox)
-			+SHorizontalBox::Slot()
-			.AutoWidth()
+			+SVerticalBox::Slot()
+			.HAlign(HAlign_Fill)
+			.VAlign(VAlign_Center)
+			.AutoHeight()
+			.Padding(10.0f, 1.0f)
+			[
+				SNew( STextBlock )
+				.Text( NSLOCTEXT("BuildProgress", "BuildProgressLabel", "Build Progress") )
+			]
+			+SVerticalBox::Slot()
+			.HAlign(HAlign_Fill)
+			.VAlign(VAlign_Center)
+			.AutoHeight()
+			.Padding( 10.0f, 7.0f, 10.0f, 7.0f )
+			[
+				SNew(SProgressBar)
+				.Percent( this, &SBuildProgressWidget::OnGetProgressFraction )	
+			]
+			+SVerticalBox::Slot()
+			.Padding(15.0f, 10.0f)
+			.HAlign(HAlign_Right)
+			.VAlign(VAlign_Bottom)
 			[
 				SNew(SButton)
+				.TextStyle(FAppStyle::Get(), "DialogButtonText")
 				.Text( NSLOCTEXT("BuildProgress", "StopBuildButtonLabel", "Stop Build") )
-				.ContentPadding( 5 )
 				.OnClicked( this, &SBuildProgressWidget::OnStopBuild )
 			]
 		]
diff --git a/Engine/Source/Editor/UnrealEd/Private/EditorServer.cpp b/Engine/Source/Editor/UnrealEd/Private/EditorServer.cpp
--- a/Engine/Source/Editor/UnrealEd/Private/EditorServer.cpp
+++ b/Engine/Source/Editor/UnrealEd/Private/EditorServer.cpp
@@ -2513,7 +2513,7 @@ bool UEditorEngine::Map_Load(const TCHAR* Str, FOutputDevice& Ar)
 				}
 				
 				// Don't show progress dialogs when loading one of our startup maps. They should load rather quickly.
-				FScopedSlowTask SlowTask(100, FText::Format( NSLOCTEXT("UnrealEd", "LoadingMapStatus_Loading", "Loading {0}"), LocalizedLoadingMap ), bShowProgress != 0);
+				FScopedSlowTask SlowTask(100, LocalizedLoadingMap, bShowProgress != 0);
 				SlowTask.MakeDialog();
 
 				SlowTask.EnterProgressFrame(10, FText::Format( NSLOCTEXT("UnrealEd", "LoadingMapStatus_CleaningUp", "{0} (Clearing existing world)"), LocalizedLoadingMap ));
diff --git a/Engine/Source/Editor/UnrealEd/Private/FeedbackContextEditor.cpp b/Engine/Source/Editor/UnrealEd/Private/FeedbackContextEditor.cpp
--- a/Engine/Source/Editor/UnrealEd/Private/FeedbackContextEditor.cpp
+++ b/Engine/Source/Editor/UnrealEd/Private/FeedbackContextEditor.cpp
@@ -51,7 +51,7 @@ public:
 		SLATE_EVENT( FReceiveUserCancelDelegate, ReceiveUserCancelDelegate)
 
 		/** The feedback scope stack that we are presenting to the user */
-		SLATE_ARGUMENT( TWeakPtr<FSlowTaskStack>, ScopeStack )
+		SLATE_ARGUMENT(TWeakPtr<FSlowTaskStack>, ScopeStack)
 
 	SLATE_END_ARGS()
 
@@ -618,7 +618,7 @@ void FFeedbackContextEditor::OnUserCancel()
 TWeakPtr<class SBuildProgressWidget> FFeedbackContextEditor::ShowBuildProgressWindow()
 {
 	TSharedRef<SWindow> BuildProgressWindowRef = SNew(SWindow)
-		.ClientSize(FVector2D(500,200))
+		.ClientSize(FVector2D(600,200))
 		.IsPopupWindow(true);
 
 	BuildProgressWidget = SNew(SBuildProgressWidget);
diff --git a/Engine/Source/Editor/UnrealEd/Public/Dialogs/SBuildProgress.h b/Engine/Source/Editor/UnrealEd/Public/Dialogs/SBuildProgress.h
--- a/Engine/Source/Editor/UnrealEd/Public/Dialogs/SBuildProgress.h
+++ b/Engine/Source/Editor/UnrealEd/Public/Dialogs/SBuildProgress.h
@@ -11,7 +11,7 @@
 #include "Widgets/DeclarativeSyntaxSupport.h"
 #include "Widgets/Layout/SBorder.h"
 
-class SBuildProgressWidget : public SBorder
+class SBuildProgressWidget : public SCompoundWidget
 {
 public:
 	SLATE_BEGIN_ARGS( SBuildProgressWidget ) {}
diff --git a/Engine/Source/Runtime/Slate/Public/Widgets/Notifications/SProgressBar.h b/Engine/Source/Runtime/Slate/Public/Widgets/Notifications/SProgressBar.h
--- a/Engine/Source/Runtime/Slate/Public/Widgets/Notifications/SProgressBar.h
+++ b/Engine/Source/Runtime/Slate/Public/Widgets/Notifications/SProgressBar.h
@@ -48,7 +48,7 @@ public:
 		, _BarFillType(EProgressBarFillType::LeftToRight)
 		, _Percent( TOptional<float>() )
 		, _FillColorAndOpacity( FLinearColor::White )
-		, _BorderPadding( FVector2D(1,0) )
+		, _BorderPadding( FVector2D(0,0) )
 		, _BackgroundImage(nullptr)
 		, _FillImage(nullptr)
 		, _MarqueeImage(nullptr)
diff --git a/Engine/Source/Runtime/SlateCore/Private/Styling/StarshipCoreStyle.cpp b/Engine/Source/Runtime/SlateCore/Private/Styling/StarshipCoreStyle.cpp
--- a/Engine/Source/Runtime/SlateCore/Private/Styling/StarshipCoreStyle.cpp
+++ b/Engine/Source/Runtime/SlateCore/Private/Styling/StarshipCoreStyle.cpp
@@ -70,6 +70,9 @@ FStyleFonts::FStyleFonts()
 	, NormalBold(FONT(10, "Bold"))
 	, Small(FONT(8,  "Regular"))
 	, SmallBold(FONT(8,  "Bold"))
+	, HeadingMedium(FONT(33, "BoldCondensed"))
+	, HeadingSmall(FONT(21, "BoldCondensed"))
+	, HeadingExtraSmall(FONT(15, "BoldCondensed"))
 {
 };
 
@@ -908,6 +911,10 @@ void FStarshipCoreStyle::SetupTextStyles(TSharedRef<FStyle>& Style)
 	Style->Set("NormalFontBold", StyleFonts.NormalBold);
 	Style->Set("SmallFontBold", StyleFonts.SmallBold);
 
+	Style->Set("HeadingMedium", StyleFonts.HeadingMedium);
+	Style->Set("HeadingSmall", StyleFonts.HeadingSmall);
+	Style->Set("HeadingExtraSmall", StyleFonts.HeadingExtraSmall);
+
 	FSlateBrush* DefaultTextUnderlineBrush = new IMAGE_BRUSH("Old/White", Icon8x8, FLinearColor::White, ESlateBrushTileType::Both);
 
 	Style->Set("DefaultTextUnderline", DefaultTextUnderlineBrush);
diff --git a/Engine/Source/Runtime/SlateCore/Public/Styling/StarshipCoreStyle.h b/Engine/Source/Runtime/SlateCore/Public/Styling/StarshipCoreStyle.h
--- a/Engine/Source/Runtime/SlateCore/Public/Styling/StarshipCoreStyle.h
+++ b/Engine/Source/Runtime/SlateCore/Public/Styling/StarshipCoreStyle.h
@@ -106,6 +106,9 @@ struct SLATECORE_API FStyleFonts
 	const FSlateFontInfo Small;
 	const FSlateFontInfo SmallBold;
 
+	const FSlateFontInfo HeadingMedium;
+	const FSlateFontInfo HeadingSmall;
+	const FSlateFontInfo HeadingExtraSmall;
 
   private: 
   	FStyleFonts();
