Report Guid Duplication in WorldPartitionConvertCommandlet

diff --git a/Engine/Source/Editor/UnrealEd/Classes/Commandlets/WorldPartitionConvertCommandlet.h b/Engine/Source/Editor/UnrealEd/Classes/Commandlets/WorldPartitionConvertCommandlet.h
--- a/Engine/Source/Editor/UnrealEd/Classes/Commandlets/WorldPartitionConvertCommandlet.h
+++ b/Engine/Source/Editor/UnrealEd/Classes/Commandlets/WorldPartitionConvertCommandlet.h
@@ -65,6 +65,8 @@ protected:
 	UHLODLayer* CreateHLODLayerFromINI(const FString& InHLODLayerName);
 	void SetupHLODLayerAssets();
 
+	void SetActorGuid(AActor* Actor, const FGuid& NewGuid);
+
 	// Conversion report
 	TSet<FString> MapsWithLevelScriptsBPs;
 	TSet<FString> MapsWithMapBuildData;
diff --git a/Engine/Source/Editor/UnrealEd/Private/Commandlets/WorldPartitionConvertCommandlet.cpp b/Engine/Source/Editor/UnrealEd/Private/Commandlets/WorldPartitionConvertCommandlet.cpp
--- a/Engine/Source/Editor/UnrealEd/Private/Commandlets/WorldPartitionConvertCommandlet.cpp
+++ b/Engine/Source/Editor/UnrealEd/Private/Commandlets/WorldPartitionConvertCommandlet.cpp
@@ -617,6 +617,11 @@ void UWorldPartitionConvertCommandlet::SetupHLODLayerAssets()
 	}
 }
 
+void UWorldPartitionConvertCommandlet::SetActorGuid(AActor* Actor, const FGuid& NewGuid)
+{
+	FSetActorGuid SetActorGuid(Actor, NewGuid);
+}
+
 void UWorldPartitionConvertCommandlet::OnWorldLoaded(UWorld* World)
 {
 	if (UWorldComposition* WorldComposition = World->WorldComposition)
@@ -1043,6 +1048,8 @@ int32 UWorldPartitionConvertCommandlet::Main(const FString& Params)
 
 	if (!bReportOnly)
 	{
+		TSet<FGuid> ActorGuids;
+		
 		for(AActor* Actor: ActorList)
 		{
 			if (!Actor || Actor->IsPendingKill() || !Actor->SupportsExternalPackaging())
@@ -1050,13 +1057,22 @@ int32 UWorldPartitionConvertCommandlet::Main(const FString& Params)
 				continue;
 			}
 
+			bool bAlreadySet = false;
+			ActorGuids.Add(Actor->GetActorGuid(), &bAlreadySet);
+			if (bAlreadySet)
+			{
+				UE_LOG(LogWorldPartitionConvertCommandlet, Error, TEXT("Duplicated guid actor %s(guid:%s) can't extract actor"), *Actor->GetName(), *Actor->GetActorGuid().ToString(EGuidFormats::Digits));
+				return 1;
+			}
+
 			if (Actor->IsPackageExternal())
 			{
 				PackagesToDelete.Add(Actor->GetPackage());
 				Actor->SetPackageExternal(false);
 			}
-
+						
 			Actor->SetPackageExternal(true);
+			
 			UPackage* ActorPackage = Actor->GetExternalPackage();
 			PackagesToSave.Add(ActorPackage);
 
diff --git a/Engine/Source/Runtime/Engine/Classes/GameFramework/Actor.h b/Engine/Source/Runtime/Engine/Classes/GameFramework/Actor.h
--- a/Engine/Source/Runtime/Engine/Classes/GameFramework/Actor.h
+++ b/Engine/Source/Runtime/Engine/Classes/GameFramework/Actor.h
@@ -3639,7 +3639,8 @@ private:
 	{
 		InActor->ActorGuid = InActorGuid;
 	}
-	friend UWorld;
+	friend class UWorld;
+	friend class UWorldPartitionConvertCommandlet;
 };
 #endif
 
