Added an animation editor option to force opening animation assets inside a new tab.
This is disabled on default.

diff --git a/Engine/Source/Developer/AssetTools/Private/AssetTypeActions/AssetTypeActions_AnimationAsset.cpp b/Engine/Source/Developer/AssetTools/Private/AssetTypeActions/AssetTypeActions_AnimationAsset.cpp
--- a/Engine/Source/Developer/AssetTools/Private/AssetTypeActions/AssetTypeActions_AnimationAsset.cpp
+++ b/Engine/Source/Developer/AssetTools/Private/AssetTypeActions/AssetTypeActions_AnimationAsset.cpp
@@ -116,7 +116,14 @@ void FAssetTypeActions_AnimationAsset::OpenAnimAssetEditor(const TArray<UObject*
 {
 	EToolkitMode::Type Mode = EditWithinLevelEditor.IsValid() ? EToolkitMode::WorldCentric : EToolkitMode::Standalone;
 
-	// Force new window if control held down
+#if WITH_EDITOR
+	// Force opening a new tab when the Animation Editor options window has the 'Always Open In New Tab' option enabled.
+	const UPersonaOptions* PersonaOptions = GetDefault<UPersonaOptions>();
+	check(PersonaOptions);
+	bForceNewEditor |= PersonaOptions->bAlwaysOpenAnimationAssetsInNewTab;
+#endif
+
+	// Force new tab when shift is held down.
 	bForceNewEditor |= FSlateApplication::Get().GetModifierKeys().IsShiftDown();
 
 	// Find all the anim assets
diff --git a/Engine/Source/Editor/UnrealEd/Classes/Preferences/PersonaOptions.h b/Engine/Source/Editor/UnrealEd/Classes/Preferences/PersonaOptions.h
--- a/Engine/Source/Editor/UnrealEd/Classes/Preferences/PersonaOptions.h
+++ b/Engine/Source/Editor/UnrealEd/Classes/Preferences/PersonaOptions.h
@@ -96,11 +96,17 @@ UCLASS(hidecategories=Object, config=EditorPerProjectUserSettings)
 class UNREALED_API UPersonaOptions : public UObject
 {
 	GENERATED_UCLASS_BODY()
-		
+
 	/** Whether or not the floor should be aligned to the Skeletal Mesh's bounds by default for the Animation Editor(s)*/
 	UPROPERTY(EditAnywhere, config, Category = "Preview Scene")
 	uint32 bAutoAlignFloorToMesh : 1;
 
+	/** Whether or not the Animation Editor opens in an additional tab when double clicking an animation asset or if it reuses an already existing Animation Editor tab.
+	  * You can also keep this disabled and hold shift pressed while double clicking the asset to open the asset inside its own tab.
+	  */
+	UPROPERTY(EditAnywhere, config, Category = "Assets")
+	uint32 bAlwaysOpenAnimationAssetsInNewTab : 1;
+
 	/** Whether or not the grid should be visible by default for the Animation Editor(s)*/
 	UPROPERTY(EditAnywhere, config, Category = "Viewport")
 	uint32 bShowGrid:1;
