WorldPartition: Support Landscape in Convert Commandlet

diff --git a/Engine/Source/Editor/UnrealEd/Classes/Commandlets/WorldPartitionConvertCommandlet.h b/Engine/Source/Editor/UnrealEd/Classes/Commandlets/WorldPartitionConvertCommandlet.h
--- a/Engine/Source/Editor/UnrealEd/Classes/Commandlets/WorldPartitionConvertCommandlet.h
+++ b/Engine/Source/Editor/UnrealEd/Classes/Commandlets/WorldPartitionConvertCommandlet.h
@@ -107,6 +107,9 @@ protected:
 	UPROPERTY(Config)
 	TArray<FHLODLayerActorMapping> HLODLayersForActorClasses;
 
+	UPROPERTY(Config)
+	uint32 LandscapeGridSize;
+
 	UPROPERTY(Transient)
 	TMap<FString, UHLODLayer*> HLODLayers;
 };
diff --git a/Engine/Source/Editor/UnrealEd/Private/Commandlets/WorldPartitionConvertCommandlet.cpp b/Engine/Source/Editor/UnrealEd/Private/Commandlets/WorldPartitionConvertCommandlet.cpp
--- a/Engine/Source/Editor/UnrealEd/Private/Commandlets/WorldPartitionConvertCommandlet.cpp
+++ b/Engine/Source/Editor/UnrealEd/Private/Commandlets/WorldPartitionConvertCommandlet.cpp
@@ -37,6 +37,10 @@
 #include "Engine/WorldComposition.h"
 #include "ActorPartition/ActorPartitionSubsystem.h"
 #include "InstancedFoliage.h"
+#include "Landscape.h"
+#include "LandscapeStreamingProxy.h"
+#include "LandscapeInfo.h"
+#include "LandscapeConfigHelper.h"
 
 DEFINE_LOG_CATEGORY(LogWorldPartitionConvertCommandlet);
 
@@ -140,6 +144,7 @@ UWorldPartitionConvertCommandlet::UWorldPartitionConvertCommandlet(const FObject
 : Super(ObjectInitializer)
 , bConversionSuffix(false)
 , ConversionSuffix(TEXT("_WP"))
+, LandscapeGridSize(4)
 {}
 
 UWorld* UWorldPartitionConvertCommandlet::LoadWorld(const FString& LevelToLoad)
@@ -821,9 +826,35 @@ int32 UWorldPartitionConvertCommandlet::Main(const FString& Params)
 		IFA->GetLevel()->GetWorld()->DestroyActor(IFA);
 	};
 
-	auto PrepareLevelActors = [this, PartitionFoliage](ULevel* Level, bool bMainLevel, EActorGridPlacement DefaultGridPlacement)
+	auto PartitionLandscape = [this, MainWorld](ULandscapeInfo* LandscapeInfo)
+	{
+		// Handle Landscapes with missing LandscapeActor(s)
+		if (!LandscapeInfo->LandscapeActor.Get())
+		{
+			// Use the first proxy as the landscape template
+			ALandscapeProxy* FirstProxy = LandscapeInfo->Proxies[0];
+
+			FActorSpawnParameters SpawnParams;
+			FTransform LandscapeTransform = FirstProxy->LandscapeActorToWorld();
+			ALandscape* NewLandscape = MainWorld->SpawnActor<ALandscape>(ALandscape::StaticClass(), LandscapeTransform, SpawnParams);
+			
+			NewLandscape->GetSharedProperties(FirstProxy);
+
+			LandscapeInfo->RegisterActor(NewLandscape);
+		}
+
+		TSet<AActor*> ActorsToDelete;
+		FLandscapeConfigHelper::ChangeGridSize(LandscapeInfo, LandscapeGridSize, ActorsToDelete);
+		for (AActor* ActorToDelete : ActorsToDelete)
+		{
+			MainWorld->DestroyActor(ActorToDelete);
+		}
+	};
+
+	auto PrepareLevelActors = [this, PartitionFoliage, PartitionLandscape](ULevel* Level, bool bMainLevel, EActorGridPlacement DefaultGridPlacement)
 	{
 		TArray<AInstancedFoliageActor*> IFAs;
+		TSet<ULandscapeInfo*> LandscapeInfos;
 		for (AActor* Actor: Level->Actors)
 		{
 			if (Actor && !Actor->IsPendingKill())
@@ -836,6 +867,12 @@ int32 UWorldPartitionConvertCommandlet::Main(const FString& Params)
 				{
 					IFAs.Add(IFA);
 				}
+				else if (ALandscapeProxy* LandscapeProxy = Cast<ALandscapeProxy>(Actor))
+				{
+					ULandscapeInfo* LandscapeInfo = LandscapeProxy->GetLandscapeInfo();
+					check(LandscapeInfo);
+					LandscapeInfos.Add(LandscapeInfo);
+				}
 				// Only override default grid placement on actors that are not marked as always loaded
 				else if (Actor->GridPlacement != EActorGridPlacement::AlwaysLoaded)
 				{
@@ -852,6 +889,11 @@ int32 UWorldPartitionConvertCommandlet::Main(const FString& Params)
 		{
 			PartitionFoliage(IFA);
 		}
+
+		for (ULandscapeInfo* LandscapeInfo : LandscapeInfos)
+		{
+			PartitionLandscape(LandscapeInfo);
+		}
 	};
 
 	// Gather and load sublevels
