Fix an issue where relative file path is not correctly converted, causing FilenameToLongPackageName() to crash while generating cooker open order

diff --git a/Engine/Source/Editor/UnrealEd/Private/Commandlets/AssetRegistryGenerator.cpp b/Engine/Source/Editor/UnrealEd/Private/Commandlets/AssetRegistryGenerator.cpp
--- a/Engine/Source/Editor/UnrealEd/Private/Commandlets/AssetRegistryGenerator.cpp
+++ b/Engine/Source/Editor/UnrealEd/Private/Commandlets/AssetRegistryGenerator.cpp
@@ -1401,17 +1401,18 @@ bool FAssetRegistryGenerator::SaveAssetRegistry(const FString& SandboxPath, bool
 
 class FPackageCookerOpenOrderVisitor : public IPlatformFile::FDirectoryVisitor
 {
-	const FString& SandboxRootPath;
-	const FString& RelativePathToRoot;
+	const FSandboxPlatformFile* SandboxFile;
+	const FString& PlatformSandboxPath;
 	const TSet<FString>& ValidExtensions;
 	TMultiMap<FString, FString>& PackageExtensions;
 public:
-	FPackageCookerOpenOrderVisitor(const FString& InSandboxRootPath,
-		const FString& InRelativePathToRoot,
+	FPackageCookerOpenOrderVisitor(
+		const FSandboxPlatformFile* InSandboxFile,
+		const FString& InPlatformSandboxPath,
 		const TSet<FString>& InValidExtensions,
 		TMultiMap<FString, FString>& OutPackageExtensions) :
-		SandboxRootPath(InSandboxRootPath),
-		RelativePathToRoot(InRelativePathToRoot),
+		SandboxFile(InSandboxFile),
+		PlatformSandboxPath(InPlatformSandboxPath),
 		ValidExtensions(InValidExtensions),
 		PackageExtensions(OutPackageExtensions)
 	{}
@@ -1426,10 +1427,12 @@ public:
 
 		if (ValidExtensions.Contains(FileExtension))
 		{
-			Filename.ReplaceInline(*SandboxRootPath, *RelativePathToRoot);
-			FString PackageName = FPackageName::FilenameToLongPackageName(Filename);
+			Filename.ReplaceInline(*PlatformSandboxPath, *SandboxFile->GetSandboxDirectory());
+			FString AssetSourcePath = SandboxFile->ConvertFromSandboxPath(*Filename);
+			FString StandardAssetSourcePath = FPaths::CreateStandardFilename(AssetSourcePath);
+			FString PackageName = FPackageName::FilenameToLongPackageName(StandardAssetSourcePath);
 
-			PackageExtensions.AddUnique(PackageName, Filename);
+			PackageExtensions.AddUnique(PackageName, StandardAssetSourcePath);
 		}
 
 		return true;
@@ -1516,14 +1519,13 @@ bool FAssetRegistryGenerator::WriteCookerOpenOrder(FSandboxPlatformFile* InSandb
 		};
 		const TSet<FString> ValidExtensionSet(ValidExtensions);
 
-		FString SandboxPath = InSandboxFile->GetSandboxDirectory();
+		const FString SandboxPath = InSandboxFile->GetSandboxDirectory();
 		const FString Platform = TargetPlatform->PlatformName();
-		SandboxPath.ReplaceInline(TEXT("[Platform]"), *Platform);
+		FString PlatformSandboxPath = SandboxPath.Replace(TEXT("[Platform]"), *Platform);
 
-		const FString RelativePathToRoot = FPaths::GetRelativePathToRoot();
 		TMultiMap<FString, FString> CookedPackageFilesMap;
-		FPackageCookerOpenOrderVisitor PackageSearch(SandboxPath, RelativePathToRoot, ValidExtensionSet, CookedPackageFilesMap);
-		IFileManager::Get().IterateDirectoryRecursively(*SandboxPath, PackageSearch);
+		FPackageCookerOpenOrderVisitor PackageSearch(InSandboxFile, PlatformSandboxPath, ValidExtensionSet, CookedPackageFilesMap);
+		IFileManager::Get().IterateDirectoryRecursively(*PlatformSandboxPath, PackageSearch);
 
 		int32 CurrentIndex = 0;
 		for (FName PackageName : FileOrder)
