Landscape: Make Gizmo actor Transient (doesn't make sense to persist this)

diff --git a/Engine/Source/Editor/UnrealEd/Private/Commandlets/WorldPartitionConvertCommandlet.cpp b/Engine/Source/Editor/UnrealEd/Private/Commandlets/WorldPartitionConvertCommandlet.cpp
--- a/Engine/Source/Editor/UnrealEd/Private/Commandlets/WorldPartitionConvertCommandlet.cpp
+++ b/Engine/Source/Editor/UnrealEd/Private/Commandlets/WorldPartitionConvertCommandlet.cpp
@@ -41,6 +41,7 @@
 #include "LandscapeStreamingProxy.h"
 #include "LandscapeInfo.h"
 #include "LandscapeConfigHelper.h"
+#include "LandscapeGizmoActor.h"
 
 DEFINE_LOG_CATEGORY(LogWorldPartitionConvertCommandlet);
 
@@ -307,7 +308,8 @@ bool UWorldPartitionConvertCommandlet::ShouldDeleteActor(AActor* Actor, bool bMa
 {
 	if (Actor->HasAllFlags(RF_Transient) ||
 		Actor->IsA<ALODActor>() ||
-		Actor->IsA<ALevelBounds>())
+		Actor->IsA<ALevelBounds>() ||
+		Actor->IsA<ALandscapeGizmoActor>())
 	{
 		return true;
 	}
diff --git a/Engine/Source/Runtime/Landscape/Classes/LandscapeGizmoActiveActor.h b/Engine/Source/Runtime/Landscape/Classes/LandscapeGizmoActiveActor.h
--- a/Engine/Source/Runtime/Landscape/Classes/LandscapeGizmoActiveActor.h
+++ b/Engine/Source/Runtime/Landscape/Classes/LandscapeGizmoActiveActor.h
@@ -50,7 +50,7 @@ struct FGizmoSelectData
 	
 };
 
-UCLASS(notplaceable, MinimalAPI)
+UCLASS(notplaceable, transient, MinimalAPI)
 class ALandscapeGizmoActiveActor : public ALandscapeGizmoActor
 {
 	GENERATED_UCLASS_BODY()
diff --git a/Engine/Source/Runtime/Landscape/Classes/LandscapeGizmoActor.h b/Engine/Source/Runtime/Landscape/Classes/LandscapeGizmoActor.h
--- a/Engine/Source/Runtime/Landscape/Classes/LandscapeGizmoActor.h
+++ b/Engine/Source/Runtime/Landscape/Classes/LandscapeGizmoActor.h
@@ -10,7 +10,7 @@
 
 class UBillboardComponent;
 
-UCLASS(notplaceable, MinimalAPI, NotBlueprintable)
+UCLASS(notplaceable, transient, MinimalAPI, NotBlueprintable)
 class ALandscapeGizmoActor : public AActor
 {
 	GENERATED_UCLASS_BODY()
