WorldPartitionConvertCommandlet
- Added optional world origin/extent that will automatically set actors outside as always loaded (defaults to HALF_WORLD_MAX).
- Always write ini file when `GenerateIni` commandline option is specified, even if the file already exist.
- Log the time spent converting Foliage and Landscape actors.

diff --git a/Engine/Source/Editor/UnrealEd/Classes/Commandlets/WorldPartitionConvertCommandlet.h b/Engine/Source/Editor/UnrealEd/Classes/Commandlets/WorldPartitionConvertCommandlet.h
--- a/Engine/Source/Editor/UnrealEd/Classes/Commandlets/WorldPartitionConvertCommandlet.h
+++ b/Engine/Source/Editor/UnrealEd/Classes/Commandlets/WorldPartitionConvertCommandlet.h
@@ -98,6 +98,12 @@ protected:
 	UPROPERTY(Config)
 	TMap<FName, EActorGridPlacement> LevelsGridPlacement;
 
+	UPROPERTY(Config)
+	FVector WorldOrigin;
+	
+	UPROPERTY(Config)
+	FVector WorldExtent;
+
 	UPROPERTY(Config)
 	FString HLODLayerAssetsPath;
 
diff --git a/Engine/Source/Editor/UnrealEd/Private/Commandlets/WorldPartitionConvertCommandlet.cpp b/Engine/Source/Editor/UnrealEd/Private/Commandlets/WorldPartitionConvertCommandlet.cpp
--- a/Engine/Source/Editor/UnrealEd/Private/Commandlets/WorldPartitionConvertCommandlet.cpp
+++ b/Engine/Source/Editor/UnrealEd/Private/Commandlets/WorldPartitionConvertCommandlet.cpp
@@ -142,10 +142,12 @@ public:
 };
 
 UWorldPartitionConvertCommandlet::UWorldPartitionConvertCommandlet(const FObjectInitializer& ObjectInitializer)
-: Super(ObjectInitializer)
-, bConversionSuffix(false)
-, ConversionSuffix(TEXT("_WP"))
-, LandscapeGridSize(4)
+	: Super(ObjectInitializer)
+	, bConversionSuffix(false)
+	, ConversionSuffix(TEXT("_WP"))
+	, WorldOrigin(FVector::ZeroVector)
+	, WorldExtent(HALF_WORLD_MAX)
+	, LandscapeGridSize(4)
 {}
 
 UWorld* UWorldPartitionConvertCommandlet::LoadWorld(const FString& LevelToLoad)
@@ -855,6 +857,8 @@ int32 UWorldPartitionConvertCommandlet::Main(const FString& Params)
 
 	auto PrepareLevelActors = [this, PartitionFoliage, PartitionLandscape](ULevel* Level, bool bMainLevel, EActorGridPlacement DefaultGridPlacement)
 	{
+		const FBox WorldBounds(WorldOrigin - WorldExtent, WorldOrigin + WorldExtent);
+
 		TArray<AInstancedFoliageActor*> IFAs;
 		TSet<ULandscapeInfo*> LandscapeInfos;
 		for (AActor* Actor: Level->Actors)
@@ -878,7 +882,20 @@ int32 UWorldPartitionConvertCommandlet::Main(const FString& Params)
 				// Only override default grid placement on actors that are not marked as always loaded
 				else if (Actor->GridPlacement != EActorGridPlacement::AlwaysLoaded)
 				{
-					Actor->GridPlacement = DefaultGridPlacement;
+					FVector BoundsLocation;
+					FVector BoundsExtent;
+					Actor->GetActorLocationBounds(/*bOnlyCollidingComponents*/false, BoundsLocation, BoundsExtent, /*bIncludeFromChildActors*/true);
+
+					const FBox ActorBounds(BoundsLocation - BoundsExtent, BoundsLocation + BoundsExtent);
+
+					if (!WorldBounds.IsInside(ActorBounds))
+					{
+						Actor->GridPlacement = EActorGridPlacement::AlwaysLoaded;
+					}
+					else
+					{
+						Actor->GridPlacement = DefaultGridPlacement;
+					}
 				}
 
 				// Clear actor layers as they are not supported yet in world partition
@@ -887,14 +904,24 @@ int32 UWorldPartitionConvertCommandlet::Main(const FString& Params)
 		}
 
 		// do loop after as it may modify Level->Actors
-		for (AInstancedFoliageActor* IFA : IFAs)
+		if (IFAs.Num())
 		{
-			PartitionFoliage(IFA);
+			UE_SCOPED_TIMER(TEXT("PartitionFoliage"), LogWorldPartitionConvertCommandlet, Display);
+
+			for (AInstancedFoliageActor* IFA : IFAs)
+			{
+				PartitionFoliage(IFA);
+			}
 		}
 
-		for (ULandscapeInfo* LandscapeInfo : LandscapeInfos)
+		if (LandscapeInfos.Num())
 		{
-			PartitionLandscape(LandscapeInfo);
+			UE_SCOPED_TIMER(TEXT("PartitionLandscape"), LogWorldPartitionConvertCommandlet, Display);
+
+			for (ULandscapeInfo* LandscapeInfo : LandscapeInfos)
+			{
+				PartitionLandscape(LandscapeInfo);
+			}
 		}
 	};
 
@@ -1206,7 +1233,7 @@ int32 UWorldPartitionConvertCommandlet::Main(const FString& Params)
 
 	if (bGenerateIni || !bReportOnly)
 	{
-		if (!FPlatformFileManager::Get().GetPlatformFile().FileExists(*LevelConfigFilename))
+		if (bGenerateIni || !FPlatformFileManager::Get().GetPlatformFile().FileExists(*LevelConfigFilename))
 		{
 			SaveConfig(CPF_Config, *LevelConfigFilename);
 			WorldPartition->EditorHash->SaveConfig(CPF_Config, *LevelConfigFilename);
