Updated Menu SpinBox Style

diff --git a/Engine/Source/Editor/LevelEditor/Private/SLevelViewportToolBar.cpp b/Engine/Source/Editor/LevelEditor/Private/SLevelViewportToolBar.cpp
--- a/Engine/Source/Editor/LevelEditor/Private/SLevelViewportToolBar.cpp
+++ b/Engine/Source/Editor/LevelEditor/Private/SLevelViewportToolBar.cpp
@@ -1247,12 +1247,18 @@ TSharedRef<SWidget> SLevelViewportToolBar::GenerateFOVMenu() const
 			.Padding( FMargin(4.0f, 0.0f, 0.0f, 0.0f) )
 			.WidthOverride( 100.0f )
 			[
-				SNew(SSpinBox<float>)
-				.Font( FEditorStyle::GetFontStyle( TEXT( "MenuItem.Font" ) ) )
-				.MinValue(FOVMin)
-				.MaxValue(FOVMax)
-				.Value( this, &SLevelViewportToolBar::OnGetFOVValue )
-				.OnValueChanged( const_cast<SLevelViewportToolBar*>(this), &SLevelViewportToolBar::OnFOVValueChanged )
+				SNew ( SBorder )
+				.BorderImage(FAppStyle::Get().GetBrush("Menu.WidgetBorder"))
+				.Padding(FMargin(1.0f))
+				[
+					SNew(SSpinBox<float>)
+					.Style(&FAppStyle::Get(), "Menu.SpinBox")
+					.Font( FEditorStyle::GetFontStyle( TEXT( "MenuItem.Font" ) ) )
+					.MinValue(FOVMin)
+					.MaxValue(FOVMax)
+					.Value( this, &SLevelViewportToolBar::OnGetFOVValue )
+					.OnValueChanged( const_cast<SLevelViewportToolBar*>(this), &SLevelViewportToolBar::OnFOVValueChanged )
+				]
 			]
 		];
 }
@@ -1297,15 +1303,21 @@ TSharedRef<SWidget> SLevelViewportToolBar::GenerateScreenPercentageMenu() const
 		[
 			SNew(SBox)
 			.Padding(FMargin(4.0f, 0.0f, 0.0f, 0.0f))
-		.WidthOverride(100.0f)
-		[
-			SNew(SSpinBox<int32>)
-			.Font(FEditorStyle::GetFontStyle(TEXT("MenuItem.Font")))
-		.MinValue(PreviewScreenPercentageMin)
-		.MaxValue(PreviewScreenPercentageMax)
-		.Value(this, &SLevelViewportToolBar::OnGetScreenPercentageValue)
-		.OnValueChanged(const_cast<SLevelViewportToolBar*>(this), &SLevelViewportToolBar::OnScreenPercentageValueChanged)
-		]
+			.WidthOverride(100.0f)
+			[
+				SNew ( SBorder )
+				.BorderImage(FAppStyle::Get().GetBrush("Menu.WidgetBorder"))
+				.Padding(FMargin(1.0f))
+				[
+					SNew(SSpinBox<int32>)
+					.Style(&FAppStyle::Get(), "Menu.SpinBox")
+					.Font(FEditorStyle::GetFontStyle(TEXT("MenuItem.Font")))
+					.MinValue(PreviewScreenPercentageMin)
+					.MaxValue(PreviewScreenPercentageMax)
+					.Value(this, &SLevelViewportToolBar::OnGetScreenPercentageValue)
+					.OnValueChanged(const_cast<SLevelViewportToolBar*>(this), &SLevelViewportToolBar::OnScreenPercentageValueChanged)
+				]
+			]
 		];
 }
 
@@ -1336,13 +1348,19 @@ TSharedRef<SWidget> SLevelViewportToolBar::GenerateFarViewPlaneMenu() const
 			.Padding(FMargin(4.0f, 0.0f, 0.0f, 0.0f))
 			.WidthOverride(100.0f)
 			[
-				SNew(SSpinBox<float>)
-				.ToolTipText(LOCTEXT("FarViewPlaneTooltip", "Distance to use as the far view plane, or zero to enable an infinite far view plane"))
-				.MinValue(0.0f)
-				.MaxValue(100000.0f)
-				.Font(FEditorStyle::GetFontStyle(TEXT("MenuItem.Font")))
-				.Value(this, &SLevelViewportToolBar::OnGetFarViewPlaneValue)
-				.OnValueChanged(const_cast<SLevelViewportToolBar*>(this), &SLevelViewportToolBar::OnFarViewPlaneValueChanged)
+				SNew ( SBorder )
+				.BorderImage(FAppStyle::Get().GetBrush("Menu.WidgetBorder"))
+				.Padding(FMargin(1.0f))
+				[
+					SNew(SSpinBox<float>)
+					.Style(&FAppStyle::Get(), "Menu.SpinBox")
+					.ToolTipText(LOCTEXT("FarViewPlaneTooltip", "Distance to use as the far view plane, or zero to enable an infinite far view plane"))
+					.MinValue(0.0f)
+					.MaxValue(100000.0f)
+					.Font(FEditorStyle::GetFontStyle(TEXT("MenuItem.Font")))
+					.Value(this, &SLevelViewportToolBar::OnGetFarViewPlaneValue)
+					.OnValueChanged(const_cast<SLevelViewportToolBar*>(this), &SLevelViewportToolBar::OnFarViewPlaneValueChanged)
+				]
 			]
 		];
 }
diff --git a/Engine/Source/Editor/UnrealEd/Private/SCommonEditorViewportToolbarBase.cpp b/Engine/Source/Editor/UnrealEd/Private/SCommonEditorViewportToolbarBase.cpp
--- a/Engine/Source/Editor/UnrealEd/Private/SCommonEditorViewportToolbarBase.cpp
+++ b/Engine/Source/Editor/UnrealEd/Private/SCommonEditorViewportToolbarBase.cpp
@@ -357,12 +357,18 @@ TSharedRef<SWidget> SCommonEditorViewportToolbarBase::GenerateFOVMenu() const
 			.Padding( FMargin(4.0f, 0.0f, 0.0f, 0.0f) )
 			.WidthOverride( 100.0f )
 			[
-				SNew(SSpinBox<float>)
-				.Font( FEditorStyle::GetFontStyle( TEXT( "MenuItem.Font" ) ) )
-				.MinValue(FOVMin)
-				.MaxValue(FOVMax)
-				.Value(this, &SCommonEditorViewportToolbarBase::OnGetFOVValue)
-				.OnValueChanged(this, &SCommonEditorViewportToolbarBase::OnFOVValueChanged)
+				SNew ( SBorder )
+				.BorderImage(FAppStyle::Get().GetBrush("Menu.WidgetBorder"))
+				.Padding(FMargin(1.0f))
+				[
+					SNew(SSpinBox<float>)
+					.Style(&FAppStyle::Get(), "Menu.SpinBox")
+					.Font( FEditorStyle::GetFontStyle( TEXT( "MenuItem.Font" ) ) )
+					.MinValue(FOVMin)
+					.MaxValue(FOVMax)
+					.Value(this, &SCommonEditorViewportToolbarBase::OnGetFOVValue)
+					.OnValueChanged(this, &SCommonEditorViewportToolbarBase::OnFOVValueChanged)
+				]
 			]
 		];
 }
@@ -392,15 +398,21 @@ TSharedRef<SWidget> SCommonEditorViewportToolbarBase::GenerateScreenPercentageMe
 		[
 			SNew(SBox)
 			.Padding(FMargin(4.0f, 0.0f, 0.0f, 0.0f))
-		.WidthOverride(100.0f)
-		[
-			SNew(SSpinBox<int32>)
-			.Font(FEditorStyle::GetFontStyle(TEXT("MenuItem.Font")))
-		.MinValue(PreviewScreenPercentageMin)
-		.MaxValue(PreviewScreenPercentageMax)
-		.Value(this, &SCommonEditorViewportToolbarBase::OnGetScreenPercentageValue)
-		.OnValueChanged(const_cast<SCommonEditorViewportToolbarBase*>(this), &SCommonEditorViewportToolbarBase::OnScreenPercentageValueChanged)
-		]
+			.WidthOverride(100.0f)
+			[
+				SNew ( SBorder )
+				.BorderImage(FAppStyle::Get().GetBrush("Menu.WidgetBorder"))
+				.Padding(FMargin(1.0f))
+				[
+					SNew(SSpinBox<int32>)
+					.Style(&FAppStyle::Get(), "Menu.SpinBox")
+					.Font(FEditorStyle::GetFontStyle(TEXT("MenuItem.Font")))
+					.MinValue(PreviewScreenPercentageMin)
+					.MaxValue(PreviewScreenPercentageMax)
+					.Value(this, &SCommonEditorViewportToolbarBase::OnGetScreenPercentageValue)
+					.OnValueChanged(const_cast<SCommonEditorViewportToolbarBase*>(this), &SCommonEditorViewportToolbarBase::OnScreenPercentageValueChanged)
+				]
+			]
 		];
 }
 
@@ -432,13 +444,19 @@ TSharedRef<SWidget> SCommonEditorViewportToolbarBase::GenerateFarViewPlaneMenu()
 			.Padding(FMargin(4.0f, 0.0f, 0.0f, 0.0f))
 			.WidthOverride(100.0f)
 			[
-				SNew(SSpinBox<float>)
-				.ToolTipText(LOCTEXT("FarViewPlaneTooltip", "Distance to use as the far view plane, or zero to enable an infinite far view plane"))
-				.MinValue(0.0f)
-				.MaxValue(100000.0f)
-				.Font(FEditorStyle::GetFontStyle(TEXT("MenuItem.Font")))
-				.Value(this, &SCommonEditorViewportToolbarBase::OnGetFarViewPlaneValue)
-				.OnValueChanged(const_cast<SCommonEditorViewportToolbarBase*>(this), &SCommonEditorViewportToolbarBase::OnFarViewPlaneValueChanged)
+				SNew ( SBorder )
+				.BorderImage(FAppStyle::Get().GetBrush("Menu.WidgetBorder"))
+				.Padding(FMargin(1.0f))
+				[
+					SNew(SSpinBox<float>)
+					.Style(&FAppStyle::Get(), "Menu.SpinBox")
+					.ToolTipText(LOCTEXT("FarViewPlaneTooltip", "Distance to use as the far view plane, or zero to enable an infinite far view plane"))
+					.MinValue(0.0f)
+					.MaxValue(100000.0f)
+					.Font(FEditorStyle::GetFontStyle(TEXT("MenuItem.Font")))
+					.Value(this, &SCommonEditorViewportToolbarBase::OnGetFarViewPlaneValue)
+					.OnValueChanged(const_cast<SCommonEditorViewportToolbarBase*>(this), &SCommonEditorViewportToolbarBase::OnFarViewPlaneValueChanged)
+				]
 			]
 		];
 }
diff --git a/Engine/Source/Editor/UnrealEd/Private/SEditorViewport.cpp b/Engine/Source/Editor/UnrealEd/Private/SEditorViewport.cpp
--- a/Engine/Source/Editor/UnrealEd/Private/SEditorViewport.cpp
+++ b/Engine/Source/Editor/UnrealEd/Private/SEditorViewport.cpp
@@ -591,14 +591,20 @@ TSharedRef<SWidget> SEditorViewport::BuildFixedEV100Menu()  const
 			.Padding( FMargin(0.0f, 0.0f, 0.0f, 0.0f) )
 			.WidthOverride( 100.0f )
 			[
-				SNew(SSpinBox<float>)
-				.Font( FEditorStyle::GetFontStyle( TEXT( "MenuItem.Font" ) ) )
-				.MinValue(EV100Min)
-				.MaxValue(EV100Max)
-				.Value( this, &SEditorViewport::OnGetFixedEV100Value )
-				.OnValueChanged( const_cast<SEditorViewport*>(this), &SEditorViewport::OnFixedEV100ValueChanged )
-				.ToolTipText(LOCTEXT( "EV100ToolTip", "Sets the exposure value of the camera using the specified EV100. Exposure = 1 / (1.2 * 2^EV100)"))
-				.IsEnabled( this, &SEditorViewport::IsFixedEV100Enabled )
+				SNew ( SBorder )
+				.BorderImage(FAppStyle::Get().GetBrush("Menu.WidgetBorder"))
+				.Padding(FMargin(1.0f))
+				[
+					SNew(SSpinBox<float>)
+					.Style(&FAppStyle::Get(), "Menu.SpinBox")
+					.Font( FEditorStyle::GetFontStyle( TEXT( "MenuItem.Font" ) ) )
+					.MinValue(EV100Min)
+					.MaxValue(EV100Max)
+					.Value( this, &SEditorViewport::OnGetFixedEV100Value )
+					.OnValueChanged( const_cast<SEditorViewport*>(this), &SEditorViewport::OnFixedEV100ValueChanged )
+					.ToolTipText(LOCTEXT( "EV100ToolTip", "Sets the exposure value of the camera using the specified EV100. Exposure = 1 / (1.2 * 2^EV100)"))
+					.IsEnabled( this, &SEditorViewport::IsFixedEV100Enabled )
+				]
 			]
 		];
 };
diff --git a/Engine/Source/Runtime/SlateCore/Private/Styling/StarshipCoreStyle.cpp b/Engine/Source/Runtime/SlateCore/Private/Styling/StarshipCoreStyle.cpp
--- a/Engine/Source/Runtime/SlateCore/Private/Styling/StarshipCoreStyle.cpp
+++ b/Engine/Source/Runtime/SlateCore/Private/Styling/StarshipCoreStyle.cpp
@@ -1800,6 +1800,19 @@ void FStarshipCoreStyle::SetupMultiboxStyles(TSharedRef<FStyle>& Style)
 
 	// MenuBar
 	{
+		Style->Set("Menu.WidgetBorder", new FSlateRoundedBoxBrush(FStyleColors::Input, 5.0f));
+		Style->Set("Menu.SpinBox", FSpinBoxStyle()
+			.SetBackgroundBrush(FSlateRoundedBoxBrush(FStyleColors::Input, InputFocusRadius, FStyleColors::Secondary, InputFocusThickness))
+			.SetHoveredBackgroundBrush(FSlateRoundedBoxBrush(FStyleColors::Input, InputFocusRadius, FStyleColors::Hover, InputFocusThickness))
+
+			.SetActiveFillBrush(FSlateRoundedBoxBrush(FStyleColors::Hover, InputFocusRadius, FLinearColor::Transparent, InputFocusThickness))
+			.SetInactiveFillBrush(FSlateRoundedBoxBrush(FStyleColors::Secondary, InputFocusRadius, FLinearColor::Transparent, InputFocusThickness))
+			.SetArrowsImage(FSlateNoResource())
+			.SetForegroundColor(FStyleColors::ForegroundHover)
+			.SetTextPadding(FMargin(10.f, 3.5f, 10.f, 4.f))
+		);
+
+
 		Style->Set("Menu.Background", new FSlateColorBrush(FStyleColors::Dropdown));
 		Style->Set("Menu.Outline", new BORDER_BRUSH("Common/Window/WindowOutline", FMargin(1.0f / 32.0f), WindowHighlight));
 		Style->Set("Menu.Icon", new IMAGE_BRUSH("Icons/icon_tab_toolbar_16px", Icon16x16));
