Fix unbalanced macros as a result of CL# 14517515

diff --git a/Engine/Source/Editor/UnrealEd/Classes/Editor/EditorEngine.h b/Engine/Source/Editor/UnrealEd/Classes/Editor/EditorEngine.h
--- a/Engine/Source/Editor/UnrealEd/Classes/Editor/EditorEngine.h
+++ b/Engine/Source/Editor/UnrealEd/Classes/Editor/EditorEngine.h
@@ -1822,9 +1822,7 @@ public:
 	 * @param	bInSimulateInEditor	True to start an in-editor simulation session, or false to start a play-in-editor session
 	 */
 	UE_DEPRECATED(4.25, "Override StartPlayInEditorSession(FRequestPlaySessionParams& InRequestParams) instead. InWorld is GetEditorWorldContext().World(), and the other arguments come from InRequestParams now.") 
-	PRAGMA_DISABLE_DEPRECATION_WARNINGS
 	virtual void PlayInEditor( UWorld* InWorld, bool bInSimulateInEditor, FPlayInEditorOverrides Overrides = FPlayInEditorOverrides());
-	PRAGMA_ENABLE_DEPRECATION_WARNINGS
 
 	UE_DEPRECATED(4.25, "Override CreateInnerProcessPIEGameInstance instead.")
 	virtual UGameInstance* CreatePIEGameInstance(int32 InPIEInstance, bool bInSimulateInEditor, bool bAnyBlueprintErrors, bool bStartInSpectatorMode, bool bPlayNetDedicated, bool bPlayStereoscopic, float PIEStartTime);
