Add a few more test cases to property editor test object.
Also fix property editor test object details panel falling behind the level editor

diff --git a/Engine/Source/Editor/UnrealEd/Classes/Editor/PropertyEditorTestObject.h b/Engine/Source/Editor/UnrealEd/Classes/Editor/PropertyEditorTestObject.h
--- a/Engine/Source/Editor/UnrealEd/Classes/Editor/PropertyEditorTestObject.h
+++ b/Engine/Source/Editor/UnrealEd/Classes/Editor/PropertyEditorTestObject.h
@@ -660,6 +660,12 @@ class UPropertyEditorTestObject : public UObject
 	UPROPERTY(EditAnywhere, Category=DateTime)
 	FDateTime DateTime;
 
+	UPROPERTY(EditAnywhere, Category = DateTime)
+	FTimespan Timespan;
+
+	UPROPERTY(EditAnywhere, Category = AdvancedProperties)
+	FGuid Guid;
+
 	UPROPERTY()
 	bool bInlineEditConditionWithoutMetaToggle;
 
diff --git a/Engine/Source/Editor/UnrealEd/Private/EditorServer.cpp b/Engine/Source/Editor/UnrealEd/Private/EditorServer.cpp
--- a/Engine/Source/Editor/UnrealEd/Private/EditorServer.cpp
+++ b/Engine/Source/Editor/UnrealEd/Private/EditorServer.cpp
@@ -6008,7 +6008,15 @@ bool UEditorEngine::HandleTestPropsCommand( const TCHAR* Str, FOutputDevice& Ar
 		DetailsView->SetObjects( Objects );
 	}
 
-	FSlateApplication::Get().AddWindow( Window );
+	// Parent to the main frame window
+	if (FModuleManager::Get().IsModuleLoaded("MainFrame"))
+	{
+		IMainFrameModule& MainFrame = FModuleManager::LoadModuleChecked<IMainFrameModule>("MainFrame");
+		TSharedPtr<SWindow> ParentWindow = MainFrame.GetParentWindow();
+
+		FSlateApplication::Get().AddWindowAsNativeChild(Window, ParentWindow.ToSharedRef());
+	}
+
 
 	return true;
 }
