Updated USelection to be a legacy shim around a UTypedElementSelectionSet

This change removes the original UObject-array based USelection implementation that was still being used for generic UObject selection, in favor of always using a UTypedElementSelectionSet to back the USelection.

 - Added an "Object" element type, so that generic objects can be stored in a UTypedElementSelectionSet.
 - Added UEngineElementsLibrary to manage the creation and destruction of "Object" elements, as CoreUObject cannot directly use the TypedElementFramework.
   - Also migrated "Actor" and "Component" element creation and destruction to UEngineElementsLibrary for consistency.
 - Fixed a race condition in the AcquireEditorXElementHandle functions due to the TTypedElementOwnerStore lock being released between the find and the add (see TTypedElementOwnerStore::FindOrRegisterElementOwner).
 - Moved the GetObject function onto a new UTypedElementObjectInterface, as that function will likely be needed by legacy code outside of selection.
 - Removed UTypedElementSelectionSet::GetMutableElementList in favor of selectively exposing certain functions from the underlying UTypedElementList, which should avoid potential abuse or misuse of the underlying list.

#rnx

diff --git a/Engine/Source/Editor/EditorFramework/Private/Elements/TypedElementSelectionSet.cpp b/Engine/Source/Editor/EditorFramework/Private/Elements/TypedElementSelectionSet.cpp
--- a/Engine/Source/Editor/EditorFramework/Private/Elements/TypedElementSelectionSet.cpp
+++ b/Engine/Source/Editor/EditorFramework/Private/Elements/TypedElementSelectionSet.cpp
@@ -1,7 +1,6 @@
 // Copyright Epic Games, Inc. All Rights Reserved.
 
 #include "Elements/TypedElementSelectionSet.h"
-#include "TypedElementList.h"
 #include "TypedElementRegistry.h"
 
 UTypedElementSelectionSet::UTypedElementSelectionSet()
@@ -9,6 +8,8 @@ UTypedElementSelectionSet::UTypedElementSelectionSet()
 	if (!HasAnyFlags(RF_ClassDefaultObject))
 	{
 		ElementList = UTypedElementRegistry::GetInstance()->CreateElementList();
+		ElementList->OnPreChange().AddUObject(this, &UTypedElementSelectionSet::OnElementListPreChange);
+		ElementList->OnChanged().AddUObject(this, &UTypedElementSelectionSet::OnElementListChanged);
 	}
 }
 
diff --git a/Engine/Source/Editor/EditorFramework/Public/Elements/TypedElementSelectionSet.h b/Engine/Source/Editor/EditorFramework/Public/Elements/TypedElementSelectionSet.h
--- a/Engine/Source/Editor/EditorFramework/Public/Elements/TypedElementSelectionSet.h
+++ b/Engine/Source/Editor/EditorFramework/Public/Elements/TypedElementSelectionSet.h
@@ -2,11 +2,10 @@
 
 #pragma once
 
+#include "TypedElementList.h"
 #include "TypedElementSelectionInterface.h"
 #include "TypedElementSelectionSet.generated.h"
 
-class UTypedElementList;
-
 /**
  * Proxy type used to allow asset editors (such as the level editor) to override the base behavior of element selection,
  * by injecting extra pre/post selection logic around the call into the selection interface for an element type.
@@ -155,18 +154,63 @@ public:
 	FTypedElementHandle GetSelectionElement(const FTypedElementHandle& InElementHandle, const ETypedElementSelectionMethod InSelectionMethod) const;
 
 	/**
-	 * Get the underlying element list holding the selection state.
+	 * Access the delegate that is invoked whenever this element list is potentially about to change.
+	 * @note This may be called even if no actual change happens, though once a change does happen it won't be called again until after the next call to NotifyPendingChanges.
 	 */
-	const UTypedElementList* GetElementList() const
+	DECLARE_EVENT_OneParam(UTypedElementSelectionSet, FOnPreChange, const UTypedElementSelectionSet* /*InElementSelectionSet*/);
+	FOnPreChange& OnPreChange()
 	{
-		return ElementList;
+		return OnPreChangeDelegate;
+	}
+
+	/**
+	 * Access the delegate that is invoked whenever the underlying element list has been changed.
+	 * @note This is called automatically at the end of each frame, but can also be manually invoked by NotifyPendingChanges.
+	 */
+	DECLARE_EVENT_OneParam(UTypedElementSelectionSet, FOnChanged, const UTypedElementSelectionSet* /*InElementSelectionSet*/);
+	FOnChanged& OnChanged()
+	{
+		return OnChangedDelegate;
+	}
+
+	/**
+	 * Invoke the delegate called whenever the underlying element list has been changed.
+	 */
+	void NotifyPendingChanges()
+	{
+		ElementList->NotifyPendingChanges();
+	}
+
+	/**
+	 * Clear whether there are pending changes for OnChangedDelegate to notify for, without emitting a notification.
+	 */
+	void ClearPendingChanges()
+	{
+		ElementList->ClearPendingChanges();
+	}
+
+	/**
+	 * Access the interface to allow external systems (such as USelection) to receive immediate sync notifications as the underlying element list is changed.
+	 * This exists purely as a bridging mechanism and shouldn't be relied on for new code. It is lazily created as needed.
+	 */
+	FTypedElementListLegacySync& Legacy_GetElementListSync()
+	{
+		return ElementList->Legacy_GetSync();
+	}
+
+	/**
+	 * Access the interface to allow external systems (such as USelection) to receive immediate sync notifications as the underlying element list is changed.
+	 * This exists purely as a bridging mechanism and shouldn't be relied on for new code. This will return null if no legacy sync has been created for this instance.
+	 */
+	FTypedElementListLegacySync* Legacy_GetElementListSyncPtr()
+	{
+		return ElementList->Legacy_GetSyncPtr();
 	}
 
 	/**
 	 * Get the underlying element list holding the selection state.
-	 * @note Mutating the internal element list directly is a bad idea.
 	 */
-	UTypedElementList* GetMutableElementList()
+	const UTypedElementList* GetElementList() const
 	{
 		return ElementList;
 	}
@@ -182,6 +226,24 @@ private:
 	 */
 	FTypedElementSelectionSetElement ResolveSelectionSetElement(const FTypedElementHandle& InElementHandle) const;
 
+	/**
+	 * Proxy the internal OnPreChange event from the underlying element list.
+	 */
+	void OnElementListPreChange(const UTypedElementList* InElementList) const
+	{
+		check(InElementList == ElementList);
+		OnPreChangeDelegate.Broadcast(this);
+	}
+
+	/**
+	 * Proxy the internal OnChanged event from the underlying element list.
+	 */
+	void OnElementListChanged(const UTypedElementList* InElementList) const
+	{
+		check(InElementList == ElementList);
+		OnChangedDelegate.Broadcast(this);
+	}
+
 	/** Underlying element list holding the selection state. */
 	UPROPERTY()
 	UTypedElementList* ElementList = nullptr;
@@ -189,4 +251,10 @@ private:
 	/** Array of registered asset editor selection proxies, indexed by ElementTypeId-1. */
 	UPROPERTY()
 	UTypedElementAssetEditorSelectionProxy* RegisteredAssetEditorSelectionProxies[TypedHandleMaxTypeId - 1];
+
+	/** Delegate that is invoked whenever the underlying element list is potentially about to change. */
+	FOnPreChange OnPreChangeDelegate;
+
+	/** Delegate that is invoked whenever the underlying element list has been changed. */
+	FOnChanged OnChangedDelegate;
 };
diff --git a/Engine/Source/Editor/LevelEditor/Private/Elements/Actor/ActorElementLevelEditorSelectionProxy.cpp b/Engine/Source/Editor/LevelEditor/Private/Elements/Actor/ActorElementLevelEditorSelectionProxy.cpp
--- a/Engine/Source/Editor/LevelEditor/Private/Elements/Actor/ActorElementLevelEditorSelectionProxy.cpp
+++ b/Engine/Source/Editor/LevelEditor/Private/Elements/Actor/ActorElementLevelEditorSelectionProxy.cpp
@@ -5,6 +5,7 @@
 #include "GameFramework/Actor.h"
 #include "Components/ActorComponent.h"
 #include "TypedElementList.h"
+#include "Elements/EngineElementsLibrary.h"
 
 #include "Elements/Actor/ActorElementSelectionInterface.h"
 #include "Elements/Component/ComponentElementSelectionInterface.h"
@@ -55,7 +56,7 @@ FTypedElementHandle UActorElementLevelEditorSelectionProxy::GetSelectionElement(
 		{
 			ConsideredActor = ConsideredActor->GetParentActor();
 		}
-		return ConsideredActor->AcquireEditorElementHandle();
+		return UEngineElementsLibrary::AcquireEditorActorElementHandle(ConsideredActor);
 	}
 	return InElementSelectionHandle;
 }
@@ -155,7 +156,7 @@ bool UActorElementLevelEditorSelectionProxy::SelectActorElement(const TTypedElem
 
 	// Select the desired actor
 	{
-		TTypedElement<UTypedElementSelectionInterface> ActorSelectionHandle = InSelectionSet->GetElement<UTypedElementSelectionInterface>(Actor->AcquireEditorElementHandle());
+		TTypedElement<UTypedElementSelectionInterface> ActorSelectionHandle = InSelectionSet->GetElement<UTypedElementSelectionInterface>(UEngineElementsLibrary::AcquireEditorActorElementHandle(Actor));
 		if (!ActorSelectionHandle.SelectElement(InSelectionSet, InSelectionOptions))
 		{
 			return bSelectionChanged;
@@ -164,6 +165,9 @@ bool UActorElementLevelEditorSelectionProxy::SelectActorElement(const TTypedElem
 
 	UE_LOG(LogActorLevelEditorSelection, Verbose, TEXT("Selected Actor: %s"), *Actor->GetClass()->GetName());
 
+	// Update the annotation state
+	GSelectedActorAnnotation.Set(Actor);
+	
 	// Bind the override delegates for the components on the selected actor
 	for (UActorComponent* Component : Actor->GetComponents())
 	{
@@ -211,7 +215,7 @@ bool UActorElementLevelEditorSelectionProxy::DeselectActorElement(const TTypedEl
 
 	// Deselect the desired actor
 	{
-		TTypedElement<UTypedElementSelectionInterface> ActorSelectionHandle = InSelectionSet->GetElement<UTypedElementSelectionInterface>(Actor->AcquireEditorElementHandle());
+		TTypedElement<UTypedElementSelectionInterface> ActorSelectionHandle = InSelectionSet->GetElement<UTypedElementSelectionInterface>(UEngineElementsLibrary::AcquireEditorActorElementHandle(Actor));
 		if (!ActorSelectionHandle.DeselectElement(InSelectionSet, InSelectionOptions))
 		{
 			return bSelectionChanged;
@@ -220,13 +224,16 @@ bool UActorElementLevelEditorSelectionProxy::DeselectActorElement(const TTypedEl
 	
 	UE_LOG(LogActorLevelEditorSelection, Verbose, TEXT("Deselected Actor: %s"), *Actor->GetClass()->GetName());
 	
+	// Update the annotation state
+	GSelectedActorAnnotation.Clear(Actor);
+	
 	// Deselect and unbind the override delegates for the components on the selected actor
 	{
 		FTypedElementListLegacySyncScopedBatch LegacySyncBatch(InSelectionSet, InSelectionOptions.AllowLegacyNotifications());
 
 		for (UActorComponent* Component : Actor->GetComponents())
 		{
-			TTypedElement<UTypedElementSelectionInterface> ComponentSelectionHandle = InSelectionSet->GetElement<UTypedElementSelectionInterface>(Component->AcquireEditorElementHandle());
+			TTypedElement<UTypedElementSelectionInterface> ComponentSelectionHandle = InSelectionSet->GetElement<UTypedElementSelectionInterface>(UEngineElementsLibrary::AcquireEditorComponentElementHandle(Component));
 			ComponentSelectionHandle.DeselectElement(InSelectionSet, InSelectionOptions);
 
 			if (USceneComponent* SceneComponent = Cast<USceneComponent>(Component))
@@ -261,7 +268,7 @@ bool UActorElementLevelEditorSelectionProxy::SelectActorGroup(AGroupActor* InGro
 		InGroupActor->GetGroupActors(GroupActors);
 		for (AActor* Actor : GroupActors)
 		{
-			TTypedElement<UTypedElementSelectionInterface> ActorSelectionHandle = InSelectionSet->GetElement<UTypedElementSelectionInterface>(Actor->AcquireEditorElementHandle());
+			TTypedElement<UTypedElementSelectionInterface> ActorSelectionHandle = InSelectionSet->GetElement<UTypedElementSelectionInterface>(UEngineElementsLibrary::AcquireEditorActorElementHandle(Actor));
 			if (CanSelectActorElement(ActorSelectionHandle, GroupSelectionOptions))
 			{
 				bSelectionChanged |= SelectActorElement(ActorSelectionHandle, InSelectionSet, GroupSelectionOptions);
@@ -288,7 +295,7 @@ bool UActorElementLevelEditorSelectionProxy::DeselectActorGroup(AGroupActor* InG
 		InGroupActor->GetGroupActors(GroupActors);
 		for (AActor* Actor : GroupActors)
 		{
-			TTypedElement<UTypedElementSelectionInterface> ActorSelectionHandle = InSelectionSet->GetElement<UTypedElementSelectionInterface>(Actor->AcquireEditorElementHandle());
+			TTypedElement<UTypedElementSelectionInterface> ActorSelectionHandle = InSelectionSet->GetElement<UTypedElementSelectionInterface>(UEngineElementsLibrary::AcquireEditorActorElementHandle(Actor));
 			if (CanDeselectActorElement(ActorSelectionHandle, GroupSelectionOptions))
 			{
 				bSelectionChanged |= DeselectActorElement(ActorSelectionHandle, InSelectionSet, GroupSelectionOptions);
diff --git a/Engine/Source/Editor/LevelEditor/Private/Elements/Component/ComponentElementLevelEditorSelectionProxy.cpp b/Engine/Source/Editor/LevelEditor/Private/Elements/Component/ComponentElementLevelEditorSelectionProxy.cpp
--- a/Engine/Source/Editor/LevelEditor/Private/Elements/Component/ComponentElementLevelEditorSelectionProxy.cpp
+++ b/Engine/Source/Editor/LevelEditor/Private/Elements/Component/ComponentElementLevelEditorSelectionProxy.cpp
@@ -5,6 +5,7 @@
 #include "GameFramework/Actor.h"
 #include "Components/ActorComponent.h"
 #include "TypedElementList.h"
+#include "Elements/EngineElementsLibrary.h"
 
 #include "Elements/Actor/ActorElementSelectionInterface.h"
 #include "Elements/Component/ComponentElementSelectionInterface.h"
@@ -55,7 +56,7 @@ FTypedElementHandle UComponentElementLevelEditorSelectionProxy::GetSelectionElem
 				ConsideredComponent = ConsideredComponent->GetAttachParent();
 			}
 
-			const FTypedElementHandle ConsideredActorHandle = ConsideredActor->AcquireEditorElementHandle();
+			const FTypedElementHandle ConsideredActorHandle = UEngineElementsLibrary::AcquireEditorActorElementHandle(ConsideredActor);
 
 			// We want to process the click on the component only if:
 			// 1. The actor clicked is already selected and is the only actor selected
@@ -71,7 +72,7 @@ FTypedElementHandle UComponentElementLevelEditorSelectionProxy::GetSelectionElem
 
 			if (bSelectComponent && ConsideredComponent)
 			{
-				return ConsideredComponent->AcquireEditorElementHandle();
+				return UEngineElementsLibrary::AcquireEditorComponentElementHandle(ConsideredComponent);
 			}
 			else
 			{
@@ -109,6 +110,9 @@ bool UComponentElementLevelEditorSelectionProxy::SelectComponentElement(const TT
 
 	UE_LOG(LogComponentLevelEditorSelection, Verbose, TEXT("Selected Component: %s"), *Component->GetClass()->GetName());
 	
+	// Update the annotation state
+	GSelectedComponentAnnotation.Set(Component);
+	
 	// Make sure the override delegate is bound properly
 	if (USceneComponent* SceneComponent = Cast<USceneComponent>(Component))
 	{
@@ -118,7 +122,7 @@ bool UComponentElementLevelEditorSelectionProxy::SelectComponentElement(const TT
 	if (AActor* ComponentOwner = Component->GetOwner())
 	{
 		// Selecting a component requires that its owner actor be selected too
-		TTypedElement<UTypedElementSelectionInterface> ActorSelectionHandle = InSelectionSet->GetElement<UTypedElementSelectionInterface>(ComponentOwner->AcquireEditorElementHandle());
+		TTypedElement<UTypedElementSelectionInterface> ActorSelectionHandle = InSelectionSet->GetElement<UTypedElementSelectionInterface>(UEngineElementsLibrary::AcquireEditorActorElementHandle(ComponentOwner));
 		ActorSelectionHandle.SelectElement(InSelectionSet, InSelectionOptions);
 
 		// Update the selection visualization
@@ -145,10 +149,13 @@ bool UComponentElementLevelEditorSelectionProxy::DeselectComponentElement(const
 
 	UE_LOG(LogComponentLevelEditorSelection, Verbose, TEXT("Deselected Component: %s"), *Component->GetClass()->GetName());
 
+	// Update the annotation state
+	GSelectedComponentAnnotation.Clear(Component);
+	
 	// Update the selection visualization
 	if (AActor* ComponentOwner = Component->GetOwner())
 	{
-		TTypedElement<UTypedElementSelectionInterface> ActorSelectionHandle = InSelectionSet->GetElement<UTypedElementSelectionInterface>(ComponentOwner->AcquireEditorElementHandle());
+		TTypedElement<UTypedElementSelectionInterface> ActorSelectionHandle = InSelectionSet->GetElement<UTypedElementSelectionInterface>(UEngineElementsLibrary::AcquireEditorActorElementHandle(ComponentOwner));
 		if (!ActorSelectionHandle.IsElementSelected(InSelectionSet, FTypedElementIsSelectedOptions().SetAllowIndirect(true)))
 		{
 			// Make sure the override delegate is unbound properly
diff --git a/Engine/Source/Editor/LevelEditor/Private/SLevelEditor.cpp b/Engine/Source/Editor/LevelEditor/Private/SLevelEditor.cpp
--- a/Engine/Source/Editor/LevelEditor/Private/SLevelEditor.cpp
+++ b/Engine/Source/Editor/LevelEditor/Private/SLevelEditor.cpp
@@ -340,7 +340,10 @@ SLevelEditor::~SLevelEditor()
 	}
 	if (UObjectInitialized())
 	{
-		SelectedElements->ClearSelection(FTypedElementSelectionOptions());
+		if (!SelectedElements->HasAnyFlags(RF_BeginDestroyed | RF_FinishDestroyed))
+		{
+			SelectedElements->ClearSelection(FTypedElementSelectionOptions());
+		}
 		SelectedElements->RemoveFromRoot();
 		SelectedElements = nullptr;
 	}
diff --git a/Engine/Source/Editor/UnrealEd/Classes/Editor/UnrealEdEngine.h b/Engine/Source/Editor/UnrealEd/Classes/Editor/UnrealEdEngine.h
--- a/Engine/Source/Editor/UnrealEd/Classes/Editor/UnrealEdEngine.h
+++ b/Engine/Source/Editor/UnrealEd/Classes/Editor/UnrealEdEngine.h
@@ -28,7 +28,7 @@ class UPrimitiveComponent;
 class UTexture2D;
 class UUnrealEdOptions;
 class USelection;
-class UTypedElementList;
+class UTypedElementSelectionSet;
 class FName;
 typedef FName FEditorModeID;
 struct FTypedElementSelectionOptions;
@@ -803,11 +803,11 @@ protected:
 	/** Called when global editor selection changes */
 	void OnEditorSelectionChanged(UObject* SelectionThatChanged);
 
-	/** Called when the element list pointer set on the global editor selection changes */
-	void OnEditorSelectionElementListPtrChanged(USelection* Selection, UTypedElementList* OldElementList, UTypedElementList* NewElementList);
+	/** Called when the element selection set pointer set on the global editor selection changes */
+	void OnEditorElementSelectionPtrChanged(USelection* Selection, UTypedElementSelectionSet* OldSelectionSet, UTypedElementSelectionSet* NewSelectionSet);
 
-	/** Called when the element list associated with the global editor selection changes */
-	void OnEditorSelectionElementListChanged(const UTypedElementList* ElementList);
+	/** Called when the element selection set associated with the global editor selection changes */
+	void OnEditorElementSelectionChanged(const UTypedElementSelectionSet* SelectionSet);
 
 	/** The package auto-saver instance used by the editor */
 	TUniquePtr<IPackageAutoSaver> PackageAutoSaver;
diff --git a/Engine/Source/Editor/UnrealEd/Private/AssetEditorModeManager.cpp b/Engine/Source/Editor/UnrealEd/Private/AssetEditorModeManager.cpp
--- a/Engine/Source/Editor/UnrealEd/Private/AssetEditorModeManager.cpp
+++ b/Engine/Source/Editor/UnrealEd/Private/AssetEditorModeManager.cpp
@@ -3,26 +3,28 @@
 #include "AssetEditorModeManager.h"
 #include "Engine/Selection.h"
 #include "PreviewScene.h"
-#include "TypedElementRegistry.h"
+#include "Elements/TypedElementSelectionSet.h"
 
 //////////////////////////////////////////////////////////////////////////
 // FAssetEditorModeManager
 
 FAssetEditorModeManager::FAssetEditorModeManager()
 {
-	SelectedElements = UTypedElementRegistry::GetInstance()->CreateElementList();
-	SelectedElements->AddToRoot();
+	{
+		UTypedElementSelectionSet* ActorAndComponentsSelectionSet = NewObject<UTypedElementSelectionSet>();
 
-	ActorSet = USelection::CreateActorSelection(nullptr, GetTransientPackage(), NAME_None, RF_Transactional);
-	ActorSet->SetElementList(SelectedElements);
-	ActorSet->AddToRoot();
+		ActorSet = USelection::CreateActorSelection(GetTransientPackage(), NAME_None, RF_Transactional);
+		ActorSet->SetElementSelectionSet(ActorAndComponentsSelectionSet);
+		ActorSet->AddToRoot();
 
-	ObjectSet = USelection::CreateObjectSelection(nullptr, GetTransientPackage(), NAME_None, RF_Transactional);
-	ObjectSet->AddToRoot();
+		ComponentSet = USelection::CreateComponentSelection(GetTransientPackage(), NAME_None, RF_Transactional);
+		ComponentSet->SetElementSelectionSet(ActorAndComponentsSelectionSet);
+		ComponentSet->AddToRoot();
+	}
 
-	ComponentSet = USelection::CreateComponentSelection(nullptr, GetTransientPackage(), NAME_None, RF_Transactional);
-	ComponentSet->SetElementList(SelectedElements);
-	ComponentSet->AddToRoot();
+	ObjectSet = USelection::CreateObjectSelection(GetTransientPackage(), NAME_None, RF_Transactional);
+	ObjectSet->SetElementSelectionSet(NewObject<UTypedElementSelectionSet>(ObjectSet));
+	ObjectSet->AddToRoot();
 }
 
 FAssetEditorModeManager::~FAssetEditorModeManager()
@@ -33,20 +35,31 @@ FAssetEditorModeManager::~FAssetEditorModeManager()
 	// which would mean that these instances will be garbage
 	if (UObjectInitialized())
 	{
-		ActorSet->SetElementList(nullptr);
+		if (!ActorSet->HasAnyFlags(RF_BeginDestroyed | RF_FinishDestroyed))
+		{
+			check(ActorSet->GetElementSelectionSet() == ComponentSet->GetElementSelectionSet());
+			if (UTypedElementSelectionSet* ActorAndComponentsSelectionSet = ActorSet->GetElementSelectionSet())
+			{
+				ActorAndComponentsSelectionSet->ClearSelection(FTypedElementSelectionOptions());
+			}
+		}
+
 		ActorSet->RemoveFromRoot();
 		ActorSet = nullptr;
 
-		ObjectSet->RemoveFromRoot();
-		ObjectSet = nullptr;
-
-		ComponentSet->SetElementList(nullptr);
 		ComponentSet->RemoveFromRoot();
 		ComponentSet = nullptr;
 
-		SelectedElements->Empty();
-		SelectedElements->RemoveFromRoot();
-		SelectedElements = nullptr;
+		if (!ObjectSet->HasAnyFlags(RF_BeginDestroyed | RF_FinishDestroyed))
+		{
+			if (UTypedElementSelectionSet* ObjectSelectionSet = ObjectSet->GetElementSelectionSet())
+			{
+				ObjectSelectionSet->ClearSelection(FTypedElementSelectionOptions());
+			}
+		}
+
+		ObjectSet->RemoveFromRoot();
+		ObjectSet = nullptr;
 	}
 }
 
diff --git a/Engine/Source/Editor/UnrealEd/Private/EditorEngine.cpp b/Engine/Source/Editor/UnrealEd/Private/EditorEngine.cpp
--- a/Engine/Source/Editor/UnrealEd/Private/EditorEngine.cpp
+++ b/Engine/Source/Editor/UnrealEd/Private/EditorEngine.cpp
@@ -129,6 +129,9 @@
 
 #include "Editor/ActorPositioning.h"
 
+#include "Elements/TypedElementSelectionSet.h"
+#include "Elements/Object/ObjectElementEditorSelectionProxy.h"
+
 #include "IDirectoryWatcher.h"
 #include "DirectoryWatcherModule.h"
 
@@ -243,77 +246,55 @@ DEFINE_LOG_CATEGORY_STATIC(LogEditor, Log, All);
 //////////////////////////////////////////////////////////////////////////
 // Globals
 
-static inline USelection*& PrivateGetSelectedActors()
+namespace PrivateEditorSelection
 {
-	static USelection* SSelectedActors = NULL;
-	return SSelectedActors;
-};
 
-static inline USelection*& PrivateGetSelectedComponents()
-{
-	static USelection* SSelectedComponents = NULL;
-	return SSelectedComponents;
-}
+static USelection* GActorSelection = nullptr;
+static USelection* GComponentSelection = nullptr;
+static USelection* GObjectSelection = nullptr;
 
-static inline USelection*& PrivateGetSelectedObjects()
+void InitSelectionSets()
 {
-	static USelection* SSelectedObjects = NULL;
-	return SSelectedObjects;
-};
+	// Note: The actor and component typed element selection set is set and owned by the level editor, so it is deliberately left null here
+	GActorSelection = USelection::CreateActorSelection(GetTransientPackage(), TEXT("SelectedActors"), RF_Transactional);
+	GActorSelection->AddToRoot();
+
+	GComponentSelection = USelection::CreateComponentSelection(GetTransientPackage(), TEXT("SelectedComponents"), RF_Transactional);
+	GComponentSelection->AddToRoot();
 
-static void OnObjectSelected(UObject* Object)
+	GObjectSelection = USelection::CreateObjectSelection(GetTransientPackage(), TEXT("SelectedObjects"), RF_Transactional);
+	GObjectSelection->AddToRoot();
+
+	UTypedElementSelectionSet* ObjectSelectionSet = NewObject<UTypedElementSelectionSet>(GObjectSelection);
+	ObjectSelectionSet->RegisterAssetEditorSelectionProxy(NAME_Object, NewObject<UObjectElementEditorSelectionProxy>());
+	GObjectSelection->SetElementSelectionSet(ObjectSelectionSet);
+}
+
+void DestroySelectionSets()
 {
-	// Whenever an actor is unselected we must remove its components from the components selection
-	if (!Object->IsSelected())
+	// We may be destroyed after the UObject system has already shutdown, 
+	// which would mean that these instances will be garbage
+	if (UObjectInitialized())
 	{
-		TArray<UActorComponent*> ComponentsToDeselect;
-		for (FSelectionIterator It(*PrivateGetSelectedComponents()); It; ++It)
-		{
-			UActorComponent* Component = CastChecked<UActorComponent>(*It);
-			if (Component->GetOwner() == Object)
-			{
-				ComponentsToDeselect.Add(Component);
-			}
-		}
-		if (ComponentsToDeselect.Num() > 0)
+		GActorSelection->RemoveFromRoot();
+		GComponentSelection->RemoveFromRoot();
+		GObjectSelection->RemoveFromRoot();
+
+		if (!GObjectSelection->HasAnyFlags(RF_BeginDestroyed | RF_FinishDestroyed))
 		{
-			PrivateGetSelectedComponents()->Modify();
-			PrivateGetSelectedComponents()->BeginBatchSelectOperation();
-			for (UActorComponent* Component : ComponentsToDeselect)
+			if (UTypedElementSelectionSet* ObjectSelectionSet = GObjectSelection->GetElementSelectionSet())
 			{
-				PrivateGetSelectedComponents()->Deselect(Component);
+				ObjectSelectionSet->ClearSelection(FTypedElementSelectionOptions());
 			}
-			PrivateGetSelectedComponents()->EndBatchSelectOperation();
 		}
 	}
-}
-
-static void PrivateInitSelectedSets()
-{
-	PrivateGetSelectedActors() = USelection::CreateActorSelection(&GSelectedActorAnnotation, GetTransientPackage(), TEXT("SelectedActors"), RF_Transactional);
-	PrivateGetSelectedActors()->AddToRoot();
-	PrivateGetSelectedActors()->SelectObjectEvent.AddStatic(&OnObjectSelected);
-
-	PrivateGetSelectedComponents() = USelection::CreateComponentSelection(&GSelectedComponentAnnotation, GetTransientPackage(), TEXT("SelectedComponents"), RF_Transactional);
-	PrivateGetSelectedComponents()->AddToRoot();
 
-	PrivateGetSelectedObjects() = USelection::CreateObjectSelection(&GSelectedObjectAnnotation, GetTransientPackage(), TEXT("SelectedObjects"), RF_Transactional);
-	PrivateGetSelectedObjects()->AddToRoot();
+	GActorSelection = nullptr;
+	GComponentSelection = nullptr;
+	GObjectSelection = nullptr;
 }
 
-static void PrivateDestroySelectedSets()
-{
-#if 0
-	PrivateGetSelectedActors()->RemoveFromRoot();
-	PrivateGetSelectedActors() = NULL;
-
-	PrivateGetSelectedComponents()->RemoveFromRoot();
-	PrivateGetSelectedComponents() = NULL;
-
-	PrivateGetSelectedObjects()->RemoveFromRoot();
-	PrivateGetSelectedObjects() = NULL;
-#endif
-}
+} // namespace PrivateEditorSelection
 
 /**
 * A mapping of all startup packages to whether or not we have warned the user about editing them
@@ -404,7 +385,7 @@ int32 UEditorEngine::GetSelectedActorCount() const
 
 USelection* UEditorEngine::GetSelectedActors() const
 {
-	return PrivateGetSelectedActors();
+	return PrivateEditorSelection::GActorSelection;
 }
 
 bool UEditorEngine::IsWorldSettingsSelected() 
@@ -454,12 +435,12 @@ FSelectedEditableComponentIterator UEditorEngine::GetSelectedEditableComponentIt
 
 USelection* UEditorEngine::GetSelectedComponents() const
 {
-	return PrivateGetSelectedComponents();
+	return PrivateEditorSelection::GComponentSelection;
 }
 
 USelection* UEditorEngine::GetSelectedObjects() const
 {
-	return PrivateGetSelectedObjects();
+	return PrivateEditorSelection::GObjectSelection;
 }
 
 void UEditorEngine::GetContentBrowserSelectionClasses(TArray<UClass*>& Selection) const
@@ -647,7 +628,7 @@ void UEditorEngine::InitEditor(IEngineLoop* InEngineLoop)
 	}
 
 	// Create selection sets.
-	PrivateInitSelectedSets();
+	PrivateEditorSelection::InitSelectionSets();
 
 	// Set slate options
 	FMultiBoxSettings::UseSmallToolBarIcons = TAttribute<bool>::Create(TAttribute<bool>::FGetter::CreateStatic(&GetSmallToolBarIcons));
@@ -1262,7 +1243,7 @@ void UEditorEngine::FinishDestroy()
 		}
 
 		// Destroy selection sets.
-		PrivateDestroySelectedSets();
+		PrivateEditorSelection::DestroySelectionSets();
 
 		extern void TearDownDistanceFieldBuildNotification();
 		TearDownDistanceFieldBuildNotification();
diff --git a/Engine/Source/Editor/UnrealEd/Private/EditorSelectUtils.cpp b/Engine/Source/Editor/UnrealEd/Private/EditorSelectUtils.cpp
--- a/Engine/Source/Editor/UnrealEd/Private/EditorSelectUtils.cpp
+++ b/Engine/Source/Editor/UnrealEd/Private/EditorSelectUtils.cpp
@@ -29,6 +29,7 @@
 #include "Subsystems/BrushEditingSubsystem.h"
 #include "TypedElementList.h"
 #include "Elements/TypedElementSelectionSet.h"
+#include "Elements/EngineElementsLibrary.h"
 
 #define LOCTEXT_NAMESPACE "EditorSelectUtils"
 
@@ -256,24 +257,24 @@ void UUnrealEdEngine::ResetPivot()
 	Selection.
 -----------------------------------------------------------------------------*/
 
-void UUnrealEdEngine::OnEditorSelectionElementListPtrChanged(USelection* Selection, UTypedElementList* OldElementList, UTypedElementList* NewElementList)
+void UUnrealEdEngine::OnEditorElementSelectionPtrChanged(USelection* Selection, UTypedElementSelectionSet* OldSelectionSet, UTypedElementSelectionSet* NewSelectionSet)
 {
 	if (Selection == GetSelectedActors())
 	{
-		if (OldElementList)
+		if (OldSelectionSet)
 		{
-			OldElementList->OnChanged().RemoveAll(this);
+			OldSelectionSet->OnChanged().RemoveAll(this);
 		}
 
-		if (NewElementList)
+		if (NewSelectionSet)
 		{
-			NewElementList->OnChanged().AddUObject(this, &UUnrealEdEngine::OnEditorSelectionElementListChanged);
+			NewSelectionSet->OnChanged().AddUObject(this, &UUnrealEdEngine::OnEditorElementSelectionChanged);
 		}
 	}
 }
 
 
-void UUnrealEdEngine::OnEditorSelectionElementListChanged(const UTypedElementList* ElementList)
+void UUnrealEdEngine::OnEditorElementSelectionChanged(const UTypedElementSelectionSet* SelectionSet)
 {
 	NoteSelectionChange();
 }
@@ -429,15 +430,24 @@ void UUnrealEdEngine::SelectGroup(AGroupActor* InGroupActor, bool bForceSelectio
 		// Select/deselect all actors within the group (if locked or forced)
 		if (bForceSelection || InGroupActor->IsLocked())
 		{
-			FTypedElementListLegacySyncScopedBatch LegacySyncBatch(SelectionSet->GetMutableElementList(), SelectionOptions.AllowLegacyNotifications());
-
 			TArray<AActor*> GroupActors;
 			InGroupActor->GetGroupActors(GroupActors);
+
+			TArray<FTypedElementHandle, TInlineAllocator<256>> GroupElements;
+			GroupElements.Reserve(GroupActors.Num());
 			for (AActor* Actor : GroupActors)
+			{
+				if (FTypedElementHandle ElementHandle = UEngineElementsLibrary::AcquireEditorActorElementHandle(Actor))
+				{
+					GroupElements.Add(MoveTemp(ElementHandle));
+				}
+			}
+
+			if (GroupElements.Num() > 0)
 			{
 				bSelectionChanged |= bInSelected
-					? SelectionSet->SelectElement(Actor->AcquireEditorElementHandle(), SelectionOptions)
-					: SelectionSet->DeselectElement(Actor->AcquireEditorElementHandle(), SelectionOptions);
+					? SelectionSet->SelectElements(GroupElements, SelectionOptions)
+					: SelectionSet->DeselectElements(GroupElements, SelectionOptions);
 			}
 		}
 
@@ -445,11 +455,11 @@ void UUnrealEdEngine::SelectGroup(AGroupActor* InGroupActor, bool bForceSelectio
 		{
 			if (bNotify)
 			{
-				SelectionSet->GetMutableElementList()->NotifyPendingChanges();
+				SelectionSet->NotifyPendingChanges();
 			}
 			else
 			{
-				SelectionSet->GetMutableElementList()->ClearPendingChanges();
+				SelectionSet->ClearPendingChanges();
 			}
 		}
 	}
@@ -471,8 +481,8 @@ bool UUnrealEdEngine::CanSelectActor(AActor* Actor, bool bInSelected, bool bSele
 			.SetWarnIfLocked(bWarnIfLevelLocked);
 
 		return bInSelected
-			? SelectionSet->CanSelectElement(Actor->AcquireEditorElementHandle(), SelectionOptions)
-			: SelectionSet->CanDeselectElement(Actor->AcquireEditorElementHandle(), SelectionOptions);
+			? SelectionSet->CanSelectElement(UEngineElementsLibrary::AcquireEditorActorElementHandle(Actor), SelectionOptions)
+			: SelectionSet->CanDeselectElement(UEngineElementsLibrary::AcquireEditorActorElementHandle(Actor), SelectionOptions);
 	}
 
 	return false;
@@ -495,18 +505,18 @@ void UUnrealEdEngine::SelectActor(AActor* Actor, bool bInSelected, bool bNotify,
 			.SetAllowLegacyNotifications(false);
 
 		const bool bSelectionChanged = bInSelected
-			? SelectionSet->SelectElement(Actor->AcquireEditorElementHandle(), SelectionOptions)
-			: SelectionSet->DeselectElement(Actor->AcquireEditorElementHandle(), SelectionOptions);
+			? SelectionSet->SelectElement(UEngineElementsLibrary::AcquireEditorActorElementHandle(Actor), SelectionOptions)
+			: SelectionSet->DeselectElement(UEngineElementsLibrary::AcquireEditorActorElementHandle(Actor), SelectionOptions);
 
 		if (bSelectionChanged)
 		{
 			if (bNotify)
 			{
-				SelectionSet->GetMutableElementList()->NotifyPendingChanges();
+				SelectionSet->NotifyPendingChanges();
 			}
 			else
 			{
-				SelectionSet->GetMutableElementList()->ClearPendingChanges();
+				SelectionSet->ClearPendingChanges();
 			}
 		}
 		else if (bNotify || bForceRefresh)
@@ -534,18 +544,18 @@ void UUnrealEdEngine::SelectComponent(UActorComponent* Component, bool bInSelect
 			.SetAllowLegacyNotifications(false);
 
 		const bool bSelectionChanged = bInSelected
-			? SelectionSet->SelectElement(Component->AcquireEditorElementHandle(), SelectionOptions)
-			: SelectionSet->DeselectElement(Component->AcquireEditorElementHandle(), SelectionOptions);
+			? SelectionSet->SelectElement(UEngineElementsLibrary::AcquireEditorComponentElementHandle(Component), SelectionOptions)
+			: SelectionSet->DeselectElement(UEngineElementsLibrary::AcquireEditorComponentElementHandle(Component), SelectionOptions);
 
 		if (bSelectionChanged)
 		{
 			if (bNotify)
 			{
-				SelectionSet->GetMutableElementList()->NotifyPendingChanges();
+				SelectionSet->NotifyPendingChanges();
 			}
 			else
 			{
-				SelectionSet->GetMutableElementList()->ClearPendingChanges();
+				SelectionSet->ClearPendingChanges();
 			}
 		}
 	}
@@ -557,7 +567,7 @@ bool UUnrealEdEngine::IsComponentSelected(const UPrimitiveComponent* PrimCompone
 
 	if (UTypedElementSelectionSet* SelectionSet = ComponentSelection->GetElementSelectionSet())
 	{
-		return SelectionSet->IsElementSelected(PrimComponent->AcquireEditorElementHandle(), FTypedElementIsSelectedOptions().SetAllowIndirect(true));
+		return SelectionSet->IsElementSelected(UEngineElementsLibrary::AcquireEditorComponentElementHandle(PrimComponent), FTypedElementIsSelectedOptions().SetAllowIndirect(true));
 	}
 
 	return false;
@@ -677,11 +687,11 @@ void UUnrealEdEngine::SelectNone(bool bNoteSelectionChange, bool bDeselectBSPSur
 		{
 			if (bNoteSelectionChange)
 			{
-				SelectionSet->GetMutableElementList()->NotifyPendingChanges();
+				SelectionSet->NotifyPendingChanges();
 			}
 			else
 			{
-				SelectionSet->GetMutableElementList()->ClearPendingChanges();
+				SelectionSet->ClearPendingChanges();
 			}
 		}
 	}
diff --git a/Engine/Source/Editor/UnrealEd/Private/Elements/Actor/ActorElementEditorSelectionInterface.cpp b/Engine/Source/Editor/UnrealEd/Private/Elements/Actor/ActorElementEditorSelectionInterface.cpp
--- a/Engine/Source/Editor/UnrealEd/Private/Elements/Actor/ActorElementEditorSelectionInterface.cpp
+++ b/Engine/Source/Editor/UnrealEd/Private/Elements/Actor/ActorElementEditorSelectionInterface.cpp
@@ -4,6 +4,7 @@
 #include "Elements/Actor/ActorElementData.h"
 #include "GameFramework/Actor.h"
 #include "TypedElementList.h"
+#include "Elements/EngineElementsLibrary.h"
 
 bool UActorElementEditorSelectionInterface::IsElementSelected(const FTypedElementHandle& InElementHandle, const UTypedElementList* InSelectionSet, const FTypedElementIsSelectedOptions& InSelectionOptions)
 {
@@ -18,7 +19,7 @@ bool UActorElementEditorSelectionInterface::IsActorSelected(const AActor* InActo
 		return false;
 	}
 
-	if (InSelectionSet->Contains(InActor->AcquireEditorElementHandle()))
+	if (InSelectionSet->Contains(UEngineElementsLibrary::AcquireEditorActorElementHandle(InActor)))
 	{
 		return true;
 	}
@@ -27,7 +28,7 @@ bool UActorElementEditorSelectionInterface::IsActorSelected(const AActor* InActo
 	{
 		if (const AActor* RootSelectionActor = InActor->GetRootSelectionParent())
 		{
-			return InSelectionSet->Contains(RootSelectionActor->AcquireEditorElementHandle());
+			return InSelectionSet->Contains(UEngineElementsLibrary::AcquireEditorActorElementHandle(RootSelectionActor));
 		}
 	}
 
diff --git a/Engine/Source/Editor/UnrealEd/Private/Elements/Component/ComponentElementEditorSelectionInterface.cpp b/Engine/Source/Editor/UnrealEd/Private/Elements/Component/ComponentElementEditorSelectionInterface.cpp
--- a/Engine/Source/Editor/UnrealEd/Private/Elements/Component/ComponentElementEditorSelectionInterface.cpp
+++ b/Engine/Source/Editor/UnrealEd/Private/Elements/Component/ComponentElementEditorSelectionInterface.cpp
@@ -5,6 +5,7 @@
 #include "Components/SceneComponent.h"
 #include "GameFramework/Actor.h"
 #include "TypedElementList.h"
+#include "Elements/EngineElementsLibrary.h"
 
 bool UComponentElementEditorSelectionInterface::IsElementSelected(const FTypedElementHandle& InElementHandle, const UTypedElementList* InSelectionSet, const FTypedElementIsSelectedOptions& InSelectionOptions)
 {
@@ -19,7 +20,7 @@ bool UComponentElementEditorSelectionInterface::IsComponentSelected(const UActor
 		return false;
 	}
 
-	if (InSelectionSet->Contains(InComponent->AcquireEditorElementHandle()))
+	if (InSelectionSet->Contains(UEngineElementsLibrary::AcquireEditorComponentElementHandle(InComponent)))
 	{
 		return true;
 	}
@@ -44,7 +45,7 @@ bool UComponentElementEditorSelectionInterface::IsComponentSelected(const UActor
 
 		if (ConsideredComponent)
 		{
-			return InSelectionSet->Contains(ConsideredComponent->AcquireEditorElementHandle());
+			return InSelectionSet->Contains(UEngineElementsLibrary::AcquireEditorComponentElementHandle(ConsideredComponent));
 		}
 	}
 
diff --git a/Engine/Source/Editor/UnrealEd/Private/Elements/Object/ObjectElementEditorSelectionProxy.cpp b/Engine/Source/Editor/UnrealEd/Private/Elements/Object/ObjectElementEditorSelectionProxy.cpp
new file mode 100644
--- /dev/null
+++ b/Engine/Source/Editor/UnrealEd/Private/Elements/Object/ObjectElementEditorSelectionProxy.cpp
@@ -0,0 +1,33 @@
+// Copyright Epic Games, Inc. All Rights Reserved.
+
+#include "Elements/Object/ObjectElementEditorSelectionProxy.h"
+#include "Elements/Object/ObjectElementData.h"
+#include "UObject/UObjectAnnotation.h"
+
+bool UObjectElementEditorSelectionProxy::SelectElement(const TTypedElement<UTypedElementSelectionInterface>& InElementSelectionHandle, UTypedElementList* InSelectionSet, const FTypedElementSelectionOptions& InSelectionOptions)
+{
+	if (const FObjectElementData* ObjectData = InElementSelectionHandle.GetData<FObjectElementData>())
+	{
+		if (InElementSelectionHandle.SelectElement(InSelectionSet, InSelectionOptions))
+		{
+			GSelectedObjectAnnotation.Set(ObjectData->Object);
+			return true;
+		}
+	}
+
+	return false;
+}
+
+bool UObjectElementEditorSelectionProxy::DeselectElement(const TTypedElement<UTypedElementSelectionInterface>& InElementSelectionHandle, UTypedElementList* InSelectionSet, const FTypedElementSelectionOptions& InSelectionOptions)
+{
+	if (const FObjectElementData* ObjectData = InElementSelectionHandle.GetData<FObjectElementData>())
+	{
+		if (InElementSelectionHandle.DeselectElement(InSelectionSet, InSelectionOptions))
+		{
+			GSelectedObjectAnnotation.Clear(ObjectData->Object);
+			return true;
+		}
+	}
+
+	return false;
+}
diff --git a/Engine/Source/Editor/UnrealEd/Private/Elements/Object/ObjectElementEditorSelectionProxy.h b/Engine/Source/Editor/UnrealEd/Private/Elements/Object/ObjectElementEditorSelectionProxy.h
new file mode 100644
--- /dev/null
+++ b/Engine/Source/Editor/UnrealEd/Private/Elements/Object/ObjectElementEditorSelectionProxy.h
@@ -0,0 +1,16 @@
+// Copyright Epic Games, Inc. All Rights Reserved.
+
+#pragma once
+
+#include "Elements/TypedElementSelectionSet.h"
+#include "ObjectElementEditorSelectionProxy.generated.h"
+
+UCLASS()
+class UObjectElementEditorSelectionProxy : public UTypedElementAssetEditorSelectionProxy
+{
+	GENERATED_BODY()
+
+public:
+	virtual bool SelectElement(const TTypedElement<UTypedElementSelectionInterface>& InElementSelectionHandle, UTypedElementList* InSelectionSet, const FTypedElementSelectionOptions& InSelectionOptions) override;
+	virtual bool DeselectElement(const TTypedElement<UTypedElementSelectionInterface>& InElementSelectionHandle, UTypedElementList* InSelectionSet, const FTypedElementSelectionOptions& InSelectionOptions) override;
+};
diff --git a/Engine/Source/Editor/UnrealEd/Private/Selection.cpp b/Engine/Source/Editor/UnrealEd/Private/Selection.cpp
--- a/Engine/Source/Editor/UnrealEd/Private/Selection.cpp
+++ b/Engine/Source/Editor/UnrealEd/Private/Selection.cpp
@@ -2,450 +2,198 @@
 
 #include "Selection.h"
 #include "UObject/Package.h"
-#include "UObject/UObjectAnnotation.h"
 #include "GameFramework/Actor.h"
-#include "TypedElementList.h"
-#include "TypedElementSelectionInterface.h"
-#include "Elements/TypedElementSelectionSet.h"
+#include "TypedElementObjectInterface.h"
+#include "Elements/EngineElementsLibrary.h"
 
 DEFINE_LOG_CATEGORY_STATIC(LogSelection, Log, All);
 
-namespace Selection_Private
-{
-
-class FSelectionStoreBase : public ISelectionStore
+class ISelectionElementBridge
 {
 public:
-	explicit FSelectionStoreBase(const UClass* InObjectBaseClass);
-
-	//~ ISelectionStore interface
-	virtual void SetSink(ISelectionStoreSink* InSink) override;
-	virtual bool IsValidObjectToSelect(const UObject* InObject) const override;
-
-protected:
-	/** Base type of objects that this selection instance manages. */
-	const UClass* ObjectBaseClass = nullptr;
-
-	/** Sink instance to notify of changes. */
-	ISelectionStoreSink* Sink = nullptr;
+	virtual ~ISelectionElementBridge() = default;
+	virtual bool IsValidObjectType(const UObject* InObject) const = 0;
+	virtual FTypedElementHandle GetElementHandleForObject(const UObject* InObject, const bool bAllowCreate = true) const = 0;
 };
 
-FSelectionStoreBase::FSelectionStoreBase(const UClass* InObjectBaseClass)
-	: ObjectBaseClass(InObjectBaseClass)
-{
-}
-
-void FSelectionStoreBase::SetSink(ISelectionStoreSink* InSink)
-{
-	Sink = InSink;
-}
-
-bool FSelectionStoreBase::IsValidObjectToSelect(const UObject* InObject) const
-{
-	return InObject && (!ObjectBaseClass || InObject->IsA(ObjectBaseClass));
-}
-
-
-class FObjectSelectionStore : public FSelectionStoreBase
+class FObjectSelectionElementBridge : public ISelectionElementBridge
 {
 public:
-	explicit FObjectSelectionStore(const UClass* InObjectBaseClass);
-
-	//~ ISelectionStore interface
-	virtual void SetElementList(UTypedElementList* InElementList) override;
-	virtual UTypedElementList* GetElementList() const override;
-	virtual int32 GetNumObjects() const override;
-	virtual UObject* GetObjectAtIndex(const int32 InIndex) const override;
-	virtual bool IsObjectSelected(const UObject* InObject) const override;
-	virtual void SelectObject(UObject* InObject) override;
-	virtual void DeselectObject(UObject* InObject) override;
-	virtual int32 DeselectObjects(TFunctionRef<bool(UObject*)> InPredicate) override;
-	virtual void BeginBatchSelection() override;
-	virtual void EndBatchSelection(const bool InNotify) override;
-	virtual bool IsBatchSelecting() const override;
-	virtual void ForceBatchDirty() override;
-
-private:
-	/** Tracks the number of active selection operations.  Allows batched selection operations to only send one notification at the end of the batch */
-	int32 SelectionMutex = 0;
-
-	/** Tracks whether the selection set changed during a batch selection operation */
-	bool bIsBatchDirty = false;
-
-	/** List of selected objects, ordered as they were selected. */
-	TArray<TWeakObjectPtr<UObject>> SelectedObjects;
-};
-
-FObjectSelectionStore::FObjectSelectionStore(const UClass* InObjectBaseClass)
-	: FSelectionStoreBase(InObjectBaseClass)
-{
-}
-
-void FObjectSelectionStore::SetElementList(UTypedElementList* InElementList)
-{
-	checkf(false, TEXT("This selection store does not support element lists! Only actor and component selections may currently use element lists!"));
-}
-
-UTypedElementList* FObjectSelectionStore::GetElementList() const
-{
-	return nullptr;
-}
-
-int32 FObjectSelectionStore::GetNumObjects() const
-{
-	return SelectedObjects.Num();
-}
-
-UObject* FObjectSelectionStore::GetObjectAtIndex(const int32 InIndex) const
-{
-	return SelectedObjects.IsValidIndex(InIndex)
-		? SelectedObjects[InIndex].Get()
-		: nullptr;
-}
-
-bool FObjectSelectionStore::IsObjectSelected(const UObject* InObject) const
-{
-	return SelectedObjects.Contains(InObject);
-}
-
-void FObjectSelectionStore::SelectObject(UObject* InObject)
-{
-	if (Sink && !bIsBatchDirty)
+	virtual bool IsValidObjectType(const UObject* InObject) const override
 	{
-		Sink->OnPreModify();
+		check(InObject);
+		return true;
 	}
 
-	bIsBatchDirty = true;
-	SelectedObjects.Add(InObject);
-	
-	if (Sink)
+	virtual FTypedElementHandle GetElementHandleForObject(const UObject* InObject, const bool bAllowCreate = true) const override
 	{
-		Sink->OnObjectSelected(InObject, !IsBatchSelecting());
+		check(InObject);
+		return UEngineElementsLibrary::AcquireEditorObjectElementHandle(InObject, bAllowCreate);
 	}
-}
+};
 
-void FObjectSelectionStore::DeselectObject(UObject* InObject)
+class FActorSelectionElementBridge : public ISelectionElementBridge
 {
-	if (Sink && !bIsBatchDirty)
+public:
+	virtual bool IsValidObjectType(const UObject* InObject) const override
 	{
-		Sink->OnPreModify();
+		check(InObject);
+		return InObject->IsA<AActor>();
 	}
 
-	bIsBatchDirty = true;
-	SelectedObjects.RemoveSingle(InObject);
-	
-	if (Sink)
+	virtual FTypedElementHandle GetElementHandleForObject(const UObject* InObject, const bool bAllowCreate = true) const override
 	{
-		Sink->OnObjectDeselected(InObject, !IsBatchSelecting());
+		check(InObject);
+		return UEngineElementsLibrary::AcquireEditorActorElementHandle(CastChecked<AActor>(InObject), bAllowCreate);
 	}
-}
+};
 
-int32 FObjectSelectionStore::DeselectObjects(TFunctionRef<bool(UObject*)> InPredicate)
+class FComponentSelectionElementBridge : public ISelectionElementBridge
 {
-	int32 RemovedCount = 0;
-	
-	if (SelectedObjects.Num() > 0)
+public:
+	virtual bool IsValidObjectType(const UObject* InObject) const override
 	{
-		if (Sink && !bIsBatchDirty)
-		{
-			Sink->OnPreModify();
-		}
-
-		for (int32 Index = SelectedObjects.Num() - 1; Index >= 0; --Index)
-		{
-			UObject* SelectedObject = SelectedObjects[Index].Get();
-			if (!SelectedObject || InPredicate(SelectedObject))
-			{
-				SelectedObjects.RemoveAt(Index);
-				if (SelectedObject)
-				{
-					if (Sink)
-					{
-						Sink->OnObjectDeselected(SelectedObject, /*bNotify*/false);
-					}
-					++RemovedCount;
-				}
-			}
-		}
-
-		if (RemovedCount > 0)
-		{
-			bIsBatchDirty = true;
-			if (Sink)
-			{
-				Sink->OnSelectedChanged(/*bSyncState*/false, !IsBatchSelecting());
-			}
-		}
+		check(InObject);
+		return InObject->IsA<UActorComponent>();
 	}
 
-	return RemovedCount;
-}
-
-void FObjectSelectionStore::BeginBatchSelection()
-{
-	if (SelectionMutex++ == 0)
+	virtual FTypedElementHandle GetElementHandleForObject(const UObject* InObject, const bool bAllowCreate = true) const override
 	{
-		bIsBatchDirty = false;
+		check(InObject);
+		return UEngineElementsLibrary::AcquireEditorComponentElementHandle(CastChecked<UActorComponent>(InObject), bAllowCreate);
 	}
-}
-
-void FObjectSelectionStore::EndBatchSelection(const bool InNotify)
-{
-	if (--SelectionMutex == 0)
-	{
-		const bool bSelectionChanged = bIsBatchDirty;
-		bIsBatchDirty = false;
+};
 
-		if (Sink && bSelectionChanged)
-		{
-			Sink->OnSelectedChanged(/*bSyncState*/false, !IsBatchSelecting() && InNotify);
-		}
-	}
-}
+USelection::FOnSelectionChanged						USelection::SelectionChangedEvent;
+USelection::FOnSelectionChanged						USelection::SelectObjectEvent;
+FSimpleMulticastDelegate							USelection::SelectNoneEvent;
+USelection::FOnSelectionElementSelectionPtrChanged	USelection::SelectionElementSelectionPtrChanged;
 
-bool FObjectSelectionStore::IsBatchSelecting() const
+USelection* USelection::CreateObjectSelection(UObject* InOuter, FName InName, EObjectFlags InFlags)
 {
-	return SelectionMutex > 0;
+	USelection* Selection = NewObject<USelection>(InOuter, InName, InFlags);
+	Selection->Initialize(MakeShared<FObjectSelectionElementBridge>());
+	return Selection;
 }
 
-void FObjectSelectionStore::ForceBatchDirty()
+USelection* USelection::CreateActorSelection(UObject* InOuter, FName InName, EObjectFlags InFlags)
 {
-	bIsBatchDirty = true;
+	USelection* Selection = NewObject<USelection>(InOuter, InName, InFlags);
+	Selection->Initialize(MakeShared<FActorSelectionElementBridge>());
+	return Selection;
 }
 
-
-class FElementSelectionStore : public FSelectionStoreBase
-{
-public:
-	explicit FElementSelectionStore(const UClass* InObjectBaseClass);
-	virtual ~FElementSelectionStore();
-
-	//~ ISelectionStore interface
-	virtual void SetElementList(UTypedElementList* InElementList) override;
-	virtual UTypedElementList* GetElementList() const override;
-	virtual int32 GetNumObjects() const override;
-	virtual UObject* GetObjectAtIndex(const int32 InIndex) const override;
-	virtual bool IsValidObjectToSelect(const UObject* InObject) const override;
-	virtual bool IsObjectSelected(const UObject* InObject) const override;
-	virtual void SelectObject(UObject* InObject) override;
-	virtual void DeselectObject(UObject* InObject) override;
-	virtual int32 DeselectObjects(TFunctionRef<bool(UObject*)> InPredicate) override;
-	virtual void BeginBatchSelection() override;
-	virtual void EndBatchSelection(const bool InNotify) override;
-	virtual bool IsBatchSelecting() const override;
-	virtual void ForceBatchDirty() override;
-
-	virtual FTypedElementHandle GetElementHandleForObject(const UObject* InObject, const bool bAllowCreate = true) const = 0;
-
-private:
-	void OnElementListPreChangeEvent(const UTypedElementList* InElementList);
-	void OnElementListSyncEvent(const UTypedElementList* InElementList, FTypedElementListLegacySync::ESyncType InSyncType, const FTypedElementHandle& InElementHandle, bool bIsWithinBatchOperation);
-	UObject* GetObjectForElementHandle(const FTypedElementHandle& InElementHandle) const;
-
-	/** External list of elements that this selection instance is bridging. This will be null unless a level editor instance exists. */
-	UTypedElementList* ElementList = nullptr;
-};
-
-FElementSelectionStore::FElementSelectionStore(const UClass* InObjectBaseClass)
-	: FSelectionStoreBase(InObjectBaseClass)
+USelection* USelection::CreateComponentSelection(UObject* InOuter, FName InName, EObjectFlags InFlags)
 {
+	USelection* Selection = NewObject<USelection>(InOuter, InName, InFlags);
+	Selection->Initialize(MakeShared<FComponentSelectionElementBridge>());
+	return Selection;
 }
 
-FElementSelectionStore::~FElementSelectionStore()
+void USelection::Initialize(TSharedRef<ISelectionElementBridge>&& InSelectionElementBridge)
 {
-	if (ElementList)
-	{
-		ElementList->OnPreChange().RemoveAll(this);
-		ElementList->Legacy_GetSync().OnSyncEvent().RemoveAll(this);
-	}
+	SelectionElementBridge = MoveTemp(InSelectionElementBridge);
+	SyncSelectedClasses();
 }
 
-void FElementSelectionStore::SetElementList(UTypedElementList* InElementList)
+void USelection::SetElementSelectionSet(UTypedElementSelectionSet* InElementSelectionSet)
 {
-	if (ElementList)
+	if (ElementSelectionSet)
 	{
-		ElementList->OnPreChange().RemoveAll(this);
-		ElementList->Legacy_GetSync().OnSyncEvent().RemoveAll(this);
+		ElementSelectionSet->OnPreChange().RemoveAll(this);
+		ElementSelectionSet->Legacy_GetElementListSync().OnSyncEvent().RemoveAll(this);
 	}
 
-	ElementList = InElementList;
+	UTypedElementSelectionSet* OldElementSelectionSet = ElementSelectionSet;
+	ElementSelectionSet = InElementSelectionSet;
 
-	if (ElementList)
+	if (ElementSelectionSet)
 	{
-		ElementList->OnPreChange().AddRaw(this, &FElementSelectionStore::OnElementListPreChangeEvent);
-		ElementList->Legacy_GetSync().OnSyncEvent().AddRaw(this, &FElementSelectionStore::OnElementListSyncEvent);
+		ElementSelectionSet->OnPreChange().AddUObject(this, &USelection::OnElementPreChangeEvent);
+		ElementSelectionSet->Legacy_GetElementListSync().OnSyncEvent().AddUObject(this, &USelection::OnElementListSyncEvent);
 	}
+	
+	USelection::SelectionElementSelectionPtrChanged.Broadcast(this, OldElementSelectionSet, ElementSelectionSet);
 }
 
-UTypedElementList* FElementSelectionStore::GetElementList() const
+UTypedElementSelectionSet* USelection::GetElementSelectionSet() const
 {
-	return ElementList;
+	return ElementSelectionSet;
 }
 
-int32 FElementSelectionStore::GetNumObjects() const
+int32 USelection::Num() const
 {
-	return ElementList
-		? ElementList->Num()
+	return ElementSelectionSet
+		? ElementSelectionSet->GetElementList()->Num()
 		: 0;
 }
 
-UObject* FElementSelectionStore::GetObjectAtIndex(const int32 InIndex) const
+UObject* USelection::GetSelectedObject(const int32 InIndex) const
 {
-	if (ElementList && ElementList->IsValidIndex(InIndex))
+	if (ElementSelectionSet)
 	{
-		const FTypedElementHandle& ElementHandle = ElementList->GetElementHandleAt(InIndex);
-		return GetObjectForElementHandle(ElementHandle);
+		const UTypedElementList* ElementList = ElementSelectionSet->GetElementList();
+		if (ElementList->IsValidIndex(InIndex))
+		{
+			const FTypedElementHandle ElementHandle = ElementList->GetElementHandleAt(InIndex);
+			return GetObjectForElementHandle(ElementHandle);
+		}
 	}
 
 	return nullptr;
 }
 
-bool FElementSelectionStore::IsValidObjectToSelect(const UObject* InObject) const
+void USelection::BeginBatchSelectOperation()
 {
-	bool bIsValid = FSelectionStoreBase::IsValidObjectToSelect(InObject);
-
-	if (bIsValid && ElementList)
+	if (ElementSelectionSet)
 	{
-		TTypedElement<UTypedElementSelectionInterface> SelectionElement = ElementList->GetElement<UTypedElementSelectionInterface>(GetElementHandleForObject(InObject));
-		if (SelectionElement)
-		{
-			bIsValid &= SelectionElement.Legacy_GetSelectionObject() == InObject;
-		}
-		else
-		{
-			// Elements must implement the selection interface in order to be selected!
-			bIsValid = false;
-		}
+		ElementSelectionSet->Legacy_GetElementListSync().BeginBatchOperation();
 	}
-
-	return bIsValid;
 }
 
-bool FElementSelectionStore::IsObjectSelected(const UObject* InObject) const
+void USelection::EndBatchSelectOperation(bool bNotify)
 {
-	if (ElementList)
+	if (ElementSelectionSet)
 	{
-		if (const FTypedElementHandle ElementHandle = GetElementHandleForObject(InObject, /*bAllowCreate*/false))
-		{
-			return ElementList->Contains(ElementHandle);
-		}
+		ElementSelectionSet->Legacy_GetElementListSync().EndBatchOperation(bNotify);
 	}
-
-	return false;
 }
 
-void FElementSelectionStore::SelectObject(UObject* InObject)
+bool USelection::IsBatchSelecting() const
 {
-	if (ElementList)
-	{
-		FTypedElementHandle ElementHandle = GetElementHandleForObject(InObject);
-		check(ElementHandle);
-		ElementList->Add(MoveTemp(ElementHandle));
-	}
+	return ElementSelectionSet && ElementSelectionSet->Legacy_GetElementListSync().IsRunningBatchOperation();
 }
 
-void FElementSelectionStore::DeselectObject(UObject* InObject)
+bool USelection::IsValidObjectToSelect(const UObject* InObject) const
 {
-	if (ElementList)
+	bool bIsValid = SelectionElementBridge->IsValidObjectType(InObject);
+
+	if (bIsValid && ElementSelectionSet)
 	{
-		if (const FTypedElementHandle ElementHandle = GetElementHandleForObject(InObject, /*bAllowCreate*/false))
+		TTypedElement<UTypedElementObjectInterface> ObjectElement = ElementSelectionSet->GetElementList()->GetElement<UTypedElementObjectInterface>(SelectionElementBridge->GetElementHandleForObject(InObject));
+		if (ObjectElement)
 		{
-			ElementList->Remove(ElementHandle);
+			bIsValid &= ObjectElement.GetObject() == InObject;
 		}
-	}
-}
-
-int32 FElementSelectionStore::DeselectObjects(TFunctionRef<bool(UObject*)> InPredicate)
-{
-	if (ElementList)
-	{
-		return ElementList->RemoveAll([this, &InPredicate](const FTypedElementHandle& InElementHandle)
+		else
 		{
-			UObject* Object = GetObjectForElementHandle(InElementHandle);
-			return Object && InPredicate(Object);
-		});
-	}
-
-	return 0;
-}
-
-void FElementSelectionStore::BeginBatchSelection()
-{
-	if (ElementList)
-	{
-		ElementList->Legacy_GetSync().BeginBatchOperation();
-	}
-}
-
-void FElementSelectionStore::EndBatchSelection(const bool InNotify)
-{
-	if (ElementList)
-	{
-		ElementList->Legacy_GetSync().EndBatchOperation(InNotify);
-	}
-}
-
-bool FElementSelectionStore::IsBatchSelecting() const
-{
-	return ElementList && ElementList->Legacy_GetSync().IsRunningBatchOperation();
-}
-
-void FElementSelectionStore::ForceBatchDirty()
-{
-	if (ElementList)
-	{
-		ElementList->Legacy_GetSync().ForceBatchOperationDirty();
+			// Elements must implement the object interface in order to be selected!
+			bIsValid = false;
+		}
 	}
-}
 
-void FElementSelectionStore::OnElementListPreChangeEvent(const UTypedElementList* InElementList)
-{
-	if (Sink)
-	{
-		Sink->OnPreModify();
-	}
+	return bIsValid;
 }
 
-void FElementSelectionStore::OnElementListSyncEvent(const UTypedElementList* InElementList, FTypedElementListLegacySync::ESyncType InSyncType, const FTypedElementHandle& InElementHandle, bool bIsWithinBatchOperation)
+UObject* USelection::GetObjectForElementHandle(const FTypedElementHandle& InElementHandle) const
 {
-	if (Sink)
-	{
-		switch (InSyncType)
-		{
-		case FTypedElementListLegacySync::ESyncType::Added:
-			if (UObject* Object = GetObjectForElementHandle(InElementHandle))
-			{
-				Sink->OnObjectSelected(Object, !bIsWithinBatchOperation);
-			}
-			break;
-
-		case FTypedElementListLegacySync::ESyncType::Removed:
-			if (UObject* Object = GetObjectForElementHandle(InElementHandle))
-			{
-				Sink->OnObjectDeselected(Object, !bIsWithinBatchOperation);
-			}
-			break;
-
-		case FTypedElementListLegacySync::ESyncType::BatchComplete:
-			Sink->OnSelectedChanged(/*bSyncState*/false, !bIsWithinBatchOperation);
-			break;
-
-		default:
-			Sink->OnSelectedChanged(/*bSyncState*/true, !bIsWithinBatchOperation);
-			break;
-		}
-	}
-}
+	check(InElementHandle && ElementSelectionSet);
 
-UObject* FElementSelectionStore::GetObjectForElementHandle(const FTypedElementHandle& InElementHandle) const
-{
-	check(InElementHandle && ElementList);
-	
-	TTypedElement<UTypedElementSelectionInterface> SelectionElement = ElementList->GetElement<UTypedElementSelectionInterface>(InElementHandle);
-	check(SelectionElement);
+	TTypedElement<UTypedElementObjectInterface> ObjectElement = ElementSelectionSet->GetElementList()->GetElement<UTypedElementObjectInterface>(InElementHandle);
+	check(ObjectElement);
 
-	UObject* Object = SelectionElement.Legacy_GetSelectionObject();
-	if (Object && Object->IsA(ObjectBaseClass))
+	UObject* Object = ObjectElement.GetObject();
+	if (Object && SelectionElementBridge->IsValidObjectType(Object))
 	{
 		return Object;
 	}
@@ -453,118 +201,45 @@ UObject* FElementSelectionStore::GetObjectForElementHandle(const FTypedElementHa
 	return nullptr;
 }
 
-
-class FActorElementSelectionStore : public FElementSelectionStore
-{
-public:
-	FActorElementSelectionStore();
-
-	//~ FElementSelectionStore interface
-	virtual FTypedElementHandle GetElementHandleForObject(const UObject* InObject, const bool bAllowCreate = true) const override;
-};
-
-FActorElementSelectionStore::FActorElementSelectionStore()
-	: FElementSelectionStore(AActor::StaticClass())
-{
-}
-
-FTypedElementHandle FActorElementSelectionStore::GetElementHandleForObject(const UObject* InObject, const bool bAllowCreate) const
-{
-	check(InObject);
-
-	const AActor* Actor = CastChecked<AActor>(InObject);
-	return Actor->AcquireEditorElementHandle(bAllowCreate);
-}
-
-
-class FComponentElementSelectionStore : public FElementSelectionStore
-{
-public:
-	FComponentElementSelectionStore();
-
-	//~ FElementSelectionStore interface
-	virtual FTypedElementHandle GetElementHandleForObject(const UObject* InObject, const bool bAllowCreate = true) const override;
-};
-
-FComponentElementSelectionStore::FComponentElementSelectionStore()
-	: FElementSelectionStore(UActorComponent::StaticClass())
-{
-}
-
-FTypedElementHandle FComponentElementSelectionStore::GetElementHandleForObject(const UObject* InObject, const bool bAllowCreate) const
-{
-	check(InObject);
-
-	const UActorComponent* Component = CastChecked<UActorComponent>(InObject);
-	return Component->AcquireEditorElementHandle(bAllowCreate);
-}
-
-} // namespace Selection_Private
-
-USelection::FOnSelectionChanged				USelection::SelectionChangedEvent;
-USelection::FOnSelectionChanged				USelection::SelectObjectEvent;
-FSimpleMulticastDelegate					USelection::SelectNoneEvent;
-USelection::FOnSelectionElementListChanged	USelection::SelectionElementListChanged;
-
-USelection* USelection::CreateObjectSelection(FUObjectAnnotationSparseBool* InSelectionAnnotation, UObject* InOuter, FName InName, EObjectFlags InFlags)
-{
-	USelection* Selection = NewObject<USelection>(InOuter, InName, InFlags);
-	Selection->Initialize(InSelectionAnnotation, MakeShared<Selection_Private::FObjectSelectionStore>(nullptr));
-	return Selection;
-}
-
-USelection* USelection::CreateActorSelection(FUObjectAnnotationSparseBool* InSelectionAnnotation, UObject* InOuter, FName InName, EObjectFlags InFlags)
-{
-	USelection* Selection = NewObject<USelection>(InOuter, InName, InFlags);
-	Selection->Initialize(InSelectionAnnotation, MakeShared<Selection_Private::FActorElementSelectionStore>());
-	return Selection;
-}
-
-USelection* USelection::CreateComponentSelection(FUObjectAnnotationSparseBool* InSelectionAnnotation, UObject* InOuter, FName InName, EObjectFlags InFlags)
+void USelection::OnElementPreChangeEvent(const UTypedElementSelectionSet* InElementSelectionSet)
 {
-	USelection* Selection = NewObject<USelection>(InOuter, InName, InFlags);
-	Selection->Initialize(InSelectionAnnotation, MakeShared<Selection_Private::FComponentElementSelectionStore>());
-	return Selection;
+	check(InElementSelectionSet == ElementSelectionSet);
+	Modify();
 }
 
-void USelection::Initialize(FUObjectAnnotationSparseBool* InSelectionAnnotation, TSharedRef<Selection_Private::ISelectionStore>&& InSelectionStore)
+void USelection::OnElementListSyncEvent(const UTypedElementList* InElementList, FTypedElementListLegacySync::ESyncType InSyncType, const FTypedElementHandle& InElementHandle, bool bIsWithinBatchOperation)
 {
-	SelectionStore = MoveTemp(InSelectionStore);
-	SelectionStore->SetSink(this);
-
-	if (InSelectionAnnotation)
+	check(InElementList == ElementSelectionSet->GetElementList());
+	switch (InSyncType)
 	{
-		SelectionAnnotation = InSelectionAnnotation;
-		bOwnsSelectionAnnotation = false;
-	}
-	else
-	{
-		SelectionAnnotation = new FUObjectAnnotationSparseBool;
-		bOwnsSelectionAnnotation = true;
-	}
+	case FTypedElementListLegacySync::ESyncType::Added:
+		if (UObject* Object = GetObjectForElementHandle(InElementHandle))
+		{
+			OnObjectSelected(Object, !bIsWithinBatchOperation);
+		}
+		break;
 
-	SyncSelectedState();
-}
+	case FTypedElementListLegacySync::ESyncType::Removed:
+		if (UObject* Object = GetObjectForElementHandle(InElementHandle))
+		{
+			OnObjectDeselected(Object, !bIsWithinBatchOperation);
+		}
+		break;
 
-void USelection::SetElementSelectionSet(UTypedElementSelectionSet* InElementSelectionSet)
-{
-	ElementSelectionSet = InElementSelectionSet;
-	UTypedElementList* OldElementList = SelectionStore->GetElementList();
-	SelectionStore->SetElementList(InElementSelectionSet ? InElementSelectionSet->GetMutableElementList() : nullptr);
-	USelection::SelectionElementListChanged.Broadcast(this, OldElementList, SelectionStore->GetElementList());
-}
+	case FTypedElementListLegacySync::ESyncType::BatchComplete:
+		OnSelectedChanged(/*bSyncState*/false, !bIsWithinBatchOperation);
+		break;
 
-void USelection::OnPreModify()
-{
-	Modify();
+	default:
+		OnSelectedChanged(/*bSyncState*/true, !bIsWithinBatchOperation);
+		break;
+	}
 }
 
 void USelection::OnObjectSelected(UObject* InObject, const bool bNotify)
 {
 	check(InObject);
 
-	SelectionAnnotation->Set(InObject);
-
 	if (FSelectedClassInfo* SelectedClassInfo = SelectedClasses.Find(InObject->GetClass()))
 	{
 		++SelectedClassInfo->SelectionCount;
@@ -586,8 +261,6 @@ void USelection::OnObjectDeselected(UObject* InObject, const bool bNotify)
 {
 	check(InObject);
 
-	SelectionAnnotation->Clear(InObject);
-
 	FSetElementId Id = SelectedClasses.FindId(InObject->GetClass());
 	if (Id.IsValidId())
 	{
@@ -612,7 +285,7 @@ void USelection::OnSelectedChanged(const bool bSyncState, const bool bNotify)
 {
 	if (bSyncState)
 	{
-		SyncSelectedState();
+		SyncSelectedClasses();
 	}
 	
 	if (bNotify)
@@ -621,17 +294,14 @@ void USelection::OnSelectedChanged(const bool bSyncState, const bool bNotify)
 	}
 }
 
-void USelection::SyncSelectedState()
+void USelection::SyncSelectedClasses()
 {
 	SelectedClasses.Reset();
-	SelectionAnnotation->ClearAll();
 
 	for (int32 Idx = 0; Idx < Num(); ++Idx)
 	{
 		if (UObject* Object = GetSelectedObject(Idx))
 		{
-			SelectionAnnotation->Set(Object);
-
 			if (FSelectedClassInfo* SelectedClassInfo = SelectedClasses.Find(Object->GetClass()))
 			{
 				++SelectedClassInfo->SelectionCount;
@@ -647,21 +317,24 @@ void USelection::SyncSelectedState()
 
 void USelection::Select(UObject* InObject)
 {
-	check(InObject && SelectionStore->IsValidObjectToSelect(InObject));
-
-	if (!SelectionAnnotation->Get(InObject))
+	check(InObject && IsValidObjectToSelect(InObject));
+	if (ElementSelectionSet)
 	{
-		SelectionStore->SelectObject(InObject);
+		FTypedElementHandle ElementHandle = SelectionElementBridge->GetElementHandleForObject(InObject);
+		check(ElementHandle);
+		ElementSelectionSet->SelectElement(ElementHandle, FTypedElementSelectionOptions().SetAllowHidden(true).SetAllowGroups(false).SetWarnIfLocked(false));
 	}
 }
 
 void USelection::Deselect(UObject* InObject)
 {
 	check(InObject);
-
-	if (SelectionAnnotation->Get(InObject))
+	if (ElementSelectionSet)
 	{
-		SelectionStore->DeselectObject(InObject);
+		if (const FTypedElementHandle ElementHandle = SelectionElementBridge->GetElementHandleForObject(InObject, /*bAllowCreate*/false))
+		{
+			ElementSelectionSet->DeselectElement(ElementHandle, FTypedElementSelectionOptions().SetAllowHidden(true).SetAllowGroups(false).SetWarnIfLocked(false));
+		}
 	}
 }
 
@@ -684,23 +357,42 @@ void USelection::ToggleSelect(UObject* InObject)
 
 void USelection::DeselectAll(UClass* InClass)
 {
-	// Fast path for deselecting all UObjects with any flags
-	if (InClass == UObject::StaticClass())
+	if (!ElementSelectionSet)
 	{
-		InClass = nullptr;
+		return;
 	}
 
-	SelectionStore->DeselectObjects([InClass](UObject* InObject)
+	UClass* ClassToDeselect = InClass;
+	if (!ClassToDeselect)
 	{
-		return !InClass || InObject->IsA(InClass);
-	});
+		ClassToDeselect = UObject::StaticClass();
+	}
+
+	TArray<UObject*> ObjectsToDeselect;
+	GetSelectedObjects(ClassToDeselect, ObjectsToDeselect);
+
+	TArray<FTypedElementHandle, TInlineAllocator<256>> ElementsToDeselect;
+	ElementsToDeselect.Reserve(ObjectsToDeselect.Num());
+	for (UObject* ObjectToDeselect : ObjectsToDeselect)
+	{
+		if (FTypedElementHandle ElementHandle = SelectionElementBridge->GetElementHandleForObject(ObjectToDeselect, /*bAllowCreate*/false))
+		{
+			ElementsToDeselect.Add(MoveTemp(ElementHandle));
+		}
+	}
+
+	if (ElementsToDeselect.Num() > 0)
+	{
+		ElementSelectionSet->DeselectElements(ElementsToDeselect, FTypedElementSelectionOptions().SetAllowHidden(true).SetAllowGroups(false).SetWarnIfLocked(false));
+	}
 }
 
 void USelection::ForceBatchDirty()
 {
 	if (IsBatchSelecting())
 	{
-		SelectionStore->ForceBatchDirty();
+		check(ElementSelectionSet);
+		ElementSelectionSet->Legacy_GetElementListSync().ForceBatchOperationDirty();
 	}
 }
 
@@ -716,7 +408,15 @@ void USelection::NoteUnknownSelectionChanged()
 
 bool USelection::IsSelected(const UObject* InObject) const
 {
-	return InObject && SelectionStore->IsObjectSelected(InObject);
+	if (InObject && ElementSelectionSet)
+	{
+		if (const FTypedElementHandle ElementHandle = SelectionElementBridge->GetElementHandleForObject(InObject, /*bAllowCreate*/false))
+		{
+			return ElementSelectionSet->IsElementSelected(ElementHandle, FTypedElementIsSelectedOptions());
+		}
+	}
+
+	return false;
 }
 
 void USelection::Serialize(FArchive& Ar)
@@ -742,10 +442,10 @@ void USelection::Serialize(FArchive& Ar)
 		{
 			if (UObject* Object = ObjectPtr.Get(true))
 			{
-				SelectionStore->SelectObject(Object);
+				Select(Object);
 			}
 		}
-		SyncSelectedState();
+		SyncSelectedClasses();
 
 		EndBatchSelectOperation(/*bNotify*/false);
 	}
@@ -765,16 +465,3 @@ bool USelection::Modify(bool bAlwaysMarkDirty/* =true */)
 
 	return Super::Modify(bAlwaysMarkDirty);
 }
-
-void USelection::BeginDestroy()
-{
-	Super::BeginDestroy();
-
-	SelectionStore.Reset();
-
-	if (bOwnsSelectionAnnotation)
-	{
-		delete SelectionAnnotation;
-		SelectionAnnotation = nullptr;
-	}
-}
diff --git a/Engine/Source/Editor/UnrealEd/Private/UnrealEdEngine.cpp b/Engine/Source/Editor/UnrealEd/Private/UnrealEdEngine.cpp
--- a/Engine/Source/Editor/UnrealEd/Private/UnrealEdEngine.cpp
+++ b/Engine/Source/Editor/UnrealEd/Private/UnrealEdEngine.cpp
@@ -123,7 +123,7 @@ void UUnrealEdEngine::Init(IEngineLoop* InEngineLoop)
 	FEditorSupportDelegates::PostWindowsMessage.AddUObject(this, &UUnrealEdEngine::OnPostWindowsMessage);
 
 	USelection::SelectionChangedEvent.AddUObject(this, &UUnrealEdEngine::OnEditorSelectionChanged);
-	USelection::SelectionElementListChanged.AddUObject(this, &UUnrealEdEngine::OnEditorSelectionElementListPtrChanged);
+	USelection::SelectionElementSelectionPtrChanged.AddUObject(this, &UUnrealEdEngine::OnEditorElementSelectionPtrChanged);
 
 	// Initialize the snap manager
 	FSnappingUtils::InitEditorSnappingTools();
diff --git a/Engine/Source/Editor/UnrealEd/Public/AssetEditorModeManager.h b/Engine/Source/Editor/UnrealEd/Public/AssetEditorModeManager.h
--- a/Engine/Source/Editor/UnrealEd/Public/AssetEditorModeManager.h
+++ b/Engine/Source/Editor/UnrealEd/Public/AssetEditorModeManager.h
@@ -6,7 +6,6 @@
 #include "EditorModeManager.h"
 
 class FPreviewScene;
-class UTypedElementList;
 
 //////////////////////////////////////////////////////////////////////////
 // FAssetEditorModeManager
@@ -32,5 +31,4 @@ protected:
 	class USelection* ObjectSet = nullptr;
 	class USelection* ComponentSet = nullptr;
 	class FPreviewScene* PreviewScene = nullptr;
-	UTypedElementList* SelectedElements = nullptr;
 };
diff --git a/Engine/Source/Editor/UnrealEd/Public/Selection.h b/Engine/Source/Editor/UnrealEd/Public/Selection.h
--- a/Engine/Source/Editor/UnrealEd/Public/Selection.h
+++ b/Engine/Source/Editor/UnrealEd/Public/Selection.h
@@ -7,129 +7,17 @@
 #include "UObject/UObjectGlobals.h"
 #include "UObject/Object.h"
 #include "Components/ActorComponent.h"
+#include "Elements/TypedElementSelectionSet.h"
 #include "Selection.generated.h"
 
-class UTypedElementList;
-class UTypedElementSelectionSet;
-
-namespace Selection_Private
-{
-
-class ISelectionStoreSink
-{
-public:
-	virtual ~ISelectionStoreSink() = default;
-
-	/**
-	 * Called prior to potential changes being made within the underlying store.
-	 */
-	virtual void OnPreModify() = 0;
-
-	/**
-	 * Called when the given object is selected within the underlying store.
-	 */
-	virtual void OnObjectSelected(UObject* InObject, const bool bNotify) = 0;
-
-	/**
-	 * Called when the given object is deselected within the underlying store.
-	 */
-	virtual void OnObjectDeselected(UObject* InObject, const bool bNotify) = 0;
-
-	/**
-	 * Called when the the underlying store changes in an unknown way.
-	 */
-	virtual void OnSelectedChanged(const bool bSyncState, const bool bNotify) = 0;
-};
-
-class ISelectionStore
-{
-public:
-	virtual ~ISelectionStore() = default;
-
-	/**
-	 * Set the notification sink instance for this store.
-	 */
-	virtual void SetSink(ISelectionStoreSink* InSink) = 0;
-
-	/**
-	 * Set the element list instance for this store.
-	 * @note Does nothing if element list stores aren't enabled. Asserts for non-element list stores if they are!
-	 */
-	virtual void SetElementList(UTypedElementList* InElementList) = 0;
-
-	/**
-	 * Get the element list instance for this selection set, if any.
-	 */
-	virtual UTypedElementList* GetElementList() const = 0;
-
-	/**
-	 * Get the number of objects within the underlying store.
-	 * This is the total number of objects within the store (for use as an upper limit of GetObjectAtIndex),
-	 * however not all of those objects may be valid so it should not be used as a public selection count.
-	 */
-	virtual int32 GetNumObjects() const = 0;
-
-	/**
-	 * Get the object at the internal index of the underlying store.
-	 * This object may be null, both in cases where the underlying store is using weak references, 
-	 * and also in the case that the object does not match the type managed by the underlying store.
-	 */
-	virtual UObject* GetObjectAtIndex(const int32 InIndex) const = 0;
-
-	/**
-	 * Test to see whether the given object is valid to be added to the underlying store.
-	 */
-	virtual bool IsValidObjectToSelect(const UObject* InObject) const = 0;
-
-	/**
-	 * Test to see whether the given object is currently in the underlying store.
-	 */
-	virtual bool IsObjectSelected(const UObject* InObject) const = 0;
-
-	/**
-	 * Add the given object to the underlying store.
-	 */
-	virtual void SelectObject(UObject* InObject) = 0;
-
-	/**
-	 * Remove the given object from the underlying store.
-	 */
-	virtual void DeselectObject(UObject* InObject) = 0;
-
-	/**
-	 * Remove any objects that match the predicate from the underlying store.
-	 */
-	virtual int32 DeselectObjects(TFunctionRef<bool(UObject*)> InPredicate) = 0;
-
-	/**
-	 * Called to begin a batch selection.
-	 */
-	virtual void BeginBatchSelection() = 0;
-
-	/**
-	 * Called to end a batch selection.
-	 */
-	virtual void EndBatchSelection(const bool InNotify) = 0;
-
-	/**
-	 * Are we currently batch selecting?
-	 */
-	virtual bool IsBatchSelecting() const = 0;
-
-	/**
-	 * Forcibly mark this batch as being dirty.
-	 */
-	virtual void ForceBatchDirty() = 0;
-};
-
-} // namespace Selection_Private
+class ISelectionElementBridge;
 
 /**
  * Manages selections of objects.
  * Used in the editor for selecting objects in the various browser windows.
  */
 UCLASS(transient)
-class UNREALED_API USelection : public UObject, public Selection_Private::ISelectionStoreSink
+class UNREALED_API USelection : public UObject
 {
 	GENERATED_BODY()
 
@@ -168,15 +56,15 @@ private:
 	friend class TSelectionIterator;
 
 public:
-	static USelection* CreateObjectSelection(FUObjectAnnotationSparseBool* InSelectionAnnotation, UObject* InOuter = GetTransientPackage(), FName InName = NAME_None, EObjectFlags InFlags = RF_NoFlags);
-	static USelection* CreateActorSelection(FUObjectAnnotationSparseBool* InSelectionAnnotation, UObject* InOuter = GetTransientPackage(), FName InName = NAME_None, EObjectFlags InFlags = RF_NoFlags);
-	static USelection* CreateComponentSelection(FUObjectAnnotationSparseBool* InSelectionAnnotation, UObject* InOuter = GetTransientPackage(), FName InName = NAME_None, EObjectFlags InFlags = RF_NoFlags);
+	static USelection* CreateObjectSelection(UObject* InOuter = GetTransientPackage(), FName InName = NAME_None, EObjectFlags InFlags = RF_NoFlags);
+	static USelection* CreateActorSelection(UObject* InOuter = GetTransientPackage(), FName InName = NAME_None, EObjectFlags InFlags = RF_NoFlags);
+	static USelection* CreateComponentSelection(UObject* InOuter = GetTransientPackage(), FName InName = NAME_None, EObjectFlags InFlags = RF_NoFlags);
 
 	/** Params: UObject* NewSelection */
 	DECLARE_EVENT_OneParam(USelection, FOnSelectionChanged, UObject*);
 
-	/** Event fired when the element list for a selection is changed */
-	DECLARE_EVENT_ThreeParams(USelection, FOnSelectionElementListChanged, USelection* /*Selection*/, UTypedElementList* /*OldElementList*/, UTypedElementList* /*NewElementList*/);
+	/** Event fired when the typed element selection set for a selection is changed */
+	DECLARE_EVENT_ThreeParams(USelection, FOnSelectionElementSelectionPtrChanged, USelection* /*Selection*/, UTypedElementSelectionSet* /*OldSelectionSet*/, UTypedElementSelectionSet* /*NewSelectionSet*/);
 
 	/** Called when selection in editor has changed */
 	static FOnSelectionChanged SelectionChangedEvent;
@@ -184,8 +72,8 @@ public:
 	static FOnSelectionChanged SelectObjectEvent;
 	/** Called to deselect everything */
 	static FSimpleMulticastDelegate SelectNoneEvent;
-	/** Called when the element list for a selection has changed */
-	static FOnSelectionElementListChanged SelectionElementListChanged;
+	/** Called when the assigned typed element selection pointer set for a selection is changed */
+	static FOnSelectionElementSelectionPtrChanged SelectionElementSelectionPtrChanged;
 
 	typedef ClassArray::TIterator TClassIterator;
 	typedef ClassArray::TConstIterator TClassConstIterator;
@@ -193,26 +81,6 @@ public:
 	TClassIterator			ClassItor()				{ return TClassIterator( SelectedClasses ); }
 	TClassConstIterator		ClassConstItor() const	{ return TClassConstIterator( SelectedClasses ); }
 
-	/**
-	 * Set the element list instance for this selection set.
-	 * @note Does nothing if element list stores aren't enabled. Asserts for non-element list stores if they are!
-	 */
-	void SetElementList(UTypedElementList* InElementList)
-	{
-		ElementSelectionSet.Reset();
-		UTypedElementList* OldElementList = SelectionStore->GetElementList();
-		SelectionStore->SetElementList(InElementList);
-		USelection::SelectionElementListChanged.Broadcast(this, OldElementList, SelectionStore->GetElementList());
-	}
-
-	/**
-	 * Get the element list instance for this selection set, if any.
-	 */
-	UTypedElementList* GetElementList() const
-	{
-		return SelectionStore->GetElementList();
-	}
-
 	/**
 	 * Set the element selection set instance for this selection set.
 	 * @note Also sets the element list instance.
@@ -222,10 +90,7 @@ public:
 	/**
 	 * Get the element selection set instance for this selection set, if any.
 	 */
-	UTypedElementSelectionSet* GetElementSelectionSet() const
-	{
-		return ElementSelectionSet.Get();
-	}
+	UTypedElementSelectionSet* GetElementSelectionSet() const;
 
 	/**
 	 * Returns the number of objects in the selection set.  This function is used by clients in
@@ -235,51 +100,28 @@ public:
 	 * 
 	 * @return		Number of objects in the selection set.
 	 */
-	int32 Num() const
-	{
-		return SelectionStore->GetNumObjects();
-	}
-
-	/**
-	 * @return	The Index'th selected objects.  May be NULL.
-	 */
-	UObject* GetSelectedObject(const int32 InIndex)
-	{
-		return SelectionStore->GetObjectAtIndex(InIndex);
-	}
+	int32 Num() const;
 
 	/**
 	 * @return	The Index'th selected objects.  May be NULL.
 	 */
-	const UObject* GetSelectedObject(const int32 InIndex) const
-	{
-		return SelectionStore->GetObjectAtIndex(InIndex);
-	}
+	UObject* GetSelectedObject(const int32 InIndex) const;
 
 	/**
 	 * Call before beginning selection operations
 	 */
-	void BeginBatchSelectOperation()
-	{
-		SelectionStore->BeginBatchSelection();
-	}
+	void BeginBatchSelectOperation();
 
 	/**
 	 * Should be called when selection operations are complete.  If all selection operations are complete, notifies all listeners
 	 * that the selection has been changed.
 	 */
-	void EndBatchSelectOperation(bool bNotify = true)
-	{
-		SelectionStore->EndBatchSelection(bNotify);
-	}
+	void EndBatchSelectOperation(bool bNotify = true);
 
 	/**
 	 * @return	Returns whether or not the selection object is currently in the middle of a batch select block.
 	 */
-	bool IsBatchSelecting() const
-	{
-		return SelectionStore->IsBatchSelecting();
-	}
+	bool IsBatchSelecting() const;
 
 	/**
 	 * Selects the specified object.
@@ -464,8 +306,6 @@ public:
 	//~ Begin UObject Interface
 	virtual void Serialize(FArchive& Ar) override;
 	virtual bool Modify( bool bAlwaysMarkDirty=true) override;
-	virtual void BeginDestroy() override;
-
 	//~ End UObject Interface
 
 
@@ -518,31 +358,30 @@ public:
 		return OutSelectedObjects.Num();
 	}
 
-protected:
-	/** Initializes the selection set with an annotation used to quickly look up selection state */
-	void Initialize(FUObjectAnnotationSparseBool* InSelectionAnnotation, TSharedRef<Selection_Private::ISelectionStore>&& InSelectionStore);
+private:
+	/** Initializes the selection set with its typed element bridge */
+	void Initialize(TSharedRef<ISelectionElementBridge>&& InSelectionElementBridge);
+
+	bool IsValidObjectToSelect(const UObject* InObject) const;
+	UObject* GetObjectForElementHandle(const FTypedElementHandle& InElementHandle) const;
+
+	void OnElementPreChangeEvent(const UTypedElementSelectionSet* InElementSelectionSet);
+	void OnElementListSyncEvent(const UTypedElementList* InElementList, FTypedElementListLegacySync::ESyncType InSyncType, const FTypedElementHandle& InElementHandle, bool bIsWithinBatchOperation);
 
-	//~ ISelectionStoreSink interface
-	virtual void OnPreModify() override;
-	virtual void OnObjectSelected(UObject* InObject, const bool bNotify) override;
-	virtual void OnObjectDeselected(UObject* InObject, const bool bNotify) override;
-	virtual void OnSelectedChanged(const bool bSyncState, const bool bNotify) override;
-	
-	/** Sync the state of the underlying selection store to the annotation and classes data */
-	void SyncSelectedState();
+	void OnObjectSelected(UObject* InObject, const bool bNotify);
+	void OnObjectDeselected(UObject* InObject, const bool bNotify);
+	void OnSelectedChanged(const bool bSyncState, const bool bNotify);
+	void SyncSelectedClasses();
 
-	/** Store of selected objects. */
-	TSharedPtr<Selection_Private::ISelectionStore> SelectionStore = nullptr;
+	/** Bridge from UObjects to their corresponding typed elements. */
+	TSharedPtr<ISelectionElementBridge> SelectionElementBridge;
 
 	/** Underlying element selection set (if any). */
-	TWeakObjectPtr<UTypedElementSelectionSet> ElementSelectionSet;
+	UPROPERTY()
+	UTypedElementSelectionSet* ElementSelectionSet = nullptr;
 
 	/** Tracks the most recently selected actor classes.  Used for UnrealEd menus. */
 	ClassArray SelectedClasses;
-	
-	/** Selection annotation for fast lookup */
-	FUObjectAnnotationSparseBool* SelectionAnnotation = nullptr;
-	bool bOwnsSelectionAnnotation = false;
 
 private:
 	// Hide IsSelected(), as calling IsSelected() on a selection set almost always indicates
diff --git a/Engine/Source/Runtime/Engine/Classes/Components/ActorComponent.h b/Engine/Source/Runtime/Engine/Classes/Components/ActorComponent.h
--- a/Engine/Source/Runtime/Engine/Classes/Components/ActorComponent.h
+++ b/Engine/Source/Runtime/Engine/Classes/Components/ActorComponent.h
@@ -15,7 +15,6 @@
 #include "Interfaces/Interface_AssetUserData.h"
 #include "UObject/UObjectAnnotation.h"
 #include "UObject/StructOnScope.h"
-#include "TypedElementHandle.h"
 #include "ComponentInstanceDataCache.h"
 #include "ActorComponent.generated.h"
 
@@ -545,14 +544,6 @@ private:
 	/** Private version without inlining that does *not* check Dedicated server build flags (which should already have been done). */
 	ENetMode InternalGetNetMode() const;
 
-public:
-	TTypedElementOwner<struct FComponentElementData> CreateComponentElement() const;
-	void DestroyComponentElement(TTypedElementOwner<struct FComponentElementData>& InOutComponentElement) const;
-#if WITH_EDITOR
-	UFUNCTION(BlueprintCallable, BlueprintPure=false, Category="TypedElementFramework|Component")
-	FTypedElementHandle AcquireEditorElementHandle(const bool bAllowCreate = true) const;
-#endif
-
 protected:
 	/** Return true if this component is in a state where it can be activated normally. */
 	virtual bool ShouldActivate() const;
diff --git a/Engine/Source/Runtime/Engine/Classes/GameFramework/Actor.h b/Engine/Source/Runtime/Engine/Classes/GameFramework/Actor.h
--- a/Engine/Source/Runtime/Engine/Classes/GameFramework/Actor.h
+++ b/Engine/Source/Runtime/Engine/Classes/GameFramework/Actor.h
@@ -15,7 +15,6 @@
 #include "ComponentInstanceDataCache.h"
 #include "Components/ChildActorComponent.h"
 #include "RenderCommandFence.h"
-#include "TypedElementHandle.h"
 #include "Misc/ITransaction.h"
 #include "Engine/Level.h"
 
@@ -3302,14 +3301,6 @@ public:
 	/** Returns the instance components array */
 	const TArray<UActorComponent*>& GetInstanceComponents() const;
 
-public:
-	TTypedElementOwner<struct FActorElementData> CreateActorElement() const;
-	void DestroyActorElement(TTypedElementOwner<struct FActorElementData>& InOutActorElement) const;
-#if WITH_EDITOR
-	UFUNCTION(BlueprintCallable, BlueprintPure=false, Category="TypedElementFramework|Actor")
-	FTypedElementHandle AcquireEditorElementHandle(const bool bAllowCreate = true) const;
-#endif
-
 	//~=============================================================================
 	// Navigation/AI related functions
 	
diff --git a/Engine/Source/Runtime/Engine/Private/Actor.cpp b/Engine/Source/Runtime/Engine/Private/Actor.cpp
--- a/Engine/Source/Runtime/Engine/Private/Actor.cpp
+++ b/Engine/Source/Runtime/Engine/Private/Actor.cpp
@@ -4,9 +4,7 @@
 #include "EngineDefines.h"
 #include "EngineStats.h"
 #include "EngineGlobals.h"
-#include "TypedElementRegistry.h"
-#include "TypedElementOwnerStore.h"
-#include "Elements/Actor/ActorElementData.h"
+#include "Elements/EngineElementsLibrary.h"
 #include "GameFramework/DamageType.h"
 #include "TimerManager.h"
 #include "GameFramework/Pawn.h"
@@ -74,7 +72,6 @@ FMakeNoiseDelegate AActor::MakeNoiseDelegate = FMakeNoiseDelegate::CreateStatic(
 
 #if WITH_EDITOR
 FUObjectAnnotationSparseBool GSelectedActorAnnotation;
-TTypedElementOwnerStore<FActorElementData, const AActor*> GActorElementOwnerStore;
 #endif
 
 #if !UE_BUILD_SHIPPING
@@ -231,50 +228,6 @@ void AActor::InitializeDefaults()
 	GridPlacement = EActorGridPlacement::Bounds;
 }
 
-TTypedElementOwner<FActorElementData> AActor::CreateActorElement() const
-{
-	UTypedElementRegistry* Registry = UTypedElementRegistry::GetInstance();
-
-	TTypedElementOwner<FActorElementData> ActorElement;
-	if (ensureMsgf(Registry, TEXT("Typed element was requested for '%s' before the registry was available! This usually means that NewObject was used instead of CreateDefaultSubobject during CDO construction."), *GetPathName()))
-	{
-		ActorElement = Registry->CreateElement<FActorElementData>(NAME_Actor);
-		ActorElement.GetDataChecked().Actor = const_cast<AActor*>(this);
-	}
-
-	return ActorElement;
-}
-
-void AActor::DestroyActorElement(TTypedElementOwner<FActorElementData>& InOutActorElement) const
-{
-	if (InOutActorElement)
-	{
-		checkf(InOutActorElement.GetDataChecked().Actor == this, TEXT("Actor element was not for this actor instance! %s"), *GetPathName());
-		UTypedElementRegistry::GetInstance()->DestroyElement(InOutActorElement);
-	}
-}
-
-#if WITH_EDITOR
-FTypedElementHandle AActor::AcquireEditorElementHandle(const bool bAllowCreate) const
-{
-	if (GIsEditor)
-	{
-		if (TTypedElementOwnerScopedAccess<FActorElementData> EditorElementScopedAccess = GActorElementOwnerStore.FindElementOwner(this))
-		{
-			return EditorElementScopedAccess->AcquireHandle();
-		}
-
-		if (bAllowCreate)
-		{
-			TTypedElementOwnerScopedAccess<FActorElementData> EditorElementScopedAccess = GActorElementOwnerStore.RegisterElementOwner(this, CreateActorElement());
-			return EditorElementScopedAccess->AcquireHandle();
-		}
-	}
-
-	return FTypedElementHandle();
-}
-#endif
-
 void FActorTickFunction::ExecuteTick(float DeltaTime, enum ELevelTick TickType, ENamedThreads::Type CurrentThread, const FGraphEventRef& MyCompletionGraphEvent)
 {
 	if (Target && !Target->IsPendingKillOrUnreachable())
@@ -429,10 +382,7 @@ void AActor::PostInitProperties()
 	Super::PostInitProperties();
 
 #if WITH_EDITOR
-	if (GIsEditor && !HasAnyFlags(RF_ClassDefaultObject | RF_ArchetypeObject))
-	{
-		GActorElementOwnerStore.RegisterElementOwner(this, CreateActorElement());
-	}
+	UEngineElementsLibrary::CreateEditorActorElement(this);
 #endif	// WITH_EDITOR
 
 	RemoteRole = (bReplicates ? ROLE_SimulatedProxy : ROLE_None);
@@ -793,10 +743,7 @@ void AActor::BeginDestroy()
 #endif // (CSV_PROFILER && !UE_BUILD_SHIPPING)
 
 #if WITH_EDITOR
-	if (TTypedElementOwner<FActorElementData> EditorElement = GActorElementOwnerStore.UnregisterElementOwner(this))
-	{
-		DestroyActorElement(EditorElement);
-	}
+	UEngineElementsLibrary::DestroyEditorActorElement(this);
 #endif	// WITH_EDITOR
 
 	Super::BeginDestroy();
diff --git a/Engine/Source/Runtime/Engine/Private/Components/ActorComponent.cpp b/Engine/Source/Runtime/Engine/Private/Components/ActorComponent.cpp
--- a/Engine/Source/Runtime/Engine/Private/Components/ActorComponent.cpp
+++ b/Engine/Source/Runtime/Engine/Private/Components/ActorComponent.cpp
@@ -13,9 +13,7 @@
 #include "Components/PrimitiveComponent.h"
 #include "AI/NavigationSystemBase.h"
 #include "Engine/BlueprintGeneratedClass.h"
-#include "TypedElementRegistry.h"
-#include "TypedElementOwnerStore.h"
-#include "Elements/Component/ComponentElementData.h"
+#include "Elements/EngineElementsLibrary.h"
 #include "ContentStreaming.h"
 #include "ComponentReregisterContext.h"
 #include "Engine/AssetUserData.h"
@@ -77,7 +75,6 @@ FAutoConsoleVariableRef GTickComponentLatentActionsWithTheComponentCVar(
 
 #if WITH_EDITOR
 FUObjectAnnotationSparseBool GSelectedComponentAnnotation;
-TTypedElementOwnerStore<FComponentElementData, const UActorComponent*> GComponentElementOwnerStore;
 #endif
 
 /** Static var indicating activity of reregister context */
@@ -296,10 +293,7 @@ void UActorComponent::PostInitProperties()
 	Super::PostInitProperties();
 
 #if WITH_EDITOR
-	if (GIsEditor && !HasAnyFlags(RF_ClassDefaultObject | RF_ArchetypeObject))
-	{
-		GComponentElementOwnerStore.RegisterElementOwner(this, CreateComponentElement());
-	}
+	UEngineElementsLibrary::CreateEditorComponentElement(this);
 #endif	// WITH_EDITOR
 
 	// Instance components will be added during the owner's initialization
@@ -346,50 +340,6 @@ void UActorComponent::PostInitProperties()
 	}
 }
 
-TTypedElementOwner<FComponentElementData> UActorComponent::CreateComponentElement() const
-{
-	UTypedElementRegistry* Registry = UTypedElementRegistry::GetInstance();
-
-	TTypedElementOwner<FComponentElementData> ComponentElement;
-	if (ensureMsgf(Registry, TEXT("Typed element was requested for '%s' before the registry was available! This usually means that NewObject was used instead of CreateDefaultSubobject during CDO construction."), *GetPathName()))
-	{
-		ComponentElement = Registry->CreateElement<FComponentElementData>(NAME_Components);
-		ComponentElement.GetDataChecked().Component = const_cast<UActorComponent*>(this);
-	}
-
-	return ComponentElement;
-}
-
-void UActorComponent::DestroyComponentElement(TTypedElementOwner<FComponentElementData>& InOutComponentElement) const
-{
-	if (InOutComponentElement)
-	{
-		checkf(InOutComponentElement.GetDataChecked().Component == this, TEXT("Component element was not for this component instance! %s"), *GetPathName());
-		UTypedElementRegistry::GetInstance()->DestroyElement(InOutComponentElement);
-	}
-}
-
-#if WITH_EDITOR
-FTypedElementHandle UActorComponent::AcquireEditorElementHandle(const bool bAllowCreate) const
-{
-	if (GIsEditor)
-	{
-		if (TTypedElementOwnerScopedAccess<FComponentElementData> EditorElementScopedAccess = GComponentElementOwnerStore.FindElementOwner(this))
-		{
-			return EditorElementScopedAccess->AcquireHandle();
-		}
-
-		if (bAllowCreate)
-		{
-			TTypedElementOwnerScopedAccess<FComponentElementData> EditorElementScopedAccess = GComponentElementOwnerStore.RegisterElementOwner(this, CreateComponentElement());
-			return EditorElementScopedAccess->AcquireHandle();
-		}
-	}
-
-	return FTypedElementHandle();
-}
-#endif
-
 void UActorComponent::PostLoad()
 {
 	Super::PostLoad();
@@ -725,10 +675,7 @@ void UActorComponent::BeginDestroy()
 	}
 
 #if WITH_EDITOR
-	if (TTypedElementOwner<FComponentElementData> EditorElement = GComponentElementOwnerStore.UnregisterElementOwner(this))
-	{
-		DestroyComponentElement(EditorElement);
-	}
+	UEngineElementsLibrary::DestroyEditorComponentElement(this);
 #endif	// WITH_EDITOR
 
 	Super::BeginDestroy();
diff --git a/Engine/Source/Runtime/Engine/Private/Elements/Actor/ActorElementObjectInterface.cpp b/Engine/Source/Runtime/Engine/Private/Elements/Actor/ActorElementObjectInterface.cpp
new file mode 100644
--- /dev/null
+++ b/Engine/Source/Runtime/Engine/Private/Elements/Actor/ActorElementObjectInterface.cpp
@@ -0,0 +1,11 @@
+// Copyright Epic Games, Inc. All Rights Reserved.
+
+#include "Elements/Actor/ActorElementObjectInterface.h"
+#include "Elements/Actor/ActorElementData.h"
+#include "GameFramework/Actor.h"
+
+UObject* UActorElementObjectInterface::GetObject(const FTypedElementHandle& InElementHandle)
+{
+	const FActorElementData* ActorData = InElementHandle.GetData<FActorElementData>();
+	return ActorData ? ActorData->Actor : nullptr;
+}
diff --git a/Engine/Source/Runtime/Engine/Private/Elements/Actor/ActorElementSelectionInterface.cpp b/Engine/Source/Runtime/Engine/Private/Elements/Actor/ActorElementSelectionInterface.cpp
--- a/Engine/Source/Runtime/Engine/Private/Elements/Actor/ActorElementSelectionInterface.cpp
+++ b/Engine/Source/Runtime/Engine/Private/Elements/Actor/ActorElementSelectionInterface.cpp
@@ -5,12 +5,6 @@
 #include "GameFramework/Actor.h"
 #include "TypedElementList.h"
 
-UObject* UActorElementSelectionInterface::Legacy_GetSelectionObject(const FTypedElementHandle& InElementHandle)
-{
-	const FActorElementData* ActorData = InElementHandle.GetData<FActorElementData>();
-	return ActorData ? ActorData->Actor : nullptr;
-}
-
 int32 UActorElementSelectionInterface::GetNumSelectedActors(const UTypedElementList* InCurrentSelection)
 {
 	int32 NumSelected = 0;
diff --git a/Engine/Source/Runtime/Engine/Private/Elements/Component/ComponentElementObjectInterface.cpp b/Engine/Source/Runtime/Engine/Private/Elements/Component/ComponentElementObjectInterface.cpp
new file mode 100644
--- /dev/null
+++ b/Engine/Source/Runtime/Engine/Private/Elements/Component/ComponentElementObjectInterface.cpp
@@ -0,0 +1,11 @@
+// Copyright Epic Games, Inc. All Rights Reserved.
+
+#include "Elements/Component/ComponentElementObjectInterface.h"
+#include "Elements/Component/ComponentElementData.h"
+#include "Components/ActorComponent.h"
+
+UObject* UComponentElementObjectInterface::GetObject(const FTypedElementHandle& InElementHandle)
+{
+	const FComponentElementData* ComponentData = InElementHandle.GetData<FComponentElementData>();
+	return ComponentData ? ComponentData->Component : nullptr;
+}
diff --git a/Engine/Source/Runtime/Engine/Private/Elements/Component/ComponentElementSelectionInterface.cpp b/Engine/Source/Runtime/Engine/Private/Elements/Component/ComponentElementSelectionInterface.cpp
--- a/Engine/Source/Runtime/Engine/Private/Elements/Component/ComponentElementSelectionInterface.cpp
+++ b/Engine/Source/Runtime/Engine/Private/Elements/Component/ComponentElementSelectionInterface.cpp
@@ -5,12 +5,6 @@
 #include "Components/ActorComponent.h"
 #include "TypedElementList.h"
 
-UObject* UComponentElementSelectionInterface::Legacy_GetSelectionObject(const FTypedElementHandle& InElementHandle)
-{
-	const FComponentElementData* ComponentData = InElementHandle.GetData<FComponentElementData>();
-	return ComponentData ? ComponentData->Component : nullptr;
-}
-
 int32 UComponentElementSelectionInterface::GetNumSelectedComponents(const UTypedElementList* InCurrentSelection)
 {
 	int32 NumSelected = 0;
diff --git a/Engine/Source/Runtime/Engine/Private/Elements/EngineElements.cpp b/Engine/Source/Runtime/Engine/Private/Elements/EngineElements.cpp
--- a/Engine/Source/Runtime/Engine/Private/Elements/EngineElements.cpp
+++ b/Engine/Source/Runtime/Engine/Private/Elements/EngineElements.cpp
@@ -5,19 +5,35 @@
 
 #include "TypedElementRegistry.h"
 
+#include "Elements/Object/ObjectElementData.h"
+#include "Elements/Object/ObjectElementObjectInterface.h"
+#include "Elements/Object/ObjectElementSelectionInterface.h"
+
 #include "Elements/Actor/ActorElementData.h"
+#include "Elements/Actor/ActorElementObjectInterface.h"
 #include "Elements/Actor/ActorElementSelectionInterface.h"
 
 #include "Elements/Component/ComponentElementData.h"
+#include "Elements/Component/ComponentElementObjectInterface.h"
 #include "Elements/Component/ComponentElementSelectionInterface.h"
 
 FSimpleMulticastDelegate OnRegisterEngineElementsDelegate;
 
+void RegisterEngineObjectElements()
+{
+	UTypedElementRegistry* Registry = UTypedElementRegistry::GetInstance();
+
+	Registry->RegisterElementType<FObjectElementData>(NAME_Object);
+	Registry->RegisterElementInterface<UTypedElementObjectInterface>(NAME_Object, NewObject<UObjectElementObjectInterface>());
+	Registry->RegisterElementInterface<UTypedElementSelectionInterface>(NAME_Object, NewObject<UObjectElementSelectionInterface>());
+}
+
 void RegisterEngineActorElements()
 {
 	UTypedElementRegistry* Registry = UTypedElementRegistry::GetInstance();
 
 	Registry->RegisterElementType<FActorElementData>(NAME_Actor);
+	Registry->RegisterElementInterface<UTypedElementObjectInterface>(NAME_Actor, NewObject<UActorElementObjectInterface>());
 	Registry->RegisterElementInterface<UTypedElementSelectionInterface>(NAME_Actor, NewObject<UActorElementSelectionInterface>());
 }
 
@@ -26,6 +42,7 @@ void RegisterEngineComponentElements()
 	UTypedElementRegistry* Registry = UTypedElementRegistry::GetInstance();
 
 	Registry->RegisterElementType<FComponentElementData>(NAME_Components);
+	Registry->RegisterElementInterface<UTypedElementObjectInterface>(NAME_Components, NewObject<UComponentElementObjectInterface>());
 	Registry->RegisterElementInterface<UTypedElementSelectionInterface>(NAME_Components, NewObject<UComponentElementSelectionInterface>());
 }
 
@@ -35,6 +52,7 @@ void RegisterEngineElements()
 	FModuleManager::Get().LoadModuleChecked("TypedElementFramework");
 	FModuleManager::Get().LoadModuleChecked("TypedElementInterfaces");
 
+	RegisterEngineObjectElements();
 	RegisterEngineActorElements();
 	RegisterEngineComponentElements();
 
diff --git a/Engine/Source/Runtime/Engine/Private/Elements/EngineElementsLibrary.cpp b/Engine/Source/Runtime/Engine/Private/Elements/EngineElementsLibrary.cpp
new file mode 100644
--- /dev/null
+++ b/Engine/Source/Runtime/Engine/Private/Elements/EngineElementsLibrary.cpp
@@ -0,0 +1,206 @@
+// Copyright Epic Games, Inc. All Rights Reserved.
+
+#include "Elements/EngineElementsLibrary.h"
+
+#include "Elements/Object/ObjectElementData.h"
+#include "UObject/Object.h"
+
+#include "Elements/Actor/ActorElementData.h"
+#include "GameFramework/Actor.h"
+
+#include "Elements/Component/ComponentElementData.h"
+#include "Components/ActorComponent.h"
+
+#include "TypedElementRegistry.h"
+#include "TypedElementOwnerStore.h"
+
+namespace EngineElementsLibraryUtil
+{
+template <typename ObjectClass, typename ElementDataType>
+TTypedElementOwner<ElementDataType> CreateTypedElement(const ObjectClass* Object, const FName ElementTypeName)
+{
+	UTypedElementRegistry* Registry = UTypedElementRegistry::GetInstance();
+
+	TTypedElementOwner<ElementDataType> TypedElement;
+	if (ensureMsgf(Registry, TEXT("Typed element was requested for '%s' before the registry was available! This usually means that NewObject was used instead of CreateDefaultSubobject during CDO construction."), *Object->GetPathName()))
+	{
+		TypedElement = Registry->CreateElement<ElementDataType>(ElementTypeName);
+	}
+
+	return TypedElement;
+}
+
+#if WITH_EDITOR
+template <typename ObjectClass, typename ElementDataType>
+void CreateEditorTypedElement(const ObjectClass* Object, TTypedElementOwnerStore<ElementDataType, const ObjectClass*>& ElementOwnerStore, TFunctionRef<TTypedElementOwner<ElementDataType>(const ObjectClass*)> CreateElement)
+{
+	if (GIsEditor && !Object->HasAnyFlags(RF_ClassDefaultObject | RF_ArchetypeObject))
+	{
+		ElementOwnerStore.RegisterElementOwner(Object, CreateElement(Object));
+	}
+}
+
+template <typename ObjectClass, typename ElementDataType>
+void DestroyEditorTypedElement(const ObjectClass* Object, TTypedElementOwnerStore<ElementDataType, const ObjectClass*>& ElementOwnerStore, TFunctionRef<void(const ObjectClass*, TTypedElementOwner<ElementDataType>&)> DestroyElement)
+{
+	if (TTypedElementOwner<ElementDataType> EditorElement = ElementOwnerStore.UnregisterElementOwner(Object))
+	{
+		DestroyElement(Object, EditorElement);
+	}
+}
+
+template <typename ObjectClass, typename ElementDataType>
+FTypedElementHandle AcquireEditorTypedElementHandle(const ObjectClass* Object, TTypedElementOwnerStore<ElementDataType, const ObjectClass*>& ElementOwnerStore, TFunctionRef<TTypedElementOwner<ElementDataType>(const ObjectClass*)> CreateElement, const bool bAllowCreate)
+{
+	if (GIsEditor)
+	{
+		return bAllowCreate
+			? ElementOwnerStore.FindOrRegisterElementOwner(Object, [Object, &CreateElement](){ return CreateElement(Object); })->AcquireHandle()
+			: ElementOwnerStore.FindElementOwner(Object)->AcquireHandle();
+	}
+
+	return FTypedElementHandle();
+}
+#endif
+} // namespace EngineElementsLibraryUtil
+
+#if WITH_EDITOR
+TTypedElementOwnerStore<FObjectElementData, const UObject*> GObjectElementOwnerStore;
+TTypedElementOwnerStore<FActorElementData, const AActor*> GActorElementOwnerStore;
+TTypedElementOwnerStore<FComponentElementData, const UActorComponent*> GComponentElementOwnerStore;
+#endif
+
+UEngineElementsLibrary::UEngineElementsLibrary()
+{
+#if WITH_EDITOR
+	if (HasAnyFlags(RF_ClassDefaultObject))
+	{
+		// UObject exists inside CoreUObject, so it cannot call through directly to clean-up any element handles that have been created
+		// Instead we rely on this GC hook to clean-up any element handles for unreachable objects prior to them being destroyed
+		FCoreUObjectDelegates::PreGarbageCollectConditionalBeginDestroy.AddStatic(&UEngineElementsLibrary::DestroyUnreachableEditorObjectElements);
+	}
+#endif
+}
+
+TTypedElementOwner<FObjectElementData> UEngineElementsLibrary::CreateObjectElement(const UObject* InObject)
+{
+	TTypedElementOwner<FObjectElementData> ObjectElement = EngineElementsLibraryUtil::CreateTypedElement<UObject, FObjectElementData>(InObject, NAME_Object);
+	if (ObjectElement)
+	{
+		ObjectElement.GetDataChecked().Object = const_cast<UObject*>(InObject);
+	}
+	return ObjectElement;
+}
+
+void UEngineElementsLibrary::DestroyObjectElement(const UObject* InObject, TTypedElementOwner<FObjectElementData>& InOutObjectElement)
+{
+	if (InOutObjectElement)
+	{
+		checkf(InOutObjectElement.GetDataChecked().Object == InObject, TEXT("Object element was not for this object instance! %s"), *InObject->GetPathName());
+		UTypedElementRegistry::GetInstance()->DestroyElement(InOutObjectElement);
+	}
+}
+
+#if WITH_EDITOR
+void UEngineElementsLibrary::CreateEditorObjectElement(const UObject* Object)
+{
+	EngineElementsLibraryUtil::CreateEditorTypedElement<UObject, FObjectElementData>(Object, GObjectElementOwnerStore, &UEngineElementsLibrary::CreateObjectElement);
+}
+
+void UEngineElementsLibrary::DestroyEditorObjectElement(const UObject* Object)
+{
+	EngineElementsLibraryUtil::DestroyEditorTypedElement<UObject, FObjectElementData>(Object, GObjectElementOwnerStore, &UEngineElementsLibrary::DestroyObjectElement);
+}
+
+void UEngineElementsLibrary::DestroyUnreachableEditorObjectElements()
+{
+	auto IsObjectElementUnreachable = [](const TTypedElementOwner<FObjectElementData>& EditorElement)
+	{
+		return EditorElement.GetDataChecked().Object->IsUnreachable();
+	};
+
+	auto DestroyUnreachableObjectElement = [Registry = UTypedElementRegistry::GetInstance()](TTypedElementOwner<FObjectElementData>&& EditorElement)
+	{
+		Registry->DestroyElement(EditorElement);
+	};
+
+	GObjectElementOwnerStore.UnregisterElementOwners(IsObjectElementUnreachable, DestroyUnreachableObjectElement);
+}
+
+FTypedElementHandle UEngineElementsLibrary::AcquireEditorObjectElementHandle(const UObject* Object, const bool bAllowCreate)
+{
+	return EngineElementsLibraryUtil::AcquireEditorTypedElementHandle<UObject, FObjectElementData>(Object, GObjectElementOwnerStore, &UEngineElementsLibrary::CreateObjectElement, bAllowCreate);
+}
+#endif
+
+TTypedElementOwner<FActorElementData> UEngineElementsLibrary::CreateActorElement(const AActor* InActor)
+{
+	TTypedElementOwner<FActorElementData> ActorElement = EngineElementsLibraryUtil::CreateTypedElement<AActor, FActorElementData>(InActor, NAME_Actor);
+	if (ActorElement)
+	{
+		ActorElement.GetDataChecked().Actor = const_cast<AActor*>(InActor);
+	}
+	return ActorElement;
+}
+
+void UEngineElementsLibrary::DestroyActorElement(const AActor* InActor, TTypedElementOwner<FActorElementData>& InOutActorElement)
+{
+	if (InOutActorElement)
+	{
+		checkf(InOutActorElement.GetDataChecked().Actor == InActor, TEXT("Actor element was not for this actor instance! %s"), *InActor->GetPathName());
+		UTypedElementRegistry::GetInstance()->DestroyElement(InOutActorElement);
+	}
+}
+
+#if WITH_EDITOR
+void UEngineElementsLibrary::CreateEditorActorElement(const AActor* Actor)
+{
+	EngineElementsLibraryUtil::CreateEditorTypedElement<AActor, FActorElementData>(Actor, GActorElementOwnerStore, &UEngineElementsLibrary::CreateActorElement);
+}
+
+void UEngineElementsLibrary::DestroyEditorActorElement(const AActor* Actor)
+{
+	EngineElementsLibraryUtil::DestroyEditorTypedElement<AActor, FActorElementData>(Actor, GActorElementOwnerStore, &UEngineElementsLibrary::DestroyActorElement);
+}
+
+FTypedElementHandle UEngineElementsLibrary::AcquireEditorActorElementHandle(const AActor* Actor, const bool bAllowCreate)
+{
+	return EngineElementsLibraryUtil::AcquireEditorTypedElementHandle<AActor, FActorElementData>(Actor, GActorElementOwnerStore, &UEngineElementsLibrary::CreateActorElement, bAllowCreate);
+}
+#endif
+
+TTypedElementOwner<FComponentElementData> UEngineElementsLibrary::CreateComponentElement(const UActorComponent* InComponent)
+{
+	TTypedElementOwner<FComponentElementData> ComponentElement = EngineElementsLibraryUtil::CreateTypedElement<UActorComponent, FComponentElementData>(InComponent, NAME_Components);
+	if (ComponentElement)
+	{
+		ComponentElement.GetDataChecked().Component = const_cast<UActorComponent*>(InComponent);
+	}
+	return ComponentElement;
+}
+
+void UEngineElementsLibrary::DestroyComponentElement(const UActorComponent* InComponent, TTypedElementOwner<FComponentElementData>& InOutComponentElement)
+{
+	if (InOutComponentElement)
+	{
+		checkf(InOutComponentElement.GetDataChecked().Component == InComponent, TEXT("Component element was not for this component instance! %s"), *InComponent->GetPathName());
+		UTypedElementRegistry::GetInstance()->DestroyElement(InOutComponentElement);
+	}
+}
+
+#if WITH_EDITOR
+void UEngineElementsLibrary::CreateEditorComponentElement(const UActorComponent* Component)
+{
+	EngineElementsLibraryUtil::CreateEditorTypedElement<UActorComponent, FComponentElementData>(Component, GComponentElementOwnerStore, &UEngineElementsLibrary::CreateComponentElement);
+}
+
+void UEngineElementsLibrary::DestroyEditorComponentElement(const UActorComponent* Component)
+{
+	EngineElementsLibraryUtil::DestroyEditorTypedElement<UActorComponent, FComponentElementData>(Component, GComponentElementOwnerStore, &UEngineElementsLibrary::DestroyComponentElement);
+}
+
+FTypedElementHandle UEngineElementsLibrary::AcquireEditorComponentElementHandle(const UActorComponent* Component, const bool bAllowCreate)
+{
+	return EngineElementsLibraryUtil::AcquireEditorTypedElementHandle<UActorComponent, FComponentElementData>(Component, GComponentElementOwnerStore, &UEngineElementsLibrary::CreateComponentElement, bAllowCreate);
+}
+#endif
diff --git a/Engine/Source/Runtime/Engine/Private/Elements/Object/ObjectElementData.cpp b/Engine/Source/Runtime/Engine/Private/Elements/Object/ObjectElementData.cpp
new file mode 100644
--- /dev/null
+++ b/Engine/Source/Runtime/Engine/Private/Elements/Object/ObjectElementData.cpp
@@ -0,0 +1,5 @@
+// Copyright Epic Games, Inc. All Rights Reserved.
+
+#include "Elements/Object/ObjectElementData.h"
+
+UE_DEFINE_TYPED_ELEMENT_DATA_RTTI(FObjectElementData);
diff --git a/Engine/Source/Runtime/Engine/Private/Elements/Object/ObjectElementObjectInterface.cpp b/Engine/Source/Runtime/Engine/Private/Elements/Object/ObjectElementObjectInterface.cpp
new file mode 100644
--- /dev/null
+++ b/Engine/Source/Runtime/Engine/Private/Elements/Object/ObjectElementObjectInterface.cpp
@@ -0,0 +1,10 @@
+// Copyright Epic Games, Inc. All Rights Reserved.
+
+#include "Elements/Object/ObjectElementObjectInterface.h"
+#include "Elements/Object/ObjectElementData.h"
+
+UObject* UObjectElementObjectInterface::GetObject(const FTypedElementHandle& InElementHandle)
+{
+	const FObjectElementData* ObjectData = InElementHandle.GetData<FObjectElementData>();
+	return ObjectData ? ObjectData->Object : nullptr;
+}
diff --git a/Engine/Source/Runtime/Engine/Private/Elements/Object/ObjectElementSelectionInterface.cpp b/Engine/Source/Runtime/Engine/Private/Elements/Object/ObjectElementSelectionInterface.cpp
new file mode 100644
--- /dev/null
+++ b/Engine/Source/Runtime/Engine/Private/Elements/Object/ObjectElementSelectionInterface.cpp
@@ -0,0 +1,4 @@
+// Copyright Epic Games, Inc. All Rights Reserved.
+
+#include "Elements/Object/ObjectElementSelectionInterface.h"
+#include "Elements/Object/ObjectElementData.h"
diff --git a/Engine/Source/Runtime/Engine/Private/EngineUtils.cpp b/Engine/Source/Runtime/Engine/Private/EngineUtils.cpp
--- a/Engine/Source/Runtime/Engine/Private/EngineUtils.cpp
+++ b/Engine/Source/Runtime/Engine/Private/EngineUtils.cpp
@@ -11,6 +11,7 @@
 #include "Misc/EngineVersion.h"
 #include "GameFramework/PlayerController.h"
 #include "GameFramework/LightWeightInstanceSubsystem.h"
+#include "Elements/EngineElementsLibrary.h"
 #include "EngineGlobals.h"
 #include "Components/PrimitiveComponent.h"
 #include "Engine/Engine.h"
@@ -46,11 +47,11 @@ FTypedElementHandle HActor::GetElementHandle() const
 #if WITH_EDITOR
 	if (PrimComponent)
 	{
-		return PrimComponent->AcquireEditorElementHandle();
+		return UEngineElementsLibrary::AcquireEditorComponentElementHandle(PrimComponent);
 	}
 	if (Actor)
 	{
-		return Actor->AcquireEditorElementHandle();
+		return UEngineElementsLibrary::AcquireEditorActorElementHandle(Actor);
 	}
 #endif	// WITH_EDITOR
 	return FTypedElementHandle();
diff --git a/Engine/Source/Runtime/Engine/Public/Elements/Actor/ActorElementObjectInterface.h b/Engine/Source/Runtime/Engine/Public/Elements/Actor/ActorElementObjectInterface.h
new file mode 100644
--- /dev/null
+++ b/Engine/Source/Runtime/Engine/Public/Elements/Actor/ActorElementObjectInterface.h
@@ -0,0 +1,15 @@
+// Copyright Epic Games, Inc. All Rights Reserved.
+
+#pragma once
+
+#include "TypedElementObjectInterface.h"
+#include "ActorElementObjectInterface.generated.h"
+
+UCLASS()
+class ENGINE_API UActorElementObjectInterface : public UTypedElementObjectInterface
+{
+	GENERATED_BODY()
+
+public:
+	virtual UObject* GetObject(const FTypedElementHandle& InElementHandle) override;
+};
diff --git a/Engine/Source/Runtime/Engine/Public/Elements/Actor/ActorElementSelectionInterface.h b/Engine/Source/Runtime/Engine/Public/Elements/Actor/ActorElementSelectionInterface.h
--- a/Engine/Source/Runtime/Engine/Public/Elements/Actor/ActorElementSelectionInterface.h
+++ b/Engine/Source/Runtime/Engine/Public/Elements/Actor/ActorElementSelectionInterface.h
@@ -11,8 +11,6 @@ class ENGINE_API UActorElementSelectionInterface : public UTypedElementSelection
 	GENERATED_BODY()
 
 public:
-	virtual UObject* Legacy_GetSelectionObject(const FTypedElementHandle& InElementHandle) override;
-
 	static int32 GetNumSelectedActors(const UTypedElementList* InCurrentSelection);
 	static bool HasSelectedActors(const UTypedElementList* InCurrentSelection);
 };
diff --git a/Engine/Source/Runtime/Engine/Public/Elements/Component/ComponentElementObjectInterface.h b/Engine/Source/Runtime/Engine/Public/Elements/Component/ComponentElementObjectInterface.h
new file mode 100644
--- /dev/null
+++ b/Engine/Source/Runtime/Engine/Public/Elements/Component/ComponentElementObjectInterface.h
@@ -0,0 +1,15 @@
+// Copyright Epic Games, Inc. All Rights Reserved.
+
+#pragma once
+
+#include "TypedElementObjectInterface.h"
+#include "ComponentElementObjectInterface.generated.h"
+
+UCLASS()
+class ENGINE_API UComponentElementObjectInterface : public UTypedElementObjectInterface
+{
+	GENERATED_BODY()
+
+public:
+	virtual UObject* GetObject(const FTypedElementHandle& InElementHandle) override;
+};
diff --git a/Engine/Source/Runtime/Engine/Public/Elements/Component/ComponentElementSelectionInterface.h b/Engine/Source/Runtime/Engine/Public/Elements/Component/ComponentElementSelectionInterface.h
--- a/Engine/Source/Runtime/Engine/Public/Elements/Component/ComponentElementSelectionInterface.h
+++ b/Engine/Source/Runtime/Engine/Public/Elements/Component/ComponentElementSelectionInterface.h
@@ -11,8 +11,6 @@ class ENGINE_API UComponentElementSelectionInterface : public UTypedElementSelec
 	GENERATED_BODY()
 
 public:
-	virtual UObject* Legacy_GetSelectionObject(const FTypedElementHandle& InElementHandle) override;
-
 	static int32 GetNumSelectedComponents(const UTypedElementList* InCurrentSelection);
 	static bool HasSelectedComponents(const UTypedElementList* InCurrentSelection);
 };
diff --git a/Engine/Source/Runtime/Engine/Public/Elements/EngineElementsLibrary.h b/Engine/Source/Runtime/Engine/Public/Elements/EngineElementsLibrary.h
new file mode 100644
--- /dev/null
+++ b/Engine/Source/Runtime/Engine/Public/Elements/EngineElementsLibrary.h
@@ -0,0 +1,57 @@
+// Copyright Epic Games, Inc. All Rights Reserved.
+
+#pragma once
+
+#include "TypedElementHandle.h"
+#include "Kismet/BlueprintFunctionLibrary.h"
+#include "EngineElementsLibrary.generated.h"
+
+class UObject;
+struct FObjectElementData;
+
+class AActor;
+struct FActorElementData;
+
+class UActorComponent;
+struct FComponentElementData;
+
+UCLASS()
+class ENGINE_API UEngineElementsLibrary : public UBlueprintFunctionLibrary
+{
+	GENERATED_BODY()
+
+public:
+	UEngineElementsLibrary();
+
+	static TTypedElementOwner<FObjectElementData> CreateObjectElement(const UObject* InObject);
+	static void DestroyObjectElement(const UObject* InObject, TTypedElementOwner<FObjectElementData>& InOutObjectElement);
+#if WITH_EDITOR
+	static void CreateEditorObjectElement(const UObject* Object);
+	static void DestroyEditorObjectElement(const UObject* Object);
+	UFUNCTION(BlueprintCallable, BlueprintPure=false, Category="TypedElementFramework|Object", meta=(ScriptMethod="AcquireEditorElementHandle"))
+	static FTypedElementHandle AcquireEditorObjectElementHandle(const UObject* Object, const bool bAllowCreate = true);
+#endif
+
+	static TTypedElementOwner<FActorElementData> CreateActorElement(const AActor* InActor);
+	static void DestroyActorElement(const AActor* InActor, TTypedElementOwner<FActorElementData>& InOutActorElement);
+#if WITH_EDITOR
+	static void CreateEditorActorElement(const AActor* Actor);
+	static void DestroyEditorActorElement(const AActor* Actor);
+	UFUNCTION(BlueprintCallable, Category="TypedElementFramework|Actor", meta=(ScriptMethod="AcquireEditorElementHandle"))
+	static FTypedElementHandle AcquireEditorActorElementHandle(const AActor* Actor, const bool bAllowCreate = true);
+#endif
+
+	static TTypedElementOwner<FComponentElementData> CreateComponentElement(const UActorComponent* InComponent);
+	static void DestroyComponentElement(const UActorComponent* InComponent, TTypedElementOwner<FComponentElementData>& InOutComponentElement);
+#if WITH_EDITOR
+	static void CreateEditorComponentElement(const UActorComponent* Component);
+	static void DestroyEditorComponentElement(const UActorComponent* Component);
+	UFUNCTION(BlueprintCallable, BlueprintPure=false, Category="TypedElementFramework|Component", meta=(ScriptMethod="AcquireEditorElementHandle"))
+	static FTypedElementHandle AcquireEditorComponentElementHandle(const UActorComponent* Component, const bool bAllowCreate = true);
+#endif
+
+private:
+#if WITH_EDITOR
+	static void DestroyUnreachableEditorObjectElements();
+#endif
+};
diff --git a/Engine/Source/Runtime/Engine/Public/Elements/Object/ObjectElementData.h b/Engine/Source/Runtime/Engine/Public/Elements/Object/ObjectElementData.h
new file mode 100644
--- /dev/null
+++ b/Engine/Source/Runtime/Engine/Public/Elements/Object/ObjectElementData.h
@@ -0,0 +1,17 @@
+// Copyright Epic Games, Inc. All Rights Reserved.
+
+#pragma once
+
+#include "TypedElementData.h"
+
+class UObject;
+
+/**
+ * Element data that represents an Object.
+ */
+struct ENGINE_API FObjectElementData
+{
+	UE_DECLARE_TYPED_ELEMENT_DATA_RTTI(FObjectElementData);
+
+	UObject* Object = nullptr;
+};
diff --git a/Engine/Source/Runtime/Engine/Public/Elements/Object/ObjectElementObjectInterface.h b/Engine/Source/Runtime/Engine/Public/Elements/Object/ObjectElementObjectInterface.h
new file mode 100644
--- /dev/null
+++ b/Engine/Source/Runtime/Engine/Public/Elements/Object/ObjectElementObjectInterface.h
@@ -0,0 +1,15 @@
+// Copyright Epic Games, Inc. All Rights Reserved.
+
+#pragma once
+
+#include "TypedElementObjectInterface.h"
+#include "ObjectElementObjectInterface.generated.h"
+
+UCLASS()
+class ENGINE_API UObjectElementObjectInterface : public UTypedElementObjectInterface
+{
+	GENERATED_BODY()
+
+public:
+	virtual UObject* GetObject(const FTypedElementHandle& InElementHandle) override;
+};
diff --git a/Engine/Source/Runtime/Engine/Public/Elements/Object/ObjectElementSelectionInterface.h b/Engine/Source/Runtime/Engine/Public/Elements/Object/ObjectElementSelectionInterface.h
new file mode 100644
--- /dev/null
+++ b/Engine/Source/Runtime/Engine/Public/Elements/Object/ObjectElementSelectionInterface.h
@@ -0,0 +1,14 @@
+// Copyright Epic Games, Inc. All Rights Reserved.
+
+#pragma once
+
+#include "TypedElementSelectionInterface.h"
+#include "ObjectElementSelectionInterface.generated.h"
+
+UCLASS()
+class ENGINE_API UObjectElementSelectionInterface : public UTypedElementSelectionInterface
+{
+	GENERATED_BODY()
+
+public:
+};
diff --git a/Engine/Source/Runtime/TypedElementFramework/Public/TypedElementOwnerStore.h b/Engine/Source/Runtime/TypedElementFramework/Public/TypedElementOwnerStore.h
--- a/Engine/Source/Runtime/TypedElementFramework/Public/TypedElementOwnerStore.h
+++ b/Engine/Source/Runtime/TypedElementFramework/Public/TypedElementOwnerStore.h
@@ -145,6 +145,40 @@ public:
 		return ElementOwner;
 	}
 
+	/**
+	 * Unregister any given element owners that match the given predicate, transferring ownership of the owner instance to the given callback.
+	 */
+	void UnregisterElementOwners(TFunctionRef<bool(const TTypedElementOwner<ElementDataType>&)> InShouldUnregisterElement, TFunctionRef<void(TTypedElementOwner<ElementDataType>&&)> InOnUnregisterElement)
+	{
+		FScopeLock Lock(&ElementOwnerMapCS);
+		for (auto It = ElementOwnerMap.CreateIterator(); It; ++It)
+		{
+			if (InShouldUnregisterElement(It->Value))
+			{
+				InOnUnregisterElement(MoveTemp(It->Value));
+				It.RemoveCurrent();
+			}
+		}
+	}
+
+	/**
+	 * Find the given element owner, or call the register callback to create one if none is found.
+	 * @note This must be paired with a call to UnregisterElementOwner.
+	 */
+	TTypedElementOwnerScopedAccess<ElementDataType> FindOrRegisterElementOwner(const KeyDataType& InKey, TFunctionRef<TTypedElementOwner<ElementDataType>()> InCreateElement)
+	{
+		TTypedElementOwnerScopedAccess<ElementDataType> ScopedAccess(&ElementOwnerMapCS);
+		if (const TTypedElementOwner<ElementDataType>* ExistingElement = ElementOwnerMap.Find(InKey))
+		{
+			ScopedAccess.Private_SetElementOwner(ExistingElement);
+		}
+		else
+		{
+			ScopedAccess.Private_SetElementOwner(&ElementOwnerMap.Add(InKey, InCreateElement()));
+		}
+		return ScopedAccess;
+	}
+
 	/**
 	 * Provide const access to the given element owner.
 	 * @note This will lock the store for writes until the scoped access struct is no longer being used. The returned struct may be unset if the key was not registered.
diff --git a/Engine/Source/Runtime/TypedElementInterfaces/Private/TypedElementObjectInterface.cpp b/Engine/Source/Runtime/TypedElementInterfaces/Private/TypedElementObjectInterface.cpp
new file mode 100644
--- /dev/null
+++ b/Engine/Source/Runtime/TypedElementInterfaces/Private/TypedElementObjectInterface.cpp
@@ -0,0 +1,3 @@
+// Copyright Epic Games, Inc. All Rights Reserved.
+
+#include "TypedElementObjectInterface.h"
diff --git a/Engine/Source/Runtime/TypedElementInterfaces/Public/TypedElementObjectInterface.h b/Engine/Source/Runtime/TypedElementInterfaces/Public/TypedElementObjectInterface.h
new file mode 100644
--- /dev/null
+++ b/Engine/Source/Runtime/TypedElementInterfaces/Public/TypedElementObjectInterface.h
@@ -0,0 +1,25 @@
+// Copyright Epic Games, Inc. All Rights Reserved.
+
+#pragma once
+
+#include "TypedElementHandle.h"
+#include "TypedElementObjectInterface.generated.h"
+
+UCLASS(Abstract)
+class TYPEDELEMENTINTERFACES_API UTypedElementObjectInterface : public UTypedElementInterface
+{
+	GENERATED_BODY()
+
+public:
+	/**
+	 * Retreive the object instance that this handle represents, if any.
+	 */
+	UFUNCTION(BlueprintPure, Category="TypedElementInterfaces|Object")
+	virtual UObject* GetObject(const FTypedElementHandle& InElementHandle) { return nullptr; }
+};
+
+template <>
+struct TTypedElement<UTypedElementObjectInterface> : public TTypedElementBase<UTypedElementObjectInterface>
+{
+	UObject* GetObject() const { return InterfacePtr->GetObject(*this); }
+};
diff --git a/Engine/Source/Runtime/TypedElementInterfaces/Public/TypedElementSelectionInterface.h b/Engine/Source/Runtime/TypedElementInterfaces/Public/TypedElementSelectionInterface.h
--- a/Engine/Source/Runtime/TypedElementInterfaces/Public/TypedElementSelectionInterface.h
+++ b/Engine/Source/Runtime/TypedElementInterfaces/Public/TypedElementSelectionInterface.h
@@ -111,13 +111,6 @@ public:
 	 */
 	UFUNCTION(BlueprintPure, Category="TypedElementInterfaces|Selection")
 	virtual FTypedElementHandle GetSelectionElement(const FTypedElementHandle& InElementHandle, const UTypedElementList* InCurrentSelection, const ETypedElementSelectionMethod InSelectionMethod) { return InElementHandle; }
-
-	/**
-	 * Retreive the object instance that should be selected via this handle.
-	 * @note This exists to allow the legacy USelection to bridge to an UTypedElementList instance. It should not be used in new code!
-	 */
-	UFUNCTION(BlueprintPure, Category="TypedElementInterfaces|Selection", DisplayName="GetSelectionObject (Legacy)")
-	virtual UObject* Legacy_GetSelectionObject(const FTypedElementHandle& InElementHandle) { return nullptr; }
 };
 
 template <>
@@ -130,5 +123,4 @@ struct TTypedElement<UTypedElementSelectionInterface> : public TTypedElementBase
 	bool DeselectElement(UTypedElementList* InSelectionSet, const FTypedElementSelectionOptions& InSelectionOptions) const { return InterfacePtr->DeselectElement(*this, InSelectionSet, InSelectionOptions); }
 	bool AllowSelectionModifiers(const UTypedElementList* InSelectionSet) const { return InterfacePtr->AllowSelectionModifiers(*this, InSelectionSet); }
 	FTypedElementHandle GetSelectionElement(const UTypedElementList* InCurrentSelection, const ETypedElementSelectionMethod InSelectionMethod) const { return InterfacePtr->GetSelectionElement(*this, InCurrentSelection, InSelectionMethod); }
-	UObject* Legacy_GetSelectionObject() const { return InterfacePtr->Legacy_GetSelectionObject(*this); }
 };
