Cooker: Add diagnostic to warn when cooker has been unable to save or load a package for 120 seconds.
#rnx

diff --git a/Engine/Source/Editor/UnrealEd/Classes/CookOnTheSide/CookOnTheFlyServer.h b/Engine/Source/Editor/UnrealEd/Classes/CookOnTheSide/CookOnTheFlyServer.h
--- a/Engine/Source/Editor/UnrealEd/Classes/CookOnTheSide/CookOnTheFlyServer.h
+++ b/Engine/Source/Editor/UnrealEd/Classes/CookOnTheSide/CookOnTheFlyServer.h
@@ -267,6 +267,8 @@ private:
 
 	/** Execute operations that need to be done after each Scheduler task, such as checking for new external requests. */
 	void TickCookStatus(UE::Cook::FTickStackData& StackData);
+	void SetSaveBusy(bool bInSaveBusy);
+	void SetLoadBusy(bool bInLoadBusy);
 	void UpdateDisplay(ECookTickFlags CookFlags, bool bForceDisplay);
 	enum class ECookAction
 	{
@@ -283,26 +285,26 @@ private:
 	ECookAction DecideNextCookAction(UE::Cook::FTickStackData& StackData);
 	/** Execute any existing external callbacks and push any existing external cook requests into the RequestQueue. */
 	void PumpExternalRequests(const UE::Cook::FCookerTimer& CookerTimer);
-	/** Inspect the next package in the RequestQueue and push it on to its next state. */
-	void PumpRequests(UE::Cook::FTickStackData& StackData);
-	/** Handle the requested PackageData that has been peeled off of the RequestQueue, e.g. by sending it on to the LoadQueue. */
-	void ProcessRequest(UE::Cook::FPackageData& PackageData, UE::Cook::FTickStackData& StackData);
-	/** Get the list of unloaded Packages the load of the given PackageData is dependent upon, and add their PackageDatas to the load queue */
-	void AddDependenciesToLoadQueue(UE::Cook::FPackageData& PackageData);
-
-	/** Load all packages in the LoadQueue until it's time to break. */
-	void PumpLoads(UE::Cook::FTickStackData& StackData, uint32 DesiredQueueLength);
+	/** Inspect the next package in the RequestQueue and push it on to its next state. Report the number of PackageDatas that were pushed to load. */
+	void PumpRequests(UE::Cook::FTickStackData& StackData, int32& OutNumPushed);
+	/** Handle the requested PackageData that has been peeled off of the RequestQueue, e.g. by sending it on to the LoadQueue. Report the number of PackageDatas that were pushed to load (includes dependencies, so may be more than 1). */
+	void ProcessRequest(UE::Cook::FPackageData& PackageData, UE::Cook::FTickStackData& StackData, int32& OutNumPushed);
+	/** Get the list of unloaded Packages the load of the given PackageData is dependent upon, and add their PackageDatas to the load queue. Report the number of PackageDatass that were pushed to the load queue */
+	void AddDependenciesToLoadQueue(UE::Cook::FPackageData& PackageData, int32& OutNumPushed);
+
+	/** Load all packages in the LoadQueue until it's time to break. Report the number of loads that were pushed to save. */
+	void PumpLoads(UE::Cook::FTickStackData& StackData, uint32 DesiredQueueLength, int32& OutNumPushed, bool& bOutBusy);
 	/** Move packages from LoadPrepare's entry queue into the PreloadingQueue until we run out of Preload slots. */
 	void PumpPreloadStarts();
 	/** Move preload-completed packages from LoadPrepare into state LoadReady until we find the first one that is not yet finished preloading. */
 	void PumpPreloadCompletes();
-	/** Load the given PackageData that was in the load queue and send it on to its next state */
-	void LoadPackageInQueue(UE::Cook::FPackageData& PackageData, uint32& ResultFlags);
+	/** Load the given PackageData that was in the load queue and send it on to its next state. Report the number of PackageDatas that were pushed to save (0 or 1) */
+	void LoadPackageInQueue(UE::Cook::FPackageData& PackageData, uint32& ResultFlags, int32& OutNumPushed);
 	/** Mark that the given PackageData failed to load and return it to idle. */
 	void RejectPackageToLoad(UE::Cook::FPackageData& PackageData, const TCHAR* Reason);
 
-	/** Try to save all packages in the SaveQueue until it's time to break. */
-	void PumpSaves(UE::Cook::FTickStackData& StackData, uint32 DesiredQueueLength);
+	/** Try to save all packages in the SaveQueue until it's time to break. Report the number of requests that were completed (either skipped or successfully saved or failed to save) */
+	void PumpSaves(UE::Cook::FTickStackData& StackData, uint32 DesiredQueueLength, int32& OutNumPushed, bool& bOutBusy);
 	/**
 	 * Inspect the given package from the PackageTracker and add it to the SaveQueue if the cooker should save it.
 	 *
@@ -1047,6 +1049,12 @@ private:
 	/** If preloading is enabled, we call TryPreload until it returns true before sending the package to LoadReady, otherwise we skip TryPreload and it goes immediately. */
 	bool bPreloadingEnabled = false;
 
+	/** Timers for tracking how long we have been busy, to manage retries and warnings of deadlock */
+	float SaveBusyTimeLastRetry = 0.f;
+	float SaveBusyTimeStarted = 0.f;
+	float LoadBusyTimeLastRetry = 0.f;
+	float LoadBusyTimeStarted = 0.f;
+
 	// These helper structs are all TUniquePtr rather than inline members so that we can keep their headers private.  See class header comments for their purpose.
 	TUniquePtr<UE::Cook::FPackageTracker> PackageTracker;
 	TUniquePtr<UE::Cook::FPackageDatas> PackageDatas;
diff --git a/Engine/Source/Editor/UnrealEd/Private/CookOnTheFlyServer.cpp b/Engine/Source/Editor/UnrealEd/Private/CookOnTheFlyServer.cpp
--- a/Engine/Source/Editor/UnrealEd/Private/CookOnTheFlyServer.cpp
+++ b/Engine/Source/Editor/UnrealEd/Private/CookOnTheFlyServer.cpp
@@ -154,6 +154,20 @@ static FAutoConsoleVariableRef CVarCookDisplayRepeatTime(
 	TEXT("Controls the time before the cooker will repeat the same progress message.\n"),
 	ECVF_Default);
 
+float GCookProgressRetryBusyTime = 1.0f;
+static FAutoConsoleVariableRef CVarCookRetryBusyTime(
+	TEXT("Cook.retrybusytime"),
+	GCookProgressRetryBusyTime,
+	TEXT("Controls the time between retry attempts at save and load when the save and load queues are busy.\n"),
+	ECVF_Default);
+
+float GCookProgressWarnBusyTime = 120.0f;
+static FAutoConsoleVariableRef CVarCookDisplayWarnBusyTime(
+	TEXT("Cook.display.warnbusytime"),
+	GCookProgressWarnBusyTime,
+	TEXT("Controls the time before the cooker will issue a warning that there is a deadlock in a busy queue.\n"),
+	ECVF_Default);
+
 
 #define PROFILE_NETWORK 0
 
@@ -1124,33 +1138,46 @@ uint32 UCookOnTheFlyServer::TickCookOnTheSide(const float TimeSlice, uint32 &Coo
 	{
 		UE_SCOPED_HIERARCHICAL_COOKTIMER(TickCookOnTheSide); // Make sure no UE_SCOPED_HIERARCHICAL_COOKTIMERs are around CookByTheBookFinishes, as that function deletes memory for them
 
-		bSaveBusy = false;
-		bLoadBusy = false;
 		bool bContinueTick = true;
 		while (bContinueTick && (!IsEngineExitRequested() || CurrentCookMode == ECookMode::CookByTheBook))
 		{
 			TickCookStatus(StackData);
 
 			ECookAction CookAction = DecideNextCookAction(StackData);
+			int32 NumPushed;
+			bool bBusy;
 			switch (CookAction)
 			{
 			case ECookAction::Request:
-				PumpRequests(StackData);
-				bLoadBusy = false;
+				PumpRequests(StackData, NumPushed);
+				if (NumPushed > 0)
+				{
+					SetLoadBusy(false);
+				}
 				break;
 			case ECookAction::Load:
-				PumpLoads(StackData, 0);
-				bSaveBusy = false;
+				PumpLoads(StackData, 0, NumPushed, bBusy);
+				SetLoadBusy(bBusy);
+				if (NumPushed > 0)
+				{
+					SetSaveBusy(false);
+				}
 				break;
 			case ECookAction::LoadLimited:
-				PumpLoads(StackData, DesiredLoadQueueLength);
-				bSaveBusy = false;
+				PumpLoads(StackData, DesiredLoadQueueLength, NumPushed, bBusy);
+				SetLoadBusy(bBusy);
+				if (NumPushed > 0)
+				{
+					SetSaveBusy(false);
+				}
 				break;
 			case ECookAction::Save:
-				PumpSaves(StackData, 0);
+				PumpSaves(StackData, 0, NumPushed, bBusy);
+				SetSaveBusy(bBusy);
 				break;
 			case ECookAction::SaveLimited:
-				PumpSaves(StackData, DesiredSaveQueueLength);
+				PumpSaves(StackData, DesiredSaveQueueLength, NumPushed, bBusy);
+				SetSaveBusy(bBusy);
 				break;
 			case ECookAction::Done:
 				bContinueTick = false;
@@ -1218,8 +1245,46 @@ void UCookOnTheFlyServer::TickCookStatus(UE::Cook::FTickStackData& StackData)
 	PumpExternalRequests(StackData.Timer);
 }
 
+void UCookOnTheFlyServer::SetSaveBusy(bool bInBusy)
+{
+	if (bSaveBusy != bInBusy)
+	{
+		bSaveBusy = bInBusy;
+		if (bSaveBusy)
+		{
+			SaveBusyTimeStarted = FPlatformTime::Seconds();
+			SaveBusyTimeLastRetry = SaveBusyTimeStarted;
+		}
+		else
+		{
+			SaveBusyTimeStarted = 0;
+			SaveBusyTimeLastRetry = 0;
+		}
+	}
+}
+
+void UCookOnTheFlyServer::SetLoadBusy(bool bInLoadBusy)
+{
+	if (bLoadBusy != bInLoadBusy)
+	{
+		bLoadBusy = bInLoadBusy;
+		if (bLoadBusy)
+		{
+			LoadBusyTimeStarted = FPlatformTime::Seconds();
+			LoadBusyTimeLastRetry = LoadBusyTimeStarted;
+		}
+		else
+		{
+			LoadBusyTimeStarted = 0;
+			LoadBusyTimeLastRetry = 0;
+		}
+	}
+}
+
 void UCookOnTheFlyServer::UpdateDisplay(ECookTickFlags TickFlags, bool bForceDisplay)
 {
+	using namespace UE::Cook;
+
 	const float CurrentTime = FPlatformTime::Seconds();
 	const float DeltaProgressDisplayTime = CurrentTime - LastProgressDisplayTime;
 	const int32 CookedPackagesCount = PackageDatas->GetNumCooked();
@@ -1252,6 +1317,55 @@ void UCookOnTheFlyServer::UpdateDisplay(ECookTickFlags TickFlags, bool bForceDis
 			TEXT("Cook Diagnostics: OpenFileHandles=%d, VirtualMemory=%dMiB"),
 			OpenFileHandles, FPlatformMemory::GetStats().UsedVirtual / 1024 / 1024);
 		LastDiagnosticsDisplayTime = CurrentTime;
+		if (bSaveBusy && CurrentTime - SaveBusyTimeStarted > GCookProgressWarnBusyTime)
+		{
+			UE_LOG(LogCook, Warning, TEXT("Cooker has been blocked from saving the current packages for %f seconds."));
+			UE_LOG(LogCook, Display, TEXT("Current packages in the savequeue : "), GCookProgressWarnBusyTime);
+			FPackageDataQueue& SaveQueue = PackageDatas->GetSaveQueue();
+			bool bFound = false;
+			for (UE::Cook::FPackageData* PackageData : SaveQueue)
+			{
+				UE_LOG(LogCook, Display, TEXT("    %s"), *PackageData->GetFileName().ToString());
+				bFound = true;
+			}
+			if (!bFound)
+			{
+				UE_LOG(LogCook, Display, TEXT("    <None>"));
+			}
+
+			UE_LOG(LogCook, Display, TEXT("Current objects that have not yet returned true from IsCachedCookedPlatformDataLoaded:"));
+			bFound = false;
+			for (const FPendingCookedPlatformData& Data : PackageDatas->GetPendingCookedPlatformDatas())
+			{
+				if (Data.Object.IsValid())
+				{
+					bFound = true;
+					UE_LOG(LogCook, Display, TEXT("    %s"), *Data.Object.Get()->GetFullName());
+				}
+			}
+			if (!bFound)
+			{
+				UE_LOG(LogCook, Display, TEXT("    <None>"));
+			}
+
+			SaveBusyTimeStarted = CurrentTime;
+			SaveBusyTimeLastRetry = CurrentTime;
+		}
+		if (bLoadBusy && CurrentTime - LoadBusyTimeStarted > GCookProgressWarnBusyTime)
+		{
+			UE_LOG(LogCook, Warning, TEXT("Cooker has been blocked from loading the current packages for %f seconds. Current packages in the loadqueue:"), GCookProgressWarnBusyTime);
+			FLoadPrepareQueue& LoadPrepareQueue = PackageDatas->GetLoadPrepareQueue();
+			for (FPackageData* PackageData : LoadPrepareQueue.PreloadingQueue)
+			{
+				UE_LOG(LogCook, Display, TEXT("%s"), *PackageData->GetFileName().ToString());
+			}
+			for (FPackageData* PackageData : LoadPrepareQueue.EntryQueue)
+			{
+				UE_LOG(LogCook, Display, TEXT("%s"), *PackageData->GetFileName().ToString());
+			}
+			LoadBusyTimeStarted = CurrentTime;
+			LoadBusyTimeLastRetry = CurrentTime;
+		}
 	}
 }
 UCookOnTheFlyServer::ECookAction UCookOnTheFlyServer::DecideNextCookAction(UE::Cook::FTickStackData& StackData)
@@ -1296,7 +1410,6 @@ UCookOnTheFlyServer::ECookAction UCookOnTheFlyServer::DecideNextCookAction(UE::C
 		}
 	}
 
-
 	int32 NumSaves = PackageDatas->GetSaveQueue().Num();
 	bool bSaveAvailable = ((!bSaveBusy) & (NumSaves > 0)) != 0;
 	if (bSaveAvailable & (NumSaves > static_cast<int32>(DesiredSaveQueueLength)))
@@ -1328,6 +1441,25 @@ UCookOnTheFlyServer::ECookAction UCookOnTheFlyServer::DecideNextCookAction(UE::C
 		return ECookAction::Load;
 	}
 
+	if (bSaveBusy & (NumSaves > 0))
+	{
+		const float CurrentTime = FPlatformTime::Seconds();
+		if (CurrentTime - SaveBusyTimeLastRetry > GCookProgressRetryBusyTime)
+		{
+			SaveBusyTimeLastRetry = CurrentTime;
+			return ECookAction::Save;
+		}
+	}
+	if (bLoadBusy & (NumLoads > 0))
+	{
+		const float CurrentTime = FPlatformTime::Seconds();
+		if (CurrentTime - LoadBusyTimeLastRetry > GCookProgressRetryBusyTime)
+		{
+			LoadBusyTimeLastRetry = CurrentTime;
+			return ECookAction::Load;
+		}
+	}
+
 	if (PackageDatas->GetMonitor().GetNumInProgress() > 0)
 	{
 		return ECookAction::YieldTick;
@@ -1397,24 +1529,28 @@ void UCookOnTheFlyServer::PumpExternalRequests(const UE::Cook::FCookerTimer& Coo
 	}
 }
 
-void UCookOnTheFlyServer::PumpRequests(UE::Cook::FTickStackData& StackData)
+void UCookOnTheFlyServer::PumpRequests(UE::Cook::FTickStackData& StackData, int32& OutNumPushed)
 {
 	UE_SCOPED_COOKTIMER(PumpRequests);
 	using namespace UE::Cook;
 
+	OutNumPushed = 0;
 	FRequestQueue& RequestQueue = PackageDatas->GetRequestQueue();
 	COOK_STAT(DetailedCookStats::PeakRequestQueueSize = FMath::Max(DetailedCookStats::PeakRequestQueueSize, static_cast<int32>(RequestQueue.Num())));
 	if (!RequestQueue.IsEmpty())
 	{
 		FPackageData* PackageData = RequestQueue.PopRequest();
 		FPoppedPackageDataScope Scope(*PackageData);
-		ProcessRequest(*PackageData, StackData);
+		int32 NumPushed;
+		ProcessRequest(*PackageData, StackData, NumPushed);
+		OutNumPushed += NumPushed;
 	}
 }
 
-void UCookOnTheFlyServer::ProcessRequest(UE::Cook::FPackageData& PackageData, UE::Cook::FTickStackData& StackData)
+void UCookOnTheFlyServer::ProcessRequest(UE::Cook::FPackageData& PackageData, UE::Cook::FTickStackData& StackData, int32& OutNumPushed)
 {
 	using namespace UE::Cook;
+	OutNumPushed = 0;
 
 	if (PackageData.HasAllCookedPlatforms(PackageData.GetRequestedPlatforms(), true /* bIncludeFailed */))
 	{
@@ -1454,16 +1590,20 @@ void UCookOnTheFlyServer::ProcessRequest(UE::Cook::FPackageData& PackageData, UE
 
 	if (!PackageData.GetIsUrgent() && (!IsCookByTheBookMode() || !CookByTheBookOptions->bSkipHardReferences))
 	{
-		AddDependenciesToLoadQueue(PackageData);
+		int32 NumPushed;
+		AddDependenciesToLoadQueue(PackageData, NumPushed);
+		OutNumPushed += NumPushed;
 	}
 	// AddDependenciesToLoadQueue is supposed to add the dependencies only and not add the passed-in packagedata, so it should still be in request
 	check(PackageData.GetState() == EPackageState::Request);
 	PackageData.SendToState(EPackageState::LoadPrepare, ESendFlags::QueueAdd);
+	++OutNumPushed;
 }
 
-void UCookOnTheFlyServer::AddDependenciesToLoadQueue(UE::Cook::FPackageData& PackageData)
+void UCookOnTheFlyServer::AddDependenciesToLoadQueue(UE::Cook::FPackageData& PackageData, int32& OutNumPushed)
 {
 	using namespace UE::Cook;
+	OutNumPushed = 0;
 
 	struct FPackageAndDependencies
 	{
@@ -1570,18 +1710,21 @@ void UCookOnTheFlyServer::AddDependenciesToLoadQueue(UE::Cook::FPackageData& Pac
 			check(CurrentPackageData.GetState() == EPackageState::Request); 
 			// Send the package to the load queue
 			CurrentPackageData.SendToState(EPackageState::LoadPrepare, ESendFlags::QueueAdd);
+			++OutNumPushed;
 			COOK_STAT(++DetailedCookStats::NumPreloadedDependencies);
 		}
 	}
 }
 
-void UCookOnTheFlyServer::PumpLoads(UE::Cook::FTickStackData& StackData, uint32 DesiredQueueLength)
+void UCookOnTheFlyServer::PumpLoads(UE::Cook::FTickStackData& StackData, uint32 DesiredQueueLength, int32& OutNumPushed, bool& bOutBusy)
 {
 	using namespace UE::Cook;
 	FPackageDataQueue& LoadReadyQueue = PackageDatas->GetLoadReadyQueue();
 	FLoadPrepareQueue& LoadPrepareQueue = PackageDatas->GetLoadPrepareQueue();
 	FPackageDataMonitor& Monitor = PackageDatas->GetMonitor();
 	bool bIsUrgentInProgress = Monitor.GetNumUrgent() > 0;
+	OutNumPushed = 0;
+	bOutBusy = false;
 
 	while (LoadReadyQueue.Num() + LoadPrepareQueue.Num() > static_cast<int32>(DesiredQueueLength))
 	{
@@ -1603,7 +1746,7 @@ void UCookOnTheFlyServer::PumpLoads(UE::Cook::FTickStackData& StackData, uint32
 			{
 				if (!LoadPrepareQueue.IsEmpty())
 				{
-					bLoadBusy = true;
+					bOutBusy = true;
 				}
 				break;
 			}
@@ -1611,7 +1754,9 @@ void UCookOnTheFlyServer::PumpLoads(UE::Cook::FTickStackData& StackData, uint32
 
 		FPackageData& PackageData(*LoadReadyQueue.PopFrontValue());
 		FPoppedPackageDataScope Scope(PackageData);
-		LoadPackageInQueue(PackageData, StackData.ResultFlags);
+		int32 NumPushed;
+		LoadPackageInQueue(PackageData, StackData.ResultFlags, NumPushed);
+		OutNumPushed += NumPushed;
 		ProcessUnsolicitedPackages(); // May add new packages into the LoadQueue
 	}
 }
@@ -1657,9 +1802,10 @@ void UCookOnTheFlyServer::PumpPreloadStarts()
 	}
 }
 
-void UCookOnTheFlyServer::LoadPackageInQueue(UE::Cook::FPackageData& PackageData, uint32& ResultFlags)
+void UCookOnTheFlyServer::LoadPackageInQueue(UE::Cook::FPackageData& PackageData, uint32& ResultFlags, int32& OutNumPushed)
 {
 	UPackage* LoadedPackage = nullptr;
+	OutNumPushed = 0;
 
 	FName PackageFileName(PackageData.GetFileName());
 	bool bLoadFullySuccessful = LoadPackageForCooking(PackageData, LoadedPackage);
@@ -1716,6 +1862,7 @@ void UCookOnTheFlyServer::LoadPackageInQueue(UE::Cook::FPackageData& PackageData
 	{
 		PackageData.SetPackage(LoadedPackage);
 		PackageData.SendToState(UE::Cook::EPackageState::Save, UE::Cook::ESendFlags::QueueAdd);
+		++OutNumPushed;
 	}
 }
 
@@ -2150,9 +2297,11 @@ UE_TRACE_EVENT_BEGIN(UE_CUSTOM_COOKTIMER_LOG, SaveCookedPackage, NoSync)
 	UE_TRACE_EVENT_FIELD(Trace::WideString, PackageName)
 UE_TRACE_EVENT_END()
 
-void UCookOnTheFlyServer::PumpSaves(UE::Cook::FTickStackData& StackData, uint32 DesiredQueueLength)
+void UCookOnTheFlyServer::PumpSaves(UE::Cook::FTickStackData& StackData, uint32 DesiredQueueLength, int32& OutNumPushed, bool& bOutBusy)
 {
 	using namespace UE::Cook;
+	OutNumPushed = 0;
+	bOutBusy = false;
 
 	UE_SCOPED_HIERARCHICAL_COOKTIMER(SavingPackages);
 	check(IsInGameThread());
@@ -2180,6 +2329,7 @@ void UCookOnTheFlyServer::PumpSaves(UE::Cook::FTickStackData& StackData, uint32
 		{
 			// We already attempted to cook this package and it's still not referenced by any non editor-only properties.
 			PackageData.SendToState(EPackageState::Idle, ESendFlags::QueueAdd);
+			++OutNumPushed;
 			continue;
 		}
 
@@ -2191,6 +2341,7 @@ void UCookOnTheFlyServer::PumpSaves(UE::Cook::FTickStackData& StackData, uint32
 		{
 			// refuse to save this package, it's clearly one of the undesirables
 			PackageData.SendToState(EPackageState::Idle, ESendFlags::QueueAdd);
+			++OutNumPushed;
 			continue;
 		}
 
@@ -2202,6 +2353,7 @@ void UCookOnTheFlyServer::PumpSaves(UE::Cook::FTickStackData& StackData, uint32
 			// All places that add a package to the save queue check for existence of incomplete platforms before adding
 			UE_LOG(LogCook, Warning, TEXT("Package '%s' in SaveQueue has no more platforms left to cook; this should not be possible!"), *PackageData.GetFileName().ToString());
 			PackageData.SendToState(EPackageState::Idle, ESendFlags::QueueAdd);
+			++OutNumPushed;
 			continue;
 		}
 
@@ -2291,7 +2443,7 @@ void UCookOnTheFlyServer::PumpSaves(UE::Cook::FTickStackData& StackData, uint32
 			if (!AllObjectsCookedDataCached)
 			{
 				StackData.ResultFlags |= COSR_WaitingOnCache;
-				bSaveBusy = true;
+				bOutBusy = true;
 				SaveQueue.AddFront(&PackageData);
 				return;
 			}
@@ -2468,6 +2620,7 @@ void UCookOnTheFlyServer::PumpSaves(UE::Cook::FTickStackData& StackData, uint32
 		}
 
 		PackageData.SendToState(EPackageState::Idle, ESendFlags::QueueAdd);
+		++OutNumPushed;
 	}
 }
 
@@ -6223,8 +6376,8 @@ void UCookOnTheFlyServer::CancelAllQueues()
 		RequestQueue.PopRequest()->SendToState(EPackageState::Idle, ESendFlags::QueueAdd);
 	}
 
-	bLoadBusy = false;
-	bSaveBusy = false;
+	SetLoadBusy(false);
+	SetSaveBusy(false);
 }
 
 
