Fix save discrepencies between save2 and save1 by allowing objects of transient classes
Add an additional check in the linker diff commandlet to check the old save algo against itself and catch potential save byproduct
Add a consistent vs algo linker diff mode for the cooker, consistent mode will detect save byproduct on the same save algorithm while algo mode will compare the two algorithm

diff --git a/Engine/Source/Editor/UnrealEd/Private/Commandlets/SavePackageUtilitiesCommandlet.cpp b/Engine/Source/Editor/UnrealEd/Private/Commandlets/SavePackageUtilitiesCommandlet.cpp
--- a/Engine/Source/Editor/UnrealEd/Private/Commandlets/SavePackageUtilitiesCommandlet.cpp
+++ b/Engine/Source/Editor/UnrealEd/Private/Commandlets/SavePackageUtilitiesCommandlet.cpp
@@ -82,6 +82,16 @@ int32 USavePackageUtilitiesCommandlet::Main(const FString& Params)
 				FDateTime::MinValue(), false, /*DiffMap*/ nullptr,
 				nullptr);
 		}
+
+		// Old Save Package
+		FSavePackageResultStruct OldResultCheck;
+		{
+			EnableNewSave->Set(0);
+			OldResultCheck = GEditor->Save(Package, Asset, SaveArgs.TopLevelFlags, *Filename,
+				GError, nullptr, false, true, SaveArgs.SaveFlags, SaveArgs.TargetPlatform,
+				FDateTime::MinValue(), false, /*DiffMap*/ nullptr,
+				nullptr);
+		}
 		EnableNewSave->Set(EnableNewSavePreviousValue);
 
 		if (OldResult.LinkerSave && NewResult.LinkerSave)
@@ -90,7 +100,15 @@ int32 USavePackageUtilitiesCommandlet::Main(const FString& Params)
 			FLinkerDiff LinkerDiff = FLinkerDiff::CompareLinkers(OldResult.LinkerSave.Get(), NewResult.LinkerSave.Get());
 			LinkerDiff.PrintDiff(*GWarn);
 		}
-		
+
+		// Add a old save against itself check to test potential byproduct, doesn't catch them all, since oftentimes byproduct are caused by the first save
+		if (OldResultCheck.LinkerSave && OldResult.LinkerSave)
+		{
+			// Compare Linker Save info
+			FLinkerDiff LinkerDiff = FLinkerDiff::CompareLinkers(OldResultCheck.LinkerSave.Get(), OldResult.LinkerSave.Get());
+			LinkerDiff.PrintDiff(*GWarn);
+		}
+
 		// Delete the temp filename
 		IFileManager::Get().Delete(*Filename);
 	}
diff --git a/Engine/Source/Editor/UnrealEd/Private/CookOnTheFlyServer.cpp b/Engine/Source/Editor/UnrealEd/Private/CookOnTheFlyServer.cpp
--- a/Engine/Source/Editor/UnrealEd/Private/CookOnTheFlyServer.cpp
+++ b/Engine/Source/Editor/UnrealEd/Private/CookOnTheFlyServer.cpp
@@ -3307,9 +3307,21 @@ bool UCookOnTheFlyServer::HasRecompileShaderRequests() const
 
 class FDiffModeCookServerUtils
 {
+public:
+	/**
+	 *  Mode for the linker diff test
+	 */
+	enum ELinkerDiffMode
+	{
+		LDM_None = 0,
+		LDM_Algo = 1,
+		LDM_Consistent = 2,
+	};
+
+private:
 	/** Misc / common settings */
 	bool bDiffEnabled;
-	bool bLinkerDiffEnabled;
+	ELinkerDiffMode LinkerDiffMode;
 	FString PackageFilter;
 
 	/** DumpObjList settings */
@@ -3322,10 +3334,24 @@ class FDiffModeCookServerUtils
 
 public:
 
+
 	FDiffModeCookServerUtils()
 	{
 		bDiffEnabled = FParse::Param(FCommandLine::Get(), TEXT("DIFFONLY"));
-		bLinkerDiffEnabled = FParse::Param(FCommandLine::Get(), TEXT("LINKERDIFF"));
+
+		LinkerDiffMode = LDM_None;
+		FString DiffMode;
+		if (FParse::Value(FCommandLine::Get(), TEXT("-LINKERDIFF="), DiffMode))
+		{
+			if (DiffMode == TEXT("1") || DiffMode == TEXT("algo"))
+			{
+				LinkerDiffMode = LDM_Algo;
+			}
+			else if (DiffMode == TEXT("2") || DiffMode == TEXT("consistent"))
+			{
+				LinkerDiffMode = LDM_Consistent;
+			}
+		}
 
 		bDumpObjList = false;
 		bDumpObjects = false;
@@ -3341,7 +3367,12 @@ public:
 
 	bool IsRunningCookLinkerDiff() const
 	{
-		return bLinkerDiffEnabled;
+		return LinkerDiffMode != LDM_None;
+	}
+
+	ELinkerDiffMode GetLinkerDiffMode() const
+	{
+		return LinkerDiffMode;
 	}
 
 	void ProcessPackage(UPackage* InPackage)
@@ -3670,12 +3701,24 @@ void UCookOnTheFlyServer::SaveCookedPackage(UE::Cook::FPackageData& PackageData,
 					{
 						static IConsoleVariable* EnableNewSave = IConsoleManager::Get().FindConsoleVariable(TEXT("SavePackage.EnableNewSave"));
 						bool bPreviousCvarValue = EnableNewSave->GetBool();
-						EnableNewSave->Set(!bPreviousCvarValue);
+
+						// If the linker diff is comparing the two save algo, switch the cvar to the other algo before saving again,
+						// Otherwise the linker diff mode is tracking if the save is consistent across multiple save
+						bool bAlgoLinkerDiff = DiffModeHelper.GetLinkerDiffMode() == FDiffModeCookServerUtils::LDM_Algo;
+						if (bAlgoLinkerDiff)
+						{
+							EnableNewSave->Set(!bPreviousCvarValue);
+						}
+
 						FSavePackageResultStruct NewResult = GEditor->Save(Package, World, FlagsToCook, *PlatFilename,
 							GError, nullptr, bSwap, false, SaveFlags|SAVE_DiffOnly, Target,
 							FDateTime::MinValue(), false, /*DiffMap*/ nullptr,
 							nullptr);
-						EnableNewSave->Set(bPreviousCvarValue);
+
+						if (bAlgoLinkerDiff)
+						{
+							EnableNewSave->Set(bPreviousCvarValue);
+						}
 
 						if (Result.LinkerSave && NewResult.LinkerSave)
 						{
diff --git a/Engine/Source/Runtime/CoreUObject/Private/UObject/SavePackage/SaveContext.cpp b/Engine/Source/Runtime/CoreUObject/Private/UObject/SavePackage/SaveContext.cpp
--- a/Engine/Source/Runtime/CoreUObject/Private/UObject/SavePackage/SaveContext.cpp
+++ b/Engine/Source/Runtime/CoreUObject/Private/UObject/SavePackage/SaveContext.cpp
@@ -21,7 +21,9 @@ bool FSaveContext::IsUnsaveable(UObject* InObject) const
 	while (Obj)
 	{
 		// if the object class is abstract, has been marked as deprecated, there is a newer version that exist or the class is maked transient, then the object is unsaveable
-		if (Obj->GetClass()->HasAnyClassFlags(CLASS_Abstract | CLASS_Deprecated | CLASS_NewerVersionExists | CLASS_Transient) && !Obj->HasAnyFlags(RF_ClassDefaultObject))
+		// @note: Although object instance of transient class should definitely be unsaveable, it results in discrepencies with the old save algorithm and currently load problems
+		if (Obj->GetClass()->HasAnyClassFlags(CLASS_Abstract | CLASS_Deprecated | CLASS_NewerVersionExists /*| CLASS_Transient*/) 
+			&& !Obj->HasAnyFlags(RF_ClassDefaultObject))
 		{
 			if (!InObject->IsPendingKill() 
 				&& InObject->GetOutermost() == GetPackage()
diff --git a/Engine/Source/Runtime/CoreUObject/Private/UObject/SavePackage/SaveContext.h b/Engine/Source/Runtime/CoreUObject/Private/UObject/SavePackage/SaveContext.h
--- a/Engine/Source/Runtime/CoreUObject/Private/UObject/SavePackage/SaveContext.h
+++ b/Engine/Source/Runtime/CoreUObject/Private/UObject/SavePackage/SaveContext.h
@@ -575,7 +575,7 @@ private:
 	TMap<UObject*, TSet<UObject*>> ExportObjectDependencies;
 	// Map of objects to their native dependencies
 	TMap<UObject*, TSet<UObject*>> ExportNativeObjectDependencies;
-	// Set of harvested prestream packages
+	// Set of harvested prestream packages, should be deprecated, not implementated in Save2 yet
 	TSet<UPackage*> PrestreamPackages;
 	// Harvested custom versions
 	FCustomVersionContainer CustomVersions;
