When changing the name of a repnotify blueprint variable, a warning is displayed and the OnRep function is no longer called.

diff --git a/Engine/Source/Editor/UnrealEd/Private/Kismet2/BlueprintEditorUtils.cpp b/Engine/Source/Editor/UnrealEd/Private/Kismet2/BlueprintEditorUtils.cpp
--- a/Engine/Source/Editor/UnrealEd/Private/Kismet2/BlueprintEditorUtils.cpp
+++ b/Engine/Source/Editor/UnrealEd/Private/Kismet2/BlueprintEditorUtils.cpp
@@ -4843,6 +4843,21 @@ void FBlueprintEditorUtils::RenameMemberVariable(UBlueprint* Blueprint, const FN
 			// Update the name
 			FBPVariableDescription& Variable = Blueprint->NewVariables[VarIndex];
 			Variable.VarName = NewName;
+
+			// If the variable has an associated OnRep function, warn the user and break the association if the name is changed
+			FName OnRepFuncName = Blueprint->NewVariables[VarIndex].RepNotifyFunc;
+			if (OnRepFuncName != NAME_None)
+			{
+				if (!VerifyUserWantsRepNotifyVariableNameChanged(OldName, OnRepFuncName))
+				{
+					// Showing the warning dialog causes the variable name text box to lose focus, which can result in this function being called again.
+					// The VarName is set before verifying to skip over the second call, preventing the dialog from appearing twice. 
+					Variable.VarName = OldName;
+					return;
+				}
+				Blueprint->NewVariables[VarIndex].RepNotifyFunc = NAME_None;
+			}
+
 			Variable.FriendlyName = FName::NameToDisplayString( NewName.ToString(), (Variable.VarType.PinCategory == UEdGraphSchema_K2::PC_Boolean) ? true : false );
 
 			// Update any existing references to the old name
@@ -4951,6 +4966,25 @@ bool FBlueprintEditorUtils::VerifyUserWantsVariableTypeChanged(const FName& InVa
 	return RetCode == FSuppressableWarningDialog::Confirm || RetCode == FSuppressableWarningDialog::Suppressed;
 }
 
+bool FBlueprintEditorUtils::VerifyUserWantsRepNotifyVariableNameChanged(const FName& InVarName, const FName& InFuncName)
+{
+	FFormatNamedArguments Args;
+	Args.Add(TEXT("VariableName"), FText::FromName(InVarName));
+	Args.Add(TEXT("FuncName"), FText::FromName(InFuncName));
+
+	FText ConfirmRename = FText::Format(LOCTEXT("ConfirmChangeRepNotifyVarName",
+		"Variable '{VariableName}' is linked to the OnRep function '{FuncName}'. Renaming it will still allow this variable to be replicated, but the function will not be called. Do you wish to proceed?"), Args);
+
+	FSuppressableWarningDialog::FSetupInfo Info(ConfirmRename, LOCTEXT("ChangeRepNotifyVariableName", "Change RepNotify Variable Name"), "ChangeRepNotifyVariableName_Warning");
+	Info.ConfirmText = LOCTEXT("ChangeRepNotifyVariableName_Yes", "Yes");
+	Info.CancelText = LOCTEXT("ChangeRepNotifyVariableName_No", "No");
+
+	FSuppressableWarningDialog ChangeRepNotifyVariableName(Info);
+
+	FSuppressableWarningDialog::EResult RetCode = ChangeRepNotifyVariableName.ShowModal();
+	return RetCode == FSuppressableWarningDialog::Confirm || RetCode == FSuppressableWarningDialog::Suppressed;
+}
+
 void FBlueprintEditorUtils::GetLoadedChildBlueprints(UBlueprint* InBlueprint, TArray<UBlueprint*>& OutBlueprints)
 {
 	// Iterate over currently-loaded Blueprints and potentially adjust their variable names if they conflict with the parent
diff --git a/Engine/Source/Editor/UnrealEd/Public/Kismet2/BlueprintEditorUtils.h b/Engine/Source/Editor/UnrealEd/Public/Kismet2/BlueprintEditorUtils.h
--- a/Engine/Source/Editor/UnrealEd/Public/Kismet2/BlueprintEditorUtils.h
+++ b/Engine/Source/Editor/UnrealEd/Public/Kismet2/BlueprintEditorUtils.h
@@ -1737,6 +1737,15 @@ protected:
 	 */
 	static bool VerifyUserWantsVariableTypeChanged(const FName& InVarName);
 
+	/**
+	 * Helper function to warn user of the results of changing a RepNotify variable name by displaying a suppressible dialog
+	 *
+	 * @param InVarName		Variable name to display in the dialog message
+	 * @param InFuncName	Associated OnRep function name to display in the dialog message
+	 * @return				TRUE if the user wants to change the variable name
+	 */
+	static bool VerifyUserWantsRepNotifyVariableNameChanged(const FName& InVarName, const FName& InFuncName);
+
 	/**
 	 * Helper function to get all loaded Blueprints that are children (or using as an interface) the passed Blueprint
 	 *
