Created a default implementation for property widget and legacy gizmos to aid in transition from FEdMode to UEdModes.

Updated default ed mode to UEdMode using the legacy interfaces.

diff --git a/Engine/Source/Editor/UnrealEd/Private/EdMode.cpp b/Engine/Source/Editor/UnrealEd/Private/EdMode.cpp
--- a/Engine/Source/Editor/UnrealEd/Private/EdMode.cpp
+++ b/Engine/Source/Editor/UnrealEd/Private/EdMode.cpp
@@ -15,192 +15,12 @@
 
 #include "CanvasTypes.h"
 
-/** Hit proxy used for editable properties */
-struct HPropertyWidgetProxy : public HHitProxy
-{
-	DECLARE_HIT_PROXY();
-
-	/** Name of property this is the widget for */
-	FString	PropertyName;
-
-	/** If the property is an array property, the index into that array that this widget is for */
-	int32	PropertyIndex;
-
-	/** This property is a transform */
-	bool	bPropertyIsTransform;
-
-	HPropertyWidgetProxy(FString InPropertyName, int32 InPropertyIndex, bool bInPropertyIsTransform)
-		: HHitProxy(HPP_Foreground)
-		, PropertyName(InPropertyName)
-		, PropertyIndex(InPropertyIndex)
-		, bPropertyIsTransform(bInPropertyIsTransform)
-	{}
-
-	/** Show cursor as cross when over this handle */
-	virtual EMouseCursor::Type GetMouseCursor() override
-	{
-		return EMouseCursor::Crosshairs;
-	}
-};
-
-IMPLEMENT_HIT_PROXY(HPropertyWidgetProxy, HHitProxy);
-
-
-namespace
-{
-	/**
-	 * Returns a reference to the named property value data in the given container.
-	 */
-	template<typename T>
-	T* GetPropertyValuePtrByName(const UStruct* InStruct, void* InContainer, FString PropertyName, int32 ArrayIndex, FProperty*& OutProperty)
-	{
-		T* ValuePtr = NULL;
-
-		// Extract the vector ptr recursively using the property name
-		int32 DelimPos = PropertyName.Find(TEXT("."));
-		if(DelimPos != INDEX_NONE)
-		{
-			// Parse the property name and (optional) array index
-			int32 SubArrayIndex = 0;
-			FString NameToken = PropertyName.Left(DelimPos);
-			int32 ArrayPos = NameToken.Find(TEXT("["));
-			if(ArrayPos != INDEX_NONE)
-			{
-				FString IndexToken = NameToken.RightChop(ArrayPos + 1).LeftChop(1);
-				SubArrayIndex = FCString::Atoi(*IndexToken);
-
-				NameToken = PropertyName.Left(ArrayPos);
-			}
-
-			// Obtain the property info from the given structure definition
-			FProperty* CurrentProp = FindFProperty<FProperty>(InStruct, FName(*NameToken));
-
-			// Check first to see if this is a simple structure (i.e. not an array of structures)
-			FStructProperty* StructProp = CastField<FStructProperty>(CurrentProp);
-			if(StructProp != NULL)
-			{
-				// Recursively call back into this function with the structure property and container value
-				ValuePtr = GetPropertyValuePtrByName<T>(StructProp->Struct, StructProp->ContainerPtrToValuePtr<void>(InContainer), PropertyName.RightChop(DelimPos + 1), ArrayIndex, OutProperty);
-			}
-			else
-			{
-				// Check to see if this is an array
-				FArrayProperty* ArrayProp = CastField<FArrayProperty>(CurrentProp);
-				if(ArrayProp != NULL)
-				{
-					// It is an array, now check to see if this is an array of structures
-					StructProp = CastField<FStructProperty>(ArrayProp->Inner);
-					if(StructProp != NULL)
-					{
-						FScriptArrayHelper_InContainer ArrayHelper(ArrayProp, InContainer);
-						if(ArrayHelper.IsValidIndex(SubArrayIndex))
-						{
-							// Recursively call back into this function with the array element and container value
-							ValuePtr = GetPropertyValuePtrByName<T>(StructProp->Struct, ArrayHelper.GetRawPtr(SubArrayIndex), PropertyName.RightChop(DelimPos + 1), ArrayIndex, OutProperty);
-						}
-					}
-				}
-			}
-		}
-		else
-		{
-			FProperty* Prop = FindFProperty<FProperty>(InStruct, FName(*PropertyName));
-			if(Prop != NULL)
-			{
-				if( FArrayProperty* ArrayProp = CastField<FArrayProperty>(Prop) )
-				{
-					check(ArrayIndex != INDEX_NONE);
-
-					// Property is an array property, so make sure we have a valid index specified
-					FScriptArrayHelper_InContainer ArrayHelper(ArrayProp, InContainer);
-					if( ArrayHelper.IsValidIndex(ArrayIndex) )
-					{
-						ValuePtr = (T*)ArrayHelper.GetRawPtr(ArrayIndex);
-					}
-				}
-				else
-				{
-					// Property is a vector property, so access directly
-					ValuePtr = Prop->ContainerPtrToValuePtr<T>(InContainer);
-				}
-
-				OutProperty = Prop;
-			}
-		}
-
-		return ValuePtr;
-	}
-
-	/**
-	 * Returns the value of the property with the given name in the given Actor instance.
-	 */
-	template<typename T>
-	T GetPropertyValueByName(UObject* Object, FString PropertyName, int32 PropertyIndex)
-	{
-		T Value;
-		FProperty* DummyProperty = NULL;
-		if (T* ValuePtr = GetPropertyValuePtrByName<T>(Object->GetClass(), Object, PropertyName, PropertyIndex, DummyProperty))
-		{
-			Value = *ValuePtr;
-		}
-		return Value;
-	}
-
-	/**
-	 * Sets the property with the given name in the given Actor instance to the given value.
-	 */
-	template<typename T>
-	void SetPropertyValueByName(UObject* Object, FString PropertyName, int32 PropertyIndex, const T& InValue, FProperty*& OutProperty)
-	{
-		if (T* ValuePtr = GetPropertyValuePtrByName<T>(Object->GetClass(), Object, PropertyName, PropertyIndex, OutProperty))
-		{
-			*ValuePtr = InValue;
-		}
-	}
-}
-
-//////////////////////////////////
-// FEdMode
-
-void FEdMode::FPropertyWidgetInfo::GetTransformAndColor(UObject* BestSelectedItem, bool bIsSelected, FTransform& OutLocalTransform, FString& OutValidationMessage, FColor& OutDrawColor) const
-{
-	// Determine the desired position
-	if (bIsTransform)
-	{
-		OutLocalTransform = GetPropertyValueByName<FTransform>(BestSelectedItem, PropertyName, PropertyIndex);
-	}
-	else
-	{
-		OutLocalTransform = FTransform(GetPropertyValueByName<FVector>(BestSelectedItem, PropertyName, PropertyIndex));
-	}
-
-	// Determine the desired color
-	OutDrawColor = bIsSelected ? FColor::White : FColor(128, 128, 255);
-	if (PropertyValidationName != NAME_None)
-	{
-		if (UFunction* ValidateFunc = BestSelectedItem->FindFunction(PropertyValidationName))
-		{
-			BestSelectedItem->ProcessEvent(ValidateFunc, &OutValidationMessage);
-
-			// if we have a negative result, the widget color is red.
-			OutDrawColor = OutValidationMessage.IsEmpty() ? OutDrawColor : FColor::Red;
-		}
-	}
-}
-
 //////////////////////////////////
 // FEdMode
 
-const FName FEdMode::MD_MakeEditWidget(TEXT("MakeEditWidget"));
-const FName FEdMode::MD_ValidateWidgetUsing(TEXT("ValidateWidgetUsing"));
-
 FEdMode::FEdMode()
 	: bPendingDeletion(false)
-	, CurrentWidgetAxis(EAxisList::None)
 	, CurrentTool(nullptr)
-	, Owner(nullptr)
-	, EditedPropertyIndex(INDEX_NONE)
-	, bEditedPropertyIsTransform(false)
 {
 	bDrawKillZ = true;
 }
@@ -323,61 +143,9 @@ bool FEdMode::InputAxis(FEditorViewportClient* InViewportClient, FViewport* View
 
 bool FEdMode::InputDelta(FEditorViewportClient* InViewportClient, FViewport* InViewport, FVector& InDrag, FRotator& InRot, FVector& InScale)
 {	
-	if (UsesPropertyWidgets())
+	if (FLegacyEdModeWidgetHelper::InputDelta(InViewportClient, InViewport, InDrag, InRot, InScale))
 	{
-		FTransform DisplayWidgetToWorld;
-		UObject* BestSelectedItem = GetItemToTryDisplayingWidgetsFor(/*out*/ DisplayWidgetToWorld);
-
-		if ((BestSelectedItem != nullptr) && (InViewportClient->GetCurrentWidgetAxis() != EAxisList::None))
-		{
-			GEditor->NoteActorMovement();
-
-			if (!EditedPropertyName.IsEmpty())
-			{
-				FTransform LocalTM = FTransform::Identity;
-
-				if (bEditedPropertyIsTransform)
-				{
-					LocalTM = GetPropertyValueByName<FTransform>(BestSelectedItem, EditedPropertyName, EditedPropertyIndex);
-				}
-				else
-				{					
-					FVector LocalPos = GetPropertyValueByName<FVector>(BestSelectedItem, EditedPropertyName, EditedPropertyIndex);
-					LocalTM = FTransform(LocalPos);
-				}
-
-				// Calculate world transform
-				FTransform WorldTM = LocalTM * DisplayWidgetToWorld;
-				// Calc delta specified by drag
-				//FTransform DeltaTM(InRot.Quaternion(), InDrag);
-				// Apply delta in world space
-				WorldTM.SetTranslation(WorldTM.GetTranslation() + InDrag);
-				WorldTM.SetRotation(InRot.Quaternion() * WorldTM.GetRotation());
-				// Convert new world transform back into local space
-				LocalTM = WorldTM.GetRelativeTransform(DisplayWidgetToWorld);
-				// Apply delta scale
-				LocalTM.SetScale3D(LocalTM.GetScale3D() + InScale);
-
-				BestSelectedItem->PreEditChange(NULL);
-
-				// Property that we actually change
-				FProperty* SetProperty = NULL;
-
-				if (bEditedPropertyIsTransform)
-				{
-					SetPropertyValueByName<FTransform>(BestSelectedItem, EditedPropertyName, EditedPropertyIndex, LocalTM, SetProperty);
-				}
-				else
-				{
-					SetPropertyValueByName<FVector>(BestSelectedItem, EditedPropertyName, EditedPropertyIndex, LocalTM.GetLocation(), SetProperty);
-				}
-
-				FPropertyChangedEvent PropertyChangeEvent(SetProperty);
-				BestSelectedItem->PostEditChangeProperty(PropertyChangeEvent);
-
-				return true;
-			}
-		}
+		return true;
 	}
 
 	if (GetCurrentTool())
@@ -398,118 +166,11 @@ bool FEdMode::UsesTransformWidget() const
 	return true;
 }
 
-bool FEdMode::UsesTransformWidget(UE::Widget::EWidgetMode CheckMode) const
-{
-	if (UsesPropertyWidgets())
-	{
-		FTransform DisplayWidgetToWorld;
-		UObject* BestSelectedItem = GetItemToTryDisplayingWidgetsFor(/*out*/ DisplayWidgetToWorld);
-
-		if (BestSelectedItem != nullptr)
-		{
-			// If editing a vector (not a transform)
-			if (!EditedPropertyName.IsEmpty() && !bEditedPropertyIsTransform)
-			{
-				return (CheckMode == UE::Widget::WM_Translate);
-			}
-		}
-	}
-
-	return true;
-}
-
-FVector FEdMode::GetWidgetLocation() const
+bool FEdMode::UsesPropertyWidgets() const
 {
-	if (UsesPropertyWidgets())
-	{
-		FTransform DisplayWidgetToWorld;
-		UObject* BestSelectedItem = GetItemToTryDisplayingWidgetsFor(/*out*/ DisplayWidgetToWorld);
-		if (BestSelectedItem)
-		{
-			if (!EditedPropertyName.IsEmpty())
-			{
-				FVector LocalPos = FVector::ZeroVector;
-
-				if (bEditedPropertyIsTransform)
-				{
-					FTransform LocalTM = GetPropertyValueByName<FTransform>(BestSelectedItem, EditedPropertyName, EditedPropertyIndex);
-					LocalPos = LocalTM.GetLocation();
-				}
-				else
-				{
-					LocalPos = GetPropertyValueByName<FVector>(BestSelectedItem, EditedPropertyName, EditedPropertyIndex);
-				}
-
-				const FVector WorldPos = DisplayWidgetToWorld.TransformPosition(LocalPos);
-				return WorldPos;
-			}
-		}
-	}
-
-	return Owner->PivotLocation;
-}
-
-bool FEdMode::GetCustomDrawingCoordinateSystem(FMatrix& InMatrix, void* InData)
-{
-	if (UsesPropertyWidgets())
-	{
-		FTransform DisplayWidgetToWorld;
-		UObject* BestSelectedItem = GetItemToTryDisplayingWidgetsFor(/*out*/ DisplayWidgetToWorld);
-		if (BestSelectedItem)
-		{
-			if (EditedPropertyName != TEXT(""))
-			{
-				if (bEditedPropertyIsTransform)
-				{
-					FTransform LocalTM = GetPropertyValueByName<FTransform>(BestSelectedItem, EditedPropertyName, EditedPropertyIndex);
-					InMatrix = FRotationMatrix::Make((LocalTM * DisplayWidgetToWorld).GetRotation());
-					return true;
-				}
-				else
-				{
-					InMatrix = FRotationMatrix::Make(DisplayWidgetToWorld.GetRotation());
-					return true;
-				}
-			}
-		}
-	}
-
 	return false;
 }
 
-bool FEdMode::ShouldDrawWidget() const
-{
-	bool bDrawWidget = false;
-
-	bool bHadSelectableComponents = false;
-	if (Owner->GetSelectedComponents()->Num() > 0)
-	{
-		// when components are selected, only show the widget when one or more are scene components
-		for (FSelectedEditableComponentIterator It(*Owner->GetSelectedComponents()); It; ++It)
-		{
-			bHadSelectableComponents = true;
-			if (It->IsA<USceneComponent>())
-			{
-				bDrawWidget = true;
-				break;
-			}
-		}
-	}
-
-	if (!bHadSelectableComponents)
-	{
-		// when actors are selected, only show the widget when all selected actors have scene components
-		bDrawWidget = Owner->SelectionHasSceneComponent();
-	}
-
-	return bDrawWidget;
-}
-
-EAxisList::Type FEdMode::GetWidgetAxisToDraw( UE::Widget::EWidgetMode InWidgetMode ) const
-{
-	return EAxisList::All;
-}
-
 bool FEdMode::BoxSelect( FBox& InBox, bool InSelect )
 {
 	bool bResult = false;
@@ -546,32 +207,11 @@ void FEdMode::Tick(FEditorViewportClient* ViewportClient,float DeltaTime)
 	}
 }
 
-void FEdMode::ActorSelectionChangeNotify()
-{
-	EditedPropertyName = TEXT("");
-	EditedPropertyIndex = INDEX_NONE;
-	bEditedPropertyIsTransform = false;
-}
-
 bool FEdMode::HandleClick(FEditorViewportClient* InViewportClient, HHitProxy *HitProxy, const FViewportClick &Click)
 {
-	if (UsesPropertyWidgets() && (HitProxy != nullptr))
+	if (FLegacyEdModeWidgetHelper::HandleClick(InViewportClient, HitProxy, Click))
 	{
-		if( HitProxy->IsA(HPropertyWidgetProxy::StaticGetType()) )
-		{
-			HPropertyWidgetProxy* PropertyProxy = (HPropertyWidgetProxy*)HitProxy;
-			EditedPropertyName = PropertyProxy->PropertyName;
-			EditedPropertyIndex = PropertyProxy->PropertyIndex;
-			bEditedPropertyIsTransform = PropertyProxy->bPropertyIsTransform;
-			return true;
-		}
-		// Left clicking on an actor, stop editing a property
-		else if( HitProxy->IsA(HActor::StaticGetType()) )
-		{
-			EditedPropertyName = FString();
-			EditedPropertyIndex = INDEX_NONE;
-			bEditedPropertyIsTransform = false;
-		}
+		return true;
 	}
 
 	return false;
@@ -640,41 +280,7 @@ void FEdMode::Render(const FSceneView* View,FViewport* Viewport,FPrimitiveDrawIn
 		Tool->Render( View, Viewport, PDI );
 	}
 
-	if (UsesPropertyWidgets())
-	{
-		const bool bHitTesting = PDI->IsHitTesting();
-
-		FTransform DisplayWidgetToWorld;
-		UObject* BestSelectedItem = GetItemToTryDisplayingWidgetsFor(/*out*/ DisplayWidgetToWorld);
-		
-		if (BestSelectedItem != nullptr)
-		{
-			UClass* Class = BestSelectedItem->GetClass();
-			TArray<FPropertyWidgetInfo> WidgetInfos;
-			GetPropertyWidgetInfos(Class, BestSelectedItem, WidgetInfos);
-			FEditorScriptExecutionGuard ScriptGuard;
-			for (const FPropertyWidgetInfo& WidgetInfo : WidgetInfos)
-			{
-				const bool bSelected = (WidgetInfo.PropertyName == EditedPropertyName) && (WidgetInfo.PropertyIndex == EditedPropertyIndex);
-
-				FTransform LocalWidgetTransform;
-				FString ValidationMessage;
-				FColor WidgetColor;
-				WidgetInfo.GetTransformAndColor(BestSelectedItem, bSelected, /*out*/ LocalWidgetTransform, /*out*/ ValidationMessage, /*out*/ WidgetColor);
-
-				const FTransform WorldWidgetTransform = LocalWidgetTransform * DisplayWidgetToWorld;
-				const FMatrix WidgetTM(WorldWidgetTransform.ToMatrixWithScale());
-
-				const float WidgetSize = 0.035f;
-				const float ZoomFactor = FMath::Min<float>(View->ViewMatrices.GetProjectionMatrix().M[0][0], View->ViewMatrices.GetProjectionMatrix().M[1][1]);
-				const float WidgetRadius = View->Project(WorldWidgetTransform.GetTranslation()).W * (WidgetSize / ZoomFactor);
-
-				if (bHitTesting) PDI->SetHitProxy(new HPropertyWidgetProxy(WidgetInfo.PropertyName, WidgetInfo.PropertyIndex, WidgetInfo.bIsTransform));
-				DrawWireDiamond(PDI, WidgetTM, WidgetRadius, WidgetColor, SDPG_Foreground );
-				if(bHitTesting) PDI->SetHitProxy( NULL );
-			}
-		}
-	}
+	FLegacyEdModeWidgetHelper::Render(View, Viewport, PDI);
 }
 
 void FEdMode::DrawHUD(FEditorViewportClient* ViewportClient,FViewport* Viewport,const FSceneView* View,FCanvas* Canvas)
@@ -785,47 +391,7 @@ void FEdMode::DrawHUD(FEditorViewportClient* ViewportClient,FViewport* Viewport,
 		}
 	}
 
-	if (UsesPropertyWidgets())
-	{
-		FTransform DisplayWidgetToWorld;
-		UObject* BestSelectedItem = GetItemToTryDisplayingWidgetsFor(/*out*/ DisplayWidgetToWorld);
-		if (BestSelectedItem != nullptr)
-		{
-			FEditorScriptExecutionGuard ScriptGuard;
-
-			FIntPoint SizeXYWithDPIScale = Viewport->GetSizeXY() / Canvas->GetDPIScale();
-			const int32 HalfX = 0.5f * SizeXYWithDPIScale.X;
-			const int32 HalfY = 0.5f * SizeXYWithDPIScale.Y;
-
-			UClass* Class = BestSelectedItem->GetClass();		
-			TArray<FPropertyWidgetInfo> WidgetInfos;
-			GetPropertyWidgetInfos(Class, BestSelectedItem, WidgetInfos);
-			for (const FPropertyWidgetInfo& WidgetInfo : WidgetInfos)
-			{
-				FTransform LocalWidgetTransform;
-				FString ValidationMessage;
-				FColor IgnoredWidgetColor;
-				WidgetInfo.GetTransformAndColor(BestSelectedItem, /*bSelected=*/ false, /*out*/ LocalWidgetTransform, /*out*/ ValidationMessage, /*out*/ IgnoredWidgetColor);
-
-				const FTransform WorldWidgetTransform = LocalWidgetTransform * DisplayWidgetToWorld;
-
-				const FPlane Proj = View->Project(WorldWidgetTransform.GetTranslation());
-				if (Proj.W > 0.f)
-				{
-					// do some string fixing
-					const uint32 VectorIndex = WidgetInfo.PropertyIndex;
-					const FString WidgetDisplayName = WidgetInfo.DisplayName + ((VectorIndex != INDEX_NONE) ? FString::Printf(TEXT("[%d]"), VectorIndex) : TEXT(""));
-					const FString DisplayString = ValidationMessage.IsEmpty() ? WidgetDisplayName : ValidationMessage;
-
-					const int32 XPos = HalfX + ( HalfX * Proj.X );
-					const int32 YPos = HalfY + ( HalfY * (Proj.Y * -1.f) );
-					FCanvasTextItem TextItem( FVector2D( XPos + 5, YPos), FText::FromString( DisplayString ), GEngine->GetSmallFont(), FLinearColor::White );
-					TextItem.EnableShadow(FLinearColor::Black);
-					Canvas->DrawItem( TextItem );
-				}
-			}
-		}
-	}
+	FLegacyEdModeWidgetHelper::DrawHUD(ViewportClient, Viewport, View, Canvas);
 }
 
 void FEdMode::DrawBrackets( FEditorViewportClient* ViewportClient, FViewport* Viewport, const FSceneView* View, FCanvas* Canvas )
@@ -881,250 +447,7 @@ bool FEdMode::EndTracking(FEditorViewportClient* InViewportClient, FViewport* In
 	return bResult;
 }
 
-FVector FEdMode::GetWidgetNormalFromCurrentAxis( void* InData )
-{
-	// Figure out the proper coordinate system.
-
-	FMatrix matrix = FMatrix::Identity;
-	if( Owner->GetCoordSystem() == COORD_Local )
-	{
-		GetCustomDrawingCoordinateSystem( matrix, InData );
-	}
-
-	// Get a base normal from the current axis.
-
-	FVector BaseNormal(1,0,0);		// Default to X axis
-	switch( CurrentWidgetAxis )
-	{
-		case EAxisList::Y:	BaseNormal = FVector(0,1,0);	break;
-		case EAxisList::Z:	BaseNormal = FVector(0,0,1);	break;
-		case EAxisList::XY:	BaseNormal = FVector(1,1,0);	break;
-		case EAxisList::XZ:	BaseNormal = FVector(1,0,1);	break;
-		case EAxisList::YZ:	BaseNormal = FVector(0,1,1);	break;
-		case EAxisList::XYZ:	BaseNormal = FVector(1,1,1);	break;
-	}
-
-	// Transform the base normal into the proper coordinate space.
-	return matrix.TransformPosition( BaseNormal );
-}
-
-AActor* FEdMode::GetFirstSelectedActorInstance() const
-{
-	return Owner->GetSelectedActors()->GetTop<AActor>();
-}
-
-bool FEdMode::CanCreateWidgetForStructure(const UStruct* InPropStruct)
-{
-	return InPropStruct && (InPropStruct->GetFName() == NAME_Vector || InPropStruct->GetFName() == NAME_Transform);
-}
-
-bool FEdMode::CanCreateWidgetForProperty(FProperty* InProp)
-{
-	FStructProperty* TestProperty = CastField<FStructProperty>(InProp);
-	if( !TestProperty )
-	{
-		FArrayProperty* ArrayProperty = CastField<FArrayProperty>(InProp);
-		if( ArrayProperty )
-		{
-			TestProperty = CastField<FStructProperty>(ArrayProperty->Inner);
-		}
-	}
-	return (TestProperty != NULL) && CanCreateWidgetForStructure(TestProperty->Struct);
-}
-
-bool FEdMode::ShouldCreateWidgetForProperty(FProperty* InProp)
-{
-	return CanCreateWidgetForProperty(InProp) && InProp->HasMetaData(MD_MakeEditWidget);
-}
-
-static bool IsTransformProperty(FProperty* InProp)
-{
-	FStructProperty* StructProp = CastField<FStructProperty>(InProp);
-	return (StructProp != NULL && StructProp->Struct->GetFName() == NAME_Transform);
-
-}
-
-struct FPropertyWidgetInfoChainElement
-{
-	FProperty* Property;
-	int32 Index;
-
-	FPropertyWidgetInfoChainElement(FProperty* InProperty = nullptr, int32 InIndex = INDEX_NONE)
-		: Property(InProperty), Index(InIndex)
-	{}
-
-	static bool ShouldCreateWidgetSomwhereInBranch(FProperty* InProp)
-	{
-		FStructProperty* StructProperty = CastField<FStructProperty>(InProp);
-		if (!StructProperty)
-		{
-			FArrayProperty* ArrayProperty = CastField<FArrayProperty>(InProp);
-			if (ArrayProperty)
-			{
-				StructProperty = CastField<FStructProperty>(ArrayProperty->Inner);
-			}
-		}
-
-		if (StructProperty)
-		{
-			if (FEdMode::CanCreateWidgetForStructure(StructProperty->Struct) && InProp->HasMetaData(FEdMode::MD_MakeEditWidget))
-			{
-				return true;
-			}
-
-			for (TFieldIterator<FProperty> PropertyIt(StructProperty->Struct, EFieldIteratorFlags::IncludeSuper); PropertyIt; ++PropertyIt)
-			{
-				if (ShouldCreateWidgetSomwhereInBranch(*PropertyIt))
-				{
-					return true;
-				}
-			}
-		}
-
-		return false;
-	}
-
-	static FEdMode::FPropertyWidgetInfo CreateWidgetInfo(const TArray<FPropertyWidgetInfoChainElement>& Chain, bool bIsTransform, FProperty* CurrentProp, int32 Index = INDEX_NONE)
-	{
-		check(CurrentProp);
-		FEdMode::FPropertyWidgetInfo WidgetInfo;
-		WidgetInfo.PropertyValidationName = FName(*CurrentProp->GetMetaData(FEdMode::MD_ValidateWidgetUsing));
-		WidgetInfo.bIsTransform = bIsTransform;
-		WidgetInfo.PropertyIndex = Index;
-
-		const FString SimplePostFix(TEXT("."));
-		for (int32 ChainIndex = 0; ChainIndex < Chain.Num(); ++ChainIndex)
-		{
-			const FPropertyWidgetInfoChainElement& Element = Chain[ChainIndex];
-			check(Element.Property);
-			const FString Postfix = (Element.Index != INDEX_NONE) ? FString::Printf(TEXT("[%d]."), Element.Index) : SimplePostFix;
-			const FString PropertyName = Element.Property->GetName() + Postfix;
-			const FString& DisplayName = Element.Property->GetMetaData(TEXT("DisplayName"));
-
-			WidgetInfo.PropertyName += PropertyName;
-			WidgetInfo.DisplayName += (!DisplayName.IsEmpty()) ? (DisplayName + Postfix) : PropertyName;
-		}
-
-		{
-			const FString PropertyName = CurrentProp->GetName();
-			const FString& DisplayName = CurrentProp->GetMetaData(TEXT("DisplayName"));
-
-			WidgetInfo.PropertyName += PropertyName;
-			WidgetInfo.DisplayName += (!DisplayName.IsEmpty()) ? DisplayName : PropertyName;
-		}
-		return WidgetInfo;
-	}
-
-	static void RecursiveGet(const FEdMode& EdMode, const UStruct* InStruct, const void* InContainer, TArray<FEdMode::FPropertyWidgetInfo>& OutInfos, TArray<FPropertyWidgetInfoChainElement>& Chain)
-	{
-		for (TFieldIterator<FProperty> PropertyIt(InStruct, EFieldIteratorFlags::IncludeSuper); PropertyIt; ++PropertyIt)
-		{
-			FProperty* CurrentProp = *PropertyIt;
-			check(CurrentProp);
-
-			if (EdMode.ShouldCreateWidgetForProperty(CurrentProp))
-			{
-				if (FArrayProperty* ArrayProp = CastField<FArrayProperty>(CurrentProp))
-				{
-					check(InContainer);
-					FScriptArrayHelper_InContainer ArrayHelper(ArrayProp, InContainer);
-					// See how many widgets we need to make for the array property
-					const uint32 ArrayDim = ArrayHelper.Num();
-					for (uint32 Index = 0; Index < ArrayDim; ++Index)
-					{
-						OutInfos.Add(FPropertyWidgetInfoChainElement::CreateWidgetInfo(Chain, IsTransformProperty(ArrayProp->Inner), CurrentProp, Index));
-					}
-				}
-				else
-				{
-					OutInfos.Add(FPropertyWidgetInfoChainElement::CreateWidgetInfo(Chain, IsTransformProperty(CurrentProp), CurrentProp));
-				}
-			}
-			else if (FStructProperty* StructProp = CastField<FStructProperty>(CurrentProp))
-			{
-				// Recursively traverse into structures, looking for additional vector properties to expose
-				Chain.Push(FPropertyWidgetInfoChainElement(StructProp));
-				RecursiveGet(EdMode
-					, StructProp->Struct
-					, StructProp->ContainerPtrToValuePtr<void>(InContainer)
-					, OutInfos
-					, Chain);
-				Chain.Pop(false);
-			}
-			else if (FArrayProperty* ArrayProp = CastField<FArrayProperty>(CurrentProp))
-			{
-				// Recursively traverse into arrays of structures, looking for additional vector properties to expose
-				FStructProperty* InnerStructProp = CastField<FStructProperty>(ArrayProp->Inner);
-				if (InnerStructProp)
-				{
-					FScriptArrayHelper_InContainer ArrayHelper(ArrayProp, InContainer);
-
-					// If the array is not empty the do additional check to tell if iteration is necessary
-					if (ArrayHelper.Num() && ShouldCreateWidgetSomwhereInBranch(InnerStructProp))
-					{
-						for (int32 ArrayIndex = 0; ArrayIndex < ArrayHelper.Num(); ++ArrayIndex)
-						{
-							Chain.Push(FPropertyWidgetInfoChainElement(ArrayProp, ArrayIndex));
-							RecursiveGet(EdMode
-								, InnerStructProp->Struct
-								, ArrayHelper.GetRawPtr(ArrayIndex)
-								, OutInfos
-								, Chain);
-							Chain.Pop(false);
-						}
-					}
-				}
-			}
-		}
-	}
-};
-
-UObject* FEdMode::GetItemToTryDisplayingWidgetsFor(FTransform& OutLocalToWorld) const
-{
-	// Determine what is selected, preferring a component over an actor
-	USceneComponent* SelectedComponent = Owner->GetSelectedComponents()->GetTop<USceneComponent>();
-	UObject* BestSelectedItem = SelectedComponent;
-
-	if (SelectedComponent == nullptr)
-	{
-		AActor* SelectedActor = GetFirstSelectedActorInstance();
-		if (SelectedActor != nullptr)
-		{
-			if (USceneComponent* RootComponent = SelectedActor->GetRootComponent())
-			{
-				BestSelectedItem = SelectedActor;
-				OutLocalToWorld = RootComponent->GetComponentToWorld();
-			}
-		}
-	}
-	else
-	{
-		OutLocalToWorld = SelectedComponent->GetComponentToWorld();
-	}
-
-	return BestSelectedItem;
-}
-
-void FEdMode::GetPropertyWidgetInfos(const UStruct* InStruct, const void* InContainer, TArray<FPropertyWidgetInfo>& OutInfos) const
-{
-	TArray<FPropertyWidgetInfoChainElement> Chain;
-	FPropertyWidgetInfoChainElement::RecursiveGet(*this, InStruct, InContainer, OutInfos, Chain);
-}
-
 bool FEdMode::IsSnapRotationEnabled()
 {
 	return GetDefault<ULevelEditorViewportSettings>()->RotGridEnabled;
 }
-
-/*------------------------------------------------------------------------------
-	Default.
-------------------------------------------------------------------------------*/
-
-FEdModeDefault::FEdModeDefault()
-{
-	bDrawGrid = true;
-	bDrawPivot = false;
-	bDrawBaseInfo = false;
-	bDrawWorldBox = false;
-	bDrawKillZ = false;
-}
diff --git a/Engine/Source/Editor/UnrealEd/Private/EditorModeManager.cpp b/Engine/Source/Editor/UnrealEd/Private/EditorModeManager.cpp
--- a/Engine/Source/Editor/UnrealEd/Private/EditorModeManager.cpp
+++ b/Engine/Source/Editor/UnrealEd/Private/EditorModeManager.cpp
@@ -33,6 +33,7 @@
 #include "InputRouter.h"
 #include "InteractiveGizmoManager.h"
 #include "EdModeInteractiveToolsContext.h"
+#include "Tools/LegacyEdModeInterfaces.h"
 
 /*------------------------------------------------------------------------------
 	FEditorModeTools.
@@ -959,7 +960,7 @@ FMatrix FEditorModeTools::GetLocalCoordinateSystem()
 	bool CustomCoordinateSystemProvided = false;
 	for (UEdMode* Mode : ActiveScriptableModes)
 	{
-		FEdMode* LegacyMode = Mode->AsLegacyMode();
+		ILegacyEdModeWidgetInterface* LegacyMode = Cast<ILegacyEdModeWidgetInterface>(Mode);
 		if (LegacyMode && LegacyMode->GetCustomDrawingCoordinateSystem(Matrix, nullptr))
 		{
 			CustomCoordinateSystemProvided = true;
@@ -999,7 +1000,7 @@ EAxisList::Type FEditorModeTools::GetWidgetAxisToDraw( UE::Widget::EWidgetMode I
 	EAxisList::Type OutAxis = EAxisList::All;
 	for( int Index = ActiveScriptableModes.Num() - 1; Index >= 0 ; Index-- )
 	{
-		FEdMode* Mode = ActiveScriptableModes[Index]->AsLegacyMode();
+		ILegacyEdModeWidgetInterface* Mode = Cast<ILegacyEdModeWidgetInterface>(ActiveScriptableModes[Index]);
 		if ( Mode && Mode->ShouldDrawWidget() )
 		{
 			OutAxis = Mode->GetWidgetAxisToDraw( InWidgetMode );
@@ -1047,7 +1048,7 @@ bool FEditorModeTools::AllowsViewportDragTool() const
 	bool bCanUseDragTool = false;
 	for (UEdMode* Mode : ActiveScriptableModes)
 	{
-		if (FEdMode* LegacyMode = Mode->AsLegacyMode())
+		if (ILegacyEdModeSelectInterface* LegacyMode = Cast<ILegacyEdModeSelectInterface>(Mode))
 		{
 			bCanUseDragTool |= LegacyMode->AllowsViewportDragTool();
 		}
@@ -1079,7 +1080,7 @@ bool FEditorModeTools::BoxSelect( FBox& InBox, bool InSelect )
 	bool bHandled = false;
 	for (UEdMode* Mode : ActiveScriptableModes)
 	{
-		if (FEdMode* LegacyMode = Mode->AsLegacyMode())
+		if (ILegacyEdModeSelectInterface* LegacyMode = Cast<ILegacyEdModeSelectInterface>(Mode))
 		{
 			bHandled |= LegacyMode->BoxSelect(InBox, InSelect);
 		}
@@ -1093,7 +1094,7 @@ bool FEditorModeTools::FrustumSelect( const FConvexVolume& InFrustum, FEditorVie
 	bool bHandled = false;
 	for (UEdMode* Mode : ActiveScriptableModes)
 	{
-		if (FEdMode* LegacyMode = Mode->AsLegacyMode())
+		if (ILegacyEdModeSelectInterface* LegacyMode = Cast<ILegacyEdModeSelectInterface>(Mode))
 		{
 			bHandled |= LegacyMode->FrustumSelect(InFrustum, InViewportClient, InSelect);
 		}
@@ -1108,7 +1109,7 @@ bool FEditorModeTools::UsesTransformWidget() const
 	bool bUsesTransformWidget = false;
 	for (UEdMode* Mode : ActiveScriptableModes)
 	{
-		if (FEdMode* LegacyMode = Mode->AsLegacyMode())
+		if (ILegacyEdModeWidgetInterface* LegacyMode = Cast<ILegacyEdModeWidgetInterface>(Mode))
 		{
 			bUsesTransformWidget |= LegacyMode->UsesTransformWidget();
 		}
@@ -1123,7 +1124,7 @@ bool FEditorModeTools::UsesTransformWidget( UE::Widget::EWidgetMode CheckMode )
 	bool bUsesTransformWidget = false;
 	for (UEdMode* Mode : ActiveScriptableModes)
 	{
-		if (FEdMode* LegacyMode = Mode->AsLegacyMode())
+		if (ILegacyEdModeWidgetInterface* LegacyMode = Cast<ILegacyEdModeWidgetInterface>(Mode))
 		{
 			bUsesTransformWidget |= LegacyMode->UsesTransformWidget(CheckMode);
 		}
@@ -1137,7 +1138,7 @@ void FEditorModeTools::SetCurrentWidgetAxis( EAxisList::Type NewAxis )
 {
 	for (UEdMode* Mode : ActiveScriptableModes)
 	{
-		if (FEdMode* LegacyMode = Mode->AsLegacyMode())
+		if (ILegacyEdModeWidgetInterface* LegacyMode = Cast<ILegacyEdModeWidgetInterface>(Mode))
 		{
 			LegacyMode->SetCurrentWidgetAxis(NewAxis);
 		}
@@ -1395,7 +1396,7 @@ bool FEditorModeTools::AllowWidgetMove() const
 	bool bAllow = false;
 	for (UEdMode* Mode : ActiveScriptableModes)
 	{
-		if (FEdMode* LegacyMode = Mode->AsLegacyMode())
+		if (ILegacyEdModeWidgetInterface* LegacyMode = Cast<ILegacyEdModeWidgetInterface>(Mode))
 		{
 			bAllow |= LegacyMode->AllowWidgetMove();
 		}
@@ -1459,7 +1460,7 @@ bool FEditorModeTools::GetShowWidget() const
 	// Check to see of any active modes support widget drawing
 	for (UEdMode* Mode : ActiveScriptableModes)
 	{
-		if (FEdMode* LegacyMode = Mode->AsLegacyMode())
+		if (ILegacyEdModeWidgetInterface* LegacyMode = Cast<ILegacyEdModeWidgetInterface>(Mode))
 		{
 			bDrawModeSupportsWidgetDrawing |= LegacyMode->ShouldDrawWidget();
 		}
@@ -1524,7 +1525,7 @@ FVector FEditorModeTools::GetWidgetLocation() const
 {
 	for (int Index = ActiveScriptableModes.Num() - 1; Index >= 0; Index--)
 	{
-		if (FEdMode* LegacyMode = ActiveScriptableModes[Index]->AsLegacyMode())
+		if (ILegacyEdModeWidgetInterface* LegacyMode = Cast<ILegacyEdModeWidgetInterface>(ActiveScriptableModes[Index]))
 		{
 			if (LegacyMode->UsesTransformWidget())
 			{
@@ -1617,7 +1618,7 @@ const FEdMode* FEditorModeTools::GetActiveMode( FEditorModeID InID ) const
 
 const FModeTool* FEditorModeTools::GetActiveTool( FEditorModeID InID ) const
 {
-	const FEdMode* ActiveMode = GetActiveMode( InID );
+	ILegacyEdModeToolInterface* ActiveMode = Cast<ILegacyEdModeToolInterface>(GetActiveScriptableMode( InID ));
 	const FModeTool* Tool = nullptr;
 	if( ActiveMode )
 	{
@@ -1628,15 +1629,7 @@ const FModeTool* FEditorModeTools::GetActiveTool( FEditorModeID InID ) const
 
 bool FEditorModeTools::IsModeActive( FEditorModeID InID ) const
 {
-	if (GetActiveMode(InID) != nullptr)
-	{
-		return true;
-	}
-	else if (GetActiveScriptableMode(InID) != nullptr)
-	{
-		return true;
-	}
-	return false;
+	return (GetActiveScriptableMode(InID) != nullptr);
 }
 
 bool FEditorModeTools::IsDefaultModeActive() const
@@ -1657,7 +1650,7 @@ bool FEditorModeTools::CanCycleWidgetMode() const
 {
 	for (UEdMode* Mode : ActiveScriptableModes)
 	{
-		if (FEdMode* LegacyMode = Mode->AsLegacyMode())
+		if (ILegacyEdModeWidgetInterface* LegacyMode = Cast<ILegacyEdModeWidgetInterface>(Mode))
 		{
 			if (LegacyMode->CanCycleWidgetMode())
 			{
diff --git a/Engine/Source/Editor/UnrealEd/Private/EditorModeRegistry.cpp b/Engine/Source/Editor/UnrealEd/Private/EditorModeRegistry.cpp
--- a/Engine/Source/Editor/UnrealEd/Private/EditorModeRegistry.cpp
+++ b/Engine/Source/Editor/UnrealEd/Private/EditorModeRegistry.cpp
@@ -51,11 +51,6 @@ void FEditorModeRegistry::Initialize()
 	}
 
 	// Add default editor modes
-	RegisterMode<FEdModeDefault>(
-		FBuiltinEditorModes::EM_Default,
-		NSLOCTEXT("DefaultMode", "DisplayName", "Select"),
-		FSlateIcon(FEditorStyle::GetStyleSetName(), "LevelEditor.SelectMode", "LevelEditor.SelectMode.Small"),
-		true, 0);
 	RegisterMode<FEdModeInterpEdit>(FBuiltinEditorModes::EM_InterpEdit);
 
 	FModuleManager::LoadModuleChecked<FActorPickerModeModule>(TEXT("ActorPickerMode"));
diff --git a/Engine/Source/Editor/UnrealEd/Private/Tools/DefaultEdMode.cpp b/Engine/Source/Editor/UnrealEd/Private/Tools/DefaultEdMode.cpp
new file mode 100644
--- /dev/null
+++ b/Engine/Source/Editor/UnrealEd/Private/Tools/DefaultEdMode.cpp
@@ -0,0 +1,31 @@
+// Copyright Epic Games, Inc. All Rights Reserved.
+
+#include "Tools/DefaultEdMode.h"
+
+#include "EditorModes.h"
+#include "EditorStyleSet.h"
+#include "Textures/SlateIcon.h"
+
+UEdModeDefault::UEdModeDefault()
+{
+	bDrawGrid = true;
+	bDrawPivot = false;
+	bDrawBaseInfo = false;
+	bDrawWorldBox = false;
+	bDrawKillZ = false;
+	Info = FEditorModeInfo(
+		FBuiltinEditorModes::EM_Default,
+		NSLOCTEXT("DefaultMode", "DisplayName", "Select"),
+		FSlateIcon(FEditorStyle::GetStyleSetName(), "LevelEditor.SelectMode", "LevelEditor.SelectMode.Small"),
+		true, 0);
+}
+
+bool UEdModeDefault::UsesPropertyWidgets() const
+{
+	return true;
+}
+
+bool UEdModeDefault::UsesToolkits() const
+{
+	return false;
+}
diff --git a/Engine/Source/Editor/UnrealEd/Private/Tools/DefaultEdMode.h b/Engine/Source/Editor/UnrealEd/Private/Tools/DefaultEdMode.h
new file mode 100644
--- /dev/null
+++ b/Engine/Source/Editor/UnrealEd/Private/Tools/DefaultEdMode.h
@@ -0,0 +1,19 @@
+// Copyright Epic Games, Inc. All Rights Reserved.
+
+#pragma once
+
+#include "Tools/LegacyEdModeWidgetHelpers.h"
+
+#include "DefaultEdMode.generated.h"
+
+UCLASS(Transient)
+class UEdModeDefault : public UBaseLegacyWidgetEdMode
+{
+	GENERATED_BODY()
+
+public:
+	UEdModeDefault();
+
+	virtual bool UsesPropertyWidgets() const override;
+	virtual bool UsesToolkits() const override;
+};
diff --git a/Engine/Source/Editor/UnrealEd/Private/Tools/LegacyEdMode.cpp b/Engine/Source/Editor/UnrealEd/Private/Tools/LegacyEdMode.cpp
--- a/Engine/Source/Editor/UnrealEd/Private/Tools/LegacyEdMode.cpp
+++ b/Engine/Source/Editor/UnrealEd/Private/Tools/LegacyEdMode.cpp
@@ -10,10 +10,10 @@
 #include "Math/Vector.h"
 #include "Math/Rotator.h"
 #include "EditorModeManager.h"
+#include "EditorModeTools.h"
 
 ULegacyEdModeWrapper::ULegacyEdModeWrapper()
 {
-
 }
 
 bool ULegacyEdModeWrapper::CreateLegacyMode(FEditorModeID ModeID, FEditorModeTools& ModeManager)
@@ -316,6 +316,117 @@ UTexture2D* ULegacyEdModeWrapper::GetVertexTexture()
 	return LegacyEditorMode->GetVertexTexture();
 }
 
+bool ULegacyEdModeWrapper::AllowWidgetMove()
+{
+	return LegacyEditorMode->AllowWidgetMove();
+}
+
+bool ULegacyEdModeWrapper::CanCycleWidgetMode() const
+{
+	return LegacyEditorMode->CanCycleWidgetMode();
+}
+
+bool ULegacyEdModeWrapper::ShowModeWidgets() const
+{
+	return LegacyEditorMode->ShowModeWidgets();
+}
+
+EAxisList::Type ULegacyEdModeWrapper::GetWidgetAxisToDraw(UE::Widget::EWidgetMode InWidgetMode) const
+{
+	return LegacyEditorMode->GetWidgetAxisToDraw(InWidgetMode);
+}
+
+FVector ULegacyEdModeWrapper::GetWidgetLocation() const
+{
+	return LegacyEditorMode->GetWidgetLocation();
+}
+
+bool ULegacyEdModeWrapper::ShouldDrawWidget() const
+{
+	return LegacyEditorMode->ShouldDrawWidget();
+}
+
+bool ULegacyEdModeWrapper::UsesTransformWidget() const
+{
+	return LegacyEditorMode->UsesTransformWidget();
+}
+
+bool ULegacyEdModeWrapper::UsesTransformWidget(UE::Widget::EWidgetMode CheckMode) const
+{
+	return LegacyEditorMode->UsesTransformWidget(CheckMode);
+}
+
+FVector ULegacyEdModeWrapper::GetWidgetNormalFromCurrentAxis(void* InData)
+{
+	return LegacyEditorMode->GetWidgetNormalFromCurrentAxis(InData);
+}
+
+bool ULegacyEdModeWrapper::BoxSelect(FBox& InBox, bool InSelect)
+{
+	return LegacyEditorMode->BoxSelect(InBox, InSelect);
+}
+
+bool ULegacyEdModeWrapper::FrustumSelect(const FConvexVolume& InFrustum, FEditorViewportClient* InViewportClient, bool InSelect)
+{
+	return LegacyEditorMode->FrustumSelect(InFrustum, InViewportClient, InSelect);
+}
+
+void ULegacyEdModeWrapper::SetCurrentWidgetAxis(EAxisList::Type InAxis)
+{
+	LegacyEditorMode->SetCurrentWidgetAxis(InAxis);
+}
+
+EAxisList::Type ULegacyEdModeWrapper::GetCurrentWidgetAxis() const
+{
+	return LegacyEditorMode->GetCurrentWidgetAxis();
+}
+
+bool ULegacyEdModeWrapper::UsesPropertyWidgets() const
+{
+	return LegacyEditorMode->UsesPropertyWidgets();
+}
+
+bool ULegacyEdModeWrapper::GetCustomDrawingCoordinateSystem(FMatrix& InMatrix, void* InData)
+{
+	return LegacyEditorMode->GetCustomDrawingCoordinateSystem(InMatrix, InData);
+}
+
+bool ULegacyEdModeWrapper::GetCustomInputCoordinateSystem(FMatrix& InMatrix, void* InData)
+{
+	return LegacyEditorMode->GetCustomInputCoordinateSystem(InMatrix, InData);
+}
+
+void ULegacyEdModeWrapper::SetCurrentTool(EModeTools InID)
+{
+	LegacyEditorMode->SetCurrentTool(InID);
+}
+
+void ULegacyEdModeWrapper::SetCurrentTool(FModeTool* InModeTool)
+{
+	LegacyEditorMode->SetCurrentTool(InModeTool);
+}
+
+FModeTool* ULegacyEdModeWrapper::FindTool(EModeTools InID)
+{
+	return LegacyEditorMode->FindTool(InID);
+}
+
+const TArray<FModeTool*>& ULegacyEdModeWrapper::GetTools() const
+{
+	return LegacyEditorMode->GetTools();
+}
+
+FModeTool* ULegacyEdModeWrapper::GetCurrentTool()
+{
+	return LegacyEditorMode->GetCurrentTool();
+}
+
+const FModeTool* ULegacyEdModeWrapper::GetCurrentTool() const
+{
+	return LegacyEditorMode->GetCurrentTool();
+}
+
+
 void ULegacyEdModeWrapper::Draw(const FSceneView* View, FPrimitiveDrawInterface* PDI)
 {
 	LegacyEditorMode->Draw(View, PDI);
diff --git a/Engine/Source/Editor/UnrealEd/Private/Tools/LegacyEdMode.h b/Engine/Source/Editor/UnrealEd/Private/Tools/LegacyEdMode.h
--- a/Engine/Source/Editor/UnrealEd/Private/Tools/LegacyEdMode.h
+++ b/Engine/Source/Editor/UnrealEd/Private/Tools/LegacyEdMode.h
@@ -6,6 +6,7 @@
 #include "Tools/Modes.h"
 #include "Engine/EngineBaseTypes.h"
 #include "Templates/SharedPointer.h"
+#include "Tools/LegacyEdModeInterfaces.h"
 
 #include "LegacyEdMode.generated.h"
 
@@ -22,7 +23,7 @@ class FViewport;
 class UPrimitiveComponent;
 
 UCLASS()
-class UNREALED_API ULegacyEdModeWrapper final : public UEdMode
+class UNREALED_API ULegacyEdModeWrapper final : public UEdMode, public ILegacyEdModeWidgetInterface, public ILegacyEdModeToolInterface, public ILegacyEdModeSelectInterface
 {
 	GENERATED_BODY()
 
@@ -93,6 +94,34 @@ public:
 	virtual UTexture2D* GetVertexTexture() override;
 	// End UEdMode overrides
 
+	// ILegacyEdModeWidgetInterface
+	virtual bool AllowWidgetMove() override;
+	virtual bool CanCycleWidgetMode() const override;
+	virtual bool ShowModeWidgets() const override;
+	virtual EAxisList::Type GetWidgetAxisToDraw(UE::Widget::EWidgetMode InWidgetMode) const override;
+	virtual FVector GetWidgetLocation() const override;
+	virtual bool ShouldDrawWidget() const override;
+	virtual bool UsesTransformWidget() const override;
+	virtual bool UsesTransformWidget(UE::Widget::EWidgetMode CheckMode) const override;
+	virtual FVector GetWidgetNormalFromCurrentAxis(void* InData) override;
+	virtual bool BoxSelect(FBox& InBox, bool InSelect = true) override;
+	virtual bool FrustumSelect(const FConvexVolume& InFrustum, FEditorViewportClient* InViewportClient, bool InSelect = true) override;
+	virtual void SetCurrentWidgetAxis(EAxisList::Type InAxis) override;
+	virtual EAxisList::Type GetCurrentWidgetAxis() const override;
+	virtual bool UsesPropertyWidgets() const override;
+	virtual bool GetCustomDrawingCoordinateSystem(FMatrix& InMatrix, void* InData) override;
+	virtual bool GetCustomInputCoordinateSystem(FMatrix& InMatrix, void* InData) override;
+	// End ILegacyEdModeWidgetInterface
+
+	// ILegacyEdModeToolInterface overrides
+	virtual void SetCurrentTool(EModeTools InID) override;
+	virtual void SetCurrentTool(FModeTool* InModeTool) override;
+	virtual FModeTool* FindTool(EModeTools InID) override;
+	virtual const TArray<FModeTool*>& GetTools() const override;
+	virtual FModeTool* GetCurrentTool() override;
+	virtual const FModeTool* GetCurrentTool() const override;
+	// End ILegacyEdModeToolInterface overrides
+
 	// Start FCommonDrawHelper overrides
 	virtual void Draw(const FSceneView* View, FPrimitiveDrawInterface* PDI) override;
 	// End FCommonDrawHelper overrides
diff --git a/Engine/Source/Editor/UnrealEd/Private/Tools/LegacyEdModeWidgetHelpers.cpp b/Engine/Source/Editor/UnrealEd/Private/Tools/LegacyEdModeWidgetHelpers.cpp
new file mode 100644
--- /dev/null
+++ b/Engine/Source/Editor/UnrealEd/Private/Tools/LegacyEdModeWidgetHelpers.cpp
@@ -0,0 +1,874 @@
+// Copyright Epic Games, Inc. All Rights Reserved.
+
+#include "Tools/LegacyEdModeWidgetHelpers.h"
+#include "HitProxies.h"
+#include "Components/SceneComponent.h"
+#include "GameFramework/Actor.h"
+#include "EditorModeManager.h"
+#include "Selection.h"
+#include "EditorViewportClient.h"
+#include "EngineUtils.h"
+#include "CanvasTypes.h"
+#include "CanvasItem.h"
+
+const FName FLegacyEdModeWidgetHelper::MD_MakeEditWidget(TEXT("MakeEditWidget"));
+const FName FLegacyEdModeWidgetHelper::MD_ValidateWidgetUsing(TEXT("ValidateWidgetUsing"));
+
+/** Hit proxy used for editable properties */
+struct HPropertyWidgetProxy : public HHitProxy
+{
+	DECLARE_HIT_PROXY();
+
+	/** Name of property this is the widget for */
+	FString	PropertyName;
+
+	/** If the property is an array property, the index into that array that this widget is for */
+	int32	PropertyIndex;
+
+	/** This property is a transform */
+	bool	bPropertyIsTransform;
+
+	HPropertyWidgetProxy(FString InPropertyName, int32 InPropertyIndex, bool bInPropertyIsTransform)
+		: HHitProxy(HPP_Foreground)
+		, PropertyName(InPropertyName)
+		, PropertyIndex(InPropertyIndex)
+		, bPropertyIsTransform(bInPropertyIsTransform)
+	{}
+
+	/** Show cursor as cross when over this handle */
+	virtual EMouseCursor::Type GetMouseCursor() override
+	{
+		return EMouseCursor::Crosshairs;
+	}
+};
+
+IMPLEMENT_HIT_PROXY(HPropertyWidgetProxy, HHitProxy);
+
+
+namespace
+{
+	/**
+	 * Returns a reference to the named property value data in the given container.
+	 */
+	template<typename T>
+	T* GetPropertyValuePtrByName(const UStruct* InStruct, void* InContainer, FString PropertyName, int32 ArrayIndex, FProperty*& OutProperty)
+	{
+		T* ValuePtr = NULL;
+
+		// Extract the vector ptr recursively using the property name
+		int32 DelimPos = PropertyName.Find(TEXT("."));
+		if (DelimPos != INDEX_NONE)
+		{
+			// Parse the property name and (optional) array index
+			int32 SubArrayIndex = 0;
+			FString NameToken = PropertyName.Left(DelimPos);
+			int32 ArrayPos = NameToken.Find(TEXT("["));
+			if (ArrayPos != INDEX_NONE)
+			{
+				FString IndexToken = NameToken.RightChop(ArrayPos + 1).LeftChop(1);
+				SubArrayIndex = FCString::Atoi(*IndexToken);
+
+				NameToken = PropertyName.Left(ArrayPos);
+			}
+
+			// Obtain the property info from the given structure definition
+			FProperty* CurrentProp = FindFProperty<FProperty>(InStruct, FName(*NameToken));
+
+			// Check first to see if this is a simple structure (i.e. not an array of structures)
+			FStructProperty* StructProp = CastField<FStructProperty>(CurrentProp);
+			if (StructProp != NULL)
+			{
+				// Recursively call back into this function with the structure property and container value
+				ValuePtr = GetPropertyValuePtrByName<T>(StructProp->Struct, StructProp->ContainerPtrToValuePtr<void>(InContainer), PropertyName.RightChop(DelimPos + 1), ArrayIndex, OutProperty);
+			}
+			else
+			{
+				// Check to see if this is an array
+				FArrayProperty* ArrayProp = CastField<FArrayProperty>(CurrentProp);
+				if (ArrayProp != NULL)
+				{
+					// It is an array, now check to see if this is an array of structures
+					StructProp = CastField<FStructProperty>(ArrayProp->Inner);
+					if (StructProp != NULL)
+					{
+						FScriptArrayHelper_InContainer ArrayHelper(ArrayProp, InContainer);
+						if (ArrayHelper.IsValidIndex(SubArrayIndex))
+						{
+							// Recursively call back into this function with the array element and container value
+							ValuePtr = GetPropertyValuePtrByName<T>(StructProp->Struct, ArrayHelper.GetRawPtr(SubArrayIndex), PropertyName.RightChop(DelimPos + 1), ArrayIndex, OutProperty);
+						}
+					}
+				}
+			}
+		}
+		else
+		{
+			FProperty* Prop = FindFProperty<FProperty>(InStruct, FName(*PropertyName));
+			if (Prop != NULL)
+			{
+				if (FArrayProperty* ArrayProp = CastField<FArrayProperty>(Prop))
+				{
+					check(ArrayIndex != INDEX_NONE);
+
+					// Property is an array property, so make sure we have a valid index specified
+					FScriptArrayHelper_InContainer ArrayHelper(ArrayProp, InContainer);
+					if (ArrayHelper.IsValidIndex(ArrayIndex))
+					{
+						ValuePtr = (T*)ArrayHelper.GetRawPtr(ArrayIndex);
+					}
+				}
+				else
+				{
+					// Property is a vector property, so access directly
+					ValuePtr = Prop->ContainerPtrToValuePtr<T>(InContainer);
+				}
+
+				OutProperty = Prop;
+			}
+		}
+
+		return ValuePtr;
+	}
+
+	/**
+	 * Returns the value of the property with the given name in the given Actor instance.
+	 */
+	template<typename T>
+	T GetPropertyValueByName(UObject* Object, FString PropertyName, int32 PropertyIndex)
+	{
+		T Value;
+		FProperty* DummyProperty = NULL;
+		if (T* ValuePtr = GetPropertyValuePtrByName<T>(Object->GetClass(), Object, PropertyName, PropertyIndex, DummyProperty))
+		{
+			Value = *ValuePtr;
+		}
+		return Value;
+	}
+
+	/**
+	 * Sets the property with the given name in the given Actor instance to the given value.
+	 */
+	template<typename T>
+	void SetPropertyValueByName(UObject* Object, FString PropertyName, int32 PropertyIndex, const T& InValue, FProperty*& OutProperty)
+	{
+		if (T* ValuePtr = GetPropertyValuePtrByName<T>(Object->GetClass(), Object, PropertyName, PropertyIndex, OutProperty))
+		{
+			*ValuePtr = InValue;
+		}
+	}
+}
+
+static bool IsTransformProperty(FProperty* InProp)
+{
+	FStructProperty* StructProp = CastField<FStructProperty>(InProp);
+	return (StructProp != NULL && StructProp->Struct->GetFName() == NAME_Transform);
+
+}
+
+struct FPropertyWidgetInfoChainElement
+{
+	FProperty* Property;
+	int32 Index;
+
+	FPropertyWidgetInfoChainElement(FProperty* InProperty = nullptr, int32 InIndex = INDEX_NONE)
+		: Property(InProperty), Index(InIndex)
+	{}
+
+	static bool ShouldCreateWidgetSomwhereInBranch(FProperty* InProp)
+	{
+		FStructProperty* StructProperty = CastField<FStructProperty>(InProp);
+		if (!StructProperty)
+		{
+			FArrayProperty* ArrayProperty = CastField<FArrayProperty>(InProp);
+			if (ArrayProperty)
+			{
+				StructProperty = CastField<FStructProperty>(ArrayProperty->Inner);
+			}
+		}
+
+		if (StructProperty)
+		{
+			if (FEdMode::CanCreateWidgetForStructure(StructProperty->Struct) && InProp->HasMetaData(FEdMode::MD_MakeEditWidget))
+			{
+				return true;
+			}
+
+			for (TFieldIterator<FProperty> PropertyIt(StructProperty->Struct, EFieldIteratorFlags::IncludeSuper); PropertyIt; ++PropertyIt)
+			{
+				if (ShouldCreateWidgetSomwhereInBranch(*PropertyIt))
+				{
+					return true;
+				}
+			}
+		}
+
+		return false;
+	}
+
+	static FLegacyEdModeWidgetHelper::FPropertyWidgetInfo CreateWidgetInfo(const TArray<FPropertyWidgetInfoChainElement>& Chain, bool bIsTransform, FProperty* CurrentProp, int32 Index = INDEX_NONE)
+	{
+		check(CurrentProp);
+		FEdMode::FPropertyWidgetInfo WidgetInfo;
+		WidgetInfo.PropertyValidationName = FName(*CurrentProp->GetMetaData(FEdMode::MD_ValidateWidgetUsing));
+		WidgetInfo.bIsTransform = bIsTransform;
+		WidgetInfo.PropertyIndex = Index;
+
+		const FString SimplePostFix(TEXT("."));
+		for (int32 ChainIndex = 0; ChainIndex < Chain.Num(); ++ChainIndex)
+		{
+			const FPropertyWidgetInfoChainElement& Element = Chain[ChainIndex];
+			check(Element.Property);
+			const FString Postfix = (Element.Index != INDEX_NONE) ? FString::Printf(TEXT("[%d]."), Element.Index) : SimplePostFix;
+			const FString PropertyName = Element.Property->GetName() + Postfix;
+			const FString& DisplayName = Element.Property->GetMetaData(TEXT("DisplayName"));
+
+			WidgetInfo.PropertyName += PropertyName;
+			WidgetInfo.DisplayName += (!DisplayName.IsEmpty()) ? (DisplayName + Postfix) : PropertyName;
+		}
+
+		{
+			const FString PropertyName = CurrentProp->GetName();
+			const FString& DisplayName = CurrentProp->GetMetaData(TEXT("DisplayName"));
+
+			WidgetInfo.PropertyName += PropertyName;
+			WidgetInfo.DisplayName += (!DisplayName.IsEmpty()) ? DisplayName : PropertyName;
+		}
+		return WidgetInfo;
+	}
+
+	static void RecursiveGet(const FLegacyEdModeWidgetHelper& EdMode, const UStruct* InStruct, const void* InContainer, TArray<FEdMode::FPropertyWidgetInfo>& OutInfos, TArray<FPropertyWidgetInfoChainElement>& Chain)
+	{
+		for (TFieldIterator<FProperty> PropertyIt(InStruct, EFieldIteratorFlags::IncludeSuper); PropertyIt; ++PropertyIt)
+		{
+			FProperty* CurrentProp = *PropertyIt;
+			check(CurrentProp);
+
+			if (EdMode.ShouldCreateWidgetForProperty(CurrentProp))
+			{
+				if (FArrayProperty* ArrayProp = CastField<FArrayProperty>(CurrentProp))
+				{
+					check(InContainer);
+					FScriptArrayHelper_InContainer ArrayHelper(ArrayProp, InContainer);
+					// See how many widgets we need to make for the array property
+					const uint32 ArrayDim = ArrayHelper.Num();
+					for (uint32 Index = 0; Index < ArrayDim; ++Index)
+					{
+						OutInfos.Add(FPropertyWidgetInfoChainElement::CreateWidgetInfo(Chain, IsTransformProperty(ArrayProp->Inner), CurrentProp, Index));
+					}
+				}
+				else
+				{
+					OutInfos.Add(FPropertyWidgetInfoChainElement::CreateWidgetInfo(Chain, IsTransformProperty(CurrentProp), CurrentProp));
+				}
+			}
+			else if (FStructProperty* StructProp = CastField<FStructProperty>(CurrentProp))
+			{
+				// Recursively traverse into structures, looking for additional vector properties to expose
+				Chain.Push(FPropertyWidgetInfoChainElement(StructProp));
+				RecursiveGet(EdMode
+					, StructProp->Struct
+					, StructProp->ContainerPtrToValuePtr<void>(InContainer)
+					, OutInfos
+					, Chain);
+				Chain.Pop(false);
+			}
+			else if (FArrayProperty* ArrayProp = CastField<FArrayProperty>(CurrentProp))
+			{
+				// Recursively traverse into arrays of structures, looking for additional vector properties to expose
+				FStructProperty* InnerStructProp = CastField<FStructProperty>(ArrayProp->Inner);
+				if (InnerStructProp)
+				{
+					FScriptArrayHelper_InContainer ArrayHelper(ArrayProp, InContainer);
+
+					// If the array is not empty the do additional check to tell if iteration is necessary
+					if (ArrayHelper.Num() && ShouldCreateWidgetSomwhereInBranch(InnerStructProp))
+					{
+						for (int32 ArrayIndex = 0; ArrayIndex < ArrayHelper.Num(); ++ArrayIndex)
+						{
+							Chain.Push(FPropertyWidgetInfoChainElement(ArrayProp, ArrayIndex));
+							RecursiveGet(EdMode
+								, InnerStructProp->Struct
+								, ArrayHelper.GetRawPtr(ArrayIndex)
+								, OutInfos
+								, Chain);
+							Chain.Pop(false);
+						}
+					}
+				}
+			}
+		}
+	}
+};
+
+void FLegacyEdModeWidgetHelper::FPropertyWidgetInfo::GetTransformAndColor(UObject* BestSelectedItem, bool bIsSelected, FTransform& OutLocalTransform, FString& OutValidationMessage, FColor& OutDrawColor) const
+{
+	// Determine the desired position
+	if (bIsTransform)
+	{
+		OutLocalTransform = GetPropertyValueByName<FTransform>(BestSelectedItem, PropertyName, PropertyIndex);
+	}
+	else
+	{
+		OutLocalTransform = FTransform(GetPropertyValueByName<FVector>(BestSelectedItem, PropertyName, PropertyIndex));
+	}
+
+	// Determine the desired color
+	OutDrawColor = bIsSelected ? FColor::White : FColor(128, 128, 255);
+	if (PropertyValidationName != NAME_None)
+	{
+		if (UFunction* ValidateFunc = BestSelectedItem->FindFunction(PropertyValidationName))
+		{
+			BestSelectedItem->ProcessEvent(ValidateFunc, &OutValidationMessage);
+
+			// if we have a negative result, the widget color is red.
+			OutDrawColor = OutValidationMessage.IsEmpty() ? OutDrawColor : FColor::Red;
+		}
+	}
+}
+
+FLegacyEdModeWidgetHelper::FLegacyEdModeWidgetHelper()
+	: EditedPropertyIndex(INDEX_NONE)
+	, bEditedPropertyIsTransform(false)
+	, CurrentWidgetAxis(EAxisList::None)
+	, Owner(nullptr)
+{
+}
+
+bool FLegacyEdModeWidgetHelper::AllowWidgetMove()
+{
+	return true;
+}
+
+bool FLegacyEdModeWidgetHelper::CanCycleWidgetMode() const
+{
+	return true;
+}
+
+bool FLegacyEdModeWidgetHelper::ShowModeWidgets() const
+{
+	return true;
+}
+
+EAxisList::Type FLegacyEdModeWidgetHelper::GetWidgetAxisToDraw(UE::Widget::EWidgetMode InWidgetMode) const
+{
+	return EAxisList::All;
+}
+
+FVector FLegacyEdModeWidgetHelper::GetWidgetLocation() const
+{
+	if (UsesPropertyWidgets())
+	{
+		FTransform DisplayWidgetToWorld;
+		UObject* BestSelectedItem = GetItemToTryDisplayingWidgetsFor(/*out*/ DisplayWidgetToWorld);
+		if (BestSelectedItem)
+		{
+			if (!EditedPropertyName.IsEmpty())
+			{
+				FVector LocalPos = FVector::ZeroVector;
+
+				if (bEditedPropertyIsTransform)
+				{
+					FTransform LocalTM = GetPropertyValueByName<FTransform>(BestSelectedItem, EditedPropertyName, EditedPropertyIndex);
+					LocalPos = LocalTM.GetLocation();
+				}
+				else
+				{
+					LocalPos = GetPropertyValueByName<FVector>(BestSelectedItem, EditedPropertyName, EditedPropertyIndex);
+				}
+
+				const FVector WorldPos = DisplayWidgetToWorld.TransformPosition(LocalPos);
+				return WorldPos;
+			}
+		}
+	}
+
+	return Owner ? Owner->PivotLocation : FVector();
+}
+
+bool FLegacyEdModeWidgetHelper::ShouldDrawWidget() const
+{
+	if (!Owner)
+	{
+		return false;
+	}
+
+	bool bDrawWidget = false;
+	bool bHadSelectableComponents = false;
+	if (Owner->GetSelectedComponents()->Num() > 0)
+	{
+		// when components are selected, only show the widget when one or more are scene components
+		for (FSelectedEditableComponentIterator It(*Owner->GetSelectedComponents()); It; ++It)
+		{
+			bHadSelectableComponents = true;
+			if (It->IsA<USceneComponent>())
+			{
+				bDrawWidget = true;
+				break;
+			}
+		}
+	}
+
+	if (!bHadSelectableComponents)
+	{
+		// when actors are selected, only show the widget when all selected actors have scene components
+		bDrawWidget = Owner->SelectionHasSceneComponent();
+	}
+
+	return bDrawWidget;
+}
+
+bool FLegacyEdModeWidgetHelper::UsesTransformWidget() const
+{
+	return true;
+}
+
+bool FLegacyEdModeWidgetHelper::UsesTransformWidget(UE::Widget::EWidgetMode CheckMode) const
+{
+	if (UsesPropertyWidgets())
+	{
+		FTransform DisplayWidgetToWorld;
+		UObject* BestSelectedItem = GetItemToTryDisplayingWidgetsFor(/*out*/ DisplayWidgetToWorld);
+
+		if (BestSelectedItem != nullptr)
+		{
+			// If editing a vector (not a transform)
+			if (!EditedPropertyName.IsEmpty() && !bEditedPropertyIsTransform)
+			{
+				return (CheckMode == UE::Widget::WM_Translate);
+			}
+		}
+	}
+
+	return true;
+}
+
+FVector FLegacyEdModeWidgetHelper::GetWidgetNormalFromCurrentAxis(void* InData)
+{
+	// Figure out the proper coordinate system.
+
+	FMatrix Matrix = FMatrix::Identity;
+	if (Owner && (Owner->GetCoordSystem() == COORD_Local))
+	{
+		GetCustomDrawingCoordinateSystem(Matrix, InData);
+	}
+
+	// Get a base normal from the current axis.
+
+	FVector BaseNormal(1, 0, 0);		// Default to X axis
+	switch (CurrentWidgetAxis)
+	{
+	case EAxisList::Y:	BaseNormal = FVector(0, 1, 0);	break;
+	case EAxisList::Z:	BaseNormal = FVector(0, 0, 1);	break;
+	case EAxisList::XY:	BaseNormal = FVector(1, 1, 0);	break;
+	case EAxisList::XZ:	BaseNormal = FVector(1, 0, 1);	break;
+	case EAxisList::YZ:	BaseNormal = FVector(0, 1, 1);	break;
+	case EAxisList::XYZ:	BaseNormal = FVector(1, 1, 1);	break;
+	}
+
+	// Transform the base normal into the proper coordinate space.
+	return Matrix.TransformPosition(BaseNormal);
+}
+
+void FLegacyEdModeWidgetHelper::SetCurrentWidgetAxis(EAxisList::Type InAxis)
+{
+	CurrentWidgetAxis = InAxis;
+}
+
+EAxisList::Type FLegacyEdModeWidgetHelper::GetCurrentWidgetAxis() const
+{
+	return CurrentWidgetAxis;
+}
+
+bool FLegacyEdModeWidgetHelper::UsesPropertyWidgets() const
+{
+	return false;
+}
+
+bool FLegacyEdModeWidgetHelper::GetCustomDrawingCoordinateSystem(FMatrix& InMatrix, void* InData)
+{
+	if (UsesPropertyWidgets())
+	{
+		FTransform DisplayWidgetToWorld;
+		UObject* BestSelectedItem = GetItemToTryDisplayingWidgetsFor(/*out*/ DisplayWidgetToWorld);
+		if (BestSelectedItem)
+		{
+			if (EditedPropertyName != TEXT(""))
+			{
+				if (bEditedPropertyIsTransform)
+				{
+					FTransform LocalTM = GetPropertyValueByName<FTransform>(BestSelectedItem, EditedPropertyName, EditedPropertyIndex);
+					InMatrix = FRotationMatrix::Make((LocalTM * DisplayWidgetToWorld).GetRotation());
+					return true;
+				}
+				else
+				{
+					InMatrix = FRotationMatrix::Make(DisplayWidgetToWorld.GetRotation());
+					return true;
+				}
+			}
+		}
+	}
+
+	return false;
+}
+
+bool FLegacyEdModeWidgetHelper::GetCustomInputCoordinateSystem(FMatrix& InMatrix, void* InData)
+{
+	return false;
+}
+
+void FLegacyEdModeWidgetHelper::ActorSelectionChangeNotify()
+{
+	EditedPropertyName = TEXT("");
+	EditedPropertyIndex = INDEX_NONE;
+	bEditedPropertyIsTransform = false;
+}
+
+bool FLegacyEdModeWidgetHelper::InputDelta(FEditorViewportClient* InViewportClient, FViewport* InViewport, FVector& InDrag, FRotator& InRot, FVector& InScale)
+{
+	if (UsesPropertyWidgets())
+	{
+		FTransform DisplayWidgetToWorld;
+		UObject* BestSelectedItem = GetItemToTryDisplayingWidgetsFor(/*out*/ DisplayWidgetToWorld);
+
+		if ((BestSelectedItem != nullptr) && (InViewportClient->GetCurrentWidgetAxis() != EAxisList::None))
+		{
+			GEditor->NoteActorMovement();
+
+			if (!EditedPropertyName.IsEmpty())
+			{
+				FTransform LocalTM = FTransform::Identity;
+
+				if (bEditedPropertyIsTransform)
+				{
+					LocalTM = GetPropertyValueByName<FTransform>(BestSelectedItem, EditedPropertyName, EditedPropertyIndex);
+				}
+				else
+				{
+					FVector LocalPos = GetPropertyValueByName<FVector>(BestSelectedItem, EditedPropertyName, EditedPropertyIndex);
+					LocalTM = FTransform(LocalPos);
+				}
+
+				// Calculate world transform
+				FTransform WorldTM = LocalTM * DisplayWidgetToWorld;
+				// Calc delta specified by drag
+				//FTransform DeltaTM(InRot.Quaternion(), InDrag);
+				// Apply delta in world space
+				WorldTM.SetTranslation(WorldTM.GetTranslation() + InDrag);
+				WorldTM.SetRotation(InRot.Quaternion() * WorldTM.GetRotation());
+				// Convert new world transform back into local space
+				LocalTM = WorldTM.GetRelativeTransform(DisplayWidgetToWorld);
+				// Apply delta scale
+				LocalTM.SetScale3D(LocalTM.GetScale3D() + InScale);
+
+				BestSelectedItem->PreEditChange(NULL);
+
+				// Property that we actually change
+				FProperty* SetProperty = NULL;
+
+				if (bEditedPropertyIsTransform)
+				{
+					SetPropertyValueByName<FTransform>(BestSelectedItem, EditedPropertyName, EditedPropertyIndex, LocalTM, SetProperty);
+				}
+				else
+				{
+					SetPropertyValueByName<FVector>(BestSelectedItem, EditedPropertyName, EditedPropertyIndex, LocalTM.GetLocation(), SetProperty);
+				}
+
+				FPropertyChangedEvent PropertyChangeEvent(SetProperty);
+				BestSelectedItem->PostEditChangeProperty(PropertyChangeEvent);
+
+				return true;
+			}
+		}
+	}
+
+	return false;
+}
+
+bool FLegacyEdModeWidgetHelper::HandleClick(FEditorViewportClient* InViewportClient, HHitProxy* HitProxy, const FViewportClick& Click)
+{
+	if (UsesPropertyWidgets() && (HitProxy != nullptr))
+	{
+		if (HitProxy->IsA(HPropertyWidgetProxy::StaticGetType()))
+		{
+			HPropertyWidgetProxy* PropertyProxy = (HPropertyWidgetProxy*)HitProxy;
+			EditedPropertyName = PropertyProxy->PropertyName;
+			EditedPropertyIndex = PropertyProxy->PropertyIndex;
+			bEditedPropertyIsTransform = PropertyProxy->bPropertyIsTransform;
+			return true;
+		}
+		// Left clicking on an actor, stop editing a property
+		else if (HitProxy->IsA(HActor::StaticGetType()))
+		{
+			EditedPropertyName = FString();
+			EditedPropertyIndex = INDEX_NONE;
+			bEditedPropertyIsTransform = false;
+		}
+	}
+
+	return false;
+}
+
+void FLegacyEdModeWidgetHelper::Render(const FSceneView* View, FViewport* Viewport, FPrimitiveDrawInterface* PDI)
+{
+	if (UsesPropertyWidgets())
+	{
+		const bool bHitTesting = PDI->IsHitTesting();
+
+		FTransform DisplayWidgetToWorld;
+		UObject* BestSelectedItem = GetItemToTryDisplayingWidgetsFor(/*out*/ DisplayWidgetToWorld);
+
+		if (BestSelectedItem != nullptr)
+		{
+			UClass* Class = BestSelectedItem->GetClass();
+			TArray<FPropertyWidgetInfo> WidgetInfos;
+			GetPropertyWidgetInfos(Class, BestSelectedItem, WidgetInfos);
+			FEditorScriptExecutionGuard ScriptGuard;
+			for (const FPropertyWidgetInfo& WidgetInfo : WidgetInfos)
+			{
+				const bool bSelected = (WidgetInfo.PropertyName == EditedPropertyName) && (WidgetInfo.PropertyIndex == EditedPropertyIndex);
+
+				FTransform LocalWidgetTransform;
+				FString ValidationMessage;
+				FColor WidgetColor;
+				WidgetInfo.GetTransformAndColor(BestSelectedItem, bSelected, /*out*/ LocalWidgetTransform, /*out*/ ValidationMessage, /*out*/ WidgetColor);
+
+				const FTransform WorldWidgetTransform = LocalWidgetTransform * DisplayWidgetToWorld;
+				const FMatrix WidgetTM(WorldWidgetTransform.ToMatrixWithScale());
+
+				const float WidgetSize = 0.035f;
+				const float ZoomFactor = FMath::Min<float>(View->ViewMatrices.GetProjectionMatrix().M[0][0], View->ViewMatrices.GetProjectionMatrix().M[1][1]);
+				const float WidgetRadius = View->Project(WorldWidgetTransform.GetTranslation()).W * (WidgetSize / ZoomFactor);
+
+				if (bHitTesting) PDI->SetHitProxy(new HPropertyWidgetProxy(WidgetInfo.PropertyName, WidgetInfo.PropertyIndex, WidgetInfo.bIsTransform));
+				DrawWireDiamond(PDI, WidgetTM, WidgetRadius, WidgetColor, SDPG_Foreground);
+				if (bHitTesting) PDI->SetHitProxy(NULL);
+			}
+		}
+	}
+}
+
+void FLegacyEdModeWidgetHelper::DrawHUD(FEditorViewportClient* ViewportClient, FViewport* Viewport, const FSceneView* View, FCanvas* Canvas)
+{
+	if (UsesPropertyWidgets())
+	{
+		FTransform DisplayWidgetToWorld;
+		UObject* BestSelectedItem = GetItemToTryDisplayingWidgetsFor(/*out*/ DisplayWidgetToWorld);
+		if (BestSelectedItem != nullptr)
+		{
+			FEditorScriptExecutionGuard ScriptGuard;
+
+			FIntPoint SizeXYWithDPIScale = Viewport->GetSizeXY() / Canvas->GetDPIScale();
+			const int32 HalfX = 0.5f * SizeXYWithDPIScale.X;
+			const int32 HalfY = 0.5f * SizeXYWithDPIScale.Y;
+
+			UClass* Class = BestSelectedItem->GetClass();
+			TArray<FPropertyWidgetInfo> WidgetInfos;
+			GetPropertyWidgetInfos(Class, BestSelectedItem, WidgetInfos);
+			for (const FPropertyWidgetInfo& WidgetInfo : WidgetInfos)
+			{
+				FTransform LocalWidgetTransform;
+				FString ValidationMessage;
+				FColor IgnoredWidgetColor;
+				WidgetInfo.GetTransformAndColor(BestSelectedItem, /*bSelected=*/ false, /*out*/ LocalWidgetTransform, /*out*/ ValidationMessage, /*out*/ IgnoredWidgetColor);
+
+				const FTransform WorldWidgetTransform = LocalWidgetTransform * DisplayWidgetToWorld;
+
+				const FPlane Proj = View->Project(WorldWidgetTransform.GetTranslation());
+				if (Proj.W > 0.f)
+				{
+					// do some string fixing
+					const uint32 VectorIndex = WidgetInfo.PropertyIndex;
+					const FString WidgetDisplayName = WidgetInfo.DisplayName + ((VectorIndex != INDEX_NONE) ? FString::Printf(TEXT("[%d]"), VectorIndex) : TEXT(""));
+					const FString DisplayString = ValidationMessage.IsEmpty() ? WidgetDisplayName : ValidationMessage;
+
+					const int32 XPos = HalfX + (HalfX * Proj.X);
+					const int32 YPos = HalfY + (HalfY * (Proj.Y * -1.f));
+					FCanvasTextItem TextItem(FVector2D(XPos + 5, YPos), FText::FromString(DisplayString), GEngine->GetSmallFont(), FLinearColor::White);
+					TextItem.EnableShadow(FLinearColor::Black);
+					Canvas->DrawItem(TextItem);
+				}
+			}
+		}
+	}
+}
+
+bool FLegacyEdModeWidgetHelper::CanCreateWidgetForStructure(const UStruct* InPropStruct)
+{
+	return InPropStruct && (InPropStruct->GetFName() == NAME_Vector || InPropStruct->GetFName() == NAME_Transform);
+}
+
+bool FLegacyEdModeWidgetHelper::CanCreateWidgetForProperty(FProperty* InProp)
+{
+	FStructProperty* TestProperty = CastField<FStructProperty>(InProp);
+	if (!TestProperty)
+	{
+		FArrayProperty* ArrayProperty = CastField<FArrayProperty>(InProp);
+		if (ArrayProperty)
+		{
+			TestProperty = CastField<FStructProperty>(ArrayProperty->Inner);
+		}
+	}
+	return (TestProperty != NULL) && CanCreateWidgetForStructure(TestProperty->Struct);
+}
+
+bool FLegacyEdModeWidgetHelper::ShouldCreateWidgetForProperty(FProperty* InProp)
+{
+	return CanCreateWidgetForProperty(InProp) && InProp->HasMetaData(MD_MakeEditWidget);
+}
+
+AActor* FLegacyEdModeWidgetHelper::GetFirstSelectedActorInstance() const
+{
+	return Owner ? Owner->GetSelectedActors()->GetTop<AActor>() : nullptr;
+}
+
+void FLegacyEdModeWidgetHelper::GetPropertyWidgetInfos(const UStruct* InStruct, const void* InContainer, TArray<FPropertyWidgetInfo>& OutInfos) const
+{
+	TArray<FPropertyWidgetInfoChainElement> Chain;
+	FPropertyWidgetInfoChainElement::RecursiveGet(*this, InStruct, InContainer, OutInfos, Chain);
+}
+
+UObject* FLegacyEdModeWidgetHelper::GetItemToTryDisplayingWidgetsFor(FTransform& OutLocalToWorld) const
+{
+	// Determine what is selected, preferring a component over an actor
+	USceneComponent* SelectedComponent = Owner ? Owner->GetSelectedComponents()->GetTop<USceneComponent>() : nullptr;
+	UObject* BestSelectedItem = SelectedComponent;
+
+	if (SelectedComponent == nullptr)
+	{
+		AActor* SelectedActor = GetFirstSelectedActorInstance();
+		if (SelectedActor != nullptr)
+		{
+			if (USceneComponent* RootComponent = SelectedActor->GetRootComponent())
+			{
+				BestSelectedItem = SelectedActor;
+				OutLocalToWorld = RootComponent->GetComponentToWorld();
+			}
+		}
+	}
+	else
+	{
+		OutLocalToWorld = SelectedComponent->GetComponentToWorld();
+	}
+
+	return BestSelectedItem;
+}
+
+void UBaseLegacyWidgetEdMode::Initialize()
+{
+	WidgetHelper = MakeShared<FLegacyEdModeWidgetHelper>();
+	WidgetHelper->Owner = this->Owner;
+
+	UEdMode::Initialize();
+}
+
+bool UBaseLegacyWidgetEdMode::InputDelta(FEditorViewportClient* InViewportClient, FViewport* InViewport, FVector& InDrag, FRotator& InRot, FVector& InScale)
+{
+	if (WidgetHelper->InputDelta(InViewportClient, InViewport, InDrag, InRot, InScale))
+	{
+		return true;
+	}
+
+	return UEdMode::InputDelta(InViewportClient, InViewport, InDrag, InRot, InScale);
+}
+
+bool UBaseLegacyWidgetEdMode::HandleClick(FEditorViewportClient* InViewportClient, HHitProxy* HitProxy, const FViewportClick& Click)
+{
+	if (WidgetHelper->HandleClick(InViewportClient, HitProxy, Click))
+	{
+		return true;
+	}
+
+	return UEdMode::HandleClick(InViewportClient, HitProxy, Click);
+}
+
+void UBaseLegacyWidgetEdMode::Render(const FSceneView* View, FViewport* Viewport, FPrimitiveDrawInterface* PDI)
+{
+	WidgetHelper->Render(View, Viewport, PDI);
+
+	UEdMode::Render(View, Viewport, PDI);
+}
+
+void UBaseLegacyWidgetEdMode::DrawHUD(FEditorViewportClient* ViewportClient, FViewport* Viewport, const FSceneView* View, FCanvas* Canvas)
+{
+	WidgetHelper->DrawHUD(ViewportClient, Viewport, View, Canvas);
+
+	UEdMode::DrawHUD(ViewportClient, Viewport, View, Canvas);
+}
+
+bool UBaseLegacyWidgetEdMode::AllowWidgetMove()
+{
+	return WidgetHelper->AllowWidgetMove();
+}
+
+bool UBaseLegacyWidgetEdMode::CanCycleWidgetMode() const
+{
+	return WidgetHelper->CanCycleWidgetMode();
+}
+
+bool UBaseLegacyWidgetEdMode::ShowModeWidgets() const
+{
+	return WidgetHelper->ShowModeWidgets();
+}
+
+EAxisList::Type UBaseLegacyWidgetEdMode::GetWidgetAxisToDraw(UE::Widget::EWidgetMode InWidgetMode) const
+{
+	return WidgetHelper->GetWidgetAxisToDraw(InWidgetMode);
+}
+
+FVector UBaseLegacyWidgetEdMode::GetWidgetLocation() const
+{
+	return WidgetHelper->GetWidgetLocation();
+}
+
+bool UBaseLegacyWidgetEdMode::ShouldDrawWidget() const
+{
+	return WidgetHelper->ShouldDrawWidget();
+}
+
+bool UBaseLegacyWidgetEdMode::UsesTransformWidget() const
+{
+	return WidgetHelper->UsesTransformWidget();
+}
+
+bool UBaseLegacyWidgetEdMode::UsesTransformWidget(UE::Widget::EWidgetMode CheckMode) const
+{
+	return WidgetHelper->UsesTransformWidget(CheckMode);
+}
+
+FVector UBaseLegacyWidgetEdMode::GetWidgetNormalFromCurrentAxis(void* InData)
+{
+	return WidgetHelper->GetWidgetNormalFromCurrentAxis(InData);
+}
+
+void UBaseLegacyWidgetEdMode::SetCurrentWidgetAxis(EAxisList::Type InAxis)
+{
+	WidgetHelper->SetCurrentWidgetAxis(InAxis);
+}
+
+EAxisList::Type UBaseLegacyWidgetEdMode::GetCurrentWidgetAxis() const
+{
+	return WidgetHelper->GetCurrentWidgetAxis();
+}
+
+bool UBaseLegacyWidgetEdMode::UsesPropertyWidgets() const
+{
+	return WidgetHelper->UsesPropertyWidgets();
+}
+
+bool UBaseLegacyWidgetEdMode::GetCustomDrawingCoordinateSystem(FMatrix& InMatrix, void* InData)
+{
+	return WidgetHelper->GetCustomDrawingCoordinateSystem(InMatrix, InData);
+}
+
+bool UBaseLegacyWidgetEdMode::GetCustomInputCoordinateSystem(FMatrix& InMatrix, void* InData)
+{
+	return WidgetHelper->GetCustomInputCoordinateSystem(InMatrix, InData);
+}
+
+void UBaseLegacyWidgetEdMode::ActorSelectionChangeNotify()
+{
+	WidgetHelper->ActorSelectionChangeNotify();
+}
+
diff --git a/Engine/Source/Editor/UnrealEd/Private/Tools/UEdMode.cpp b/Engine/Source/Editor/UnrealEd/Private/Tools/UEdMode.cpp
--- a/Engine/Source/Editor/UnrealEd/Private/Tools/UEdMode.cpp
+++ b/Engine/Source/Editor/UnrealEd/Private/Tools/UEdMode.cpp
@@ -139,11 +139,15 @@ bool UEdMode::InputKey(FEditorViewportClient* ViewportClient, FViewport* Viewpor
 
 	if (Event != IE_Released)
 	{
-		if (ToolCommandList->ProcessCommandBindings(Key, FSlateApplication::Get().GetModifierKeys(), false/*Event == IE_Repeat*/))
+		if (ToolsContext->ShouldIgnoreHotkeys() == false)		// allow the context to capture keyboard input if necessary
 		{
-			return true;
+			if (ToolCommandList->ProcessCommandBindings(Key, FSlateApplication::Get().GetModifierKeys(), false/*Event == IE_Repeat*/))
+			{
+				return true;
+			}
 		}
 	}
+
 	// Next pass input to the mode toolkit
 	if (Toolkit.IsValid() && ((Event == IE_Pressed) || (Event == IE_Repeat)))
 	{
@@ -223,9 +227,10 @@ void UEdMode::Enter()
 
 	// Now that the context is ready, make the toolkit
 	CreateToolkit();
-	check(Toolkit.IsValid());
-
-	Toolkit->Init(Owner->GetToolkitHost(), this);
+	if (Toolkit.IsValid())
+	{
+		Toolkit->Init(Owner->GetToolkitHost(), this);
+	}
 
 	BindCommands();
 
@@ -238,7 +243,11 @@ void UEdMode::Enter()
 	if (SettingsObject)
 	{
 		SettingsObject->LoadConfig();
-		Toolkit->SetModeSettingsObject(SettingsObject);
+
+		if (Toolkit.IsValid())
+		{
+			Toolkit->SetModeSettingsObject(SettingsObject);
+		}
 	}
 
 	FEditorDelegates::EditorModeIDEnter.Broadcast(GetID());
@@ -247,6 +256,11 @@ void UEdMode::Enter()
 
 void UEdMode::RegisterTool(TSharedPtr<FUICommandInfo> UICommand, FString ToolIdentifier, UInteractiveToolBuilder* Builder)
 {
+	if (!Toolkit.IsValid())
+	{
+		return;
+	}
+
 	const TSharedRef<FUICommandList>& CommandList = Toolkit->GetToolkitCommands();
 	ToolsContext->ToolManager->RegisterToolType(ToolIdentifier, Builder);
 	CommandList->MapAction(UICommand,
@@ -268,19 +282,19 @@ void UEdMode::Exit()
 		SettingsObject->SaveConfig();
 	}
 
-	const TSharedRef<FUICommandList>& CommandList = Toolkit->GetToolkitCommands();
-	for (auto& RegisteredTool : RegisteredTools)
-	{
-		CommandList->UnmapAction(RegisteredTool.Key);
-		ToolsContext->ToolManager->UnregisterToolType(RegisteredTool.Value);
-	}
-	RegisteredTools.SetNum(0);
-
 	if (Toolkit.IsValid())
 	{
+		const TSharedRef<FUICommandList>& CommandList = Toolkit->GetToolkitCommands();
+		for (auto& RegisteredTool : RegisteredTools)
+		{
+			CommandList->UnmapAction(RegisteredTool.Key);
+			ToolsContext->ToolManager->UnregisterToolType(RegisteredTool.Value);
+		}
+
 		FToolkitManager::Get().CloseToolkit(Toolkit.ToSharedRef());
 		Toolkit.Reset();
 	}
+	RegisteredTools.SetNum(0);
 
 	ToolsContext = nullptr;
 
@@ -484,6 +498,11 @@ TWeakObjectPtr<UEdModeInteractiveToolsContext> UEdMode::GetInteractiveToolsConte
 
 void UEdMode::CreateToolkit()
 {
+	if (!UsesToolkits())
+	{
+		return;
+	}
+
 	check(!Toolkit.IsValid())
 	Toolkit = MakeShareable(new FModeToolkit);
 }
diff --git a/Engine/Source/Editor/UnrealEd/Public/EdMode.h b/Engine/Source/Editor/UnrealEd/Public/EdMode.h
--- a/Engine/Source/Editor/UnrealEd/Public/EdMode.h
+++ b/Engine/Source/Editor/UnrealEd/Public/EdMode.h
@@ -11,6 +11,7 @@
 #include "EditorModeRegistry.h"
 #include "Tools/UEdMode.h"
 #include "Templates/SharedPointer.h"
+#include "Tools/LegacyEdModeWidgetHelpers.h"
 
 class FCanvas;
 class FEditorModeTools;
@@ -35,7 +36,7 @@ struct FViewportClick;
 /**
  * Base class for all editor modes.
  */
-class UNREALED_API FEdMode : public TSharedFromThis<FEdMode>, public FGCObject, public FEditorCommonDrawHelper
+class UNREALED_API FEdMode : public TSharedFromThis<FEdMode>, public FGCObject, public FEditorCommonDrawHelper, public FLegacyEdModeWidgetHelper
 {
 public:
 	/** Friends so it can access mode's internals on construction */
@@ -93,16 +94,8 @@ public:
 	virtual bool IsCompatibleWith(FEditorModeID OtherModeID) const { return false; }
 	virtual void ActorMoveNotify() {}
 	virtual void ActorsDuplicatedNotify(TArray<AActor*>& PreDuplicateSelection, TArray<AActor*>& PostDuplicateSelection, bool bOffsetLocations) {}
-	virtual void ActorSelectionChangeNotify();
 	virtual void ActorPropChangeNotify() {}
 	virtual void MapChangeNotify() {}
-	virtual bool ShowModeWidgets() const { return 1; }
-
-	/** If the EdMode is handling InputDelta (i.e., returning true from it), this allows a mode to indicated whether or not the Widget should also move. */
-	virtual bool AllowWidgetMove() { return true; }
-	
-	/** Check to see if the current widget mode can be cycled */
-	virtual bool CanCycleWidgetMode() const { return true; }
 
 	/** If the Edmode is handling its own mouse deltas, it can disable the MouseDeltaTacker */
 	virtual bool DisallowMouseDeltaTracking() const { return false; }
@@ -130,8 +123,6 @@ public:
 	virtual bool PostConvertMouseMovement(FEditorViewportClient* InViewportClient) { return false;}
 
 	virtual bool ShouldDrawBrushWireframe( AActor* InActor ) const { return true; }
-	virtual bool GetCustomDrawingCoordinateSystem(FMatrix& InMatrix, void* InData);
-	virtual bool GetCustomInputCoordinateSystem( FMatrix& InMatrix, void* InData ) { return 0; }
 
 	/** If Rotation Snap should be enabled for this mode*/ 
 	virtual bool IsSnapRotationEnabled();
@@ -142,45 +133,21 @@ public:
 	* @return					True if you have overridden the value
 	*/
 	virtual bool SnapRotatorToGridOverride(FRotator& Rotation){ return false; };
-
-	/**
-	 * Allows each mode to customize the axis pieces of the widget they want drawn.
-	 *
-	 * @param	InwidgetMode	The current widget mode
-	 *
-	 * @return					A bitfield comprised of AXIS_* values
-	 */
-	virtual EAxisList::Type GetWidgetAxisToDraw( UE::Widget::EWidgetMode InWidgetMode ) const;
-
-	/**
-	 * Allows each mode/tool to determine a good location for the widget to be drawn at.
-	 */
-	virtual FVector GetWidgetLocation() const;
-
-	/**
-	 * Lets the mode determine if it wants to draw the widget or not.
-	 */
-	virtual bool ShouldDrawWidget() const;
 	virtual void UpdateInternalData() {}
 
-	virtual FVector GetWidgetNormalFromCurrentAxis( void* InData );
-
 	virtual void AddReferencedObjects( FReferenceCollector& Collector ) override {}
 
 	virtual void Enter();
 	virtual void Exit();
 	virtual UTexture2D* GetVertexTexture();
-	
+
 	/**
 	 * Lets each tool determine if it wants to use the editor widget or not.  If the tool doesn't want to use it,
 	 * it will be fed raw mouse delta information (not snapped or altered in any way).
 	 */
-	virtual bool UsesTransformWidget() const;
-
-	/**
-	 * Lets each mode selectively exclude certain widget types.
-	 */
-	virtual bool UsesTransformWidget(UE::Widget::EWidgetMode CheckMode) const;
+	virtual bool UsesTransformWidget() const override;
+	using FLegacyEdModeWidgetHelper::UsesTransformWidget;
+	virtual bool UsesPropertyWidgets() const override;
 
 	virtual void PostUndo() {}
 
@@ -252,12 +219,6 @@ public:
 	const FModeTool* GetCurrentTool() const	{ return CurrentTool; }
 	//@}
 
-	/** @name Current widget axis. */
-	//@{
-	void SetCurrentWidgetAxis(EAxisList::Type InAxis)		{ CurrentWidgetAxis = InAxis; }
-	EAxisList::Type GetCurrentWidgetAxis() const			{ return CurrentWidgetAxis; }
-	//@}
-
 	/** @name Rendering */
 	//@{
 	/** Draws translucent polygons on brushes and volumes. */
@@ -298,24 +259,12 @@ public:
 	 */
 	virtual void BuildModeToolbar(class FToolBarBuilder& ToolbarBuilder) {}
 
-	// Property Widgets
-
-	/**
-	 * Lets each mode selectively enable widgets for editing properties tagged with 'Show 3D Widget' metadata.
-	 */
-	virtual bool UsesPropertyWidgets() const { return false; }
-
-	// @TODO: Find a better home for these?
-	/** Value of UPROPERTY can be edited with a widget in the editor (translation, rotation) */
-	static const FName MD_MakeEditWidget;
-	/** Specifies a function used for validation of the current value of a property.  The function returns a string that is empty if the value is valid, or contains an error description if the value is invalid */
-	static const FName MD_ValidateWidgetUsing;
-	/** Returns true if this structure can support creating a widget in the editor */
-	static bool CanCreateWidgetForStructure(const UStruct* InPropStruct);
-	/** Returns true if this property can support creating a widget in the editor */
-	static bool CanCreateWidgetForProperty(FProperty* InProp);
-	/** See if we should create a widget for the supplied property when selecting an actor instance */
-	static bool ShouldCreateWidgetForProperty(FProperty* InProp);
+	using FLegacyEdModeWidgetHelper::MD_MakeEditWidget;
+	using FLegacyEdModeWidgetHelper::MD_ValidateWidgetUsing;
+	using FLegacyEdModeWidgetHelper::CanCreateWidgetForStructure;
+	using FLegacyEdModeWidgetHelper::CanCreateWidgetForProperty;
+	using FLegacyEdModeWidgetHelper::ShouldCreateWidgetForProperty;
+	using FPropertyWidgetInfo = FLegacyEdModeWidgetHelper::FPropertyWidgetInfo;
 
 public:
 
@@ -333,9 +282,6 @@ private:
 	void OnModeUnregistered(FEditorModeID ModeID);
 
 protected:
-	/** The current axis that is being dragged on the widget. */
-	EAxisList::Type CurrentWidgetAxis;
-
 	/** Optional array of tools for this mode. */
 	TArray<FModeTool*> Tools;
 
@@ -347,72 +293,4 @@ protected:
 
 	/** Editor Mode Toolkit that is associated with this toolkit mode */
 	TSharedPtr<class FModeToolkit> Toolkit;
-
-	/** Pointer back to the mode tools that we are registered with */
-	FEditorModeTools* Owner;
-
-	// Property Widgets
-public:
-	/** Structure that holds info about our optional property widget */
-	struct FPropertyWidgetInfo
-	{
-		FString PropertyName;
-		int32 PropertyIndex;
-		FName PropertyValidationName;
-		FString DisplayName;
-		bool bIsTransform;
-
-		FPropertyWidgetInfo()
-			: PropertyIndex(INDEX_NONE)
-			, bIsTransform(false)
-		{
-		}
-		
-		void GetTransformAndColor(UObject* BestSelectedItem, bool bIsSelected, FTransform& OutLocalTransform, FString& OutValidationMessage, FColor& OutDrawColor) const;
-	};
-
-protected:
-	/**
-	 * Returns the first selected Actor, or NULL if there is no selection.
-	 */
-	AActor* GetFirstSelectedActorInstance() const;
-
-	/**
-	 * Gets an array of property widget info structures for the given struct/class type for the given container.
-	 *
-	 * @param InStruct The type of structure/class to access widget info structures for.
-	 * @param InContainer The container of the given type.
-	 * @param OutInfos An array of widget info structures (output).
-	 */
-	void GetPropertyWidgetInfos(const UStruct* InStruct, const void* InContainer, TArray<FPropertyWidgetInfo>& OutInfos) const;
-
-	/** Finds the best item to display widgets for (preferring selected components over selected actors) */
-	virtual UObject* GetItemToTryDisplayingWidgetsFor(FTransform& OutWidgetToWorld) const;
-
-	/** Name of the property currently being edited */
-	FString EditedPropertyName;
-
-	/** If the property being edited is an array property, this is the index of the element we're currently dealing with */
-	int32 EditedPropertyIndex;
-
-	/** Indicates  */
-	bool bEditedPropertyIsTransform;
-};
-
-/*------------------------------------------------------------------------------
-	Default.
-------------------------------------------------------------------------------*/
-
-/**
- * The default editing mode. 
- */
-class FEdModeDefault : public FEdMode
-{
-public:
-	FEdModeDefault();
-
-	// FEdMode interface
-	virtual bool UsesPropertyWidgets() const override { return true; }
-	// End of FEdMode interface
 };
-
diff --git a/Engine/Source/Editor/UnrealEd/Public/Tools/LegacyEdModeInterfaces.h b/Engine/Source/Editor/UnrealEd/Public/Tools/LegacyEdModeInterfaces.h
new file mode 100644
--- /dev/null
+++ b/Engine/Source/Editor/UnrealEd/Public/Tools/LegacyEdModeInterfaces.h
@@ -0,0 +1,135 @@
+// Copyright Epic Games, Inc. All Rights Reserved.
+
+#pragma once
+
+#include "UnrealWidgetFwd.h"
+#include "Math/Axis.h"
+#include "Math/Vector.h"
+#include "UObject/Interface.h"
+
+#include "LegacyEdModeInterfaces.generated.h"
+
+class FEditorViewportClient;
+struct FConvexVolume;
+enum EModeTools : int8;
+class FModeTool;
+struct FBox;
+struct FMatrix;
+
+UINTERFACE(NotBlueprintable, MinimalAPI)
+class ULegacyEdModeSelectInterface : public UInterface
+{
+	GENERATED_BODY()
+};
+
+class UNREALED_API ILegacyEdModeSelectInterface
+{
+	GENERATED_BODY()
+public:
+	/**
+	 * Lets each mode/tool handle box selection in its own way.
+	 *
+	 * @param	InBox	The selection box to use, in worldspace coordinates.
+	 * @return		true if something was selected/deselected, false otherwise.
+	 */
+	virtual bool BoxSelect(FBox& InBox, bool InSelect = true) = 0;
+
+	/**
+	 * Lets each mode/tool handle frustum selection in its own way.
+	 *
+	 * @param	InFrustum	The selection box to use, in worldspace coordinates.
+	 * @return	true if something was selected/deselected, false otherwise.
+	 */
+	virtual bool FrustumSelect(const FConvexVolume& InFrustum, FEditorViewportClient* InViewportClient, bool InSelect = true) = 0;
+
+	/** @return True if this mode allows the viewport to use a drag tool */
+	virtual bool AllowsViewportDragTool() const = 0;
+};
+
+UINTERFACE(NotBlueprintable, MinimalAPI)
+class ULegacyEdModeWidgetInterface : public UInterface
+{
+	GENERATED_BODY()
+};
+
+class UNREALED_API ILegacyEdModeWidgetInterface
+{
+	GENERATED_BODY()
+
+public:
+	/** If the EdMode is handling InputDelta (i.e., returning true from it), this allows a mode to indicated whether or not the Widget should also move. */
+	virtual bool AllowWidgetMove() = 0;
+
+	/** Check to see if the current widget mode can be cycled */
+	virtual bool CanCycleWidgetMode() const = 0;
+
+	virtual bool ShowModeWidgets() const = 0;
+	/**
+	 * Allows each mode to customize the axis pieces of the widget they want drawn.
+	 *
+	 * @param	InwidgetMode	The current widget mode
+	 *
+	 * @return					A bitfield comprised of AXIS_* values
+	 */
+	virtual EAxisList::Type GetWidgetAxisToDraw(UE::Widget::EWidgetMode InWidgetMode) const = 0;
+
+	/**
+	 * Allows each mode/tool to determine a good location for the widget to be drawn at.
+	 */
+	virtual FVector GetWidgetLocation() const = 0;
+
+	/**
+	 * Lets the mode determine if it wants to draw the widget or not.
+	 */
+	virtual bool ShouldDrawWidget() const = 0;
+
+	/**
+	 * Lets each tool determine if it wants to use the editor widget or not.  If the tool doesn't want to use it,
+	 * it will be fed raw mouse delta information (not snapped or altered in any way).
+	 */
+	virtual bool UsesTransformWidget() const = 0;
+
+	/**
+	 * Lets each mode selectively exclude certain widget types.
+	 */
+	virtual bool UsesTransformWidget(UE::Widget::EWidgetMode CheckMode) const = 0;
+
+	virtual FVector GetWidgetNormalFromCurrentAxis(void* InData) = 0;
+
+	/** @name Current widget axis. */
+	//@{
+	virtual void SetCurrentWidgetAxis(EAxisList::Type InAxis) = 0;
+	virtual EAxisList::Type GetCurrentWidgetAxis() const = 0;
+	//@}
+
+	/**
+	 * Lets each mode selectively enable widgets for editing properties tagged with 'Show 3D Widget' metadata.
+	 */
+	virtual bool UsesPropertyWidgets() const = 0;
+
+	virtual bool GetCustomDrawingCoordinateSystem(FMatrix& InMatrix, void* InData) = 0;
+	virtual bool GetCustomInputCoordinateSystem(FMatrix& InMatrix, void* InData) = 0;
+};
+
+UINTERFACE(NotBlueprintable, MinimalAPI)
+class ULegacyEdModeToolInterface : public UInterface
+{
+	GENERATED_BODY()
+};
+
+class UNREALED_API ILegacyEdModeToolInterface
+{
+	GENERATED_BODY()
+
+public:
+	// Tools
+	virtual void SetCurrentTool(EModeTools InID) = 0;
+	virtual void SetCurrentTool(FModeTool* InModeTool) = 0;
+	virtual FModeTool* FindTool(EModeTools InID) = 0;
+
+	virtual const TArray<FModeTool*>& GetTools() const = 0;
+
+	/** Returns the current tool. */
+	virtual FModeTool* GetCurrentTool() = 0;
+	virtual const FModeTool* GetCurrentTool() const = 0;
+};
diff --git a/Engine/Source/Editor/UnrealEd/Public/Tools/LegacyEdModeWidgetHelpers.h b/Engine/Source/Editor/UnrealEd/Public/Tools/LegacyEdModeWidgetHelpers.h
new file mode 100644
--- /dev/null
+++ b/Engine/Source/Editor/UnrealEd/Public/Tools/LegacyEdModeWidgetHelpers.h
@@ -0,0 +1,161 @@
+// Copyright Epic Games, Inc. All Rights Reserved.
+
+#pragma once
+
+#include "Tools/LegacyEdModeInterfaces.h"
+#include "Math/Vector.h"
+#include "Tools/UEdMode.h"
+
+#include "LegacyEdModeWidgetHelpers.generated.h"
+
+class FCanvas;
+class FEditorModeTools;
+class FEditorViewportClient;
+class FModeTool;
+class FPrimitiveDrawInterface;
+class FSceneView;
+class FViewport;
+struct FConvexVolume;
+struct FMatrix;
+
+enum EModeTools : int8;
+class HHitProxy;
+struct FViewportClick;
+class UObject;
+
+class UNREALED_API FLegacyEdModeWidgetHelper
+{
+	friend class FEditorModeRegistry;
+	friend class UBaseLegacyWidgetEdMode;
+
+public:
+	FLegacyEdModeWidgetHelper();
+	virtual ~FLegacyEdModeWidgetHelper() = default;
+
+	virtual bool AllowWidgetMove();
+	virtual bool CanCycleWidgetMode() const;
+	virtual bool ShowModeWidgets() const;
+	virtual EAxisList::Type GetWidgetAxisToDraw(UE::Widget::EWidgetMode InWidgetMode) const;
+	virtual FVector GetWidgetLocation() const;
+	virtual bool ShouldDrawWidget() const;
+	virtual bool UsesTransformWidget() const;
+	virtual bool UsesTransformWidget(UE::Widget::EWidgetMode CheckMode) const;
+	virtual FVector GetWidgetNormalFromCurrentAxis(void* InData);
+	virtual void SetCurrentWidgetAxis(EAxisList::Type InAxis) final;
+	virtual EAxisList::Type GetCurrentWidgetAxis() const final;
+	virtual bool UsesPropertyWidgets() const;
+	virtual bool GetCustomDrawingCoordinateSystem(FMatrix& InMatrix, void* InData);
+	virtual bool GetCustomInputCoordinateSystem(FMatrix& InMatrix, void* InData);
+	
+	virtual void ActorSelectionChangeNotify();
+
+	// Property Widgets
+	/** Structure that holds info about our optional property widget */
+	struct FPropertyWidgetInfo
+	{
+		FString PropertyName;
+		int32 PropertyIndex;
+		FName PropertyValidationName;
+		FString DisplayName;
+		bool bIsTransform;
+
+		FPropertyWidgetInfo()
+			: PropertyIndex(INDEX_NONE)
+			, bIsTransform(false)
+		{
+		}
+
+		void GetTransformAndColor(UObject* BestSelectedItem, bool bIsSelected, FTransform& OutLocalTransform, FString& OutValidationMessage, FColor& OutDrawColor) const;
+	};
+
+	virtual bool InputDelta(FEditorViewportClient* InViewportClient, FViewport* InViewport, FVector& InDrag, FRotator& InRot, FVector& InScale);
+	virtual bool HandleClick(FEditorViewportClient* InViewportClient, HHitProxy* HitProxy, const FViewportClick& Click);
+	virtual void Render(const FSceneView* View, FViewport* Viewport, FPrimitiveDrawInterface* PDI);
+	virtual void DrawHUD(FEditorViewportClient* ViewportClient, FViewport* Viewport, const FSceneView* View, FCanvas* Canvas);
+
+	// @TODO: Find a better home for these?
+	/** Value of UPROPERTY can be edited with a widget in the editor (translation, rotation) */
+	static const FName MD_MakeEditWidget;
+	/** Specifies a function used for validation of the current value of a property.  The function returns a string that is empty if the value is valid, or contains an error description if the value is invalid */
+	static const FName MD_ValidateWidgetUsing;
+	/** Returns true if this structure can support creating a widget in the editor */
+	static bool CanCreateWidgetForStructure(const UStruct* InPropStruct);
+	/** Returns true if this property can support creating a widget in the editor */
+	static bool CanCreateWidgetForProperty(FProperty* InProp);
+	/** See if we should create a widget for the supplied property when selecting an actor instance */
+	static bool ShouldCreateWidgetForProperty(FProperty* InProp);
+
+protected:
+	/**
+	 * Returns the first selected Actor, or NULL if there is no selection.
+	 */
+	AActor* GetFirstSelectedActorInstance() const;
+
+	/**
+	 * Gets an array of property widget info structures for the given struct/class type for the given container.
+	 *
+	 * @param InStruct The type of structure/class to access widget info structures for.
+	 * @param InContainer The container of the given type.
+	 * @param OutInfos An array of widget info structures (output).
+	 */
+	void GetPropertyWidgetInfos(const UStruct* InStruct, const void* InContainer, TArray<FPropertyWidgetInfo>& OutInfos) const;
+
+	/** Finds the best item to display widgets for (preferring selected components over selected actors) */
+	virtual UObject* GetItemToTryDisplayingWidgetsFor(FTransform& OutWidgetToWorld) const;
+
+	/** Name of the property currently being edited */
+	FString EditedPropertyName;
+
+	/** If the property being edited is an array property, this is the index of the element we're currently dealing with */
+	int32 EditedPropertyIndex;
+
+	/** Indicates  */
+	bool bEditedPropertyIsTransform;
+
+	/** The current axis that is being dragged on the widget. */
+	EAxisList::Type CurrentWidgetAxis;
+
+	/** Pointer back to the mode tools that we are registered with */
+	FEditorModeTools* Owner;
+};
+
+// This class is to aid transitioning from native FEdModes to UEdModes, in the case that the FEdMode used property widgets and/or transform widgets
+// To use this class:
+//   1. Sublcass FLegacyEdModeWidgetHelper, and override the methods needed for your native FEdMode implementation
+//   2. Transition your native FEdMode to a UObject, by inheriting from this class
+//   3. Override CreateWidgetHelper function to return a SharedRef to the class you created in step 1.
+UCLASS(Abstract)
+class UNREALED_API UBaseLegacyWidgetEdMode : public UEdMode, public ILegacyEdModeWidgetInterface
+{
+	GENERATED_BODY()
+
+public:
+	// UEdMode overrides
+	// If you need to do any initialization in your mode, be sure to still call through to this function.
+	// It creates the WidgetHelper and hooks up the mode manager pointer for you.
+	virtual void Initialize() override;
+	virtual bool InputDelta(FEditorViewportClient* InViewportClient, FViewport* InViewport, FVector& InDrag, FRotator& InRot, FVector& InScale) override;
+	virtual bool HandleClick(FEditorViewportClient* InViewportClient, HHitProxy* HitProxy, const FViewportClick& Click) override;
+	virtual void Render(const FSceneView* View, FViewport* Viewport, FPrimitiveDrawInterface* PDI) override;
+	virtual void DrawHUD(FEditorViewportClient* ViewportClient, FViewport* Viewport, const FSceneView* View, FCanvas* Canvas) override;
+
+	// ILegacyEdModeWidgetInterface overrides are just a pass through to the FLegacyEdModeWidgetHelper
+	virtual bool AllowWidgetMove() override;
+	virtual bool CanCycleWidgetMode() const override;
+	virtual bool ShowModeWidgets() const override;
+	virtual EAxisList::Type GetWidgetAxisToDraw(UE::Widget::EWidgetMode InWidgetMode) const override;
+	virtual FVector GetWidgetLocation() const override;
+	virtual bool ShouldDrawWidget() const override;
+	virtual bool UsesTransformWidget() const override;
+	virtual bool UsesTransformWidget(UE::Widget::EWidgetMode CheckMode) const override;
+	virtual FVector GetWidgetNormalFromCurrentAxis(void* InData) override;
+	virtual void SetCurrentWidgetAxis(EAxisList::Type InAxis) override;
+	virtual EAxisList::Type GetCurrentWidgetAxis() const override;
+	virtual bool UsesPropertyWidgets() const override;
+	virtual bool GetCustomDrawingCoordinateSystem(FMatrix& InMatrix, void* InData) override;
+	virtual bool GetCustomInputCoordinateSystem(FMatrix& InMatrix, void* InData) override;
+	virtual void ActorSelectionChangeNotify() override;
+
+protected:
+	TSharedPtr<FLegacyEdModeWidgetHelper> WidgetHelper;
+};
diff --git a/Engine/Source/Editor/UnrealEd/Public/Tools/UEdMode.h b/Engine/Source/Editor/UnrealEd/Public/Tools/UEdMode.h
--- a/Engine/Source/Editor/UnrealEd/Public/Tools/UEdMode.h
+++ b/Engine/Source/Editor/UnrealEd/Public/Tools/UEdMode.h
@@ -196,9 +196,6 @@ public:
 	/** Check to see if an actor can be selected in this mode - no side effects */
 	virtual bool IsSelectionAllowed(AActor* InActor, bool bInSelection) const { return true; }
 
-	/** @return True if this mode allows the viewport to use a drag tool */
-	virtual bool AllowsViewportDragTool() const { return true; }
-
 	/** Returns the editor mode identifier. */
 	FEditorModeID GetID() const { return Info.ID; }
 
@@ -306,4 +303,3 @@ protected:
 	UObject* SettingsObject;
 
 };
-
