WorldPartition: WorldPartitionConvertCommandlet now also converts Actor's Layer into non-dynamically loaded DataLayers

diff --git a/Engine/Source/Editor/DataLayerEditor/Private/DataLayer/DataLayerEditorSubsystem.cpp b/Engine/Source/Editor/DataLayerEditor/Private/DataLayer/DataLayerEditorSubsystem.cpp
--- a/Engine/Source/Editor/DataLayerEditor/Private/DataLayer/DataLayerEditorSubsystem.cpp
+++ b/Engine/Source/Editor/DataLayerEditor/Private/DataLayer/DataLayerEditorSubsystem.cpp
@@ -1107,31 +1107,7 @@ const AWorldDataLayers* UDataLayerEditorSubsystem::GetWorldDataLayers() const
 
 AWorldDataLayers* UDataLayerEditorSubsystem::GetWorldDataLayers(bool bCreateIfNotFound)
 {
-	UWorld* World = GetWorld();
-	FName WorldDataLayersName = AWorldDataLayers::StaticClass()->GetFName();
-	AWorldDataLayers* WorldDataLayers = nullptr;
-	if (UObject* ExistingObject = StaticFindObject(nullptr, World->PersistentLevel, *WorldDataLayersName.ToString()))
-	{
-		WorldDataLayers = CastChecked<AWorldDataLayers>(ExistingObject);
-		if (WorldDataLayers->IsPendingKill())
-		{
-			// Handle the case where the actor already exists, but it's pending kill
-			WorldDataLayers->Rename(nullptr, nullptr, REN_DontCreateRedirectors | REN_DoNotDirty | REN_NonTransactional | REN_ForceNoResetLoaders);
-			WorldDataLayers = nullptr;
-		}
-	}
-
-	if (!WorldDataLayers && bCreateIfNotFound)
-	{
-		FActorSpawnParameters SpawnParams;
-		SpawnParams.OverrideLevel = World->PersistentLevel;
-		SpawnParams.bHideFromSceneOutliner = true;
-		SpawnParams.Name = WorldDataLayersName;
-		SpawnParams.NameMode = FActorSpawnParameters::ESpawnActorNameMode::Required_Fatal;
-		WorldDataLayers = World->SpawnActor<AWorldDataLayers>(AWorldDataLayers::StaticClass(), SpawnParams);
-	}
-
-	return WorldDataLayers;
+	return AWorldDataLayers::Get(GetWorld(), bCreateIfNotFound);
 }
 
 void UDataLayerEditorSubsystem::AddAllDataLayersTo(TArray<TWeakObjectPtr<UDataLayer>>& OutDataLayers) const
diff --git a/Engine/Source/Editor/UnrealEd/Private/Commandlets/WorldPartitionConvertCommandlet.cpp b/Engine/Source/Editor/UnrealEd/Private/Commandlets/WorldPartitionConvertCommandlet.cpp
--- a/Engine/Source/Editor/UnrealEd/Private/Commandlets/WorldPartitionConvertCommandlet.cpp
+++ b/Engine/Source/Editor/UnrealEd/Private/Commandlets/WorldPartitionConvertCommandlet.cpp
@@ -44,6 +44,8 @@
 #include "LandscapeInfo.h"
 #include "LandscapeConfigHelper.h"
 #include "LandscapeGizmoActor.h"
+#include "WorldPartition/DataLayer/DataLayer.h"
+#include "WorldPartition/DataLayer/WorldDataLayers.h"
 
 DEFINE_LOG_CATEGORY(LogWorldPartitionConvertCommandlet);
 
@@ -805,6 +807,7 @@ int32 UWorldPartitionConvertCommandlet::Main(const FString& Params)
 	}
 
 	UPackage* MainPackage = MainLevel->GetPackage();
+	AWorldDataLayers* MainWorldDataLayers = AWorldDataLayers::Get(MainWorld, /*bCreateIfNotFound*/true);
 
 	OnWorldLoaded(MainWorld);
 
@@ -870,7 +873,7 @@ int32 UWorldPartitionConvertCommandlet::Main(const FString& Params)
 		}
 	};
 
-	auto PrepareLevelActors = [this, PartitionFoliage, PartitionLandscape](ULevel* Level, bool bMainLevel, EActorGridPlacement DefaultGridPlacement)
+	auto PrepareLevelActors = [this, PartitionFoliage, PartitionLandscape, MainWorldDataLayers](ULevel* Level, bool bMainLevel, EActorGridPlacement DefaultGridPlacement)
 	{
 		const FBox WorldBounds(WorldOrigin - WorldExtent, WorldOrigin + WorldExtent);
 
@@ -913,6 +916,18 @@ int32 UWorldPartitionConvertCommandlet::Main(const FString& Params)
 					}
 				}
 
+				// Convert Layers into DataLayers with DynamicallyLoaded flag disabled
+				for (FName Layer : Actor->Layers)
+				{
+					UDataLayer* DataLayer = const_cast<UDataLayer*>(MainWorldDataLayers->GetDataLayerFromLabel(Layer));
+					if (!DataLayer)
+					{
+						DataLayer = MainWorldDataLayers->CreateDataLayer();
+						DataLayer->SetDataLayerLabel(Layer);
+						DataLayer->SetIsDynamicallyLoaded(false);
+					}
+					Actor->AddDataLayer(DataLayer);
+				}
 				// Clear actor layers as they are not supported yet in world partition
 				Actor->Layers.Empty();
 			}
diff --git a/Engine/Source/Runtime/Engine/Private/WorldPartition/DataLayer/WorldDataLayers.cpp b/Engine/Source/Runtime/Engine/Private/WorldPartition/DataLayer/WorldDataLayers.cpp
--- a/Engine/Source/Runtime/Engine/Private/WorldPartition/DataLayer/WorldDataLayers.cpp
+++ b/Engine/Source/Runtime/Engine/Private/WorldPartition/DataLayer/WorldDataLayers.cpp
@@ -31,6 +31,40 @@ const AWorldDataLayers* AWorldDataLayers::Get(UWorld* World)
 }
 
 #if WITH_EDITOR
+AWorldDataLayers* AWorldDataLayers::Get(UWorld* World, bool bCreateIfNotFound)
+{
+	if (bCreateIfNotFound)
+	{
+		AWorldDataLayers* WorldDataLayers = nullptr;
+		static FName WorldDataLayersName = AWorldDataLayers::StaticClass()->GetFName();
+		if (UObject* ExistingObject = StaticFindObject(nullptr, World->PersistentLevel, *WorldDataLayersName.ToString()))
+		{
+			WorldDataLayers = CastChecked<AWorldDataLayers>(ExistingObject);
+			if (WorldDataLayers->IsPendingKill())
+			{
+				// Handle the case where the actor already exists, but it's pending kill
+				WorldDataLayers->Rename(nullptr, nullptr, REN_DontCreateRedirectors | REN_DoNotDirty | REN_NonTransactional | REN_ForceNoResetLoaders);
+				WorldDataLayers = nullptr;
+			}
+		}
+
+		if (!WorldDataLayers)
+		{
+			FActorSpawnParameters SpawnParams;
+			SpawnParams.OverrideLevel = World->PersistentLevel;
+			SpawnParams.bHideFromSceneOutliner = true;
+			SpawnParams.Name = WorldDataLayersName;
+			SpawnParams.NameMode = FActorSpawnParameters::ESpawnActorNameMode::Required_Fatal;
+			WorldDataLayers = World->SpawnActor<AWorldDataLayers>(AWorldDataLayers::StaticClass(), SpawnParams);
+		}
+		return WorldDataLayers;
+	}
+	else
+	{
+		return const_cast<AWorldDataLayers*>(Get(World));
+	}
+}
+
 FName AWorldDataLayers::GenerateUniqueDataLayerLabel(const FName& InDataLayerLabel) const
 {
 	int32 DataLayerIndex = 0;
diff --git a/Engine/Source/Runtime/Engine/Public/WorldPartition/DataLayer/WorldDataLayers.h b/Engine/Source/Runtime/Engine/Public/WorldPartition/DataLayer/WorldDataLayers.h
--- a/Engine/Source/Runtime/Engine/Public/WorldPartition/DataLayer/WorldDataLayers.h
+++ b/Engine/Source/Runtime/Engine/Public/WorldPartition/DataLayer/WorldDataLayers.h
@@ -22,6 +22,7 @@ public:
 	virtual void PostLoad() override;
 
 #if WITH_EDITOR
+	static AWorldDataLayers* Get(UWorld* World, bool bCreateIfNotFound);
 	UDataLayer* CreateDataLayer();
 	bool RemoveDataLayer(const UDataLayer* InDataLayer);
 	bool RemoveDataLayers(const TArray<const UDataLayer*>& InDataLayers);
