Explicitly pass through to prevent function hiding in legacy wrapper inheritance from CL 14628583.

#rnx

diff --git a/Engine/Source/Editor/UnrealEd/Private/EdMode.cpp b/Engine/Source/Editor/UnrealEd/Private/EdMode.cpp
--- a/Engine/Source/Editor/UnrealEd/Private/EdMode.cpp
+++ b/Engine/Source/Editor/UnrealEd/Private/EdMode.cpp
@@ -166,6 +166,11 @@ bool FEdMode::UsesTransformWidget() const
 	return true;
 }
 
+bool FEdMode::UsesTransformWidget(UE::Widget::EWidgetMode CheckMode) const
+{
+	return FLegacyEdModeWidgetHelper::UsesTransformWidget(CheckMode);
+}
+
 bool FEdMode::UsesPropertyWidgets() const
 {
 	return false;
diff --git a/Engine/Source/Editor/UnrealEd/Public/EdMode.h b/Engine/Source/Editor/UnrealEd/Public/EdMode.h
--- a/Engine/Source/Editor/UnrealEd/Public/EdMode.h
+++ b/Engine/Source/Editor/UnrealEd/Public/EdMode.h
@@ -146,7 +146,7 @@ public:
 	 * it will be fed raw mouse delta information (not snapped or altered in any way).
 	 */
 	virtual bool UsesTransformWidget() const override;
-	using FLegacyEdModeWidgetHelper::UsesTransformWidget;
+	virtual bool UsesTransformWidget(UE::Widget::EWidgetMode CheckMode) const override;
 	virtual bool UsesPropertyWidgets() const override;
 
 	virtual void PostUndo() {}
