WorldPartition: Fixed WorldPartitionConvertCommandlet not skipping some actors when converting Layers to DataLayers

diff --git a/Engine/Source/Editor/DataLayerEditor/Private/DataLayer/DataLayerEditorSubsystem.cpp b/Engine/Source/Editor/DataLayerEditor/Private/DataLayer/DataLayerEditorSubsystem.cpp
--- a/Engine/Source/Editor/DataLayerEditor/Private/DataLayer/DataLayerEditorSubsystem.cpp
+++ b/Engine/Source/Editor/DataLayerEditor/Private/DataLayer/DataLayerEditorSubsystem.cpp
@@ -200,24 +200,7 @@ void UDataLayerEditorSubsystem::PostUndoRedo()
 
 bool UDataLayerEditorSubsystem::IsActorValidForDataLayer(AActor* Actor)
 {
-	if (!Actor || Actor->GetClass() == nullptr)
-	{
-		return false;
-	}
-
-	UWorld* World = Actor->GetWorld();
-	if (!World)
-	{
-		return false;
-	}
-
-	const bool bIsPartitionedActor = UWorld::HasSubsystem<UWorldPartitionSubsystem>(World);
-	const bool bIsInEditorWorld = World->WorldType == EWorldType::Editor;
-	const bool bIsBuilderBrush = FActorEditorUtils::IsABuilderBrush(Actor);
-	const bool bIsHidden = Actor->GetClass()->GetDefaultObject<AActor>()->bHiddenEd;
-	const bool bIsValid = !bIsHidden && !bIsBuilderBrush && bIsInEditorWorld && bIsPartitionedActor;
-
-	return bIsValid;
+	return Actor && Actor->IsValidForDataLayer();
 }
 
 void UDataLayerEditorSubsystem::InitializeNewActorDataLayers(AActor* Actor)
@@ -1046,7 +1029,7 @@ void UDataLayerEditorSubsystem::SetDataLayersIsDynamicallyLoaded(const TArray<UD
 			DataLayerChanged.Broadcast(EDataLayerAction::Modify, DataLayer, "bIsDynamicallyLoaded");
 			bChangeOccurred = true;
 		}
-	}
+		}
 
 	if (bChangeOccurred)
 	{
diff --git a/Engine/Source/Editor/UnrealEd/Private/Commandlets/WorldPartitionConvertCommandlet.cpp b/Engine/Source/Editor/UnrealEd/Private/Commandlets/WorldPartitionConvertCommandlet.cpp
--- a/Engine/Source/Editor/UnrealEd/Private/Commandlets/WorldPartitionConvertCommandlet.cpp
+++ b/Engine/Source/Editor/UnrealEd/Private/Commandlets/WorldPartitionConvertCommandlet.cpp
@@ -887,49 +887,55 @@ int32 UWorldPartitionConvertCommandlet::Main(const FString& Params)
 				{
 					Level->GetWorld()->DestroyActor(Actor);
 				}
-				else if (AInstancedFoliageActor* IFA = Cast<AInstancedFoliageActor>(Actor))
+				else 
 				{
-					IFAs.Add(IFA);
-				}
-				else if (ALandscapeProxy* LandscapeProxy = Cast<ALandscapeProxy>(Actor))
-				{
-					ULandscapeInfo* LandscapeInfo = LandscapeProxy->GetLandscapeInfo();
-					check(LandscapeInfo);
-					LandscapeInfos.Add(LandscapeInfo);
-				}
-				// Only override default grid placement on actors that are not marked as always loaded
-				else if (Actor->GridPlacement != EActorGridPlacement::AlwaysLoaded)
-				{
-					FVector BoundsLocation;
-					FVector BoundsExtent;
-					Actor->GetActorLocationBounds(/*bOnlyCollidingComponents*/false, BoundsLocation, BoundsExtent, /*bIncludeFromChildActors*/true);
-
-					const FBox ActorBounds(BoundsLocation - BoundsExtent, BoundsLocation + BoundsExtent);
-
-					if (!WorldBounds.IsInside(ActorBounds))
+					if (AInstancedFoliageActor* IFA = Cast<AInstancedFoliageActor>(Actor))
 					{
-						Actor->GridPlacement = EActorGridPlacement::AlwaysLoaded;
+						IFAs.Add(IFA);
 					}
-					else
+					else if (ALandscapeProxy* LandscapeProxy = Cast<ALandscapeProxy>(Actor))
 					{
-						Actor->GridPlacement = DefaultGridPlacement;
+						ULandscapeInfo* LandscapeInfo = LandscapeProxy->GetLandscapeInfo();
+						check(LandscapeInfo);
+						LandscapeInfos.Add(LandscapeInfo);
 					}
-				}
+					// Only override default grid placement on actors that are not marked as always loaded
+					else if (Actor->GridPlacement != EActorGridPlacement::AlwaysLoaded)
+					{
+						FVector BoundsLocation;
+						FVector BoundsExtent;
+						Actor->GetActorLocationBounds(/*bOnlyCollidingComponents*/false, BoundsLocation, BoundsExtent, /*bIncludeFromChildActors*/true);
 
-				// Convert Layers into DataLayers with DynamicallyLoaded flag disabled
-				for (FName Layer : Actor->Layers)
-				{
-					UDataLayer* DataLayer = const_cast<UDataLayer*>(MainWorldDataLayers->GetDataLayerFromLabel(Layer));
-					if (!DataLayer)
+						const FBox ActorBounds(BoundsLocation - BoundsExtent, BoundsLocation + BoundsExtent);
+
+						if (!WorldBounds.IsInside(ActorBounds))
+						{
+							Actor->GridPlacement = EActorGridPlacement::AlwaysLoaded;
+						}
+						else
+						{
+							Actor->GridPlacement = DefaultGridPlacement;
+						}
+					}
+
+					// Convert Layers into DataLayers with DynamicallyLoaded flag disabled
+					if (Actor->IsValidForDataLayer())
 					{
-						DataLayer = MainWorldDataLayers->CreateDataLayer();
-						DataLayer->SetDataLayerLabel(Layer);
-						DataLayer->SetIsDynamicallyLoaded(false);
+						for (FName Layer : Actor->Layers)
+						{
+							UDataLayer* DataLayer = const_cast<UDataLayer*>(MainWorldDataLayers->GetDataLayerFromLabel(Layer));
+							if (!DataLayer)
+							{
+								DataLayer = MainWorldDataLayers->CreateDataLayer();
+								DataLayer->SetDataLayerLabel(Layer);
+								DataLayer->SetIsDynamicallyLoaded(false);
+							}
+							Actor->AddDataLayer(DataLayer);
+						}
 					}
-					Actor->AddDataLayer(DataLayer);
+					// Clear actor layers as they are not supported yet in world partition
+					Actor->Layers.Empty();
 				}
-				// Clear actor layers as they are not supported yet in world partition
-				Actor->Layers.Empty();
 			}
 		}
 
diff --git a/Engine/Source/Runtime/Engine/Classes/GameFramework/Actor.h b/Engine/Source/Runtime/Engine/Classes/GameFramework/Actor.h
--- a/Engine/Source/Runtime/Engine/Classes/GameFramework/Actor.h
+++ b/Engine/Source/Runtime/Engine/Classes/GameFramework/Actor.h
@@ -1005,6 +1005,7 @@ public:
 	TArray<FName> GetDataLayerNames() const;
 	TArray<const UDataLayer*> GetDataLayerObjects() const;
 	bool IsPropertyChangedAffectingDataLayers(FPropertyChangedEvent& PropertyChangedEvent) const;
+	bool IsValidForDataLayer() const;
 	void FixupDataLayers();
 #endif
 
diff --git a/Engine/Source/Runtime/Engine/Private/ActorEditor.cpp b/Engine/Source/Runtime/Engine/Private/ActorEditor.cpp
--- a/Engine/Source/Runtime/Engine/Private/ActorEditor.cpp
+++ b/Engine/Source/Runtime/Engine/Private/ActorEditor.cpp
@@ -1297,6 +1297,23 @@ bool AActor::IsPropertyChangedAffectingDataLayers(FPropertyChangedEvent& Propert
 	return false;
 }
 
+bool AActor::IsValidForDataLayer() const
+{
+	UWorld* World = GetWorld();
+	if (!World)
+	{
+		return false;
+	}
+
+	const bool bIsPartitionedActor = UWorld::HasSubsystem<UWorldPartitionSubsystem>(World);
+	const bool bIsInEditorWorld = World->WorldType == EWorldType::Editor;
+	const bool bIsBuilderBrush = FActorEditorUtils::IsABuilderBrush(this);
+	const bool bIsHidden = GetClass()->GetDefaultObject<AActor>()->bHiddenEd;
+	const bool bIsValid = !bIsHidden && !bIsBuilderBrush && bIsInEditorWorld && bIsPartitionedActor;
+
+	return bIsValid;
+}
+
 // DataLayers (end)
 //---------------------------------------------------------------------------
 
