UEdModes are no longer draw helpers, which will prevent double drawing of the grid when the default mode is active in a viewport.

# The level editor viewport client relies on its draw helper instead of individual modes' draw helpers, but mode managers and FEdModes retain their draw helpers for legacy reasons.

diff --git a/Engine/Source/Editor/UnrealEd/Private/EditorModeManager.cpp b/Engine/Source/Editor/UnrealEd/Private/EditorModeManager.cpp
--- a/Engine/Source/Editor/UnrealEd/Private/EditorModeManager.cpp
+++ b/Engine/Source/Editor/UnrealEd/Private/EditorModeManager.cpp
@@ -1352,7 +1352,10 @@ void FEditorModeTools::DrawActiveModes( const FSceneView* InView, FPrimitiveDraw
 {
 	for (UEdMode* Mode : ActiveScriptableModes)
 	{
-		Mode->Draw(InView, PDI);
+		if (ILegacyEdModeDrawHelperInterface* DrawHelper = Cast<ILegacyEdModeDrawHelperInterface>(Mode))
+		{
+			DrawHelper->Draw(InView, PDI);
+		}
 	}
 }
 
diff --git a/Engine/Source/Editor/UnrealEd/Private/LevelEditorViewport.cpp b/Engine/Source/Editor/UnrealEd/Private/LevelEditorViewport.cpp
--- a/Engine/Source/Editor/UnrealEd/Private/LevelEditorViewport.cpp
+++ b/Engine/Source/Editor/UnrealEd/Private/LevelEditorViewport.cpp
@@ -1760,8 +1760,12 @@ FLevelEditorViewportClient::FLevelEditorViewportClient(const TSharedPtr<SLevelVi
 	// Add a delegate so we get informed when an actor has moved.
 	GEngine->OnActorMoved().AddRaw(this, &FLevelEditorViewportClient::OnActorMoved);
 
-	// GEditorModeTools serves as our draw helper
-	bUsesDrawHelper = false;
+	// Set up defaults for the draw helper.
+	DrawHelper.bDrawGrid = true;
+	DrawHelper.bDrawPivot = false;
+	DrawHelper.bDrawBaseInfo = false;
+	DrawHelper.bDrawWorldBox = false;
+	DrawHelper.bDrawKillZ = false;
 
 	InitializeVisibilityFlags();
 
diff --git a/Engine/Source/Editor/UnrealEd/Private/Tools/DefaultEdMode.cpp b/Engine/Source/Editor/UnrealEd/Private/Tools/DefaultEdMode.cpp
--- a/Engine/Source/Editor/UnrealEd/Private/Tools/DefaultEdMode.cpp
+++ b/Engine/Source/Editor/UnrealEd/Private/Tools/DefaultEdMode.cpp
@@ -8,11 +8,6 @@
 
 UEdModeDefault::UEdModeDefault()
 {
-	bDrawGrid = true;
-	bDrawPivot = false;
-	bDrawBaseInfo = false;
-	bDrawWorldBox = false;
-	bDrawKillZ = false;
 	Info = FEditorModeInfo(
 		FBuiltinEditorModes::EM_Default,
 		NSLOCTEXT("DefaultMode", "DisplayName", "Select"),
diff --git a/Engine/Source/Editor/UnrealEd/Private/Tools/LegacyEdMode.h b/Engine/Source/Editor/UnrealEd/Private/Tools/LegacyEdMode.h
--- a/Engine/Source/Editor/UnrealEd/Private/Tools/LegacyEdMode.h
+++ b/Engine/Source/Editor/UnrealEd/Private/Tools/LegacyEdMode.h
@@ -23,7 +23,7 @@ class FViewport;
 class UPrimitiveComponent;
 
 UCLASS()
-class UNREALED_API ULegacyEdModeWrapper final : public UEdMode, public ILegacyEdModeWidgetInterface, public ILegacyEdModeToolInterface, public ILegacyEdModeSelectInterface
+class UNREALED_API ULegacyEdModeWrapper final : public UEdMode, public ILegacyEdModeWidgetInterface, public ILegacyEdModeToolInterface, public ILegacyEdModeSelectInterface, public ILegacyEdModeDrawHelperInterface
 {
 	GENERATED_BODY()
 
diff --git a/Engine/Source/Editor/UnrealEd/Private/Tools/UEdMode.cpp b/Engine/Source/Editor/UnrealEd/Private/Tools/UEdMode.cpp
--- a/Engine/Source/Editor/UnrealEd/Private/Tools/UEdMode.cpp
+++ b/Engine/Source/Editor/UnrealEd/Private/Tools/UEdMode.cpp
@@ -65,7 +65,6 @@ UEdMode::UEdMode()
 	: bPendingDeletion(false)
 	, Owner(nullptr)
 {
-	bDrawKillZ = true;
 	ToolsContext = nullptr;
 	ToolCommandList = MakeShareable(new FUICommandList);
 }
diff --git a/Engine/Source/Editor/UnrealEd/Public/Tools/UEdMode.h b/Engine/Source/Editor/UnrealEd/Public/Tools/UEdMode.h
--- a/Engine/Source/Editor/UnrealEd/Public/Tools/UEdMode.h
+++ b/Engine/Source/Editor/UnrealEd/Public/Tools/UEdMode.h
@@ -56,7 +56,7 @@ namespace EEditAction
  * Base class for all editor modes.
  */
 UCLASS(Abstract)
-class UNREALED_API UEdMode : public UObject, public FEditorCommonDrawHelper
+class UNREALED_API UEdMode : public UObject
 {
 	GENERATED_BODY()
 
