Missed file from CL 14649186

#rnx

diff --git a/Engine/Source/Editor/UnrealEd/Public/Tools/LegacyEdModeInterfaces.h b/Engine/Source/Editor/UnrealEd/Public/Tools/LegacyEdModeInterfaces.h
--- a/Engine/Source/Editor/UnrealEd/Public/Tools/LegacyEdModeInterfaces.h
+++ b/Engine/Source/Editor/UnrealEd/Public/Tools/LegacyEdModeInterfaces.h
@@ -133,3 +133,17 @@ public:
 	virtual FModeTool* GetCurrentTool() = 0;
 	virtual const FModeTool* GetCurrentTool() const = 0;
 };
+
+UINTERFACE(NotBlueprintable, MinimalAPI)
+class ULegacyEdModeDrawHelperInterface : public UInterface
+{
+	GENERATED_BODY()
+};
+
+class UNREALED_API ILegacyEdModeDrawHelperInterface
+{
+	GENERATED_BODY()
+
+public:
+	virtual void Draw(const FSceneView* View, FPrimitiveDrawInterface* PDI) = 0;
+};
