Fix marquee selection in the default ed mode.

# Move AllowsViewportDragTool to the legacy widget interface so that the viewport client creates the drag tool again when the default mode is active.

#rnx

diff --git a/Engine/Source/Editor/UnrealEd/Private/EditorModeManager.cpp b/Engine/Source/Editor/UnrealEd/Private/EditorModeManager.cpp
--- a/Engine/Source/Editor/UnrealEd/Private/EditorModeManager.cpp
+++ b/Engine/Source/Editor/UnrealEd/Private/EditorModeManager.cpp
@@ -1048,7 +1048,7 @@ bool FEditorModeTools::AllowsViewportDragTool() const
 	bool bCanUseDragTool = false;
 	for (UEdMode* Mode : ActiveScriptableModes)
 	{
-		if (ILegacyEdModeSelectInterface* LegacyMode = Cast<ILegacyEdModeSelectInterface>(Mode))
+		if (ILegacyEdModeWidgetInterface* LegacyMode = Cast<ILegacyEdModeWidgetInterface>(Mode))
 		{
 			bCanUseDragTool |= LegacyMode->AllowsViewportDragTool();
 		}
diff --git a/Engine/Source/Editor/UnrealEd/Private/Tools/LegacyEdModeWidgetHelpers.cpp b/Engine/Source/Editor/UnrealEd/Private/Tools/LegacyEdModeWidgetHelpers.cpp
--- a/Engine/Source/Editor/UnrealEd/Private/Tools/LegacyEdModeWidgetHelpers.cpp
+++ b/Engine/Source/Editor/UnrealEd/Private/Tools/LegacyEdModeWidgetHelpers.cpp
@@ -524,6 +524,11 @@ void FLegacyEdModeWidgetHelper::ActorSelectionChangeNotify()
 	bEditedPropertyIsTransform = false;
 }
 
+bool FLegacyEdModeWidgetHelper::AllowsViewportDragTool() const
+{
+	return true;
+}
+
 bool FLegacyEdModeWidgetHelper::InputDelta(FEditorViewportClient* InViewportClient, FViewport* InViewport, FVector& InDrag, FRotator& InRot, FVector& InScale)
 {
 	if (UsesPropertyWidgets())
@@ -872,3 +877,8 @@ void UBaseLegacyWidgetEdMode::ActorSelectionChangeNotify()
 	WidgetHelper->ActorSelectionChangeNotify();
 }
 
+bool UBaseLegacyWidgetEdMode::AllowsViewportDragTool() const
+{
+	return WidgetHelper->AllowsViewportDragTool();
+}
+
diff --git a/Engine/Source/Editor/UnrealEd/Public/EdMode.h b/Engine/Source/Editor/UnrealEd/Public/EdMode.h
--- a/Engine/Source/Editor/UnrealEd/Public/EdMode.h
+++ b/Engine/Source/Editor/UnrealEd/Public/EdMode.h
@@ -194,9 +194,6 @@ public:
 	/** Check to see if an actor can be selected in this mode - no side effects */
 	virtual bool IsSelectionAllowed( AActor* InActor, bool bInSelection ) const { return true; }
 
-	/** @return True if this mode allows the viewport to use a drag tool */
-	virtual bool AllowsViewportDragTool() const { return true; }
-
 	/** Returns the editor mode identifier. */
 	FEditorModeID GetID() const { return Info.ID; }
 
diff --git a/Engine/Source/Editor/UnrealEd/Public/Tools/LegacyEdModeWidgetHelpers.h b/Engine/Source/Editor/UnrealEd/Public/Tools/LegacyEdModeWidgetHelpers.h
--- a/Engine/Source/Editor/UnrealEd/Public/Tools/LegacyEdModeWidgetHelpers.h
+++ b/Engine/Source/Editor/UnrealEd/Public/Tools/LegacyEdModeWidgetHelpers.h
@@ -48,6 +48,7 @@ public:
 	virtual bool GetCustomInputCoordinateSystem(FMatrix& InMatrix, void* InData);
 	
 	virtual void ActorSelectionChangeNotify();
+	virtual bool AllowsViewportDragTool() const;
 
 	// Property Widgets
 	/** Structure that holds info about our optional property widget */
@@ -155,6 +156,7 @@ public:
 	virtual bool GetCustomDrawingCoordinateSystem(FMatrix& InMatrix, void* InData) override;
 	virtual bool GetCustomInputCoordinateSystem(FMatrix& InMatrix, void* InData) override;
 	virtual void ActorSelectionChangeNotify() override;
+	virtual bool AllowsViewportDragTool() const override;
 
 protected:
 	TSharedPtr<FLegacyEdModeWidgetHelper> WidgetHelper;
