Fix missing value in log message.
#rnx

diff --git a/Engine/Source/Editor/UnrealEd/Private/CookOnTheFlyServer.cpp b/Engine/Source/Editor/UnrealEd/Private/CookOnTheFlyServer.cpp
--- a/Engine/Source/Editor/UnrealEd/Private/CookOnTheFlyServer.cpp
+++ b/Engine/Source/Editor/UnrealEd/Private/CookOnTheFlyServer.cpp
@@ -1319,7 +1319,7 @@ void UCookOnTheFlyServer::UpdateDisplay(ECookTickFlags TickFlags, bool bForceDis
 		LastDiagnosticsDisplayTime = CurrentTime;
 		if (bSaveBusy && CurrentTime - SaveBusyTimeStarted > GCookProgressWarnBusyTime)
 		{
-			UE_LOG(LogCook, Warning, TEXT("Cooker has been blocked from saving the current packages for %f seconds."));
+			UE_LOG(LogCook, Warning, TEXT("Cooker has been blocked from saving the current packages for %f seconds."), GCookProgressWarnBusyTime);
 			UE_LOG(LogCook, Display, TEXT("Current packages in the savequeue : "), GCookProgressWarnBusyTime);
 			FPackageDataQueue& SaveQueue = PackageDatas->GetSaveQueue();
 			bool bFound = false;
