Unblocked crash due to curious UV data in FBX files exported from Houdini. This needs further investigation as to whether the UV layer format being used is currently unsupported by the importer.

diff --git a/Engine/Source/Editor/UnrealEd/Private/Fbx/FbxStaticMeshImport.cpp b/Engine/Source/Editor/UnrealEd/Private/Fbx/FbxStaticMeshImport.cpp
--- a/Engine/Source/Editor/UnrealEd/Private/Fbx/FbxStaticMeshImport.cpp
+++ b/Engine/Source/Editor/UnrealEd/Private/Fbx/FbxStaticMeshImport.cpp
@@ -1106,6 +1106,9 @@ bool UnFbx::FFbxImporter::BuildStaticMeshFromGeometry(FbxNode* Node, UStaticMesh
 
 				for (int32 UVLayerIndex = 0; UVLayerIndex < FBXUVs.UniqueUVCount; UVLayerIndex++)
 				{
+					int32 UVDirectArrayCount = FBXUVs.LayerElementUV[UVLayerIndex]->GetDirectArray().GetCount();
+					int32 UVIndexArrayCount = FBXUVs.LayerElementUV[UVLayerIndex]->GetIndexArray().GetCount();
+
 					FUVID UVIDs[3] = {INDEX_NONE, INDEX_NONE, INDEX_NONE};
 					if (FBXUVs.LayerElementUV[UVLayerIndex] != NULL)
 					{
@@ -1116,7 +1119,12 @@ bool UnFbx::FFbxImporter::BuildStaticMeshFromGeometry(FbxNode* Node, UStaticMesh
 								: SkippedVertexInstance + CurrentVertexInstanceIndex + VertexIndex;
 							int32 UVIndex = (FBXUVs.UVReferenceMode[UVLayerIndex] == FbxLayerElement::eDirect)
 								? UVMapIndex
-								: FBXUVs.LayerElementUV[UVLayerIndex]->GetIndexArray().GetAt(UVMapIndex);
+								: (UVMapIndex < UVIndexArrayCount ? FBXUVs.LayerElementUV[UVLayerIndex]->GetIndexArray().GetAt(UVMapIndex) : -1);
+
+							if (UVIndex >= FBXUVs.LayerElementUV[UVLayerIndex]->GetDirectArray().GetCount())
+							{
+								UVIndex = -1;
+							}
 
 							if (UVIndex != -1)
 							{
