Redefine Content Browser's "RealTime Thumbnails" to be only the item the cursor is hovering over instead of every thumbnail in an SAssetView.

This change removes the use of boolean attribute AreRealTimeThumbnailsAllowed. This attribute is less than ideal, design-wise, since it uses the thumbnail pool to drive the behavior of the thumbnails. Instead we build in default hover behavior to the thumbnail slate widgets, and implement plumbing for custom pathways to enable/disable real time behavior per thumbnail.

-> Add default OnMouseEnter/OnMouseLeave events on SAssetThumbnail for default hover behavior (over the actual thumbnail image), which supports detail panel, customizations, and things like FoliageEditor and Niagara Stack Overview thumbnail strip.
-> Add thumbnail initialize pathway support for disabling default hover behavior for use cases where the logical asset is more than just the thumbnail image (border elements, text elements with the asset name, etc).
-> Add custom OnMouseEnter/OnMouseLeave handlers on SAssetTileItem, SAssetListItem to perform real time behavior on entire logical asset item in Content Browser and Asset Picker.

diff --git a/Engine/Plugins/Editor/AssetSearch/Source/Private/Widgets/SSearchBrowser.cpp b/Engine/Plugins/Editor/AssetSearch/Source/Private/Widgets/SSearchBrowser.cpp
--- a/Engine/Plugins/Editor/AssetSearch/Source/Private/Widgets/SSearchBrowser.cpp
+++ b/Engine/Plugins/Editor/AssetSearch/Source/Private/Widgets/SSearchBrowser.cpp
@@ -75,8 +75,7 @@ void SSearchBrowser::Construct( const FArguments& InArgs )
 	SortByColumn = SSearchTreeRow::NAME_ColumnName;
 	SortMode = EColumnSortMode::Ascending;
 
-	const bool bAreRealTimeThumbnailsAllowed = false;
-	ThumbnailPool = MakeShared<FAssetThumbnailPool>(AssetSearchConstants::ThumbnailPoolSize, bAreRealTimeThumbnailsAllowed);
+	ThumbnailPool = MakeShared<FAssetThumbnailPool>(AssetSearchConstants::ThumbnailPoolSize);
 
 	AssetRegistry = &FModuleManager::LoadModuleChecked<FAssetRegistryModule>("AssetRegistry").Get();
 
diff --git a/Engine/Plugins/FX/Niagara/Source/NiagaraEditorWidgets/Private/SNiagaraOverviewStackNode.cpp b/Engine/Plugins/FX/Niagara/Source/NiagaraEditorWidgets/Private/SNiagaraOverviewStackNode.cpp
--- a/Engine/Plugins/FX/Niagara/Source/NiagaraEditorWidgets/Private/SNiagaraOverviewStackNode.cpp
+++ b/Engine/Plugins/FX/Niagara/Source/NiagaraEditorWidgets/Private/SNiagaraOverviewStackNode.cpp
@@ -42,7 +42,7 @@ void SNiagaraOverviewStackNode::Construct(const FArguments& InArgs, UNiagaraOver
 	CurrentIssueIndex = -1;
 
 	EmitterHandleViewModelWeak.Reset();
-	ThumbnailPool = MakeShared<FAssetThumbnailPool>(10, TAttribute<bool>::Create(TAttribute<bool>::FGetter::CreateSP(this, &SNiagaraOverviewStackNode::IsHoveringThumbnail)));
+	ThumbnailPool = MakeShared<FAssetThumbnailPool>(10);
 	ThumbnailBar = SNew(SHorizontalBox);
 	if (OverviewStackNode->GetOwningSystem() != nullptr)
 	{
diff --git a/Engine/Source/Editor/ContentBrowser/Private/AssetViewWidgets.cpp b/Engine/Source/Editor/ContentBrowser/Private/AssetViewWidgets.cpp
--- a/Engine/Source/Editor/ContentBrowser/Private/AssetViewWidgets.cpp
+++ b/Engine/Source/Editor/ContentBrowser/Private/AssetViewWidgets.cpp
@@ -34,7 +34,6 @@
 #include "ContentBrowserLog.h"
 #include "ObjectTools.h"
 #include "AssetThumbnail.h"
-#include "Settings/ContentBrowserSettings.h"
 #include "ContentBrowserModule.h"
 #include "ContentBrowserUtils.h"
 #include "ContentBrowserDataSource.h"
@@ -1558,6 +1557,26 @@ void SAssetListItem::OnAssetDataChanged()
 	}
 }
 
+void SAssetListItem::OnMouseEnter(const FGeometry& MyGeometry, const FPointerEvent& MouseEvent)
+{
+	SAssetViewItem::OnMouseEnter(MyGeometry, MouseEvent);
+
+	if (AssetThumbnail.IsValid())
+	{
+		AssetThumbnail->SetRealTime(true);
+	}
+}
+
+void SAssetListItem::OnMouseLeave(const FPointerEvent& MouseEvent)
+{
+	SAssetViewItem::OnMouseLeave(MouseEvent);
+
+	if (AssetThumbnail.IsValid())
+	{
+		AssetThumbnail->SetRealTime(false);
+	}
+}
+
 float SAssetListItem::GetExtraStateIconWidth() const
 {
 	return GetStateIconImageSize().Get();
@@ -1616,6 +1635,7 @@ void SAssetTileItem::Construct( const FArguments& InArgs )
 		FAssetThumbnailConfig ThumbnailConfig;
 		ThumbnailConfig.bAllowFadeIn = true;
 		ThumbnailConfig.bAllowHintText = InArgs._AllowThumbnailHintLabel;
+		ThumbnailConfig.bAllowRealTimeOnHovered = false; // we use our own OnMouseEnter/Leave for logical asset item
 		ThumbnailConfig.bForceGenericThumbnail = AssetItem->GetItem().GetItemTemporaryReason() == EContentBrowserItemFlags::Temporary_Creation;
 		ThumbnailConfig.bAllowAssetSpecificThumbnailOverlay = !ThumbnailConfig.bForceGenericThumbnail;
 		ThumbnailConfig.ThumbnailLabel = InArgs._ThumbnailLabel;
@@ -1753,6 +1773,26 @@ void SAssetTileItem::OnAssetDataChanged()
 	}
 }
 
+void SAssetTileItem::OnMouseEnter(const FGeometry& MyGeometry, const FPointerEvent& MouseEvent)
+{
+	SAssetViewItem::OnMouseEnter(MyGeometry, MouseEvent);
+
+	if (AssetThumbnail.IsValid())
+	{
+		AssetThumbnail->SetRealTime( true );
+	}
+}
+
+void SAssetTileItem::OnMouseLeave(const FPointerEvent& MouseEvent)
+{
+	SAssetViewItem::OnMouseLeave(MouseEvent);
+
+	if (AssetThumbnail.IsValid())
+	{
+		AssetThumbnail->SetRealTime( false );
+	}
+}
+
 const FSlateBrush* SAssetTileItem::GetBorderImage() const
 {
 	const bool bIsSelected = IsSelected.IsBound() ? IsSelected.Execute() : false;
diff --git a/Engine/Source/Editor/ContentBrowser/Private/AssetViewWidgets.h b/Engine/Source/Editor/ContentBrowser/Private/AssetViewWidgets.h
--- a/Engine/Source/Editor/ContentBrowser/Private/AssetViewWidgets.h
+++ b/Engine/Source/Editor/ContentBrowser/Private/AssetViewWidgets.h
@@ -413,6 +413,12 @@ public:
 	/** Handles committing a name change */
 	virtual void OnAssetDataChanged() override;
 
+	/** Handles realtime thumbnails */
+	virtual void OnMouseEnter(const FGeometry& MyGeometry, const FPointerEvent& MouseEvent) override;
+
+	/** Handles realtime thumbnails */
+	virtual void OnMouseLeave(const FPointerEvent& MouseEvent) override;
+
 	/** Whether the widget should allow primitive tools to be displayed */
 	bool CanDisplayPrimitiveTools() const { return false; }
 
@@ -524,6 +530,12 @@ public:
 	/** Handles committing a name change */
 	virtual void OnAssetDataChanged() override;
 
+	/** Handles realtime thumbnails */
+	virtual void OnMouseEnter(const FGeometry& MyGeometry, const FPointerEvent& MouseEvent) override;
+
+	/** Handles realtime thumbnails */
+	virtual void OnMouseLeave(const FPointerEvent& MouseEvent) override;
+
 	/** Whether the widget should allow primitive tools to be displayed */
 	bool CanDisplayPrimitiveTools() const { return true; }
 
diff --git a/Engine/Source/Editor/ContentBrowser/Private/SAssetPicker.cpp b/Engine/Source/Editor/ContentBrowser/Private/SAssetPicker.cpp
--- a/Engine/Source/Editor/ContentBrowser/Private/SAssetPicker.cpp
+++ b/Engine/Source/Editor/ContentBrowser/Private/SAssetPicker.cpp
@@ -298,7 +298,6 @@ void SAssetPicker::Construct( const FArguments& InArgs )
 		.OnGetCustomAssetToolTip(InArgs._AssetPickerConfig.OnGetCustomAssetToolTip)
 		.OnVisualizeAssetToolTip(InArgs._AssetPickerConfig.OnVisualizeAssetToolTip)
 		.OnAssetToolTipClosing(InArgs._AssetPickerConfig.OnAssetToolTipClosing)
-		.AreRealTimeThumbnailsAllowed(this, &SAssetPicker::IsHovered)
 		.FrontendFilters(FrontendFilters)
 		.InitialSourcesData(CurrentSourcesData)
 		.InitialBackendFilter(CurrentBackendFilter)
diff --git a/Engine/Source/Editor/ContentBrowser/Private/SAssetView.cpp b/Engine/Source/Editor/ContentBrowser/Private/SAssetView.cpp
--- a/Engine/Source/Editor/ContentBrowser/Private/SAssetView.cpp
+++ b/Engine/Source/Editor/ContentBrowser/Private/SAssetView.cpp
@@ -219,7 +219,7 @@ void SAssetView::Construct( const FArguments& InArgs )
 	ThumbnailScaleRangeScalar = ( DisplaySize.Y / 2160 );
 
 	// Create a thumbnail pool for rendering thumbnails	
-	AssetThumbnailPool = MakeShareable( new FAssetThumbnailPool(1024, InArgs._AreRealTimeThumbnailsAllowed) );
+	AssetThumbnailPool = MakeShared<FAssetThumbnailPool>(1024);
 	NumOffscreenThumbnails = 64;
 	ListViewThumbnailResolution = 128;
 	ListViewThumbnailSize = 64;
diff --git a/Engine/Source/Editor/ContentBrowser/Private/SAssetView.h b/Engine/Source/Editor/ContentBrowser/Private/SAssetView.h
--- a/Engine/Source/Editor/ContentBrowser/Private/SAssetView.h
+++ b/Engine/Source/Editor/ContentBrowser/Private/SAssetView.h
@@ -63,7 +63,6 @@ class SAssetView : public SCompoundWidget
 public:
 	SLATE_BEGIN_ARGS( SAssetView )
 		: _InitialCategoryFilter(EContentBrowserItemCategoryFilter::IncludeAll)
-		, _AreRealTimeThumbnailsAllowed(true)
 		, _ThumbnailLabel( EThumbnailLabel::ClassName )
 		, _AllowThumbnailHintLabel(true)
 		, _InitialViewType(EAssetViewType::Tile)
@@ -125,9 +124,6 @@ public:
 		/** The warning text to display when there are no assets to show */
 		SLATE_ATTRIBUTE( FText, AssetShowWarningText )
 
-		/** Attribute to determine if real-time thumbnails should be used */
-		SLATE_ATTRIBUTE( bool, AreRealTimeThumbnailsAllowed )
-
 		/** Attribute to determine what text should be highlighted */
 		SLATE_ATTRIBUTE( FText, HighlightedText )
 
diff --git a/Engine/Source/Editor/ContentBrowser/Private/SContentBrowser.cpp b/Engine/Source/Editor/ContentBrowser/Private/SContentBrowser.cpp
--- a/Engine/Source/Editor/ContentBrowser/Private/SContentBrowser.cpp
+++ b/Engine/Source/Editor/ContentBrowser/Private/SContentBrowser.cpp
@@ -172,7 +172,6 @@ void SContentBrowser::Construct( const FArguments& InArgs, const FName& InInstan
 		.OnItemsActivated(this, &SContentBrowser::OnItemsActivated)
 		.OnGetItemContextMenu(this, &SContentBrowser::GetItemContextMenu, EContentBrowserViewContext::AssetView)
 		.OnItemRenameCommitted(this, &SContentBrowser::OnItemRenameCommitted)
-		.AreRealTimeThumbnailsAllowed(this, &SContentBrowser::IsHovered)
 		.FrontendFilters(FrontendFilters)
 		.HighlightedText(this, &SContentBrowser::GetHighlightedText)
 		.ShowBottomToolbar(bShowBottomToolbar)
diff --git a/Engine/Source/Editor/FoliageEdit/Private/SFoliagePalette.cpp b/Engine/Source/Editor/FoliageEdit/Private/SFoliagePalette.cpp
--- a/Engine/Source/Editor/FoliageEdit/Private/SFoliagePalette.cpp
+++ b/Engine/Source/Editor/FoliageEdit/Private/SFoliagePalette.cpp
@@ -179,7 +179,7 @@ void SFoliagePalette::Construct(const FArguments& InArgs)
 
 	// Size of the thumbnail pool should be large enough to show a reasonable amount of foliage assets on screen at once,
 	// otherwise some thumbnail images will appear duplicated.
-	ThumbnailPool = MakeShareable(new FAssetThumbnailPool(64, TAttribute<bool>::Create(TAttribute<bool>::FGetter::CreateSP(this, &SFoliagePalette::IsHovered))));
+	ThumbnailPool = MakeShared<FAssetThumbnailPool>(64);
 
 	TypeFilter = MakeShareable(new FoliageTypeTextFilter(
 		FoliageTypeTextFilter::FItemToStringArray::CreateSP(this, &SFoliagePalette::GetPaletteItemFilterString)));
@@ -466,7 +466,7 @@ void SFoliagePalette::UpdateThumbnailForType(UFoliageType* FoliageType)
 		{
 			const bool bItemIsSelected = GetActiveViewWidget()->IsItemSelected(Item);
 
-			Item = MakeShareable(new FFoliagePaletteItemModel(Item->GetTypeUIInfo(), SharedThis(this), ThumbnailPool, FoliageEditMode));
+			Item = MakeShared<FFoliagePaletteItemModel>(Item->GetTypeUIInfo(), SharedThis(this), ThumbnailPool, FoliageEditMode);
 			if (bItemIsSelected)
 			{
 				GetActiveViewWidget()->SetItemSelection(Item, true);
diff --git a/Engine/Source/Editor/LevelEditor/Private/SLevelEditor.cpp b/Engine/Source/Editor/LevelEditor/Private/SLevelEditor.cpp
--- a/Engine/Source/Editor/LevelEditor/Private/SLevelEditor.cpp
+++ b/Engine/Source/Editor/LevelEditor/Private/SLevelEditor.cpp
@@ -84,8 +84,7 @@ SLevelEditor::SLevelEditor()
 	: World(nullptr)
 	, bNeedsRefresh(false)
 {
-	const bool bAreRealTimeThumbnailsAllowed = false;
-	ThumbnailPool = MakeShareable(new FAssetThumbnailPool(LevelEditorConstants::ThumbnailPoolSize, bAreRealTimeThumbnailsAllowed));
+	ThumbnailPool = MakeShared<FAssetThumbnailPool>(LevelEditorConstants::ThumbnailPoolSize);
 }
 
 void SLevelEditor::BindCommands()
diff --git a/Engine/Source/Editor/PropertyEditor/Private/SDetailsViewBase.cpp b/Engine/Source/Editor/PropertyEditor/Private/SDetailsViewBase.cpp
--- a/Engine/Source/Editor/PropertyEditor/Private/SDetailsViewBase.cpp
+++ b/Engine/Source/Editor/PropertyEditor/Private/SDetailsViewBase.cpp
@@ -500,7 +500,7 @@ TSharedPtr<FAssetThumbnailPool> SDetailsViewBase::GetThumbnailPool() const
 	if (!ThumbnailPool.IsValid())
 	{
 		// Create a thumbnail pool for the view if it doesnt exist.  This does not use resources of no thumbnails are used
-		ThumbnailPool = MakeShareable(new FAssetThumbnailPool(50, TAttribute<bool>::Create(TAttribute<bool>::FGetter::CreateSP(this, &SDetailsView::IsHovered))));
+		ThumbnailPool = MakeShared<FAssetThumbnailPool>(50);
 	}
 
 	return ThumbnailPool;
diff --git a/Engine/Source/Editor/UnrealEd/Private/AssetThumbnail.cpp b/Engine/Source/Editor/UnrealEd/Private/AssetThumbnail.cpp
--- a/Engine/Source/Editor/UnrealEd/Private/AssetThumbnail.cpp
+++ b/Engine/Source/Editor/UnrealEd/Private/AssetThumbnail.cpp
@@ -42,11 +42,12 @@ class SAssetThumbnail : public SCompoundWidget
 public:
 	SLATE_BEGIN_ARGS( SAssetThumbnail )
 		: _Style("AssetThumbnail")
-		, _ThumbnailPool(NULL)
+		, _ThumbnailPool(nullptr)
 		, _AllowFadeIn(false)
 		, _ForceGenericThumbnail(false)
 		, _AllowHintText(true)
 		, _AllowAssetSpecificThumbnailOverlay(false)
+		, _AllowRealTimeOnHovered(true)
 		, _Label(EThumbnailLabel::ClassName)
 		, _HighlightedText(FText::GetEmpty())
 		, _HintColorAndOpacity(FLinearColor(0.0f, 0.0f, 0.0f, 0.0f))
@@ -62,6 +63,7 @@ public:
 		SLATE_ARGUMENT(bool, ForceGenericThumbnail)
 		SLATE_ARGUMENT(bool, AllowHintText)
 		SLATE_ARGUMENT(bool, AllowAssetSpecificThumbnailOverlay)
+		SLATE_ARGUMENT(bool, AllowRealTimeOnHovered)
 		SLATE_ARGUMENT(EThumbnailLabel::Type, Label)
 		SLATE_ATTRIBUTE(FText, HighlightedText)
 		SLATE_ATTRIBUTE(FLinearColor, HintColorAndOpacity)
@@ -79,6 +81,7 @@ public:
 		Label = InArgs._Label;
 		HintColorAndOpacity = InArgs._HintColorAndOpacity;
 		bAllowHintText = InArgs._AllowHintText;
+		bAllowRealTimeOnHovered = InArgs._AllowRealTimeOnHovered;
 		ThumbnailBrush = nullptr;
 		ClassIconBrush = nullptr;
 		AssetThumbnail = InArgs._AssetThumbnail;
@@ -280,15 +283,23 @@ public:
 		return FSlateColor( FLinearColor( Color.R, Color.G, Color.B, FMath::Lerp( 0.0f, 0.5f, Color.A ) ) );
 	}
 
-	// SWidget implementation
 	virtual void OnMouseEnter( const FGeometry& MyGeometry, const FPointerEvent& MouseEvent ) override
 	{
-		SCompoundWidget::OnMouseEnter(MyGeometry, MouseEvent);
+		SCompoundWidget::OnMouseEnter( MyGeometry, MouseEvent );
 
-		if (!GetDefault<UContentBrowserSettings>()->RealTimeThumbnails )
+		if ( bAllowRealTimeOnHovered )
 		{
-			// Update hovered thumbnails if we are not already updating them in real-time
-			AssetThumbnail->RefreshThumbnail();
+			AssetThumbnail->SetRealTime( true );
+		}
+	}
+	
+	virtual void OnMouseLeave( const FPointerEvent& MouseEvent ) override
+	{
+		SCompoundWidget::OnMouseLeave( MouseEvent );
+
+		if ( bAllowRealTimeOnHovered )
+		{
+			AssetThumbnail->SetRealTime( false );
 		}
 	}
 
@@ -684,6 +695,7 @@ private:
 
 	TAttribute< FLinearColor > HintColorAndOpacity;
 	bool bAllowHintText;
+	bool bAllowRealTimeOnHovered;
 
 	/** The name of the thumbnail which should be used instead of the class thumbnail. */
 	FName ClassThumbnailBrushOverride;
@@ -813,6 +825,7 @@ TSharedRef<SWidget> FAssetThumbnail::MakeThumbnailWidget( const FAssetThumbnailC
 		.HighlightedText(InConfig.HighlightedText)
 		.HintColorAndOpacity(InConfig.HintColorAndOpacity)
 		.AllowHintText(InConfig.bAllowHintText)
+		.AllowRealTimeOnHovered(InConfig.bAllowRealTimeOnHovered)
 		.ClassThumbnailBrushOverride(InConfig.ClassThumbnailBrushOverride)
 		.AllowAssetSpecificThumbnailOverlay(InConfig.bAllowAssetSpecificThumbnailOverlay)
 		.AssetTypeColorOverride(InConfig.AssetTypeColorOverride)
@@ -827,9 +840,16 @@ void FAssetThumbnail::RefreshThumbnail()
 	}
 }
 
-FAssetThumbnailPool::FAssetThumbnailPool( uint32 InNumInPool, const TAttribute<bool>& InAreRealTimeThumbnailsAllowed, double InMaxFrameTimeAllowance, uint32 InMaxRealTimeThumbnailsPerFrame )
-	: AreRealTimeThumbnailsAllowed( InAreRealTimeThumbnailsAllowed )
-	, NumInPool( InNumInPool )
+void FAssetThumbnail::SetRealTime(bool bRealTime)
+{
+	if (ThumbnailPool.IsValid() && AssetData.IsValid())
+	{
+		ThumbnailPool.Pin()->SetRealTimeThumbnail( SharedThis(this), bRealTime );
+	}
+}
+
+FAssetThumbnailPool::FAssetThumbnailPool( uint32 InNumInPool, double InMaxFrameTimeAllowance, uint32 InMaxRealTimeThumbnailsPerFrame )
+	: NumInPool( InNumInPool )
 	, MaxRealTimeThumbnailsPerFrame( InMaxRealTimeThumbnailsPerFrame )
 	, MaxFrameTimeAllowance( InMaxFrameTimeAllowance )
 {
@@ -948,9 +968,9 @@ void FAssetThumbnailPool::Tick( float DeltaTime )
 		RecentlyLoadedAssets.Empty();
 	}
 
-	// If we have dynamic thumbnails are we are done rendering the last batch of dynamic thumbnails, start a new batch as long as real-time thumbnails are enabled
+	// If we have dynamic thumbnails and we are done rendering the last batch of dynamic thumbnails, start a new batch as long as real-time thumbnails are enabled
 	const bool bIsInPIEOrSimulate = GEditor->PlayWorld != NULL || GEditor->bIsSimulatingInEditor;
-	const bool bShouldUseRealtimeThumbnails = AreRealTimeThumbnailsAllowed.Get() && GetDefault<UContentBrowserSettings>()->RealTimeThumbnails && !bIsInPIEOrSimulate;
+	const bool bShouldUseRealtimeThumbnails = GetDefault<UContentBrowserSettings>()->RealTimeThumbnails && !bIsInPIEOrSimulate;
 	if ( bShouldUseRealtimeThumbnails && RealTimeThumbnails.Num() > 0 && RealTimeThumbnailsToRender.Num() == 0 )
 	{
 		double CurrentTime = FPlatformTime::Seconds();
@@ -1037,9 +1057,6 @@ void FAssetThumbnailPool::Tick( float DeltaTime )
 							}
 
 							bLoadedThumbnail = true;
-
-							// Since this was rendered, add it to the list of thumbnails that can be rendered in real-time
-							RealTimeThumbnails.AddUnique(InfoRef);
 						}
 					}
 				
@@ -1348,7 +1365,7 @@ void FAssetThumbnailPool::RefreshThumbnail( const TSharedPtr<FAssetThumbnail>& T
 	const uint32 Width = ThumbnailToRefresh->GetSize().X;
 	const uint32 Height = ThumbnailToRefresh->GetSize().Y;
 
-	if ( ensure(AssetData.ObjectPath != NAME_None) && ensure(Width > 0) && ensure(Height > 0) )
+	if ( ensureAlways(AssetData.ObjectPath != NAME_None) && ensureAlways(Width > 0) && ensureAlways(Height > 0) )
 	{
 		FThumbId ThumbId( AssetData.ObjectPath, Width, Height ) ;
 		const TSharedRef<FThumbnailInfo>* ThumbnailInfoPtr = ThumbnailToTextureMap.Find( ThumbId );
@@ -1359,6 +1376,30 @@ void FAssetThumbnailPool::RefreshThumbnail( const TSharedPtr<FAssetThumbnail>& T
 	}
 }
 
+void FAssetThumbnailPool::SetRealTimeThumbnail( const TSharedPtr<FAssetThumbnail>& Thumbnail, bool bRealTimeThumbnail )
+{
+	const FAssetData& AssetData = Thumbnail->GetAssetData();
+	const uint32 Width = Thumbnail->GetSize().X;
+	const uint32 Height = Thumbnail->GetSize().Y;
+
+	if ( ensureAlways(AssetData.ObjectPath != NAME_None) && ensureAlways(Width > 0) && ensureAlways(Height > 0) )
+	{
+		FThumbId ThumbId( AssetData.ObjectPath, Width, Height );
+		const TSharedRef<FThumbnailInfo>* ThumbnailInfoPtr = ThumbnailToTextureMap.Find( ThumbId );
+		if ( ThumbnailInfoPtr )
+		{
+			if ( bRealTimeThumbnail )
+			{
+				RealTimeThumbnails.AddUnique(*ThumbnailInfoPtr);
+			}
+			else
+			{
+				RealTimeThumbnails.Remove(*ThumbnailInfoPtr);
+			}
+		}
+	}
+}
+
 void FAssetThumbnailPool::FreeThumbnail( const FName& ObjectPath, uint32 Width, uint32 Height )
 {
 	if(ObjectPath != NAME_None && Width != 0 && Height != 0)
diff --git a/Engine/Source/Editor/UnrealEd/Public/AssetThumbnail.h b/Engine/Source/Editor/UnrealEd/Public/AssetThumbnail.h
--- a/Engine/Source/Editor/UnrealEd/Public/AssetThumbnail.h
+++ b/Engine/Source/Editor/UnrealEd/Public/AssetThumbnail.h
@@ -35,6 +35,7 @@ struct FAssetThumbnailConfig
 		: bAllowFadeIn( false )
 		, bForceGenericThumbnail( false )
 		, bAllowHintText( true )
+		, bAllowRealTimeOnHovered( true )
 		, bAllowAssetSpecificThumbnailOverlay( false )
 		, ClassThumbnailBrushOverride( NAME_None )
 		, ThumbnailLabel( EThumbnailLabel::ClassName )
@@ -48,6 +49,7 @@ struct FAssetThumbnailConfig
 	bool bAllowFadeIn;
 	bool bForceGenericThumbnail;
 	bool bAllowHintText;
+	bool bAllowRealTimeOnHovered;
 	bool bAllowAssetSpecificThumbnailOverlay;
 	FName ClassThumbnailBrushOverride;
 	EThumbnailLabel::Type ThumbnailLabel;
@@ -123,6 +125,9 @@ public:
 	/** Re-renders this thumbnail */
 	UNREALED_API void RefreshThumbnail();
 
+	/** Updates if this thumbnail should be realtime rendered via the pool */
+	UNREALED_API void SetRealTime( bool bRealTime );
+
 	DECLARE_EVENT(FAssetThumbnail, FOnAssetDataChanged);
 	FOnAssetDataChanged& OnAssetDataChanged() { return AssetDataChangedEvent; }
 
@@ -154,7 +159,7 @@ public:
 	 * @param InMaxFrameTimeAllowance			The maximum number of seconds per tick to spend rendering thumbnails
 	 * @param InMaxRealTimeThumbnailsPerFrame	The maximum number of real-time thumbnails to render per tick
 	 */
-	UNREALED_API FAssetThumbnailPool( uint32 InNumInPool, const TAttribute<bool>& InAreRealTimeThumbnailsAllowed = true, double InMaxFrameTimeAllowance = 0.005, uint32 InMaxRealTimeThumbnailsPerFrame = 3 );
+	UNREALED_API FAssetThumbnailPool( uint32 InNumInPool, double InMaxFrameTimeAllowance = 0.005, uint32 InMaxRealTimeThumbnailsPerFrame = 3 );
 
 	/** Destructor to free all remaining resources */
 	UNREALED_API ~FAssetThumbnailPool();
@@ -211,6 +216,9 @@ public:
 	/** Re-renders the specified thumbnail */
 	UNREALED_API void RefreshThumbnail( const TSharedPtr<FAssetThumbnail>& ThumbnailToRefresh );
 
+	/** Enables/disables realtime thumbnail behavior */
+	UNREALED_API void SetRealTimeThumbnail(const TSharedPtr<FAssetThumbnail>& Thumbnail, bool bRealTimeThumbnail);
+
 private:
 
 	/**
