Fixing preview editor level to not autogen shaders when the platform is None.

diff --git a/Engine/Source/Editor/UnrealEd/Private/EditorEngine.cpp b/Engine/Source/Editor/UnrealEd/Private/EditorEngine.cpp
--- a/Engine/Source/Editor/UnrealEd/Private/EditorEngine.cpp
+++ b/Engine/Source/Editor/UnrealEd/Private/EditorEngine.cpp
@@ -7668,6 +7668,7 @@ void UEditorEngine::SetPreviewPlatform(const FPreviewPlatformInfo& NewPreviewPla
 		NewPreviewPlatform.bPreviewFeatureLevelActive != PreviewPlatform.bPreviewFeatureLevelActive;
 	const ERHIFeatureLevel::Type EffectiveFeatureLevel = NewPreviewPlatform.GetEffectivePreviewFeatureLevel();
 
+	if (NewPreviewPlatform.PreviewPlatformName != NAME_None)
 	{
 		// Force generation of the autogen files if they don't already exist. We'll ignore the returned data.
 		FGBufferInfo GBufferInfo = FShaderCompileUtilities::FetchGBufferInfoAndWriteAutogen(ShaderPlatform, MaxFeatureLevel);
