Make sure the Interchange manager is destroy in the game thread and not in the main thread.

diff --git a/Engine/Source/Editor/UnrealEd/Private/Subsystems/ImportSubsystem.cpp b/Engine/Source/Editor/UnrealEd/Private/Subsystems/ImportSubsystem.cpp
--- a/Engine/Source/Editor/UnrealEd/Private/Subsystems/ImportSubsystem.cpp
+++ b/Engine/Source/Editor/UnrealEd/Private/Subsystems/ImportSubsystem.cpp
@@ -129,9 +129,9 @@ void UImportSubsystem::Initialize(FSubsystemCollectionBase& Collection)
 		FCoreDelegates::OnPostEngineInit.AddLambda(AttachDelegates);
 	}
 
-	//Since subsystem deinitialize after the FCoreDelegates::OnExit and interchange manager is destroy during FCoreDelegates::OnExit
-	//We want to unregister delegate using the FCoreDelegates::OnPreExit.
-	FCoreDelegates::OnPreExit.AddLambda([this]()
+	//Unregister before the Interchange manager get destroy
+	UInterchangeManager& InterchangeManager = UInterchangeManager::GetInterchangeManager();
+	InterchangeManager.OnPreDestroyInterchangeManager.AddLambda([this]()
 	{
 		//Unhook interchange manager import/reimport delegates
 		UInterchangeManager& InterchangeManager = UInterchangeManager::GetInterchangeManager();
diff --git a/Engine/Source/Runtime/Interchange/Engine/Private/InterchangeManager.cpp b/Engine/Source/Runtime/Interchange/Engine/Private/InterchangeManager.cpp
--- a/Engine/Source/Runtime/Interchange/Engine/Private/InterchangeManager.cpp
+++ b/Engine/Source/Runtime/Interchange/Engine/Private/InterchangeManager.cpp
@@ -272,13 +272,9 @@ UInterchangeManager& UInterchangeManager::GetInterchangeManager()
 			{
 				InterchangeManager->CancelAllTasksSynchronously();
 			}
-		});
-
-		//We release the singleton here so all module was able to unhook there delegates
-		FCoreDelegates::OnExit.AddLambda([]()
-		{
 			//Task should have been cancel in the Engine pre exit callback
 			ensure(InterchangeManager->ImportTasks.Num() == 0);
+			InterchangeManager->OnPreDestroyInterchangeManager.Broadcast();
 			//Release the InterchangeManager object
 			InterchangeManager.Reset();
 			InterchangeManagerScopeOfLifeEnded = true;
diff --git a/Engine/Source/Runtime/Interchange/Engine/Public/InterchangeManager.h b/Engine/Source/Runtime/Interchange/Engine/Public/InterchangeManager.h
--- a/Engine/Source/Runtime/Interchange/Engine/Public/InterchangeManager.h
+++ b/Engine/Source/Runtime/Interchange/Engine/Public/InterchangeManager.h
@@ -190,6 +190,8 @@ public:
 
 	FInterchangeOnAssetPostImport OnAssetPostImport;
 	FInterchangeOnAssetPostReimport OnAssetPostReimport;
+	// Called when before the application is exiting.
+	FSimpleMulticastDelegate OnPreDestroyInterchangeManager;
 
 	/**
 	 * Any translator must register to the manager
