Fix MoveActorsToLevel failed Clipboard Paste

diff --git a/Engine/Source/Editor/UnrealEd/Classes/Editor/EditorEngine.h b/Engine/Source/Editor/UnrealEd/Classes/Editor/EditorEngine.h
--- a/Engine/Source/Editor/UnrealEd/Classes/Editor/EditorEngine.h
+++ b/Engine/Source/Editor/UnrealEd/Classes/Editor/EditorEngine.h
@@ -2035,8 +2035,9 @@ public:
 	 * @param bShouldCut			If true, deletes the selected actors after copying them to the clipboard
 	 * @param bIsMove				If true, this cut is part of a move and the actors will be immediately pasted
 	 * @param bWarnAboutReferences	Whether or not to show a modal warning about referenced actors that may no longer function after being moved
+	 * @param DestinationData		Data that was copied to clipboard 
 	 */
-	void CopySelectedActorsToClipboard( UWorld* InWorld, const bool bShouldCut, const bool bIsMove = false, bool bWarnAboutReferences = true);
+	void CopySelectedActorsToClipboard( UWorld* InWorld, const bool bShouldCut, const bool bIsMove = false, bool bWarnAboutReferences = true, FString* DestinationData = nullptr);
 
 	/**
 	 * Checks to see whether it's possible to perform a paste operation.
diff --git a/Engine/Source/Editor/UnrealEd/Private/EditorLevelUtils.cpp b/Engine/Source/Editor/UnrealEd/Private/EditorLevelUtils.cpp
--- a/Engine/Source/Editor/UnrealEd/Private/EditorLevelUtils.cpp
+++ b/Engine/Source/Editor/UnrealEd/Private/EditorLevelUtils.cpp
@@ -147,7 +147,8 @@ int32 UEditorLevelUtils::MoveActorsToLevel(const TArray<AActor*>& ActorsToMove,
 				// We are moving the actors so cut them to remove them from the existing level
 				const bool bShoudCut = true;
 				const bool bIsMove = true;
-				GEditor->CopySelectedActorsToClipboard(OwningWorld, bShoudCut, bIsMove, bWarnAboutReferences);
+				FString DestinationData;
+				GEditor->CopySelectedActorsToClipboard(OwningWorld, bShoudCut, bIsMove, bWarnAboutReferences, &DestinationData);
 
 				const bool bLevelVisible = DestLevel->bIsVisible;
 				if (!bLevelVisible)
@@ -164,7 +165,7 @@ int32 UEditorLevelUtils::MoveActorsToLevel(const TArray<AActor*>& ActorsToMove,
 					const bool bDuplicate = false;
 					const bool bOffsetLocations = false;
 					const bool bWarnIfHidden = false;
-					GEditor->edactPasteSelected(OwningWorld, bDuplicate, bOffsetLocations, bWarnIfHidden);
+					GEditor->edactPasteSelected(OwningWorld, bDuplicate, bOffsetLocations, bWarnIfHidden, &DestinationData);
 
 					// Restore the original current level
 					OwningWorld->SetCurrentLevel(OldCurrentLevel);
diff --git a/Engine/Source/Editor/UnrealEd/Private/EditorServer.cpp b/Engine/Source/Editor/UnrealEd/Private/EditorServer.cpp
--- a/Engine/Source/Editor/UnrealEd/Private/EditorServer.cpp
+++ b/Engine/Source/Editor/UnrealEd/Private/EditorServer.cpp
@@ -3347,7 +3347,7 @@ bool UEditorEngine::CanCopySelectedActorsToClipboard( UWorld* InWorld, FCopySele
 	return false;
 }
 
-void UEditorEngine::CopySelectedActorsToClipboard( UWorld* InWorld, bool bShouldCut, bool bIsMove, bool bWarnAboutReferences )
+void UEditorEngine::CopySelectedActorsToClipboard( UWorld* InWorld, bool bShouldCut, bool bIsMove, bool bWarnAboutReferences, FString* DestinationData )
 {
 	FCopySelectedInfo CopySelected;
 	if ( !CanCopySelectedActorsToClipboard( InWorld, &CopySelected ) )
@@ -3383,7 +3383,7 @@ void UEditorEngine::CopySelectedActorsToClipboard( UWorld* InWorld, bool bShould
 		{
 			// Cut!
 			const FScopedTransaction Transaction( NSLOCTEXT("UnrealEd", "Cut", "Cut") );
-			edactCopySelected( World );
+			edactCopySelected( World, DestinationData );
 			const bool bVerifyDeletionCanHappen = true;
 
 			edactDeleteSelected( World, bVerifyDeletionCanHappen, bWarnAboutReferences, !bIsMove );
@@ -3391,7 +3391,7 @@ void UEditorEngine::CopySelectedActorsToClipboard( UWorld* InWorld, bool bShould
 		else
 		{
 			// Copy!
-			edactCopySelected( World );
+			edactCopySelected( World, DestinationData );
 		}
 
 		World->SetCurrentLevel( OldCurrentLevel );
diff --git a/Engine/Source/Runtime/ApplicationCore/Private/Windows/WindowsPlatformApplicationMisc.cpp b/Engine/Source/Runtime/ApplicationCore/Private/Windows/WindowsPlatformApplicationMisc.cpp
--- a/Engine/Source/Runtime/ApplicationCore/Private/Windows/WindowsPlatformApplicationMisc.cpp
+++ b/Engine/Source/Runtime/ApplicationCore/Private/Windows/WindowsPlatformApplicationMisc.cpp
@@ -377,6 +377,10 @@ void FWindowsPlatformApplicationMisc::ClipboardCopy(const TCHAR* Str)
 			UE_LOG(LogWindows, Fatal,TEXT("SetClipboardData failed with error code %i"), (uint32)GetLastError() );
 		verify(CloseClipboard());
 	}
+	else
+	{
+		UE_LOG(LogWindows, Error, TEXT("OpenClipboard failed with error code %i"), (uint32)GetLastError());
+	}
 }
 
 void FWindowsPlatformApplicationMisc::ClipboardPaste(class FString& Result)
@@ -420,6 +424,7 @@ void FWindowsPlatformApplicationMisc::ClipboardPaste(class FString& Result)
 	else 
 	{
 		Result=TEXT("");
+		UE_LOG(LogWindows, Error, TEXT("OpenClipboard failed with error code %i"), (uint32)GetLastError());
 	}
 }
 
