Submitted fix for importing FBX static meshes imported from Houdini. This enables a relaxed mode in the FBX SDK which will allow UVs and normals to be imported even if there are fewer polygon vertices than there are UVs or normals in the source data.

diff --git a/Engine/Source/Editor/UnrealEd/Private/Fbx/FbxMainImport.cpp b/Engine/Source/Editor/UnrealEd/Private/Fbx/FbxMainImport.cpp
--- a/Engine/Source/Editor/UnrealEd/Private/Fbx/FbxMainImport.cpp
+++ b/Engine/Source/Editor/UnrealEd/Private/Fbx/FbxMainImport.cpp
@@ -1026,8 +1026,13 @@ bool FFbxImporter::OpenFile(FString Filename)
 	// version of FBX SDK that you are using.
 	FbxManager::GetFileFormatVersion(SDKMajor, SDKMinor, SDKRevision);
 
+	if (SdkManager->GetIOSettings())
+	{
+		SdkManager->GetIOSettings()->SetBoolProp(IMP_RELAXED_FBX_CHECK, true);
+	}
+
 	// Initialize the importer by providing a filename.
-	const bool bImportStatus = Importer->Initialize(TCHAR_TO_UTF8(*Filename));
+	const bool bImportStatus = Importer->Initialize(TCHAR_TO_UTF8(*Filename), -1, SdkManager->GetIOSettings());
 	
 	FbxCreator = EFbxCreator::Unknow;
 	FbxIOFileHeaderInfo *FileHeaderInfo = Importer->GetFileHeaderInfo();
