Fix cut off text in fbx options window

diff --git a/Engine/Source/Editor/UnrealEd/Private/Fbx/FbxMainImport.cpp b/Engine/Source/Editor/UnrealEd/Private/Fbx/FbxMainImport.cpp
--- a/Engine/Source/Editor/UnrealEd/Private/Fbx/FbxMainImport.cpp
+++ b/Engine/Source/Editor/UnrealEd/Private/Fbx/FbxMainImport.cpp
@@ -239,7 +239,7 @@ FBXImportOptions* GetImportOptions( UnFbx::FFbxImporter* FbxImporter, UFbxImport
 		}
 
 		// Compute centered window position based on max window size, which include when all categories are expanded
-		const float FbxImportWindowWidth = 410.0f;
+		const float FbxImportWindowWidth = 450.0f;
 		const float FbxImportWindowHeight = 750.0f;
 		FVector2D FbxImportWindowSize = FVector2D(FbxImportWindowWidth, FbxImportWindowHeight); // Max window size it can get based on current slate
 
diff --git a/Engine/Source/Editor/UnrealEd/Private/Fbx/FbxOptionWindow.cpp b/Engine/Source/Editor/UnrealEd/Private/Fbx/FbxOptionWindow.cpp
--- a/Engine/Source/Editor/UnrealEd/Private/Fbx/FbxOptionWindow.cpp
+++ b/Engine/Source/Editor/UnrealEd/Private/Fbx/FbxOptionWindow.cpp
@@ -11,6 +11,7 @@
 #include "IDocumentation.h"
 #include "PropertyEditorModule.h"
 #include "IDetailsView.h"
+#include "Framework/Application/SlateApplication.h"
 
 #define LOCTEXT_NAMESPACE "FBXOption"
 
@@ -51,7 +52,6 @@ void SFbxOptionWindow::Construct(const FArguments& InArgs)
 					.AutoWidth()
 					[
 						SNew(STextBlock)
-						.Font(FEditorStyle::GetFontStyle("CurveEd.LabelFont"))
 						.Text(LOCTEXT("Import_CurrentFileTitle", "Current Asset: "))
 					]
 					+SHorizontalBox::Slot()
@@ -60,7 +60,6 @@ void SFbxOptionWindow::Construct(const FArguments& InArgs)
 					.VAlign(VAlign_Center)
 					[
 						SNew(STextBlock)
-						.Font(FEditorStyle::GetFontStyle("CurveEd.InfoFont"))
 						.Text(InArgs._FullPath)
 						.ToolTipText(InArgs._FullPath)
 					]
@@ -76,18 +75,14 @@ void SFbxOptionWindow::Construct(const FArguments& InArgs)
 			]
 			+ SVerticalBox::Slot()
 			.AutoHeight()
-			.HAlign(HAlign_Right)
 			.Padding(2)
 			[
 				SNew(SUniformGridPanel)
 				.SlotPadding(2)
-				+ SUniformGridPanel::Slot(0, 0)
-				[
-					IDocumentation::Get()->CreateAnchor(FString("Engine/Content/FBX/ImportOptions"))
-				]
 				+ SUniformGridPanel::Slot(1, 0)
 				[
-					SNew(SButton)
+					SAssignNew(ImportAllButton, SButton)
+					.ButtonStyle(FAppStyle::Get(), "PrimaryButton")
 					.HAlign(HAlign_Center)
 					.Text(LOCTEXT("FbxOptionWindow_ImportAll", "Import All"))
 					.ToolTipText(LOCTEXT("FbxOptionWindow_ImportAll_ToolTip", "Import all files with these same settings"))
@@ -96,7 +91,7 @@ void SFbxOptionWindow::Construct(const FArguments& InArgs)
 				]
 				+ SUniformGridPanel::Slot(2, 0)
 				[
-					SAssignNew(ImportButton, SButton)
+					SNew(SButton)
 					.HAlign(HAlign_Center)
 					.Text(LOCTEXT("FbxOptionWindow_Import", "Import"))
 					.IsEnabled(this, &SFbxOptionWindow::CanImport)
@@ -129,33 +124,50 @@ void SFbxOptionWindow::Construct(const FArguments& InArgs)
 		[
 			SNew(SHorizontalBox)
 			+ SHorizontalBox::Slot()
+			.AutoWidth()
 			.VAlign(VAlign_Center)
 			[
 				SNew(STextBlock)
 				.Text(this, &SFbxOptionWindow::GetImportTypeDisplayText)
 			]
+			+SHorizontalBox::Slot()
+			.VAlign(VAlign_Center)
+			.AutoWidth()
+			[
+				IDocumentation::Get()->CreateAnchor(FString("Engine/Content/FBX/ImportOptions"))
+			]
 			+ SHorizontalBox::Slot()
+			.HAlign(HAlign_Right)
 			[
-				SNew(SBox)
-				.HAlign(HAlign_Right)
+				SAssignNew(FbxHeaderButtons, SHorizontalBox)
+				+ SHorizontalBox::Slot()
+				.AutoWidth()
+				.Padding(FMargin(2.0f, 0.0f))
 				[
-					SAssignNew(FbxHeaderButtons, SHorizontalBox)
-					+ SHorizontalBox::Slot()
-					.AutoWidth()
-					.Padding(FMargin(2.0f, 0.0f))
-					[
-						SNew(SButton)
-						.Text(LOCTEXT("FbxOptionWindow_ResetOptions", "Reset to Default"))
-						.OnClicked(this, &SFbxOptionWindow::OnResetToDefaultClick)
-					]
+					SNew(SButton)
+					.Text(LOCTEXT("FbxOptionWindow_ResetOptions", "Reset to Default"))
+					.OnClicked(this, &SFbxOptionWindow::OnResetToDefaultClick)
 				]
 			]
 		]
 	);
 
 	DetailsView->SetObject(ImportUI);
+
+	RegisterActiveTimer(0.f, FWidgetActiveTimerDelegate::CreateSP(this, &SFbxOptionWindow::SetFocusPostConstruct));
+
 }
 
+EActiveTimerReturnType SFbxOptionWindow::SetFocusPostConstruct(double InCurrentTime, float InDeltaTime)
+{
+	if (ImportAllButton.IsValid())
+	{
+		FSlateApplication::Get().SetKeyboardFocus(ImportAllButton, EFocusCause::SetDirectly);
+	}
+
+	return EActiveTimerReturnType::Stop;
+
+}
 FReply SFbxOptionWindow::OnResetToDefaultClick() const
 {
 	ImportUI->ResetToDefault();
diff --git a/Engine/Source/Editor/UnrealEd/Private/FbxOptionWindow.h b/Engine/Source/Editor/UnrealEd/Private/FbxOptionWindow.h
--- a/Engine/Source/Editor/UnrealEd/Private/FbxOptionWindow.h
+++ b/Engine/Source/Editor/UnrealEd/Private/FbxOptionWindow.h
@@ -92,7 +92,7 @@ public:
 	{}
 		
 private:
-
+	EActiveTimerReturnType SetFocusPostConstruct(double InCurrentTime, float InDeltaTime);
 	bool CanImport() const;
 	FReply OnResetToDefaultClick() const;
 	FText GetImportTypeDisplayText() const;
@@ -101,7 +101,7 @@ private:
 	UFbxImportUI*	ImportUI;
 	TSharedPtr<class IDetailsView> DetailsView;
 	TWeakPtr< SWindow > WidgetWindow;
-	TSharedPtr< SButton > ImportButton;
+	TSharedPtr<SButton> ImportAllButton;
 	bool			bShouldImport;
 	bool			bShouldImportAll;
 	bool			bIsObjFormat;
