PR #7524: Fixed typo: HasMountWritePersmissionForPackage -> HasMountWritePermissionForPackage (Contributed by [REDACTED_PERSON])

diff --git a/Engine/Plugins/Enterprise/DatasmithImporter/Source/DatasmithImporter/Private/Utility/DatasmithImporterImpl.cpp b/Engine/Plugins/Enterprise/DatasmithImporter/Source/DatasmithImporter/Private/Utility/DatasmithImporterImpl.cpp
--- a/Engine/Plugins/Enterprise/DatasmithImporter/Source/DatasmithImporter/Private/Utility/DatasmithImporterImpl.cpp
+++ b/Engine/Plugins/Enterprise/DatasmithImporter/Source/DatasmithImporter/Private/Utility/DatasmithImporterImpl.cpp
@@ -896,7 +896,7 @@ bool FDatasmithImporterImpl::CheckAssetPersistenceValidity(const FString& Packag
 		}
 	}
 	// Verify user has privileges to write in folder where asset file will be stored
-	else if (GUnrealEd != nullptr && !GUnrealEd->HasMountWritePersmissionForPackage(PackageName))
+	else if (GUnrealEd != nullptr && !GUnrealEd->HasMountWritePermissionForPackage(PackageName))
 	{
 		OutReason = FText::Format(LOCTEXT("DatasmithImportInvalidFolder", "Cannot write in folder {0} to store asset {1}. Check access to folder."), FText::FromString( FPaths::GetPath( FakeAbsolutePathToAsset ) ), FText::FromString( PackageName ));
 	}
diff --git a/Engine/Source/Editor/UnrealEd/Classes/Editor/UnrealEdEngine.h b/Engine/Source/Editor/UnrealEd/Classes/Editor/UnrealEdEngine.h
--- a/Engine/Source/Editor/UnrealEd/Classes/Editor/UnrealEdEngine.h
+++ b/Engine/Source/Editor/UnrealEd/Classes/Editor/UnrealEdEngine.h
@@ -796,6 +796,8 @@ public:
 	bool IsComponentSelected(const UPrimitiveComponent* PrimComponent);
 
 	/** Return if we have write permission under the mount point this packages lives in. */
+	bool HasMountWritePermissionForPackage(const FString& PackageName);
+	UE_DEPRECATED(5.0, "Use HasMountWritePermissionForPackage instead")
 	bool HasMountWritePersmissionForPackage(const FString& PackageName);
 
 protected:
diff --git a/Engine/Source/Editor/UnrealEd/Private/UnrealEdEngine.cpp b/Engine/Source/Editor/UnrealEd/Private/UnrealEdEngine.cpp
--- a/Engine/Source/Editor/UnrealEd/Private/UnrealEdEngine.cpp
+++ b/Engine/Source/Editor/UnrealEd/Private/UnrealEdEngine.cpp
@@ -532,7 +532,7 @@ void UUnrealEdEngine::AttemptModifiedPackageNotification()
 
 				// if we do not have write permission under the mount point for this package log an error in the message log to link to.
 				PackageName = Package->GetName();
-				if (!HasMountWritePersmissionForPackage(PackageName))
+				if (!HasMountWritePermissionForPackage(PackageName))
 				{
 					bShowWritePermissionWarning = true;
 					EditorLog.Warning(FText::Format(NSLOCTEXT("UnrealEd", "WritePermissionFailureLog", "Insufficient writing permission to save {0}"), FText::FromString(PackageName)));
@@ -974,7 +974,7 @@ void UUnrealEdEngine::GetPackageList( TArray<UPackage*>* InPackages, UClass* InC
 
 bool UUnrealEdEngine::CanSavePackage( UPackage* PackageToSave )
 {
-	return HasMountWritePersmissionForPackage(PackageToSave->GetName());
+	return HasMountWritePermissionForPackage(PackageToSave->GetName());
 }
 
 
@@ -1464,7 +1464,7 @@ void UUnrealEdEngine::OnEditorSelectionChanged(UObject* SelectionThatChanged)
 	}
 }
 
-bool UUnrealEdEngine::HasMountWritePersmissionForPackage(const FString& PackageName)
+bool UUnrealEdEngine::HasMountWritePermissionForPackage(const FString& PackageName)
 {
 	FName MountPoint = FPackageName::GetPackageMountPoint(PackageName, false);
 	if (const bool* bWritePermission = MountPointCheckedForWritePermission.Find(MountPoint))
