Fix overlapping text in the curve editor

diff --git a/Engine/Source/Editor/UnrealEd/Private/SCurveEditor.cpp b/Engine/Source/Editor/UnrealEd/Private/SCurveEditor.cpp
--- a/Engine/Source/Editor/UnrealEd/Private/SCurveEditor.cpp
+++ b/Engine/Source/Editor/UnrealEd/Private/SCurveEditor.cpp
@@ -292,187 +292,166 @@ void SCurveEditor::Construct(const FArguments& InArgs)
 
 	ChildSlot
 	[
-		SNew( SHorizontalBox )
-
-		+ SHorizontalBox::Slot()
-		.FillWidth(1.0f)
+		SNew( SVerticalBox )
+		+ SVerticalBox::Slot()
+		.FillHeight(1.0f)
 		[
-			SNew( SVerticalBox )
-
-			+ SVerticalBox::Slot()
-			.FillHeight(1.0f)
+			SNew(SHorizontalBox)
+			.Visibility( this, &SCurveEditor::GetCurveAreaVisibility )
+			+ SHorizontalBox::Slot()
+			.AutoWidth()
+			.Padding(FMargin(30, 12, 0, 0))
 			[
-				SNew(SHorizontalBox)
-				.Visibility( this, &SCurveEditor::GetCurveAreaVisibility )
+				CurveSelector
+			]
 
-				+ SHorizontalBox::Slot()
-				.AutoWidth()
-				.Padding(FMargin(30, 12, 0, 0))
+			+ SHorizontalBox::Slot()
+			.AutoWidth()
+			[
+				SNew(SBorder)
+				.VAlign(VAlign_Top)
+				.HAlign(HAlign_Left)
+				.BorderImage( FEditorStyle::GetBrush("NoBorder") )
+				.DesiredSizeScale(FVector2D(256.0f,32.0f))
+				.Padding(FMargin(2, 12, 0, 0))
 				[
-					CurveSelector
-				]
+					SNew(SHorizontalBox)
 
-				+ SHorizontalBox::Slot()
-				.AutoWidth()
-				[
-					SNew(SBorder)
-					.VAlign(VAlign_Top)
-					.HAlign(HAlign_Left)
-					.BorderImage( FEditorStyle::GetBrush("NoBorder") )
-					.DesiredSizeScale(FVector2D(256.0f,32.0f))
-					.Padding(FMargin(2, 12, 0, 0))
+					+ SHorizontalBox::Slot()
+					.AutoWidth()
 					[
-						SNew(SHorizontalBox)
-
-						+ SHorizontalBox::Slot()
-						.AutoWidth()
+						SNew(SButton)
+						.ToolTipText(LOCTEXT("ZoomToFitHorizontal", "Zoom To Fit Horizontal"))
+						.Visibility(this, &SCurveEditor::GetZoomButtonVisibility)
+						.OnClicked(this, &SCurveEditor::ZoomToFitHorizontalClicked)
+						.ContentPadding(1)
 						[
-							SNew(SButton)
-							.ToolTipText(LOCTEXT("ZoomToFitHorizontal", "Zoom To Fit Horizontal"))
-							.Visibility(this, &SCurveEditor::GetZoomButtonVisibility)
-							.OnClicked(this, &SCurveEditor::ZoomToFitHorizontalClicked)
-							.ContentPadding(1)
-							[
-								SNew(SImage) 
-								.Image( FEditorStyle::GetBrush("CurveEd.FitHorizontal") ) 
-								.ColorAndOpacity( FSlateColor::UseForeground() ) 
-							]
+							SNew(SImage) 
+							.Image( FEditorStyle::GetBrush("CurveEd.FitHorizontal") ) 
+							.ColorAndOpacity( FSlateColor::UseForeground() ) 
 						]
+					]
 						
-						+ SHorizontalBox::Slot()
-						.AutoWidth()
+					+ SHorizontalBox::Slot()
+					.AutoWidth()
+					[
+						SNew(SButton)
+						.ToolTipText(LOCTEXT("ZoomToFitVertical", "Zoom To Fit Vertical"))
+						.Visibility(this, &SCurveEditor::GetZoomButtonVisibility)
+						.OnClicked(this, &SCurveEditor::ZoomToFitVerticalClicked)
+						.ContentPadding(1)
 						[
-							SNew(SButton)
-							.ToolTipText(LOCTEXT("ZoomToFitVertical", "Zoom To Fit Vertical"))
-							.Visibility(this, &SCurveEditor::GetZoomButtonVisibility)
-							.OnClicked(this, &SCurveEditor::ZoomToFitVerticalClicked)
-							.ContentPadding(1)
-							[
-								SNew(SImage) 
-								.Image( FEditorStyle::GetBrush("CurveEd.FitVertical") ) 
-								.ColorAndOpacity( FSlateColor::UseForeground() ) 
-							]
+							SNew(SImage) 
+							.Image( FEditorStyle::GetBrush("CurveEd.FitVertical") ) 
+							.ColorAndOpacity( FSlateColor::UseForeground() ) 
 						]
+					]
 
+					+ SHorizontalBox::Slot()
+					.AutoWidth()
+					[
+						SNew(SHorizontalBox)
 						+ SHorizontalBox::Slot()
-						.AutoWidth()
+						.Padding(6.0f, 0.0, 3.0f, 0.0f)
+						.VAlign(VAlign_Center)
 						[
-							SNew(SBorder)
-							.BorderImage(FEditorStyle::GetBrush("NoBorder"))
+							SNew(STextBlock)
 							.Visibility(this, &SCurveEditor::GetEditVisibility)
-							.VAlign(VAlign_Center)
-							[
-								SNew(SHorizontalBox)
-
-								+ SHorizontalBox::Slot()
-								.AutoWidth()
-								[
-									SNew(SNumericEntryBox<float>)
-									.IsEnabled(this, &SCurveEditor::GetInputEditEnabled)
-									.Font(FEditorStyle::GetFontStyle("CurveEd.InfoFont"))
-									.Value(this, &SCurveEditor::OnGetTime)
-									.UndeterminedString(LOCTEXT("MultipleValues", "Multiple Values"))
-									.OnValueCommitted(this, &SCurveEditor::OnTimeComitted)
-									.OnValueChanged(this, &SCurveEditor::OnTimeChanged)
-									.OnBeginSliderMovement(this, &SCurveEditor::OnBeginSliderMovement, LOCTEXT("SetTime", "Set New Time"))
-									.OnEndSliderMovement(this, &SCurveEditor::OnEndSliderMovement)
-									.LabelVAlign(VAlign_Center)
-									.AllowSpin(true)
-									.MinValue(TOptional<float>())
-									.MaxValue(TOptional<float>())
-									.MaxSliderValue(TOptional<float>())
-									.MinSliderValue(TOptional<float>())
-									.Delta(this, &SCurveEditor::GetInputNumericEntryBoxDelta)
-									.MinDesiredValueWidth(60.0f)
-									.Visibility(this, &SCurveEditor::GetTimeEditVisibility)
-									.Label()
-									[
-										SNew(STextBlock)
-										.Text(this, &SCurveEditor::GetInputAxisName)
-									]
-								]
-
-								+ SHorizontalBox::Slot()
-								.AutoWidth()
-								[
-									SNew(SNumericEntryBox<int32>)
-									.IsEnabled(this, &SCurveEditor::GetInputEditEnabled)
-									.Font(FEditorStyle::GetFontStyle("CurveEd.InfoFont"))
-									.Value(this, &SCurveEditor::OnGetTimeInFrames)
-									.UndeterminedString(LOCTEXT("MultipleValues", "Multiple Values"))
-									.OnValueCommitted(this, &SCurveEditor::OnTimeInFramesComitted)
-									.OnValueChanged(this, &SCurveEditor::OnTimeInFramesChanged)
-									.OnBeginSliderMovement(this, &SCurveEditor::OnBeginSliderMovement, LOCTEXT("SetFrame", "Set New Frame"))
-									.OnEndSliderMovement(this, &SCurveEditor::OnEndSliderMovement)
-									.LabelVAlign(VAlign_Center)
-									.AllowSpin(true)
-									.MinValue(TOptional<int32>())
-									.MaxValue(TOptional<int32>())
-									.MaxSliderValue(TOptional<int32>())
-									.MinSliderValue(TOptional<int32>())
-									.Delta(1)
-									.MinDesiredValueWidth(60.0f)
-									.Visibility(this, &SCurveEditor::GetFrameEditVisibility)
-									.Label()
-									[
-										SNew(STextBlock)
-										.Text(this, &SCurveEditor::GetInputAxisName)
-									]
-								]
-							]
+							.Text(this, &SCurveEditor::GetInputAxisName)
+							.ShadowOffset(FVector2D(1,1))
 						]
-						
 						+ SHorizontalBox::Slot()
+						.VAlign(VAlign_Center)
 						.AutoWidth()
 						[
-							SNew(SBorder)
-							.BorderImage( FEditorStyle::GetBrush("NoBorder") )
-							.Visibility(this, &SCurveEditor::GetEditVisibility)
-							.VAlign(VAlign_Center)
-							[
-								SNew(SNumericEntryBox<float>)
-								.Font(FEditorStyle::GetFontStyle("CurveEd.InfoFont"))
-								.Value(this, &SCurveEditor::OnGetValue)
-								.UndeterminedString(LOCTEXT("MultipleValues", "Multiple Values"))
-								.OnValueCommitted(this, &SCurveEditor::OnValueComitted)
-								.OnValueChanged(this, &SCurveEditor::OnValueChanged)
-								.OnBeginSliderMovement(this, &SCurveEditor::OnBeginSliderMovement, LOCTEXT("SetValue", "Set New Value"))
-								.OnEndSliderMovement(this, &SCurveEditor::OnEndSliderMovement)
-								.LabelVAlign(VAlign_Center)
-								.AllowSpin(true)
-								.MinValue(TOptional<float>())
-								.MaxValue(TOptional<float>())
-								.MaxSliderValue(TOptional<float>())
-								.MinSliderValue(TOptional<float>())
-								.Delta(this, &SCurveEditor::GetOutputNumericEntryBoxDelta)
-								.MinDesiredValueWidth(60.0f)
-								.Label()
-								[
-									SNew(STextBlock)
-									.Text(OutputAxisName)
-								]
-							]
+							SNew(SNumericEntryBox<float>)
+							.IsEnabled(this, &SCurveEditor::GetInputEditEnabled)
+							.Value(this, &SCurveEditor::OnGetTime)
+							.UndeterminedString(LOCTEXT("MultipleValues", "Multiple Values"))
+							.OnValueCommitted(this, &SCurveEditor::OnTimeComitted)
+							.OnValueChanged(this, &SCurveEditor::OnTimeChanged)
+							.OnBeginSliderMovement(this, &SCurveEditor::OnBeginSliderMovement, LOCTEXT("SetTime", "Set New Time"))
+							.OnEndSliderMovement(this, &SCurveEditor::OnEndSliderMovement)
+							.AllowSpin(true)
+							.MinValue(TOptional<float>())
+							.MaxValue(TOptional<float>())
+							.MaxSliderValue(TOptional<float>())
+							.MinSliderValue(TOptional<float>())
+							.Delta(this, &SCurveEditor::GetInputNumericEntryBoxDelta)
+							.MinDesiredValueWidth(60.0f)
+							.Visibility(this, &SCurveEditor::GetTimeEditVisibility)
+								
 						]
+						+ SHorizontalBox::Slot()
+						.Padding(3.0f, 0.0f)
+						.VAlign(VAlign_Center)
+						.AutoWidth()
+						[
+							SNew(SNumericEntryBox<int32>)
+							.IsEnabled(this, &SCurveEditor::GetInputEditEnabled)
+							.Value(this, &SCurveEditor::OnGetTimeInFrames)
+							.UndeterminedString(LOCTEXT("MultipleValues", "Multiple Values"))
+							.OnValueCommitted(this, &SCurveEditor::OnTimeInFramesComitted)
+							.OnValueChanged(this, &SCurveEditor::OnTimeInFramesChanged)
+							.OnBeginSliderMovement(this, &SCurveEditor::OnBeginSliderMovement, LOCTEXT("SetFrame", "Set New Frame"))
+							.OnEndSliderMovement(this, &SCurveEditor::OnEndSliderMovement)
+							.LabelVAlign(VAlign_Center)
+							.AllowSpin(true)
+							.MinValue(TOptional<int32>())
+							.MaxValue(TOptional<int32>())
+							.MaxSliderValue(TOptional<int32>())
+							.MinSliderValue(TOptional<int32>())
+							.Delta(1)
+							.MinDesiredValueWidth(60.0f)
+							.Visibility(this, &SCurveEditor::GetFrameEditVisibility)
+						]
+					]
+					+ SHorizontalBox::Slot()
+					.Padding(3.0f, 0.0f)
+					.VAlign(VAlign_Center)
+					[
+						SNew(STextBlock)
+						.Visibility(this, &SCurveEditor::GetEditVisibility)
+						.Text(OutputAxisName)
+						.ShadowOffset(FVector2D(1, 1))
+					]
+					+ SHorizontalBox::Slot()
+					.VAlign(VAlign_Center)
+					.AutoWidth()
+					[
+						SNew(SNumericEntryBox<float>)
+						.Visibility(this, &SCurveEditor::GetEditVisibility)
+						.Value(this, &SCurveEditor::OnGetValue)
+						.UndeterminedString(LOCTEXT("MultipleValues", "Multiple Values"))
+						.OnValueCommitted(this, &SCurveEditor::OnValueComitted)
+						.OnValueChanged(this, &SCurveEditor::OnValueChanged)
+						.OnBeginSliderMovement(this, &SCurveEditor::OnBeginSliderMovement, LOCTEXT("SetValue", "Set New Value"))
+						.OnEndSliderMovement(this, &SCurveEditor::OnEndSliderMovement)
+						.AllowSpin(true)
+						.MinValue(TOptional<float>())
+						.MaxValue(TOptional<float>())
+						.MaxSliderValue(TOptional<float>())
+						.MinSliderValue(TOptional<float>())
+						.Delta(this, &SCurveEditor::GetOutputNumericEntryBoxDelta)
+						.MinDesiredValueWidth(60.0f)
 					]
 				]
 			]
-
-
-			+ SVerticalBox::Slot()
-			.VAlign(VAlign_Bottom)
-			.FillHeight(.75f)
+		]
+		+ SVerticalBox::Slot()
+		.VAlign(VAlign_Bottom)
+		.FillHeight(.75f)
+		[
+			SNew( SBorder )
+			.Visibility( this, &SCurveEditor::GetColorGradientVisibility )
+			.BorderImage( FEditorStyle::GetBrush("ToolPanel.GroupBorder") )
+			.BorderBackgroundColor( FLinearColor( .8f, .8f, .8f, .60f ) )
+			.Padding(1.0f)
 			[
-				SNew( SBorder )
-				.Visibility( this, &SCurveEditor::GetColorGradientVisibility )
-				.BorderImage( FEditorStyle::GetBrush("ToolPanel.GroupBorder") )
-				.BorderBackgroundColor( FLinearColor( .8f, .8f, .8f, .60f ) )
-				.Padding(1.0f)
-				[
-					SAssignNew( GradientViewer, SColorGradientEditor )
-					.ViewMinInput( ViewMinInput )
-					.ViewMaxInput( ViewMaxInput )
-					.IsEditingEnabled( this, &SCurveEditor::IsEditingEnabled )
-				]
+				SAssignNew( GradientViewer, SColorGradientEditor )
+				.ViewMinInput( ViewMinInput )
+				.ViewMaxInput( ViewMaxInput )
+				.IsEditingEnabled( this, &SCurveEditor::IsEditingEnabled )
 			]
 		]
 	];
@@ -729,12 +708,20 @@ EVisibility SCurveEditor::GetZoomButtonVisibility() const
 
 EVisibility SCurveEditor::GetTimeEditVisibility() const
 {
-	return ShowTimeInFrames() ? EVisibility::Collapsed : EVisibility::Visible;
+	if (GetEditVisibility().IsVisible())
+	{
+		return ShowTimeInFrames() ? EVisibility::Collapsed : EVisibility::Visible;
+	}
+	return EVisibility::Collapsed;
 }
 
 EVisibility SCurveEditor::GetFrameEditVisibility() const
 {
-	return ShowTimeInFrames() ? EVisibility::Visible : EVisibility::Collapsed;
+	if (GetEditVisibility().IsVisible())
+	{
+		return ShowTimeInFrames() ? EVisibility::Visible : EVisibility::Collapsed;
+	}
+	return EVisibility::Collapsed;
 }
 
 bool SCurveEditor::GetInputEditEnabled() const
