Make RenameVariableReferences public instead of private for use by the BlueprintEditorLibrary

#rnx

diff --git a/Engine/Source/Editor/UnrealEd/Public/Kismet2/BlueprintEditorUtils.h b/Engine/Source/Editor/UnrealEd/Public/Kismet2/BlueprintEditorUtils.h
--- a/Engine/Source/Editor/UnrealEd/Public/Kismet2/BlueprintEditorUtils.h
+++ b/Engine/Source/Editor/UnrealEd/Public/Kismet2/BlueprintEditorUtils.h
@@ -1766,7 +1766,6 @@ public:
 	DECLARE_MULTICAST_DELEGATE_FourParams(FOnRenameVariableReferences, UBlueprint* /*Blueprint*/, UClass* /*VariableClass*/, const FName& /*OldVarName*/, const FName& /*NewVarName*/);
 	static FOnRenameVariableReferences OnRenameVariableReferencesEvent;
 
-protected:
 	/**
 	 * Looks through the specified blueprint for any references to the specified 
 	 * variable, and renames them accordingly.
