HLOD - Cleaning up cluster generation
* Added option to apply HLOD volumes only to specific HLOD level (modified implementation of Github pull request #7410)
* Unified processing of HLOD volumes, so that they apply to all HLOD level instead of only HLOD0
* Unified processing of the cluster reuse option

diff --git a/Engine/Source/Editor/UnrealEd/Private/HierarchicalLOD.cpp b/Engine/Source/Editor/UnrealEd/Private/HierarchicalLOD.cpp
--- a/Engine/Source/Editor/UnrealEd/Private/HierarchicalLOD.cpp
+++ b/Engine/Source/Editor/UnrealEd/Private/HierarchicalLOD.cpp
@@ -147,7 +147,7 @@ void FHierarchicalLODBuilder::BuildClusters(ULevel* InLevel)
 			for (int32 LODId = 0; LODId < NumHLODLevels; ++LODId)
 			{
 				// Handle HierachicalLOD volumes first
-				HandleHLODVolumes(InLevel);
+				HandleHLODVolumes(InLevel, LODId);
 
 				// Reuse clusters from previous HLOD level (only works for HLOD level 1 and beyond)
 				if (BuildLODLevelSettings[LODId].bReusePreviousLevelClusters && LODId > 0)
@@ -156,71 +156,72 @@ void FHierarchicalLODBuilder::BuildClusters(ULevel* InLevel)
 					{
 						FLODCluster PreviousActorCluster(PreviousLODActor);
 
-						if (BuildLODLevelSettings[LODId].bOnlyGenerateClustersForVolumes)
+						// Reassess whether or not actors that were excluded from the previous HLOD level should be included in this one
+						auto EvaluateRejectedActors = [this, &PreviousActorCluster, LODId](TFunctionRef<bool(const AActor*)> InPredicate)
 						{
-							AHierarchicalLODVolume** VolumePtr = HLODVolumeActors.Find(PreviousLODActor);
-							if (VolumePtr)
+							for (auto It = RejectedActorsInLevel.CreateIterator(); It; ++It)
 							{
-								AHierarchicalLODVolume* Volume = *VolumePtr;
-								FBox HLODVolumeBox = Volume->GetComponentsBoundingBox(true);
-
-								for (AActor* Actor : ValidStaticMeshActorsInLevel)
+								AActor* Actor = *It;
+								if (!ShouldGenerateCluster(Actor, LODId - 1) && ShouldGenerateCluster(Actor, LODId))
 								{
-									if (ShouldGenerateCluster(Actor, LODId))
+									if (InPredicate(Actor))
 									{
-										FBox ActorBox = Actor->GetComponentsBoundingBox(true);
-										if (HLODVolumeBox.IsInside(ActorBox) || (Volume->bIncludeOverlappingActors && HLODVolumeBox.Intersect(ActorBox)))
-										{
-											PreviousActorCluster += Actor;
-										}
+										PreviousActorCluster += Actor;
+										It.RemoveCurrent(); // Don't use it again later once it's in a cluster
 									}
 								}
+							}
+						};
 
-								// Reassess whether or not objects that were excluded from the previous HLOD level should be included in this one
-								if (BuildLODLevelSettings[LODId - 1].bAllowSpecificExclusion)
+						auto EvaluateValidActors = [this, &PreviousActorCluster, LODId](TFunctionRef<bool(const AActor*)> InPredicate)
+						{
+							for (auto It = ValidStaticMeshActorsInLevel.CreateIterator(); It; ++It)
+							{
+								AActor* Actor = *It;
+								if (ShouldGenerateCluster(Actor, LODId))
 								{
-									for (int RejectedIndex = RejectedActorsInLevel.Num() - 1; RejectedIndex >= 0; RejectedIndex--)
+									if (InPredicate(Actor))
 									{
-										AActor* Actor = RejectedActorsInLevel[RejectedIndex];
-										if (!ShouldGenerateCluster(Actor, LODId - 1) && ShouldGenerateCluster(Actor, LODId))
-										{
-											FBox ActorBox = Actor->GetComponentsBoundingBox(true);
-											if (HLODVolumeBox.IsInside(ActorBox) || (Volume->bIncludeOverlappingActors && HLODVolumeBox.Intersect(ActorBox)))
-											{
-												PreviousActorCluster += Actor;
-												RejectedActorsInLevel.RemoveAt(RejectedIndex); // Don't use it again later once it's in a cluster
-											}
-										}
+										PreviousActorCluster += Actor;
+										It.RemoveCurrent(); // Don't use it again later once it's in a cluster
 									}
 								}
 							}
+						};
+
+						if (BuildLODLevelSettings[LODId].bOnlyGenerateClustersForVolumes)
+						{
+							AHierarchicalLODVolume** VolumePtr = HLODVolumeActors.Find(PreviousLODActor);
+							if (VolumePtr && (*VolumePtr)->AppliesToHLODLevel(LODId))
+							{
+								AHierarchicalLODVolume* Volume = *VolumePtr;
+								FBox HLODVolumeBox = Volume->GetComponentsBoundingBox(true);
+
+								auto IsInVolume = [Volume, &HLODVolumeBox](const AActor* Actor)
+								{
+									FBox ActorBox = Actor->GetComponentsBoundingBox(true);
+									return HLODVolumeBox.IsInside(ActorBox) || (Volume->bIncludeOverlappingActors && HLODVolumeBox.Intersect(ActorBox));
+								};
+
+								EvaluateValidActors(IsInVolume);
+								EvaluateRejectedActors(IsInVolume);
+							}
 						}
 						else
 						{
-							// Reassess whether or not objects that were excluded from the previous HLOD level should be included in this one
 							const FBoxSphereBounds ClusterBounds(PreviousLODActor->GetComponentsBoundingBox(true));
-							if (BuildLODLevelSettings[LODId - 1].bAllowSpecificExclusion)
+
+							auto IsInCluster = [&ClusterBounds](const AActor* Actor)
 							{
-								for (int RejectedIndex = RejectedActorsInLevel.Num() - 1; RejectedIndex >= 0; RejectedIndex--)
-								{
-									AActor* Actor = RejectedActorsInLevel[RejectedIndex];
-									if (Actor && FBoxSphereBounds::SpheresIntersect(ClusterBounds, FSphere(Actor->GetActorLocation(), Actor->GetComponentsBoundingBox().GetSize().Size())))
-									{
-										if (!ShouldGenerateCluster(Actor, LODId - 1) && ShouldGenerateCluster(Actor, LODId))
-										{
-											PreviousActorCluster += Actor;
-											RejectedActorsInLevel.RemoveAt(RejectedIndex); // Don't use it again later once it's in a cluster
-										}
-									}
-								}
-							}
+								return FBoxSphereBounds::SpheresIntersect(ClusterBounds, FSphere(Actor->GetActorLocation(), Actor->GetComponentsBoundingBox().GetSize().Size()));
+							};
+
+							EvaluateRejectedActors(IsInCluster);
 						}
 
 						ALODActor* LODActor = CreateLODActor(PreviousActorCluster, InLevel, LODId);
 						LODActor->SetLODActorTag(PreviousLODActor->GetLODActorTag());
 						LODLevelLODActors[LODId].Add(LODActor);
-
-						ValidStaticMeshActorsInLevel.RemoveAll([PreviousActorCluster](AActor* InActor) { return PreviousActorCluster.Actors.Contains(InActor); });
 					}
 				}
 				else
@@ -228,16 +229,11 @@ void FHierarchicalLODBuilder::BuildClusters(ULevel* InLevel)
 					// we use meter for bound. Otherwise it's very easy to get to overflow and have problem with filling ratio because
 					// bound is too huge
 					const float DesiredBoundRadius = BuildLODLevelSettings[LODId].DesiredBoundRadius * CM_TO_METER;
-					const float DesiredFillingRatio = BuildLODLevelSettings[LODId].DesiredFillingPercentage * 0.01f;
+					const float DesiredFillingRatio = BuildLODLevelSettings[LODId].DesiredFillingPercentage * CM_TO_METER;
 					ensure(DesiredFillingRatio != 0.f);
 					const float HighestCost = FMath::Pow(DesiredBoundRadius, 3) / (DesiredFillingRatio);
 					const int32 MinNumActors = BuildLODLevelSettings[LODId].MinNumberOfActorsToBuild;
 					check(MinNumActors > 0);
-					// test parameter I was playing with to cull adding to the array
-					// intialization can have too many elements, decided to cull
-					// the problem can be that we can create disconnected tree
-					// my assumption is that if the merge cost is too high, then it's not worth merge anyway
-					static int32 CullMultiplier = 1;
 
 					// since to show progress of initialization, I'm scoping it
 					{
@@ -250,7 +246,7 @@ void FHierarchicalLODBuilder::BuildClusters(ULevel* InLevel)
 						SlowTask.MakeDialog();
 
 						// initialize Clusters
-						InitializeClusters(InLevel, LODId, HighestCost*CullMultiplier, BuildLODLevelSettings[LODId].bOnlyGenerateClustersForVolumes);
+						InitializeClusters(InLevel, LODId, HighestCost, BuildLODLevelSettings[LODId].bOnlyGenerateClustersForVolumes);
 
 						// move a half way - I know we can do this better but as of now this is small progress
 						SlowTask.EnterProgressFrame(50);
@@ -484,7 +480,7 @@ void FHierarchicalLODBuilder::InitializeClusters(ULevel* InLevel, const int32 LO
 	SCOPE_LOG_TIME(TEXT("STAT_HLOD_InitializeClusters"), nullptr);
 
 	// Check whether or not this actor falls within a HierarchicalLODVolume, if so add to the Volume's cluster and exclude from normal process
-	auto ProcessVolumeClusters = [this](AActor* InActor) -> bool
+	auto ProcessVolumeClusters = [this, LODIdx](AActor* InActor) -> bool
 	{
 		FBox ActorBox = InActor->GetComponentsBoundingBox(true);
 		for (TPair<AHierarchicalLODVolume*, FLODCluster>& Cluster : HLODVolumeClusters)
@@ -500,124 +496,95 @@ void FHierarchicalLODBuilder::InitializeClusters(ULevel* InLevel, const int32 LO
 		return false;
 	};
 
-	Clusters.Empty();
-
-	if (InLevel->Actors.Num() > 0)
+	// Actors are either handled by a volume, valid or rejected
+	auto FilterActors = [this, LODIdx, bVolumesOnly, ProcessVolumeClusters](const TArray<AActor*>& InActors)
 	{
-		if (LODIdx == 0)
+		ValidStaticMeshActorsInLevel.Reset();
+		RejectedActorsInLevel.Reset();
+
+		for (int32 ActorId = 0; ActorId < InActors.Num(); ++ActorId)
 		{
-			for (int32 ActorId = 0; ActorId < InLevel->Actors.Num(); ++ActorId)
+			AActor* Actor = InActors[ActorId];
+			const bool bShouldGenerate = ShouldGenerateCluster(Actor, LODIdx);
+			if (bShouldGenerate)
 			{
-				AActor* Actor = InLevel->Actors[ActorId];
-				const bool bShouldGenerate = ShouldGenerateCluster(Actor, LODIdx);
-				if (bShouldGenerate)
+				if (!ProcessVolumeClusters(Actor))
 				{
-					if (!ProcessVolumeClusters(Actor))
+					if (bVolumesOnly)
 					{
-						if (bVolumesOnly)
-						{
-							// Add them to the RejectedActorsInLevel to be re-considered at the next LOD in case that one isn't using bVolumesOnly
-							RejectedActorsInLevel.Add(Actor);
-						}
-						else
-						{
-							ValidStaticMeshActorsInLevel.Add(Actor);
-						}
+						// Add them to the RejectedActorsInLevel to be re-considered at the next LOD in case that one isn't using bVolumesOnly
+						RejectedActorsInLevel.Add(Actor);
 					}
-				}
-				else
-				{
-					RejectedActorsInLevel.Add(Actor);
-				}
-			}
-			
-			if (!bVolumesOnly)
-			{
-				// Create clusters using actor pairs
-				for (int32 ActorId = 0; ActorId < ValidStaticMeshActorsInLevel.Num(); ++ActorId)
-				{
-					AActor* Actor1 = ValidStaticMeshActorsInLevel[ActorId];
-
-					for (int32 SubActorId = ActorId + 1; SubActorId < ValidStaticMeshActorsInLevel.Num(); ++SubActorId)
+					else
 					{
-						AActor* Actor2 = ValidStaticMeshActorsInLevel[SubActorId];
-
-						FLODCluster NewClusterCandidate = FLODCluster(Actor1, Actor2);
-						float NewClusterCost = NewClusterCandidate.GetCost();
-
-						if (NewClusterCost <= CullCost)
-						{
-							Clusters.Add(NewClusterCandidate);
-						}
+						ValidStaticMeshActorsInLevel.Add(Actor);
 					}
 				}
 			}
-		}
-		else
-		{
-			// we filter the LOD index first
-			TArray<AActor*> Actors;
-
-			Actors.Append(LODLevelLODActors[LODIdx - 1]);
-
-			// Re-evaluate level actors
-			for (int32 Idx = 0; Idx < ValidStaticMeshActorsInLevel.Num(); ++Idx)
+			else if(Actor)
 			{
-				AActor* Actor = ValidStaticMeshActorsInLevel[Idx];
-				if (!ShouldGenerateCluster(Actor, LODIdx))
-				{
-					ValidStaticMeshActorsInLevel.RemoveAt(Idx);
-					--Idx;
-				}
+				RejectedActorsInLevel.Add(Actor);
 			}
+		}
+	};
 
-			// Re-evaluate rejected actors
-			for (AActor* Actor : RejectedActorsInLevel)
-			{
-				if (ShouldGenerateCluster(Actor, LODIdx))
-				{
-					ValidStaticMeshActorsInLevel.Add(Actor);
-				}
-			}			
-			RejectedActorsInLevel.RemoveAll([this](AActor* Actor)
-			{
-				return ValidStaticMeshActorsInLevel.Contains(Actor);
-			});
-
-			Actors.Append(ValidStaticMeshActorsInLevel);
-
-			Actors.RemoveAll(ProcessVolumeClusters);
+	// Create clusters from actor pairs
+	auto CreateClusters = [this, CullCost, bVolumesOnly](const TArray<AActor*>& InActors)
+	{
+		Clusters.Reset();
 
-			// first we generate graph with 2 pair nodes
-			// this is very expensive when we have so many actors
-			// so we'll need to optimize later @todo
-			const int32 NumActors = Actors.Num();
-			if(NumActors == 1)
+		if (!bVolumesOnly)
+		{
+			const int32 NumActors = InActors.Num();
+			if (NumActors == 1)
 			{
 				// Only one actor means a simple one-to-one relationship
-				Clusters.Add(FLODCluster(Actors[0]));
+				Clusters.Add(FLODCluster(InActors[0]));
 			}
 			else
 			{
-				for (int32 ActorId = 0; ActorId < NumActors; ++ActorId)
+				// Create clusters using actor pairs
+				for (int32 ActorId = 0; ActorId < InActors.Num(); ++ActorId)
 				{
-					AActor* Actor1 = (Actors[ActorId]);
-					for (int32 SubActorId = ActorId + 1; SubActorId < Actors.Num(); ++SubActorId)
+					AActor* Actor1 = InActors[ActorId];
+
+					for (int32 SubActorId = ActorId + 1; SubActorId < InActors.Num(); ++SubActorId)
 					{
-						AActor* Actor2 = Actors[SubActorId];
+						AActor* Actor2 = InActors[SubActorId];
 
-						// create new cluster
 						FLODCluster NewClusterCandidate = FLODCluster(Actor1, Actor2);
-						Clusters.Add(NewClusterCandidate);
+						float NewClusterCost = NewClusterCandidate.GetCost();
+
+						if (NewClusterCost <= CullCost)
+						{
+							Clusters.Add(NewClusterCandidate);
+						}
 					}
 				}
 			}
-
-			// shrink after adding actors
-			// LOD 0 has lots of actors, and subsequence LODs tend to have a lot less actors
-			// so this should save a lot more. 
-			Clusters.Shrink();
 		}
+	};
+
+	if (LODIdx == 0)
+	{
+		FilterActors(InLevel->Actors);
+		CreateClusters(ValidStaticMeshActorsInLevel);
+	}
+	else
+	{
+		// Filter actors
+		TArray<AActor*> Actors;
+		Actors.Reset(ValidStaticMeshActorsInLevel.Num() + RejectedActorsInLevel.Num());
+		Actors.Append(ValidStaticMeshActorsInLevel);
+		Actors.Append(RejectedActorsInLevel);
+		FilterActors(Actors);
+
+		// Create clusters, taking previous level LODActor into account.
+		Actors.Reset();
+		Actors.Append(LODLevelLODActors[LODIdx - 1]);
+		Actors.RemoveAll(ProcessVolumeClusters);
+		Actors.Append(ValidStaticMeshActorsInLevel);
+		CreateClusters(Actors);
 	}
 }
 
@@ -639,24 +606,27 @@ void FHierarchicalLODBuilder::FindMST()
 	}
 }
 
-void FHierarchicalLODBuilder::HandleHLODVolumes(ULevel* InLevel)
+void FHierarchicalLODBuilder::HandleHLODVolumes(ULevel* InLevel, int32 LODIdx)
 {	
 	HLODVolumeClusters.Reset();
-	for (int32 ActorId = 0; ActorId < InLevel->Actors.Num(); ++ActorId)
+
+	for (AActor* Actor : InLevel->Actors)
 	{
-		if (AHierarchicalLODVolume* Actor = Cast<AHierarchicalLODVolume>(InLevel->Actors[ActorId]))
+		if (AHierarchicalLODVolume* VolumeActor = Cast<AHierarchicalLODVolume>(Actor))
 		{
-			// Came across a HLOD volume			
-			FLODCluster& NewCluster = HLODVolumeClusters.Add(Actor);
-
-			FVector Origin, Extent;
-			Actor->GetActorBounds(false, Origin, Extent);
-			NewCluster.Bound = FSphere(Origin * CM_TO_METER, Extent.Size() * CM_TO_METER);
+			if (VolumeActor->AppliesToHLODLevel(LODIdx))
+			{
+				// Came across a HLOD volume
+				FLODCluster& NewCluster = HLODVolumeClusters.Add(VolumeActor);
 
+				FVector Origin, Extent;
+				VolumeActor->GetActorBounds(false, Origin, Extent);
+				NewCluster.Bound = FSphere(Origin * CM_TO_METER, Extent.Size() * CM_TO_METER);
 
-			// calculate new filling factor
-			NewCluster.FillingFactor = 1.f;
-			NewCluster.ClusterCost = FMath::Pow(NewCluster.Bound.W, 3) / NewCluster.FillingFactor;
+				// calculate new filling factor
+				NewCluster.FillingFactor = 1.f;
+				NewCluster.ClusterCost = FMath::Pow(NewCluster.Bound.W, 3) / NewCluster.FillingFactor;
+			}
 		}
 	}
 }
@@ -1241,6 +1211,12 @@ void FHierarchicalLODBuilder::MergeClustersAndBuildActors(ULevel* InLevel, const
 							RejectedActorsInLevel.RemoveSingleSwap(RemoveActor, false);
 						}
 					}
+					else
+					{
+						// Add them to the RejectedActorsInLevel to be re-considered at the next LOD
+						ValidStaticMeshActorsInLevel.RemoveAllSwap([&Cluster](const AActor* Actor) { return Cluster.Actors.Contains(Actor); });
+						RejectedActorsInLevel.Append(Cluster.Actors);
+					}
 				}
 			}
 		}
diff --git a/Engine/Source/Editor/UnrealEd/Private/HierarchicalLODVolume.cpp b/Engine/Source/Editor/UnrealEd/Private/HierarchicalLODVolume.cpp
--- a/Engine/Source/Editor/UnrealEd/Private/HierarchicalLODVolume.cpp
+++ b/Engine/Source/Editor/UnrealEd/Private/HierarchicalLODVolume.cpp
@@ -17,3 +17,9 @@ AHierarchicalLODVolume::AHierarchicalLODVolume(const FObjectInitializer& ObjectI
 	BrushColor.B = 255;
 	BrushColor.A = 255;
 }
+
+bool AHierarchicalLODVolume::AppliesToHLODLevel(int32 LODIdx) const
+{
+	return ApplyOnlyToSpecificHLODLevels.Num() == 0 ||
+		   ApplyOnlyToSpecificHLODLevels.Contains(LODIdx);
+}
\ No newline at end of file
diff --git a/Engine/Source/Editor/UnrealEd/Public/HierarchicalLOD.h b/Engine/Source/Editor/UnrealEd/Public/HierarchicalLOD.h
--- a/Engine/Source/Editor/UnrealEd/Public/HierarchicalLOD.h
+++ b/Engine/Source/Editor/UnrealEd/Public/HierarchicalLOD.h
@@ -159,9 +159,10 @@ private:
 
 	/* Retrieves HierarchicalLODVolumes and creates a cluster for each individual one
 	*
-	* @param InLevel - Level for which the HLODs are currently being build
+	* @param InLevel	Level for which the HLODs are currently being build
+	* @param LODIdx		LOD index to process
 	*/
-	void HandleHLODVolumes(ULevel* InLevel);
+	void HandleHLODVolumes(ULevel* InLevel, int32 LODIdx);
 
 	/**
 	* Determine whether or not this level should have HLODs built for it in the specified world
diff --git a/Engine/Source/Editor/UnrealEd/Public/HierarchicalLODVolume.h b/Engine/Source/Editor/UnrealEd/Public/HierarchicalLODVolume.h
--- a/Engine/Source/Editor/UnrealEd/Public/HierarchicalLODVolume.h
+++ b/Engine/Source/Editor/UnrealEd/Public/HierarchicalLODVolume.h
@@ -13,8 +13,15 @@ class AHierarchicalLODVolume : public AVolume
 {
 	GENERATED_UCLASS_BODY()
 
+public:
+	bool AppliesToHLODLevel(int32 LODIdx) const;
+
 public:
 	/** When set this volume will incorporate actors which bounds overlap with the volume, otherwise only actors which are completely inside of the volume are incorporated */
 	UPROPERTY(EditAnywhere, Category = "HLOD Volume")
 	bool bIncludeOverlappingActors;
+
+	/** If set, this volume will only be applied to HLOD levels contained in the array.  If empty, it will apply to ALL HLOD levels */
+	UPROPERTY(EditAnywhere, Category = "HLOD Volume")
+	TArray<int32> ApplyOnlyToSpecificHLODLevels;
 };
