Batch source control operation for the mark for add files after a save.

diff --git a/Engine/Source/Editor/UnrealEd/Classes/Editor/EditorEngine.h b/Engine/Source/Editor/UnrealEd/Classes/Editor/EditorEngine.h
--- a/Engine/Source/Editor/UnrealEd/Classes/Editor/EditorEngine.h
+++ b/Engine/Source/Editor/UnrealEd/Classes/Editor/EditorEngine.h
@@ -2287,10 +2287,10 @@ public:
 	void AddPackagesToDefaultChangelist(TArray<FString>& InPackageNames);
 
 	/**
-	 * Delegate used when a source control connection dialog has been closed.
-	 * @param	bEnabled	Whether source control was enabled or not
+	 * Will run a batch mark for add source control operation for files recently saved if source control is enabled
+	 * @param	bool	Unused param to for the function to be used as a delegate when source control dialog is closed.
 	 */
-	void OnSourceControlDialogClosed(bool bEnabled);
+	void RunDeferredMarkForAddFiles(bool = false);
 
 	/** 
 	 * @return - returns the currently selected positional snap grid setting
diff --git a/Engine/Source/Editor/UnrealEd/Private/EditorEngine.cpp b/Engine/Source/Editor/UnrealEd/Private/EditorEngine.cpp
--- a/Engine/Source/Editor/UnrealEd/Private/EditorEngine.cpp
+++ b/Engine/Source/Editor/UnrealEd/Private/EditorEngine.cpp
@@ -4330,7 +4330,9 @@ bool UEditorEngine::IsPackageValidForAutoAdding(UPackage* InPackage, const FStri
 
 	// Ensure the package exists, the user is running the editor (and not a commandlet or cooking), and that source control
 	// is enabled and expecting new files to be auto-added before attempting to test the validity of the package
-	if ( InPackage && GIsEditor && !IsRunningCommandlet() && ISourceControlModule::Get().IsEnabled() && GetDefault<UEditorLoadingSavingSettings>()->bSCCAutoAddNewFiles )
+	if (InPackage && GIsEditor && !IsRunningCommandlet() 
+		&& ISourceControlModule::Get().IsEnabled() 
+		&& GetDefault<UEditorLoadingSavingSettings>()->bSCCAutoAddNewFiles)
 	{
 		const FString CleanFilename = FPaths::GetCleanFilename(InFilename);
 
@@ -4364,26 +4366,23 @@ bool UEditorEngine::IsPackageOKToSave(UPackage* InPackage, const FString& InFile
 	return true;
 }
 
-void UEditorEngine::OnSourceControlDialogClosed(bool bEnabled)
+void UEditorEngine::RunDeferredMarkForAddFiles(bool)
 {
+	if (DeferredFilesToAddToSourceControl.IsEmpty())
+	{
+		return;
+	}
+
 	if(ISourceControlModule::Get().IsEnabled())
 	{
 		ISourceControlProvider& SourceControlProvider = ISourceControlModule::Get().GetProvider();
 		if(SourceControlProvider.IsAvailable())
 		{
-			if(DeferredFilesToAddToSourceControl.Num() > 0)
-			{
-				SourceControlProvider.Execute(ISourceControlOperation::Create<FMarkForAdd>(), SourceControlHelpers::PackageFilenames(DeferredFilesToAddToSourceControl));
-			}
-	
-			DeferredFilesToAddToSourceControl.Empty();
+			SourceControlProvider.Execute(ISourceControlOperation::Create<FMarkForAdd>(), SourceControlHelpers::PackageFilenames(DeferredFilesToAddToSourceControl));
 		}
 	}
-	else
-	{
-		// the user decided to disable source control, so clear the deferred list so we dont try to add them again at a later time
-		DeferredFilesToAddToSourceControl.Empty();
-	}
+	// Clear the list when this run whether source control is active or not, since we do not want to accumulate those if the user is running without source control
+	DeferredFilesToAddToSourceControl.Empty();
 }
 
 bool UEditorEngine::InitializePhysicsSceneForSaveIfNecessary(UWorld* World, bool &bOutForceInitialized)
@@ -4544,13 +4543,9 @@ FSavePackageResultStruct UEditorEngine::Save( UPackage* InOuter, UObject* InBase
 		{
 			// Show the login window here & store the file we are trying to add.
 			// We defer the add operation until we have a valid source control connection.
-			ISourceControlModule::Get().ShowLoginDialog(FSourceControlLoginClosed::CreateUObject(this, &UEditorEngine::OnSourceControlDialogClosed), ELoginWindowMode::Modeless);
-			DeferredFilesToAddToSourceControl.Add( Filename );
-		}
-		else
-		{
-			ISourceControlModule::Get().GetProvider().Execute(ISourceControlOperation::Create<FMarkForAdd>(), SourceControlHelpers::PackageFilename(Filename));
+			ISourceControlModule::Get().ShowLoginDialog(FSourceControlLoginClosed::CreateUObject(this, &UEditorEngine::RunDeferredMarkForAddFiles), ELoginWindowMode::Modeless);
 		}
+		DeferredFilesToAddToSourceControl.Add(Filename);
 	}
 
 	SlowTask.EnterProgressFrame(10);
diff --git a/Engine/Source/Editor/UnrealEd/Private/FileHelpers.cpp b/Engine/Source/Editor/UnrealEd/Private/FileHelpers.cpp
--- a/Engine/Source/Editor/UnrealEd/Private/FileHelpers.cpp
+++ b/Engine/Source/Editor/UnrealEd/Private/FileHelpers.cpp
@@ -1358,6 +1358,8 @@ static bool IsCheckOutSelectedDisabled()
 
 bool FEditorFileUtils::AddCheckoutPackageItems(bool bCheckDirty, TArray<UPackage*> PackagesToCheckOut, TArray<UPackage*>* OutPackagesNotNeedingCheckout, bool* bOutHavePackageToCheckOut)
 {
+	TRACE_CPUPROFILER_EVENT_SCOPE(FEditorFileUtils_AddCheckoutPackageItems);
+
 	ISourceControlProvider& SourceControlProvider = ISourceControlModule::Get().GetProvider();
 	if (ISourceControlModule::Get().IsEnabled() && SourceControlProvider.IsAvailable())
 	{
@@ -1738,6 +1740,8 @@ bool FEditorFileUtils::PromptToCheckoutPackages(bool bCheckDirty, const TArray<U
 
 ECommandResult::Type FEditorFileUtils::CheckoutPackages(const TArray<UPackage*>& PkgsToCheckOut, TArray<UPackage*>* OutPackagesCheckedOut, const bool bErrorIfAlreadyCheckedOut, const bool bConfirmPackageBranchCheckOutStatus)
 {
+	TRACE_CPUPROFILER_EVENT_SCOPE(FEditorFileUtils_CheckoutPackages);
+
 	const bool bErrorIfFileMissing = false;
 
 	ECommandResult::Type CheckOutResult = ECommandResult::Succeeded;
@@ -3329,6 +3333,12 @@ static bool InternalSavePackagesFast(const TArray<UPackage*>& PackagesToSave, bo
 	{
 		ObjectTools::CleanupAfterSuccessfulDelete(PackagesToClean, true);
 	}
+
+	// Add all files that needs to be marked for add in one command, if any
+	if (GEditor)
+	{
+		GEditor->RunDeferredMarkForAddFiles();
+	}
 	
 	GWarn->EndSlowTask();
 	SaveErrors.Flush();
@@ -3696,6 +3706,12 @@ FEditorFileUtils::EPromptReturnCode InternalPromptForCheckoutAndSave(const TArra
 		ObjectTools::CleanupAfterSuccessfulDelete(PackagesToClean, true);
 	}
 
+	// Add all files that needs to be marked for add in one command, if any
+	if (GEditor)
+	{
+		GEditor->RunDeferredMarkForAddFiles();
+	}
+
 	// If any packages were saved that weren't actually in source control but instead forcibly made writable,
 	// then warn the user about those packages
 	if (WritablePackageFiles.Num() > 0)
