Fix LevelEditor settings category name

diff --git a/Engine/Source/Editor/UnrealEd/Classes/Settings/LevelEditorPlaySettings.h b/Engine/Source/Editor/UnrealEd/Classes/Settings/LevelEditorPlaySettings.h
--- a/Engine/Source/Editor/UnrealEd/Classes/Settings/LevelEditorPlaySettings.h
+++ b/Engine/Source/Editor/UnrealEd/Classes/Settings/LevelEditorPlaySettings.h
@@ -395,15 +395,15 @@ private:
 	FString AdditionalLaunchOptions;
 
 	/** Controls the default value of the show flag ServerDrawDebug */
-	UPROPERTY(config, EditAnywhere, Category = MultiplayerOptions)
+	UPROPERTY(config, EditAnywhere, Category = "Multiplayer Options")
 	bool bShowServerDebugDrawingByDefault;
 
 	/** How strongly debug drawing originating from the server will be biased towards the tint color */
-	UPROPERTY(config, EditAnywhere, Category=MultiplayerOptions, meta=(ClampMin=0, ClampMax=1, UIMin=0, UIMax=1))
+	UPROPERTY(config, EditAnywhere, Category="Multiplayer Options", meta=(ClampMin=0, ClampMax=1, UIMin=0, UIMax=1))
 	float ServerDebugDrawingColorTintStrength;
 
 	/** Debug drawing originating from the server will be biased towards this color */
-	UPROPERTY(config, EditAnywhere, Category=MultiplayerOptions)
+	UPROPERTY(config, EditAnywhere, Category="Multiplayer Options")
 	FLinearColor ServerDebugDrawingColorTint;
 
 private:
