Make sure skeletalmesh thumbnail do not render async skeletalmesh, the flags must be removed before the thumbnail can be updated.

diff --git a/Engine/Source/Editor/UnrealEd/Private/AssetThumbnail.cpp b/Engine/Source/Editor/UnrealEd/Private/AssetThumbnail.cpp
--- a/Engine/Source/Editor/UnrealEd/Private/AssetThumbnail.cpp
+++ b/Engine/Source/Editor/UnrealEd/Private/AssetThumbnail.cpp
@@ -7,6 +7,7 @@
 #include "Widgets/DeclarativeSyntaxSupport.h"
 #include "Widgets/SOverlay.h"
 #include "Engine/GameViewportClient.h"
+#include "Interfaces/Interface_AsyncCompilation.h"
 #include "Modules/ModuleManager.h"
 #include "Animation/CurveHandle.h"
 #include "Animation/CurveSequence.h"
@@ -1037,14 +1038,20 @@ void FAssetThumbnailPool::Tick( float DeltaTime )
 					if( InfoRef->AssetData.IsAssetLoaded() && !bAssetsBeingCompiled)
 					{
 						UObject* Asset = InfoRef->AssetData.GetAsset();
-						FThumbnailRenderingInfo* RenderInfo = GUnrealEd->GetThumbnailManager()->GetRenderingInfo( Asset );
-						if ( RenderInfo != NULL && RenderInfo->Renderer != NULL )
+						
+						//Avoid rendering the thumbnail of an asset that is currently edited asynchronously
+						const IInterface_AsyncCompilation* Interface_AsyncCompilation = Cast<IInterface_AsyncCompilation>(Asset);
+						const bool bIsAssetCurrentlyEditedAsynchronously = Interface_AsyncCompilation && Interface_AsyncCompilation->IsCompiling();
+						if (!bIsAssetCurrentlyEditedAsynchronously)
 						{
-							FThumbnailInfo_RenderThread ThumbInfo = InfoRef.Get();
-							ENQUEUE_RENDER_COMMAND(SyncSlateTextureCommand)(
-								[ThumbInfo](FRHICommandListImmediate& RHICmdList)
+							FThumbnailRenderingInfo* RenderInfo = GUnrealEd->GetThumbnailManager()->GetRenderingInfo(Asset);
+							if (RenderInfo != NULL && RenderInfo->Renderer != NULL)
+							{
+								FThumbnailInfo_RenderThread ThumbInfo = InfoRef.Get();
+								ENQUEUE_RENDER_COMMAND(SyncSlateTextureCommand)(
+									[ThumbInfo](FRHICommandListImmediate& RHICmdList)
 								{
-									if ( ThumbInfo.ThumbnailTexture->GetTypedResource() != ThumbInfo.ThumbnailRenderTarget->GetTextureRHI() )
+									if (ThumbInfo.ThumbnailTexture->GetTypedResource() != ThumbInfo.ThumbnailRenderTarget->GetTextureRHI())
 									{
 										ThumbInfo.ThumbnailTexture->ClearTextureData();
 										ThumbInfo.ThumbnailTexture->ReleaseDynamicRHI();
@@ -1052,19 +1059,20 @@ void FAssetThumbnailPool::Tick( float DeltaTime )
 									}
 								});
 
-							if (InfoRef->LastUpdateTime <= 0.0f || RenderInfo->Renderer->AllowsRealtimeThumbnails(Asset))
-							{
-								//@todo: this should be done on the GPU only but it is not supported by thumbnail tools yet
-								ThumbnailTools::RenderThumbnail(
-									Asset,
-									InfoRef->Width,
-									InfoRef->Height,
-									ThumbnailTools::EThumbnailTextureFlushMode::NeverFlush,
-									InfoRef->ThumbnailRenderTarget
+								if (InfoRef->LastUpdateTime <= 0.0f || RenderInfo->Renderer->AllowsRealtimeThumbnails(Asset))
+								{
+									//@todo: this should be done on the GPU only but it is not supported by thumbnail tools yet
+									ThumbnailTools::RenderThumbnail(
+										Asset,
+										InfoRef->Width,
+										InfoRef->Height,
+										ThumbnailTools::EThumbnailTextureFlushMode::NeverFlush,
+										InfoRef->ThumbnailRenderTarget
 									);
-							}
+								}
 
-							bLoadedThumbnail = true;
+								bLoadedThumbnail = true;
+							}
 						}
 					}
 				
diff --git a/Engine/Source/Runtime/Engine/Classes/Engine/SkeletalMesh.h b/Engine/Source/Runtime/Engine/Classes/Engine/SkeletalMesh.h
--- a/Engine/Source/Runtime/Engine/Classes/Engine/SkeletalMesh.h
+++ b/Engine/Source/Runtime/Engine/Classes/Engine/SkeletalMesh.h
@@ -18,6 +18,7 @@
 #include "GPUSkinPublicDefs.h"
 #include "Animation/PreviewAssetAttachComponent.h"
 #include "BoneContainer.h"
+#include "Interfaces/Interface_AsyncCompilation.h"
 #include "Interfaces/Interface_CollisionDataProvider.h"
 #include "EngineTypes.h"
 #include "Engine/Engine.h"
@@ -504,7 +505,7 @@ namespace NSSkeletalMeshSourceFileLabels
  * @see https://docs.unrealengine.com/latest/INT/Engine/Content/Types/SkeletalMeshes/
  */
 UCLASS(hidecategories=Object, BlueprintType)
-class ENGINE_API USkeletalMesh : public UStreamableRenderAsset, public IInterface_CollisionDataProvider, public IInterface_AssetUserData, public INodeMappingProviderInterface
+class ENGINE_API USkeletalMesh : public UStreamableRenderAsset, public IInterface_CollisionDataProvider, public IInterface_AssetUserData, public INodeMappingProviderInterface, public IInterface_AsyncCompilation
 {
 	GENERATED_UCLASS_BODY()
 
@@ -1413,6 +1414,11 @@ private:
 
 #if WITH_EDITOR
 public:
+
+	/** IInterface_AsyncCompilation begin*/
+	virtual bool IsCompiling() const;
+	/** IInterface_AsyncCompilation end*/
+
 	/** Delegates for asset editor events */
 
 	FDelegateHandle RegisterOnClothingChange(const FSimpleMulticastDelegate::FDelegate& InDelegate);
diff --git a/Engine/Source/Runtime/Engine/Private/SkeletalMesh.cpp b/Engine/Source/Runtime/Engine/Private/SkeletalMesh.cpp
--- a/Engine/Source/Runtime/Engine/Private/SkeletalMesh.cpp
+++ b/Engine/Source/Runtime/Engine/Private/SkeletalMesh.cpp
@@ -3295,6 +3295,19 @@ void USkeletalMesh::MoveMaterialFlagsToSections()
 #endif // WITH_EDITORONLY_DATA
 
 #if WITH_EDITOR
+
+/** IInterface_AsyncCompilation implementation begin*/
+
+bool USkeletalMesh::IsCompiling() const
+{
+	//The code here is temporary we have to verify more then the async flag in case we are building asynchronously
+	const EInternalObjectFlags AsyncFlags = EInternalObjectFlags::Async | EInternalObjectFlags::AsyncLoading;
+	return HasAnyInternalFlags(AsyncFlags);
+}
+
+/** IInterface_AsyncCompilation implementation end*/
+
+
 FDelegateHandle USkeletalMesh::RegisterOnClothingChange(const FSimpleMulticastDelegate::FDelegate& InDelegate)
 {
 	return OnClothingChange.Add(InDelegate);
