Fix support for importing JSON strings into DataTable

diff --git a/Engine/Source/Editor/UnrealEd/Classes/Factories/CSVImportFactory.h b/Engine/Source/Editor/UnrealEd/Classes/Factories/CSVImportFactory.h
--- a/Engine/Source/Editor/UnrealEd/Classes/Factories/CSVImportFactory.h
+++ b/Engine/Source/Editor/UnrealEd/Classes/Factories/CSVImportFactory.h
@@ -48,6 +48,8 @@ struct FCSVImportSettings
 	UPROPERTY(BlueprintReadWrite, Category="Misc")
 	TEnumAsByte<ERichCurveInterpMode> ImportCurveInterpMode;
 
+	/** Despite its name, DataToImport can be JSON instead of CSV if this bool is set */
+	bool bDataIsJson = false;
 	FString DataToImport;
 };
 
diff --git a/Engine/Source/Editor/UnrealEd/Private/Factories/CSVImportFactory.cpp b/Engine/Source/Editor/UnrealEd/Private/Factories/CSVImportFactory.cpp
--- a/Engine/Source/Editor/UnrealEd/Private/Factories/CSVImportFactory.cpp
+++ b/Engine/Source/Editor/UnrealEd/Private/Factories/CSVImportFactory.cpp
@@ -167,6 +167,7 @@ UObject* UCSVImportFactory::FactoryCreateText(UClass* InClass, UObject* InParent
 		bHaveInfo = true;
 	}
 
+	ImportSettings.bDataIsJson = FString(Type).Equals(TEXT("json"), ESearchCase::IgnoreCase);
 	bool bDoImport = true;
 
 	// If we do not have the info we need, pop up window to ask for things
@@ -400,9 +401,7 @@ EReimportResult::Type UCSVImportFactory::Reimport(UObject* Obj, const FString& P
 
 TArray<FString> UCSVImportFactory::DoImportDataTable(const FCSVImportSettings& InImportSettings, UDataTable* TargetDataTable)
 {
-	// Are we importing JSON data?
-	const bool bIsJSON = CurrentFilename.EndsWith(TEXT(".json"));
-	if (bIsJSON)
+	if (InImportSettings.bDataIsJson)
 	{
 		return TargetDataTable->CreateTableFromJSONString(InImportSettings.DataToImport);
 	}
@@ -412,9 +411,7 @@ TArray<FString> UCSVImportFactory::DoImportDataTable(const FCSVImportSettings& I
 
 TArray<FString> UCSVImportFactory::DoImportCurveTable(const FCSVImportSettings& InImportSettings, UCurveTable* TargetCurveTable)
 {
-	// Are we importing JSON data?
-	const bool bIsJSON = CurrentFilename.EndsWith(TEXT(".json"));
-	if (bIsJSON)
+	if (InImportSettings.bDataIsJson)
 	{
 		return TargetCurveTable->CreateTableFromJSONString(InImportSettings.DataToImport, InImportSettings.ImportCurveInterpMode);
 	}
@@ -424,9 +421,7 @@ TArray<FString> UCSVImportFactory::DoImportCurveTable(const FCSVImportSettings&
 
 TArray<FString> UCSVImportFactory::DoImportCurve(const FCSVImportSettings& InImportSettings, UCurveBase* TargetCurve)
 {
-	// Are we importing JSON data?
-	const bool bIsJSON = CurrentFilename.EndsWith(TEXT(".json"));
-	if (bIsJSON)
+	if (InImportSettings.bDataIsJson)
 	{
 		TArray<FString> Result;
 		Result.Add(LOCTEXT("Error_CannotImportCurveFromJSON", "Cannot import a curve from JSON. Please use CSV instead.").ToString());
