Cook: Added possibility to generate packages during the cook process using CookPackageSplitter (currently only support using cook commandlet).
WorldPartition : Removed PreCookCommandlet and replaced with new FWorldPartitionCookPackageSplitter.

diff --git a/Engine/Source/Editor/UnrealEd/Classes/Commandlets/WorldPartitionPreCookCommandlet.h b/Engine/Source/Editor/UnrealEd/Classes/Commandlets/WorldPartitionPreCookCommandlet.h
deleted file mode 100644
diff --git a/Engine/Source/Editor/UnrealEd/Classes/CookOnTheSide/CookOnTheFlyServer.h b/Engine/Source/Editor/UnrealEd/Classes/CookOnTheSide/CookOnTheFlyServer.h
--- a/Engine/Source/Editor/UnrealEd/Classes/CookOnTheSide/CookOnTheFlyServer.h
+++ b/Engine/Source/Editor/UnrealEd/Classes/CookOnTheSide/CookOnTheFlyServer.h
@@ -18,7 +18,6 @@
 #include "CookOnTheSide/CookLog.h"
 #include "CookOnTheFlyServer.generated.h"
 
-
 class FAssetRegistryGenerator;
 class FAsyncIODelete;
 struct FPackageNameCache;
@@ -110,6 +109,10 @@ namespace UE
 {
 namespace Cook
 {
+	namespace Private
+	{
+		class FRegisteredCookPackageSplitter;
+	}
 	struct FCookerTimer;
 	class FExternalRequests;
 	struct FPackageData;
@@ -182,6 +185,7 @@ private:
 	ECookInitializationFlags CookFlags = ECookInitializationFlags::None;
 	TUniquePtr<class FSandboxPlatformFile> SandboxFile;
 	TUniquePtr<FAsyncIODelete> AsyncIODelete; // Helper for deleting the old cook directory asynchronously
+	FString DefaultAsyncIODeletePlatformName; // Default platform name to be used by GetAsyncIODelete()
 	bool bIsInitializingSandbox = false; // stop recursion into callbacks when we are initializing sandbox
 	bool bIsSavingPackage = false; // used to stop recursive mark package dirty functions
 	bool bSaveAsyncAllowed = false; // True if and only if command line options and all other restrictions allow the use of SAVE_Async
@@ -193,6 +197,8 @@ private:
 	/** List of additional plugin directories to remap into the sandbox as needed */
 	TArray<TSharedRef<IPlugin> > PluginsToRemap;
 
+	TMultiMap<UClass*, UE::Cook::Private::FRegisteredCookPackageSplitter*> RegisteredSplitDataClasses;
+
 	//////////////////////////////////////////////////////////////////////////
 	// precaching system
 	//
@@ -306,12 +312,15 @@ private:
 	/** Try to save all packages in the SaveQueue until it's time to break. Report the number of requests that were completed (either skipped or successfully saved or failed to save) */
 	void PumpSaves(UE::Cook::FTickStackData& StackData, uint32 DesiredQueueLength, int32& OutNumPushed, bool& bOutBusy);
 	/**
-	 * Inspect the given package from the PackageTracker and add it to the SaveQueue if the cooker should save it.
+	 * Inspect the given package and add it to the SaveQueue if the cooker should save it.
 	 *
-	 * @param bUpdatePlatforms If true, then FilterLoadedPackage is being called on all existing packages in order a change to the package filter.  Any InProgress PackageDatas will be demoted to request to update the requested platforms.
+	 * @param Package				The package to be added to the SaveQueue.
+	 * @param bIsGeneratedPackage	If true, ignore the CookByTheBookOptions::bSkipHardReferences flag when in CookByTheBook mode.
+	 * @return						Returns the PackageData for this queued package.
 	 */
-	void FilterLoadedPackage(UPackage* Package, bool bUpdatePlatforms);
-	/** Check whether the package filter has change, and if so call FilterLoadedPackage again on each existing package. */
+	UE::Cook::FPackageData* QueueDiscoveredPackage(UPackage* Package, bool bIsGeneratedPackage = false);
+
+	/** Check whether the package filter has change, and if so call QueueDiscoveredPackage again on each existing package. */
 	void UpdatePackageFilter();
 
 	/**
@@ -857,12 +866,15 @@ private:
 	bool ContainsRedirector(const FName& PackageName, TMap<FName, FName>& RedirectedPaths) const;
 	
 	/**
-	 * Calls BeginCacheForCookedPlatformData on all UObjects in the package
+	 * Prepares save by calling BeginCacheForCookedPlatformData on all UObjects in the package. 
+	 * Also splits the package if an object of this package has its class registered to a corresponding 
+	 * FRegisteredCookPackageSplitter and an instance of this splitter class returns true ShouldSplitPackage 
+	 * for this UObject.
 	 *
 	 * @param PackageData the PackageData used to gather all uobjects from
 	 * @return false if time slice was reached, true if all objects have had BeginCacheForCookedPlatformData called
 	 */
-	bool BeginPackageCacheForCookedPlatformData(UE::Cook::FPackageData& PackageData, UE::Cook::FCookerTimer& Timer);
+	bool BeginPrepareSave(UE::Cook::FPackageData& PackageData, UE::Cook::FCookerTimer& Timer);
 	
 	/**
 	 * Returns true when all objects in package have all their cooked platform data loaded
@@ -871,7 +883,7 @@ private:
 	 * @param PackageData the PackageData used to gather all uobjects from
 	 * @return false if time slice was reached, true if all return true for IsCachedCookedPlatformDataLoaded 
 	 */
-	bool FinishPackageCacheForCookedPlatformData(UE::Cook::FPackageData& PackageData, UE::Cook::FCookerTimer& Timer);
+	bool FinishPrepareSave(UE::Cook::FPackageData& PackageData, UE::Cook::FCookerTimer& Timer);
 
 	/**
 	 * Frees all the memory used to call BeginCacheForCookedPlatformData on all the objects in PackageData.
@@ -929,6 +941,12 @@ private:
 	/* Delete the sandbox directory (asynchronously) for the given platform in preparation for a clean cook */
 	void DeleteSandboxDirectory(const FString& PlatformName);
 
+	/** If not already done, set the default platform name to be used by GetAsyncIODelete(). */
+	void TrySetDefaultAsyncIODeletePlatform(const FString& PlatformName);
+
+	/* Create the delete-old-cooked-directory helper.*/
+	FAsyncIODelete& GetAsyncIODelete();
+
 	/* Create the delete-old-cooked-directory helper. PlatformName tells it where to create its temp directory. AsyncDeleteDirectory is DeleteSandboxDirectory internal use only and should otherwise be set to nullptr. */
 	FAsyncIODelete& GetAsyncIODelete(const FString& PlatformName, const FString* AsyncDeleteDirectory = nullptr);
 
@@ -1048,6 +1066,15 @@ private:
 	/** Return the PackageNameCache that is caching FileNames on disk for this CookOnTheFlyServer. */
 	const FPackageNameCache& GetPackageNameCache() const;
 
+	/**
+	* In case the given package object is a Package splitter, this function will split the package and generate other packages.
+	*
+	* @param PackageData			Package to be potentially split
+	* @param Obj					Package's object to potentially handle splitting
+	* @param bHasError				If function fails, bHasError is set to true, else false.
+	*/
+	void ConditionalSplitPackage(UE::Cook::FPackageData& PackageData, UObject* Obj, bool& bHasError);
+
 	uint32 FullLoadAndSave(uint32& CookedPackageCount);
 
 	uint32		StatLoadedPackageCount = 0;
diff --git a/Engine/Source/Editor/UnrealEd/Private/Commandlets/CookCommandlet.cpp b/Engine/Source/Editor/UnrealEd/Private/Commandlets/CookCommandlet.cpp
--- a/Engine/Source/Editor/UnrealEd/Private/Commandlets/CookCommandlet.cpp
+++ b/Engine/Source/Editor/UnrealEd/Private/Commandlets/CookCommandlet.cpp
@@ -65,8 +65,8 @@ namespace DetailedCookStats
 	extern double TickCookOnTheSideLoadPackagesTimeSec;
 	extern double TickCookOnTheSideResolveRedirectorsTimeSec;
 	extern double TickCookOnTheSideSaveCookedPackageTimeSec;
-	extern double TickCookOnTheSideBeginPackageCacheForCookedPlatformDataTimeSec;
-	extern double TickCookOnTheSideFinishPackageCacheForCookedPlatformDataTimeSec;
+	extern double TickCookOnTheSideBeginPrepareSaveTimeSec;
+	extern double TickCookOnTheSideFinishPrepareSaveTimeSec;
 	extern double GameCookModificationDelegateTimeSec;
 	double TickLoopGCTimeSec = 0.0;
 	double TickLoopRecompileShaderRequestsTimeSec = 0.0;
@@ -88,8 +88,8 @@ namespace DetailedCookStats
 		ADD_COOK_STAT_FLT(" 0. 1. 1. 0", TickCookOnTheSideLoadPackagesTimeSec);
 		ADD_COOK_STAT_FLT(" 0. 1. 1. 1", TickCookOnTheSideSaveCookedPackageTimeSec);
 		ADD_COOK_STAT_FLT(" 0. 1. 1. 1. 0", TickCookOnTheSideResolveRedirectorsTimeSec);
-		ADD_COOK_STAT_FLT(" 0. 1. 1. 2", TickCookOnTheSideBeginPackageCacheForCookedPlatformDataTimeSec);
-		ADD_COOK_STAT_FLT(" 0. 1. 1. 3", TickCookOnTheSideFinishPackageCacheForCookedPlatformDataTimeSec);
+		ADD_COOK_STAT_FLT(" 0. 1. 1. 2", TickCookOnTheSideBeginPrepareSaveTimeSec);
+		ADD_COOK_STAT_FLT(" 0. 1. 1. 3", TickCookOnTheSideFinishPrepareSaveTimeSec);
 		ADD_COOK_STAT_FLT(" 0. 1. 2", TickLoopGCTimeSec);
 		ADD_COOK_STAT_FLT(" 0. 1. 3", TickLoopRecompileShaderRequestsTimeSec);
 		ADD_COOK_STAT_FLT(" 0. 1. 4", TickLoopShaderProcessAsyncResultsTimeSec);
diff --git a/Engine/Source/Editor/UnrealEd/Private/Commandlets/WorldPartitionPreCookCommandlet.cpp b/Engine/Source/Editor/UnrealEd/Private/Commandlets/WorldPartitionPreCookCommandlet.cpp
deleted file mode 100644
diff --git a/Engine/Source/Editor/UnrealEd/Private/CookOnTheFlyServer.cpp b/Engine/Source/Editor/UnrealEd/Private/CookOnTheFlyServer.cpp
--- a/Engine/Source/Editor/UnrealEd/Private/CookOnTheFlyServer.cpp
+++ b/Engine/Source/Editor/UnrealEd/Private/CookOnTheFlyServer.cpp
@@ -13,6 +13,7 @@
 #include "Cooker/CookRequests.h"
 #include "Cooker/CookTypes.h"
 #include "Cooker/PackageTracker.h"
+#include "CookPackageSplitter.h"
 #include "Containers/RingBuffer.h"
 #include "HAL/PlatformFileManager.h"
 #include "HAL/FileManager.h"
@@ -187,8 +188,8 @@ namespace DetailedCookStats
 	double TickCookOnTheSideLoadPackagesTimeSec = 0.0;
 	double TickCookOnTheSideResolveRedirectorsTimeSec = 0.0;
 	double TickCookOnTheSideSaveCookedPackageTimeSec = 0.0;
-	double TickCookOnTheSideBeginPackageCacheForCookedPlatformDataTimeSec = 0.0;
-	double TickCookOnTheSideFinishPackageCacheForCookedPlatformDataTimeSec = 0.0;
+	double TickCookOnTheSideBeginPrepareSaveTimeSec = 0.0;
+	double TickCookOnTheSideFinishPrepareSaveTimeSec = 0.0;
 	double GameCookModificationDelegateTimeSec = 0.0;
 	int32 PeakRequestQueueSize = 0;
 	int32 PeakLoadQueueSize = 0;
@@ -1886,14 +1887,14 @@ void UCookOnTheFlyServer::RejectPackageToLoad(UE::Cook::FPackageData& PackageDat
 
 //////////////////////////////////////////////////////////////////////////
 
-void UCookOnTheFlyServer::FilterLoadedPackage(UPackage* Package, bool bUpdatePlatforms)
+UE::Cook::FPackageData* UCookOnTheFlyServer::QueueDiscoveredPackage(UPackage* Package, bool bIsGeneratedPackage)
 {
 	check(Package != nullptr);
 
 	const FName FileName = GetPackageNameCache().GetCachedStandardFileName(Package);
 	if (FileName.IsNone())
 	{
-		return;	// if we have name none that means we are in core packages or something...
+		return nullptr;	// if we have name none that means we are in core packages or something...
 	}
 	UE::Cook::FPackageData& PackageData = PackageDatas->FindOrAddPackageData(Package->GetFName(), FileName);
 
@@ -1901,28 +1902,19 @@ void UCookOnTheFlyServer::FilterLoadedPackage(UPackage* Package, bool bUpdatePla
 	if (PackageData.HasAllCookedPlatforms(TargetPlatforms, true /* bIncludeFailed */))
 	{
 		// All SessionPlatforms have already been cooked for the package, so we don't need to save it again
-		return;
+		return nullptr;
 	}
 
-	bool bIsUrgent = false;
-	if (PackageData.IsInProgress())
+	if (!PackageData.IsInProgress() && 
+		(bIsGeneratedPackage || !IsCookByTheBookMode() || !CookByTheBookOptions->bSkipHardReferences))
 	{
-		if (bUpdatePlatforms)
-		{
-			PackageData.UpdateRequestData(TargetPlatforms, bIsUrgent, UE::Cook::FCompletionCallback());
-		}
-	}
-	else
-	{
-		if (!IsCookByTheBookMode() || !CookByTheBookOptions->bSkipHardReferences)
-		{
-			// Send this unsolicited package into the LoadReadyQueue to fully load it and send it on to the SaveQueue
-			PackageData.SetRequestData(TargetPlatforms, bIsUrgent, UE::Cook::FCompletionCallback());
-			PackageData.SendToState(UE::Cook::EPackageState::LoadReady, UE::Cook::ESendFlags::QueueNone);
-			// Send it to the front of the LoadReadyQueue since it is mostly loaded already
-			PackageDatas->GetLoadReadyQueue().AddFront(&PackageData);
-		}
+		// Send this package into the LoadReadyQueue to fully load it and send it on to the SaveQueue
+		PackageData.SetRequestData(TargetPlatforms, /*bIsUrgent*/ false, UE::Cook::FCompletionCallback());
+		PackageData.SendToState(UE::Cook::EPackageState::LoadReady, UE::Cook::ESendFlags::QueueNone);
+		// Send it to the front of the LoadReadyQueue since it is mostly loaded already
+		PackageDatas->GetLoadReadyQueue().AddFront(&PackageData);
 	}
+	return &PackageData;
 }
 
 void UCookOnTheFlyServer::UpdatePackageFilter()
@@ -1934,9 +1926,14 @@ void UCookOnTheFlyServer::UpdatePackageFilter()
 	bPackageFilterDirty = false;
 
 	UE_SCOPED_COOKTIMER(UpdatePackageFilter);
+	const TArray<const ITargetPlatform*>& TargetPlatforms = PlatformManager->GetSessionPlatforms();
 	for (UPackage* Package : PackageTracker->LoadedPackages)
 	{
-		FilterLoadedPackage(Package, true);
+		UE::Cook::FPackageData* PackageData = QueueDiscoveredPackage(Package);
+		if (PackageData && PackageData->IsInProgress())
+		{
+			PackageData->UpdateRequestData(TargetPlatforms, /*bIsUrgent*/ false, UE::Cook::FCompletionCallback());
+		}
 	}
 }
 
@@ -1971,13 +1968,161 @@ void UCookOnTheFlyServer::TickNetwork()
 	}
 }
 
-bool UCookOnTheFlyServer::BeginPackageCacheForCookedPlatformData(UE::Cook::FPackageData& PackageData, UE::Cook::FCookerTimer& Timer)
+void UCookOnTheFlyServer::ConditionalSplitPackage(UE::Cook::FPackageData& PackageData, UObject* Obj, bool& bHasError)
+{
+	bHasError = false;
+
+	// Check if Package is a Generator Package that needs to be split
+	UE::Cook::Private::FRegisteredCookPackageSplitter* RegisteredSplitter = nullptr;
+	TArray<UE::Cook::Private::FRegisteredCookPackageSplitter*> FoundRegisteredSplitters;
+	RegisteredSplitDataClasses.MultiFind(Obj->GetClass(), FoundRegisteredSplitters);
+	for (UE::Cook::Private::FRegisteredCookPackageSplitter* FoundRegisteredSplitter : FoundRegisteredSplitters)
+	{
+		if (FoundRegisteredSplitter && FoundRegisteredSplitter->ShouldSplitPackage(Obj))
+		{
+			if (RegisteredSplitter)
+			{
+				UE_LOG(LogCook, Error, TEXT("Found more than one registered Cook Package Splitter for package %s."), *PackageData.GetPackageName().ToString());
+				bHasError = true;
+				return;
+			}
+
+			RegisteredSplitter = FoundRegisteredSplitter;
+		}
+	}
+	
+	// No splitter found
+	if (!RegisteredSplitter)
+	{
+		return;
+	}
+
+	// TODO: Add support for cooking in the editor
+	if (IsCookingInEditor())
+	{
+		UE_LOG(LogCook, Error, TEXT("Cooking in editor doesn't support Cook Package Splitters."));
+		bHasError = true;
+		return;
+	}
+	
+	// Package was already marked as a Generator package, only accept one registered splitter acting on one object of this package.
+	if (PackageData.GetIsGeneratorPackage())
+	{
+		UE_LOG(LogCook, Error, TEXT("Generator package %s already split by another split data object."), *PackageData.GetPackageName().ToString());
+		bHasError = true;
+		return;
+	}
+	PackageData.SetIsGeneratorPackage(true);
+
+	UE_LOG(LogCook, Display, TEXT("Splitting Package %s with class %s acting on object %s."), *PackageData.GetPackageName().ToString(), *RegisteredSplitter->GetSplitDataClass()->GetName(), *Obj->GetFullName());
+
+	// Create instance of CookPackageSplitter class
+	TUniquePtr<ICookPackageSplitter> CookPackageSplitterInstance(RegisteredSplitter->CreateInstance(Obj));
+	if (!CookPackageSplitterInstance)
+	{
+		UE_LOG(LogCook, Error, TEXT("Error instantiating Cook Package Splitter for object %s."), *Obj->GetFullName());
+		bHasError = true;
+		return;
+	}
+
+	// Initialize the instance with the data object
+	CookPackageSplitterInstance->SetDataObject(Obj);
+
+	// Prepare necessary paths
+	const TCHAR* GeneratedKeyWord = TEXT("_Generated_");
+	const UPackage* GeneratorPackage = PackageData.GetPackage();
+	const FString GeneratorPackagePath = FPackageName::GetLongPackagePath(GeneratorPackage->GetPathName());
+	const FString GeneratorPackageShortName = FPackageName::GetShortName(GeneratorPackage->GetName());
+	const FString GeneratedCookedRootPath = FPaths::RemoveDuplicateSlashes(FString::Printf(TEXT("/%s/%s/%s/"), *GeneratorPackagePath, *GeneratorPackageShortName, GeneratedKeyWord));
+	const FString GeneratedUncookedRootPath = FPaths::RemoveDuplicateSlashes(FString::Printf(TEXT("/%s%s/"), *GeneratorPackageShortName, GeneratedKeyWord));
+	const FString GeneratedUncookedContentPath = FPaths::RemoveDuplicateSlashes(FString::Printf(TEXT("%s/Cooked/%s/%s/"), *FPaths::ProjectIntermediateDir(), *GeneratorPackageShortName, GeneratedKeyWord));
+
+	// Empty the uncooked intermediate directory
+	const FString GeneratedDirectoryToDelete = FPaths::ConvertRelativePathToFull(FPaths::CreateStandardFilename(GeneratedUncookedContentPath));
+	GetAsyncIODelete().DeleteDirectory(*GeneratedDirectoryToDelete);
+
+	// Register Mount Point for Generator Package : /<GeneratorPackageShortName>_Generated_/ -> <ProjectIntermediateDir>/Cooked/<GeneratorPackageShortName>/_Generated_/
+	if (!FPackageName::MountPointExists(GeneratedUncookedRootPath))
+	{
+		FPackageName::RegisterMountPoint(GeneratedUncookedRootPath, GeneratedUncookedContentPath);
+	}
+
+	// Get packages to generate
+	TArray<ICookPackageSplitter::FGeneratedPackage> PackagesToGenerate = CookPackageSplitterInstance->GetGenerateList();
+
+	// Start Splitting
+	for (const ICookPackageSplitter::FGeneratedPackage& PackageToGenerate : PackagesToGenerate)
+	{
+		// Make package name to fit Mounted GeneratedUncookedRootPath
+		const FString GeneratedPackageName = FPaths::RemoveDuplicateSlashes(FString::Printf(TEXT("/%s/%s"), *GeneratedUncookedRootPath, *PackageToGenerate.RelativePath));
+
+		// Create package
+		check(!FindObject<UPackage>(nullptr, *GeneratedPackageName));
+		UPackage* GeneratedPackage = CreatePackage(*GeneratedPackageName);
+
+		// Populate package using CookPackageSplitterInstance and pass GeneratedPackage's cooked name for it to properly setup any internal reference to this package (SoftObjectPaths or others)
+		const FString GeneratedPackageCookedName = FPaths::RemoveDuplicateSlashes(FString::Printf(TEXT("/%s/%s"), *GeneratedCookedRootPath, *PackageToGenerate.RelativePath));
+		if (!CookPackageSplitterInstance->TryPopulatePackage(GeneratedPackage, PackageToGenerate.RelativePath, GeneratedPackageCookedName))
+		{
+			UE_LOG(LogCook, Error, TEXT("Error populating generated package %s for object %s."), *PackageToGenerate.RelativePath, *Obj->GetFullName());
+			bHasError = true;
+			return;
+		}
+
+		// Save package into the uncooked intermediate directory
+		const FString GeneratedPackageExtension = GeneratedPackage->ContainsMap() ? FPackageName::GetMapPackageExtension() : FPackageName::GetAssetPackageExtension();
+		const FString GeneratedPackageUncookedFileName = FPaths::RemoveDuplicateSlashes(FPackageName::LongPackageNameToFilename(GeneratedPackage->GetName(), GeneratedPackageExtension));
+		if (!UPackage::SavePackage(GeneratedPackage, nullptr, RF_Standalone, *GeneratedPackageUncookedFileName, GError, nullptr, false, true, SAVE_None))
+		{
+			UE_LOG(LogCook, Error, TEXT("Error saving generated package %s for object %s."), *GeneratedPackageUncookedFileName, *Obj->GetFullName());
+			bHasError = true;
+			return;
+		}
+
+		// Update Asset Registry with the generated package
+		AssetRegistry->ScanFilesSynchronous({ GeneratedPackageUncookedFileName }, true);
+
+		// Make sure PackageNameCache successfully caches the generated package
+		if (GetPackageNameCache().GetCachedStandardFileName(GeneratedPackage).IsNone())
+		{
+			UE_LOG(LogCook, Error, TEXT("Error caching generated package name %s for object %s."), *GeneratedPackageUncookedFileName, *Obj->GetFullName());
+			bHasError = true;
+			return;
+		}
+
+		// Add PackageData for GeneratedPackage
+		UE::Cook::FPackageData* GeneratedPackageData = PackageDatas->TryAddPackageDataByFileName(FName(GeneratedPackageUncookedFileName));
+		if (!GeneratedPackageData)
+		{
+			UE_LOG(LogCook, Error, TEXT("Error adding generated package data %s for object %s."), *GeneratedPackageUncookedFileName, *Obj->GetFullName());
+			bHasError = true;
+			return;
+		}
+		// Set override cooked FileName to be used when saving the cooked generated package
+		const FString GeneratePackageCookedFileName = FPaths::RemoveDuplicateSlashes(FString::Printf(TEXT("%s/%s/%s/%s%s"), 
+													  *FPaths::GetPath(PackageData.GetFileName().ToString()), *GeneratorPackageShortName, GeneratedKeyWord, *PackageToGenerate.RelativePath, *GeneratedPackageExtension));
+		GeneratedPackageData->SetCookedFileName(FPackageNameCache::GetStandardFileName(GeneratePackageCookedFileName));
+
+		// Queue the generated package
+		QueueDiscoveredPackage(GeneratedPackage, /*bIsGeneratedPackage*/true);
+		check(GeneratedPackageData->IsInProgress());
+	}
+
+	CookPackageSplitterInstance->FinalizeGeneratorPackage();
+}
+
+bool UCookOnTheFlyServer::BeginPrepareSave(UE::Cook::FPackageData& PackageData, UE::Cook::FCookerTimer& Timer)
 {
 	if (PackageData.GetCookedPlatformDataCalled())
 	{
 		return true;
 	}
 
+	if (PackageData.GetHasBeginPrepareSaveFailed())
+	{
+		return false;
+	}
+
 	if (!PackageData.GetCookedPlatformDataStarted())
 	{
 		if (PackageData.GetNumPendingCookedPlatformData() > 0)
@@ -1989,7 +2134,7 @@ bool UCookOnTheFlyServer::BeginPackageCacheForCookedPlatformData(UE::Cook::FPack
 		PackageData.SetCookedPlatformDataStarted(true);
 	}
 
-	UE_SCOPED_HIERARCHICAL_COOKTIMER_AND_DURATION(BeginPackageCacheForCookedPlatformData, DetailedCookStats::TickCookOnTheSideBeginPackageCacheForCookedPlatformDataTimeSec);
+	UE_SCOPED_HIERARCHICAL_COOKTIMER_AND_DURATION(BeginPrepareSave, DetailedCookStats::TickCookOnTheSideBeginPrepareSaveTimeSec);
 
 #if DEBUG_COOKONTHEFLY 
 	UE_LOG(LogCook, Display, TEXT("Caching objects for package %s"), *Package->GetFName().ToString());
@@ -1999,17 +2144,40 @@ bool UCookOnTheFlyServer::BeginPackageCacheForCookedPlatformData(UE::Cook::FPack
 	check(PackageData.GetState() == UE::Cook::EPackageState::Save);
 	PackageData.CreateObjectCache();
 
+	TArray<FWeakObjectPtr>& CachedObjectsInOuter = PackageData.GetCachedObjectsInOuter();
+	int32& CookedPlatformDataNextIndex = PackageData.GetCookedPlatformDataNextIndex();
+
+	if (CookedPlatformDataNextIndex == -1)
+	{
+		// Check for Splitting the package; this needs to happen before we make any of the BeginCacheForCookedPlatformData calls in the package
+		for (FWeakObjectPtr& WeakObj : CachedObjectsInOuter)
+		{
+			UObject* Obj = WeakObj.Get();
+			if (!Obj)
+			{
+				continue;
+			}
+
+			bool bHasError = false;
+			ConditionalSplitPackage(PackageData, Obj, bHasError);
+			if (bHasError)
+			{
+				// ConditionalSplitPackage failure marks the package data for PumpSaves to handle this as an error and to put it in idle state.
+				PackageData.SetHasBeginPrepareSaveFailed(true);
+				return false;
+			}
+		}
+		CookedPlatformDataNextIndex = 0;
+	}
+
 	// Note that we cache cooked data for all requested platforms, rather than only for the requested platforms that have not cooked yet.  This allows
 	// us to avoid the complexity of needing to cancel the Save and keep track of the old list of uncooked platforms whenever the cooked platforms change
-	// while BeginPackageCacheForCookedPlatformData is active.
+	// while BeginPrepareSave is active.
 	// Currently this does not cause significant cost since saving new platforms with some platforms already saved is a rare operation.
 
+	FWeakObjectPtr* CachedObjectsInOuterData = CachedObjectsInOuter.GetData();
 	const TArray<const ITargetPlatform*>& TargetPlatforms = PackageData.GetRequestedPlatforms();
 	int NumPlatforms = TargetPlatforms.Num();
-	TArray<FWeakObjectPtr>& CachedObjectsInOuter = PackageData.GetCachedObjectsInOuter();
-	int32& CookedPlatformDataNextIndex = PackageData.GetCookedPlatformDataNextIndex();
-	FWeakObjectPtr* CachedObjectsInOuterData = CachedObjectsInOuter.GetData();
-
 	int NumIndexes = CachedObjectsInOuter.Num() * NumPlatforms;
 	while (CookedPlatformDataNextIndex < NumIndexes)
 	{
@@ -2076,7 +2244,7 @@ bool UCookOnTheFlyServer::BeginPackageCacheForCookedPlatformData(UE::Cook::FPack
 	return true;
 }
 
-bool UCookOnTheFlyServer::FinishPackageCacheForCookedPlatformData(UE::Cook::FPackageData& PackageData, UE::Cook::FCookerTimer& Timer)
+bool UCookOnTheFlyServer::FinishPrepareSave(UE::Cook::FPackageData& PackageData, UE::Cook::FCookerTimer& Timer)
 {
 	if (PackageData.GetCookedPlatformDataComplete())
 	{
@@ -2085,7 +2253,7 @@ bool UCookOnTheFlyServer::FinishPackageCacheForCookedPlatformData(UE::Cook::FPac
 
 	if (!PackageData.GetCookedPlatformDataCalled())
 	{
-		if (!BeginPackageCacheForCookedPlatformData(PackageData, Timer))
+		if (!BeginPrepareSave(PackageData, Timer))
 		{
 			return false;
 		}
@@ -2105,6 +2273,24 @@ void UCookOnTheFlyServer::ReleaseCookedPlatformData(UE::Cook::FPackageData& Pack
 {
 	using namespace UE::Cook;
 
+	ON_SCOPE_EXIT
+	{
+		if (CurrentCookMode == ECookMode::CookByTheBook)
+		{
+			// For each object for which data is cached we can call FinishedCookedPlatformDataCache
+			// we can only safely call this when we are finished caching the object completely.
+			// this doesn't ever happen for cook in editor or cook on the fly mode
+			for (FWeakObjectPtr& WeakPtr : PackageData.GetCachedObjectsInOuter())
+			{
+				UObject* Obj = WeakPtr.Get();
+				if (Obj)
+				{
+					Obj->WillNeverCacheCookedPlatformDataAgain();
+				}
+			}
+		}
+	};
+
 	if (!PackageData.GetCookedPlatformDataStarted())
 	{
 		PackageData.CheckCookedPlatformDataEmpty();
@@ -2271,7 +2457,6 @@ bool UCookOnTheFlyServer::LoadPackageForCooking(UE::Cook::FPackageData& PackageD
 	return bSuccess;
 }
 
-
 void UCookOnTheFlyServer::ProcessUnsolicitedPackages()
 {
 	// Ensure sublevels are loaded by iterating all recently loaded packages and invoking
@@ -2288,7 +2473,8 @@ void UCookOnTheFlyServer::ProcessUnsolicitedPackages()
 			{
 				PostLoadPackageFixup(Package);
 			}
-			FilterLoadedPackage(Package, false);
+
+			QueueDiscoveredPackage(Package);
 		}
 	}
 }
@@ -2407,12 +2593,20 @@ void UCookOnTheFlyServer::PumpSaves(UE::Cook::FTickStackData& StackData, uint32
 		// Release any completed pending CookedPlatformDatas, so that slots in the per-class limits on calls to BeginCacheForCookedPlatformData are freed up for new objects to use
 		PackageDatas->PollPendingCookedPlatformDatas();
 
-		// Always wait for FinishPackageCacheForCookedPlatformData before attempting to save the package
-		bool AllObjectsCookedDataCached = FinishPackageCacheForCookedPlatformData(PackageData, StackData.Timer);
+		// Always wait for FinishPrepareSave before attempting to save the package
+		bool AllObjectsCookedDataCached = FinishPrepareSave(PackageData, StackData.Timer);
 
 		// If the CookPlatformData is not ready then postpone the package, exit, or wait for it as appropriate
 		if (!AllObjectsCookedDataCached)
 		{
+			if (PackageData.GetHasBeginPrepareSaveFailed())
+			{
+				ReleaseCookedPlatformData(PackageData);
+				PackageData.AddCookedPlatforms(PackageData.GetRequestedPlatforms(), false);
+				PackageData.SendToState(EPackageState::Idle, ESendFlags::QueueAdd);
+				++OutNumPushed;
+				continue;
+			}
 			// Can we postpone?
 			if (!PackageData.GetIsUrgent())
 			{
@@ -2429,14 +2623,14 @@ void UCookOnTheFlyServer::PumpSaves(UE::Cook::FTickStackData& StackData, uint32
 				UE_SCOPED_HIERARCHICAL_COOKTIMER(WaitingForCachedCookedPlatformData);
 				do
 				{
-					// FinishPackageCacheForCookedPlatformData might block on pending CookedPlatformDatas, and it might block on BeginPackageCacheForCookedPlatformData, which can
+					// FinishPrepareSave might block on pending CookedPlatformDatas, and it might block on BeginPrepareSave, which can
 					// block on resources held by other CookedPlatformDatas. Calling PollPendingCookedPlatformDatas should handle pumping all of those.
-					check(PackageDatas->GetPendingCookedPlatformDatas().Num() || !PackageData.GetCookedPlatformDataCalled()); // FinishPackageCacheForCookedPlatformData can only return false in one of these cases
+					check(PackageDatas->GetPendingCookedPlatformDatas().Num() || !PackageData.GetCookedPlatformDataCalled()); // FinishPrepareSave can only return false in one of these cases
 					// sleep for a bit
 					FPlatformProcess::Sleep(0.0f);
 					// Poll the results again and check whether we are now done
 					PackageDatas->PollPendingCookedPlatformDatas();
-					AllObjectsCookedDataCached = FinishPackageCacheForCookedPlatformData(PackageData, StackData.Timer);
+					AllObjectsCookedDataCached = FinishPrepareSave(PackageData, StackData.Timer);
 				} while (!StackData.Timer.IsTimeUp() && !AllObjectsCookedDataCached);
 			}
 			// If we couldn't postpone or wait, then we need to exit and try again later
@@ -2448,7 +2642,7 @@ void UCookOnTheFlyServer::PumpSaves(UE::Cook::FTickStackData& StackData, uint32
 				return;
 			}
 		}
-		check(AllObjectsCookedDataCached == true); // We are not allowed to save until FinishPackageCacheForCookedPlatformData returns true.  We should have early exited above if it didn't
+		check(AllObjectsCookedDataCached == true); // We are not allowed to save until FinishPrepareSave returns true.  We should have early exited above if it didn't
 
 		// precache the next few packages
 		if (!IsCookOnTheFlyMode() && SaveQueue.Num() != 0)
@@ -2463,7 +2657,7 @@ void UCookOnTheFlyServer::PumpSaves(UE::Cook::FTickStackData& StackData, uint32
 					break;
 				}
 				--LeftToPrecache;
-				BeginPackageCacheForCookedPlatformData(*NextData, StackData.Timer);
+				BeginPrepareSave(*NextData, StackData.Timer);
 			}
 
 			// If we're in RealTimeMode, check whether the precaching overflowed our timer and if so exit before we do the potentially expensive SavePackage
@@ -2557,20 +2751,6 @@ void UCookOnTheFlyServer::PumpSaves(UE::Cook::FTickStackData& StackData, uint32
 		if (!IsCookingInEditor())
 		{
 			ReleaseCookedPlatformData(PackageData);
-			if (CurrentCookMode == ECookMode::CookByTheBook)
-			{
-				// For each object for which data is cached we can call FinishedCookedPlatformDataCache
-				// we can only safely call this when we are finished caching the object completely.
-				// this doesn't ever happen for cook in editor or cook on the fly mode
-				for (FWeakObjectPtr& WeakPtr: PackageData.GetCachedObjectsInOuter())
-				{
-					UObject* Obj = WeakPtr.Get();
-					if (Obj)
-					{
-						Obj->WillNeverCacheCookedPlatformDataAgain();
-					}
-				}
-			}
 		}
 
 		FName FileName = PackageData.GetFileName();
@@ -3579,6 +3759,12 @@ void UCookOnTheFlyServer::SaveCookedPackage(UE::Cook::FPackageData& PackageData,
 			UE_LOG(LogCook, Fatal, TEXT("Package %s marked as reloading for cook by was requested to save"), *Package->GetName());
 		}
 
+		// If valid, override cooked FileName. Currently used for Generated Packages (see ConditionalSplitPackage)
+		if (PackageData.GetCookedFileName() != NAME_None)
+		{
+			Filename = PackageData.GetCookedFileName().ToString();
+		}
+
 		// Use SandboxFile to do path conversion to properly handle sandbox paths (outside of standard paths in particular).
 		Filename = ConvertToFullSandboxPath(*Filename, true);
 
@@ -3943,6 +4129,20 @@ void UCookOnTheFlyServer::Initialize( ECookMode::Type DesiredCookMode, ECookInit
 	{
 		UE_LOG(LogCook, Warning, TEXT("Cooking with Event Driven Loader disabled. Loading code will use deprecated path which will be removed in future release."));
 	}
+
+	// Prepare a map SplitDataClass to FRegisteredCookPackageSplitter* for ConditionalSplitPackage to use
+	RegisteredSplitDataClasses.Reset();
+	UE::Cook::Private::FRegisteredCookPackageSplitter::ForEach([this](UE::Cook::Private::FRegisteredCookPackageSplitter* RegisteredCookPackageSplitter)
+	{
+		UClass* SplitDataClass = RegisteredCookPackageSplitter->GetSplitDataClass();
+		for (TObjectIterator<UClass> ClassIt; ClassIt; ++ClassIt)
+		{
+			if (ClassIt->IsChildOf(SplitDataClass) && !ClassIt->HasAnyClassFlags(CLASS_Abstract | CLASS_Deprecated | CLASS_NewerVersionExists))
+			{
+				RegisteredSplitDataClasses.Add(SplitDataClass, RegisteredCookPackageSplitter);
+			}
+		}
+	});
 }
 
 bool UCookOnTheFlyServer::Exec(class UWorld* InWorld, const TCHAR* Cmd, FOutputDevice& Ar)
@@ -4849,6 +5049,20 @@ void UCookOnTheFlyServer::DeleteSandboxDirectory(const FString& PlatformName)
 	LocalAsyncIODelete.DeleteDirectory(SandboxDirectory + TEXT("AsyncDelete"));
 }
 
+void UCookOnTheFlyServer::TrySetDefaultAsyncIODeletePlatform(const FString& PlatformName)
+{
+	if (DefaultAsyncIODeletePlatformName.IsEmpty())
+	{
+		DefaultAsyncIODeletePlatformName = PlatformName;
+	}
+}
+
+FAsyncIODelete& UCookOnTheFlyServer::GetAsyncIODelete()
+{
+	check(!DefaultAsyncIODeletePlatformName.IsEmpty());
+	return GetAsyncIODelete(DefaultAsyncIODeletePlatformName);
+}
+
 FAsyncIODelete& UCookOnTheFlyServer::GetAsyncIODelete(const FString& PlatformName, const FString* AsyncDeleteDirectory)
 {
 	FAsyncIODelete* AsyncIODeletePtr = AsyncIODelete.Get();
@@ -6818,6 +7032,11 @@ void UCookOnTheFlyServer::StartCookByTheBook( const FCookByTheBookStartupOptions
 	bPackageFilterDirty = true;
 	check(PlatformManager->GetSessionPlatforms().Num() == TargetPlatforms.Num());
 
+	if (ensure(TargetPlatforms.Num() > 0))
+	{
+		TrySetDefaultAsyncIODeletePlatform(TargetPlatforms[0]->PlatformName());
+	}
+
 	// We want to set bRunning = true as early as possible, but it implies that session platforms have been selected so this is the earliest point we can set it
 	CookByTheBookOptions->bRunning = true;
 
@@ -7320,6 +7539,8 @@ bool UCookOnTheFlyServer::HandleNetworkFileServerNewConnection(const FString& Ve
 		}
 	}
 
+	TrySetDefaultAsyncIODeletePlatform(PlatformName);
+
 	UE_LOG(LogCook, Display, TEXT("Connection received of version %s local version %s"), *VersionInfo, *LocalVersionInfo);
 
 	if (LocalVersionInfo != VersionInfo)
@@ -8175,4 +8396,4 @@ uint32 UCookOnTheFlyServer::FullLoadAndSave(uint32& CookedPackageCount)
 	return Result;
 }
 
-#undef LOCTEXT_NAMESPACE
+#undef LOCTEXT_NAMESPACE
\ No newline at end of file
diff --git a/Engine/Source/Editor/UnrealEd/Private/Cooker/CookPackageData.cpp b/Engine/Source/Editor/UnrealEd/Private/Cooker/CookPackageData.cpp
--- a/Engine/Source/Editor/UnrealEd/Private/Cooker/CookPackageData.cpp
+++ b/Engine/Source/Editor/UnrealEd/Private/Cooker/CookPackageData.cpp
@@ -27,7 +27,7 @@ namespace Cook
 
 	FPackageData::FPackageData(FPackageDatas& PackageDatas, const FName& InPackageName, const FName& InFileName)
 		: PackageName(InPackageName), FileName(InFileName), PackageDatas(PackageDatas)
-		, bIsUrgent(0), bIsVisited(0), bIsPreloadAttempted(0), bIsPreloaded(0), bHasSaveCache(0), bCookedPlatformDataStarted(0), bCookedPlatformDataCalled(0), bCookedPlatformDataComplete(0), bMonitorIsCooked(0)
+		, bIsUrgent(0), bIsVisited(0), bIsPreloadAttempted(0), bIsPreloaded(0), bIsGeneratorPackage(0), bHasSaveCache(0), bHasBeginPrepareSaveFailed(0), bCookedPlatformDataStarted(0), bCookedPlatformDataCalled(0), bCookedPlatformDataComplete(0), bMonitorIsCooked(0)
 	{
 		SetState(EPackageState::Idle);
 		SendToState(EPackageState::Idle, ESendFlags::QueueAdd);
@@ -594,6 +594,8 @@ namespace Cook
 	{
 		check(GetPackage() != nullptr && GetPackage()->IsFullyLoaded());
 
+		SetHasBeginPrepareSaveFailed(false);
+		SetIsGeneratorPackage(false);
 		CheckObjectCacheEmpty();
 		CheckCookedPlatformDataEmpty();
 	}
@@ -602,6 +604,8 @@ namespace Cook
 	{
 		PackageDatas.GetCookOnTheFlyServer().ReleaseCookedPlatformData(*this);
 		ClearObjectCache();
+		SetHasBeginPrepareSaveFailed(false);
+		SetIsGeneratorPackage(false);
 	}
 
 	void FPackageData::OnEnterInProgress()
@@ -823,7 +827,7 @@ namespace Cook
 
 	void FPackageData::CheckCookedPlatformDataEmpty() const
 	{
-		check(GetCookedPlatformDataNextIndex() == 0);
+		check(GetCookedPlatformDataNextIndex() == -1);
 		check(!GetCookedPlatformDataStarted());
 		check(!GetCookedPlatformDataCalled());
 		check(!GetCookedPlatformDataComplete());
@@ -831,7 +835,7 @@ namespace Cook
 
 	void FPackageData::ClearCookedPlatformData()
 	{
-		CookedPlatformDataNextIndex = 0;
+		CookedPlatformDataNextIndex = -1;
 		// Note that GetNumPendingCookedPlatformData is not cleared; it persists across Saves and CookSessions
 		SetCookedPlatformDataStarted(false);
 		SetCookedPlatformDataCalled(false);
diff --git a/Engine/Source/Editor/UnrealEd/Private/Cooker/CookPackageData.h b/Engine/Source/Editor/UnrealEd/Private/Cooker/CookPackageData.h
--- a/Engine/Source/Editor/UnrealEd/Private/Cooker/CookPackageData.h
+++ b/Engine/Source/Editor/UnrealEd/Private/Cooker/CookPackageData.h
@@ -245,6 +245,18 @@ namespace Cook
 		/** Check whether savestate contracts on the PackageData were invalidated by by e.g. garbage collection of objects in its package. */
 		bool IsSaveInvalidated() const;
 
+		/** Get/Set the flag for whether BeginPrepareSave has been called and returned false because an error occured. */
+		bool GetHasBeginPrepareSaveFailed() const { return static_cast<bool>(bHasBeginPrepareSaveFailed); }
+		void SetHasBeginPrepareSaveFailed(bool bValue) { bHasBeginPrepareSaveFailed = bValue != 0; }
+
+		/** Get/Set override filename to be used when saving cooked package. */
+		FName GetCookedFileName() const { return CookedFileName; }
+		void SetCookedFileName(FName InCookedFileName) { CookedFileName = InCookedFileName; }
+
+		/** Get/Set the flag for whether this package is a Generator package. */
+		bool GetIsGeneratorPackage() const { return static_cast<bool>(bIsGeneratorPackage); }
+		void SetIsGeneratorPackage(bool bValue) { bIsGeneratorPackage = bValue != 0; }
+
 		/** Validate that the fields related to the BeginCacheForCookedPlatformData calls are empty, as required when not in the save state. */
 		void CheckCookedPlatformDataEmpty() const;
 		/** Clear the fields related to the BeginCacheForCookedPlatformData calls, when e.g. leaving the save state. Caller is responsible for having already executed any required cancellation steps to avoid dangling pending operations. */
@@ -325,18 +337,21 @@ namespace Cook
 		FCompletionCallback CompletionCallback;
 		FName PackageName;
 		FName FileName;
+		FName CookedFileName; // Override filename to be used when saving cooked package (used by generated packages)
 		TWeakObjectPtr<UPackage> Package;
 		FPackageDatas& PackageDatas; // The one-per-CookOnTheFlyServer owner of this PackageData
 		TSharedPtr<FPreloadableFile> PreloadableFile;
 		int32 NumPendingCookedPlatformData = 0;
-		int32 CookedPlatformDataNextIndex = 0;
+		int32 CookedPlatformDataNextIndex = -1;
 
 		uint32 State : int32(EPackageState::BitCount);
 		uint32 bIsUrgent : 1;
 		uint32 bIsVisited : 1;
 		uint32 bIsPreloadAttempted : 1;
 		uint32 bIsPreloaded : 1;
+		uint32 bIsGeneratorPackage : 1;
 		uint32 bHasSaveCache : 1;
+		uint32 bHasBeginPrepareSaveFailed : 1;
 		uint32 bCookedPlatformDataStarted : 1;
 		uint32 bCookedPlatformDataCalled : 1;
 		uint32 bCookedPlatformDataComplete : 1;
diff --git a/Engine/Source/Editor/UnrealEd/Private/Cooker/CookPackageSplitter.cpp b/Engine/Source/Editor/UnrealEd/Private/Cooker/CookPackageSplitter.cpp
new file mode 100644
--- /dev/null
+++ b/Engine/Source/Editor/UnrealEd/Private/Cooker/CookPackageSplitter.cpp
@@ -0,0 +1,44 @@
+// Copyright Epic Games, Inc. All Rights Reserved.
+
+#include "CookPackageSplitter.h"
+
+#if WITH_EDITOR
+
+namespace UE
+{
+namespace Cook
+{
+namespace Private
+{
+
+static TLinkedList<FRegisteredCookPackageSplitter*>* GRegisteredCookPackageSplitterList = nullptr;
+
+FRegisteredCookPackageSplitter::FRegisteredCookPackageSplitter()
+: GlobalListLink(this)
+{
+	GlobalListLink.LinkHead(GetRegisteredList());
+}
+
+FRegisteredCookPackageSplitter::~FRegisteredCookPackageSplitter()
+{
+	GlobalListLink.Unlink();
+}
+
+TLinkedList<FRegisteredCookPackageSplitter*>*& FRegisteredCookPackageSplitter::GetRegisteredList()
+{
+	return GRegisteredCookPackageSplitterList;
+}
+
+void FRegisteredCookPackageSplitter::ForEach(TFunctionRef<void(FRegisteredCookPackageSplitter*)> Func)
+{
+	for (TLinkedList<FRegisteredCookPackageSplitter*>::TIterator It(GetRegisteredList()); It; It.Next())
+	{
+		Func(*It);
+	}
+}
+
+}
+}
+}
+
+#endif
\ No newline at end of file
diff --git a/Engine/Source/Editor/UnrealEd/Public/CookPackageSplitter.h b/Engine/Source/Editor/UnrealEd/Public/CookPackageSplitter.h
new file mode 100644
--- /dev/null
+++ b/Engine/Source/Editor/UnrealEd/Public/CookPackageSplitter.h
@@ -0,0 +1,114 @@
+// Copyright Epic Games, Inc. All Rights Reserved.
+
+#pragma once
+
+#if WITH_EDITOR
+
+#include "CoreMinimal.h"
+
+#include "Containers/List.h"
+#include "UObject/Object.h"
+#include "UObject/ObjectMacros.h"
+
+/**
+ * This class is used for packages that need to be split into multiple runtime packages.
+ * It provides the instructions to the cooker for how to split the package.
+ */
+class ICookPackageSplitter
+{
+public:
+	virtual ~ICookPackageSplitter() {}
+
+	/** Data sent to the cooker to describe each desired generated package */
+	struct FGeneratedPackage
+	{
+		FString RelativePath;		// Generated package relative to the Parent/_Generated_ root
+		TArray<FName> Dependencies; // LongPackageNames that the generated package references
+	};
+	
+	/** 
+	 * Return whether the CookPackageSplitter subclass should handle the given SplitDataClass instance. 
+	 * Note that this is a static function referenced by macros, not part of the virtual api.
+	 */
+	static bool ShouldSplit(UObject* SplitData) { return false; }
+	
+	/** Sets the SplitDataClass instance on which the CookPackageSplitter instance should work. */
+	virtual void SetDataObject(UObject* SplitData) = 0;
+
+	/** Return the list of packages to generate. */
+	virtual TArray<FGeneratedPackage> GetGenerateList() = 0;
+
+	/**
+	 * Try to populate a generated package.
+	 *
+	 * Receive an empty UPackage generated from an element in GetGenerateList and populate it
+	 * After returning, the given package will be queued for saving into the TargetDomain
+	 * Note that TryPopulatePackage will not be called for every package on
+	 * an iterative cook; it will only be called for the packages with changed dependencies. 
+	 *
+	 * @param GeneratedPackage			The package the cooker generated for the given FGeneratedPackage
+	 * @param RelativePath				The relative path from the given FGeneratedPackage
+	 * @param GeneratedPackageCookName	The name that will be given for the cooked package (differs from the uncooked package)
+	 * @return							True if successfully populates,  false on error (this will cause a cook error).
+	 */
+	virtual bool TryPopulatePackage(UPackage* GeneratedPackage, const FStringView& RelativePath, const FStringView& GeneratedPackageCookName) = 0;
+
+	/**
+	 * Called after calling all TryPopulatePackage.
+	 * Make any remaining adjustments to the parent package before returning and allowing
+	 * it to be saved into the target domain.
+	 */
+	virtual void FinalizeGeneratorPackage() {}
+};
+
+namespace UE
+{
+namespace Cook
+{
+namespace Private
+{
+
+/** Interface for internal use only (used by REGISTER_COOKPACKAGE_SPLITTER to register an ICookPackageSplitter for a class) */
+class UNREALED_API FRegisteredCookPackageSplitter
+{
+public:
+	FRegisteredCookPackageSplitter();
+	virtual ~FRegisteredCookPackageSplitter();
+
+	virtual UClass* GetSplitDataClass() const = 0;
+	virtual bool ShouldSplitPackage(UObject* Object) const = 0;
+	virtual ICookPackageSplitter* CreateInstance(UObject* Object) const = 0;
+
+	static void ForEach(TFunctionRef<void(FRegisteredCookPackageSplitter*)> Func);
+
+private:
+	
+	static TLinkedList<FRegisteredCookPackageSplitter*>*& GetRegisteredList();
+	TLinkedList<FRegisteredCookPackageSplitter*> GlobalListLink;
+};
+
+}
+}
+}
+
+/**
+ * Used to Register an ICookPackageSplitter for a class
+ *
+ * Example usage:
+ *
+ * class FMyCookPackageSplitter : public ICookPackageSplitter { ... }
+ * REGISTER_COOKPACKAGE_SPLITTER(FMyCookPackageSplitter, UMySplitDataClass);
+ */
+#define REGISTER_COOKPACKAGE_SPLITTER(SplitterClass, SplitDataClass) \
+class PREPROCESSOR_JOIN(PREPROCESSOR_JOIN(SplitterClass, SplitDataClass), _Register) : public UE::Cook::Private::FRegisteredCookPackageSplitter \
+{ \
+	virtual UClass* GetSplitDataClass() const override { return SplitDataClass::StaticClass(); } \
+	virtual bool ShouldSplitPackage(UObject* Object) const override { return SplitterClass::ShouldSplit(Object); } \
+	virtual ICookPackageSplitter* CreateInstance(UObject* SplitData) const override { return new SplitterClass(); } \
+}; \
+namespace PREPROCESSOR_JOIN(SplitterClass, SplitDataClass) \
+{ \
+	static PREPROCESSOR_JOIN(PREPROCESSOR_JOIN(SplitterClass, SplitDataClass), _Register) DefaultObject; \
+}
+
+#endif
\ No newline at end of file
diff --git a/Engine/Source/Runtime/Core/Private/Misc/CoreGlobals.cpp b/Engine/Source/Runtime/Core/Private/Misc/CoreGlobals.cpp
--- a/Engine/Source/Runtime/Core/Private/Misc/CoreGlobals.cpp
+++ b/Engine/Source/Runtime/Core/Private/Misc/CoreGlobals.cpp
@@ -9,6 +9,7 @@
 #include "Misc/CoreStats.h"
 #include "Misc/Compression.h"
 #include "Misc/LazySingleton.h"
+#include "Misc/CommandLine.h"
 #include "ProfilingDebugging/MiscTrace.h"
 #include "GenericPlatform/GenericPlatformCrashContext.h"
 
@@ -422,6 +423,13 @@ void CORE_API EnableEmitDrawEventsOnlyOnCommandlist()
 	GEmitDrawEvents = true;
 }
 
+bool CORE_API IsRunningCookCommandlet()
+{
+	FString Commandline = FCommandLine::Get();
+	const bool bIsCookCommandlet = IsRunningCommandlet() && Commandline.Contains(TEXT("run=cook"));
+	return bIsCookCommandlet;
+}
+
 void ToggleGDebugPUCrashedFlag(const TArray<FString>& Args)
 {
 	GIsGPUCrashed = !GIsGPUCrashed;
diff --git a/Engine/Source/Runtime/Core/Private/Misc/Paths.cpp b/Engine/Source/Runtime/Core/Private/Misc/Paths.cpp
--- a/Engine/Source/Runtime/Core/Private/Misc/Paths.cpp
+++ b/Engine/Source/Runtime/Core/Private/Misc/Paths.cpp
@@ -1154,6 +1154,13 @@ bool FPaths::CollapseRelativeDirectories(FString& InPath)
 	return true;
 }
 
+FString FPaths::RemoveDuplicateSlashes(const FString& InPath)
+{
+	FString Result = InPath;
+	RemoveDuplicateSlashes(Result);
+	return MoveTemp(Result);
+}
+
 void FPaths::RemoveDuplicateSlashes(FString& InPath)
 {
 	TCHAR* Text = InPath.GetCharArray().GetData();
diff --git a/Engine/Source/Runtime/Core/Public/CoreGlobals.h b/Engine/Source/Runtime/Core/Public/CoreGlobals.h
--- a/Engine/Source/Runtime/Core/Public/CoreGlobals.h
+++ b/Engine/Source/Runtime/Core/Public/CoreGlobals.h
@@ -208,6 +208,11 @@ FORCEINLINE bool IsRunningCommandlet()
 #endif
 }
 
+/**
+* Check to see if this executable is running the cookcommandlet
+*/
+CORE_API bool IsRunningCookCommandlet();
+
 /**
  * Check to see if we should initialise RHI and set up scene for rendering even when running a commandlet.
  */
diff --git a/Engine/Source/Runtime/Core/Public/Misc/Paths.h b/Engine/Source/Runtime/Core/Public/Misc/Paths.h
--- a/Engine/Source/Runtime/Core/Public/Misc/Paths.h
+++ b/Engine/Source/Runtime/Core/Public/Misc/Paths.h
@@ -481,6 +481,9 @@ public:
 	 */
 	static void RemoveDuplicateSlashes(FString& InPath);
 
+	 /** Returns a copy of the given path on which duplicate slashes were removed (see the inplace version for more details). */
+	static FString RemoveDuplicateSlashes(const FString& InPath);
+
 	/**
 	 * Make fully standard "Unreal" pathname:
 	 *    - Normalizes path separators [NormalizeFilename]
diff --git a/Engine/Source/Runtime/Engine/Classes/Engine/LevelStreaming.h b/Engine/Source/Runtime/Engine/Classes/Engine/LevelStreaming.h
--- a/Engine/Source/Runtime/Engine/Classes/Engine/LevelStreaming.h
+++ b/Engine/Source/Runtime/Engine/Classes/Engine/LevelStreaming.h
@@ -233,6 +233,12 @@ public:
 	uint8 bDrawOnLevelStatusMap : 1;
 
 #if WITH_EDITORONLY_DATA
+	/** If True, IsValidStreamingLevel will return true. This is necessary when cooking/saving a World with StreamingLevels
+	 *  referencing generated packages (done during cook process) that doesn't exist yet on disk (i.e. the workspace
+	 *  domain differs from the target domain).
+	 */
+	uint8 bForceIsValidStreamingLevel : 1;
+
 	/** Deprecated level color used for visualization. */
 	UPROPERTY()
 	FColor DrawColor_DEPRECATED;
diff --git a/Engine/Source/Runtime/Engine/Private/LevelStreaming.cpp b/Engine/Source/Runtime/Engine/Private/LevelStreaming.cpp
--- a/Engine/Source/Runtime/Engine/Private/LevelStreaming.cpp
+++ b/Engine/Source/Runtime/Engine/Private/LevelStreaming.cpp
@@ -277,6 +277,7 @@ ULevelStreaming::ULevelStreaming(const FObjectInitializer& ObjectInitializer)
 {
 # if WITH_EDITORONLY_DATA
 	bShouldBeVisibleInEditor = true;
+	bForceIsValidStreamingLevel = false;
 #endif
 	LevelColor = FLinearColor::White;
 	LevelTransform = FTransform::Identity;
@@ -1735,11 +1736,16 @@ ALevelScriptActor* ULevelStreaming::GetLevelScriptActor()
 
 bool ULevelStreaming::IsValidStreamingLevel() const
 {
-	const bool PIESession = GetWorld()->WorldType == EWorldType::PIE || GetOutermost()->HasAnyPackageFlags(PKG_PlayInEditor);
-	if (!PIESession && !WorldAsset.IsNull())
+#if WITH_EDITOR
+	if (!bForceIsValidStreamingLevel)
+#endif
 	{
-		const FString WorldPackageName = GetWorldAssetPackageName();
-		return FPackageName::DoesPackageExist(WorldPackageName);
+		const bool PIESession = GetWorld()->WorldType == EWorldType::PIE || GetOutermost()->HasAnyPackageFlags(PKG_PlayInEditor);
+		if (!PIESession && !WorldAsset.IsNull())
+		{
+			const FString WorldPackageName = GetWorldAssetPackageName();
+			return FPackageName::DoesPackageExist(WorldPackageName);
+		}
 	}
 	return true;
 }
diff --git a/Engine/Source/Runtime/Engine/Private/World.cpp b/Engine/Source/Runtime/Engine/Private/World.cpp
--- a/Engine/Source/Runtime/Engine/Private/World.cpp
+++ b/Engine/Source/Runtime/Engine/Private/World.cpp
@@ -7861,9 +7861,9 @@ void UWorld::GetAssetRegistryTags(TArray<FAssetRegistryTag>& OutTags) const
 
 		FBox LevelBounds;
 
-		if (UWorldPartitionSubsystem* WorldPartitionSubsystem = GetSubsystem<UWorldPartitionSubsystem>())
+		if (UWorldPartition* WorldPartition = GetWorldPartition())
 		{
-			LevelBounds = WorldPartitionSubsystem->GetWorldBounds();
+			LevelBounds = WorldPartition->GetWorldBounds();
 		}
 		else
 		{
diff --git a/Engine/Source/Runtime/Engine/Private/WorldPartition/DataLayer/WorldDataLayers.cpp b/Engine/Source/Runtime/Engine/Private/WorldPartition/DataLayer/WorldDataLayers.cpp
--- a/Engine/Source/Runtime/Engine/Private/WorldPartition/DataLayer/WorldDataLayers.cpp
+++ b/Engine/Source/Runtime/Engine/Private/WorldPartition/DataLayer/WorldDataLayers.cpp
@@ -22,10 +22,20 @@ const AWorldDataLayers* AWorldDataLayers::Get(UWorld* World)
 {
 	if (World)
 	{
-		for (AWorldDataLayers* Actor : TActorRange<AWorldDataLayers>(World))
+		EActorIteratorFlags Flags = EActorIteratorFlags::OnlyActiveLevels | EActorIteratorFlags::SkipPendingKill;
+#if WITH_EDITOR
+		if (IsRunningCookCommandlet())
+		{
+			// Don't use default Flags because EActorIteratorFlags::OnlyActiveLevels will make this code return no AWorldDataLayers actor
+			// when cooking (because world is not initialized)
+			Flags &= ~EActorIteratorFlags::OnlyActiveLevels;
+		}
+#endif
+		for (AWorldDataLayers* Actor : TActorRange<AWorldDataLayers>(World, AWorldDataLayers::StaticClass(), Flags))
 		{
-			if (Actor && !Actor->IsPendingKill())
+			if (Actor)
 			{
+				check(!Actor->IsPendingKill());
 				return Actor;
 			}
 		}
diff --git a/Engine/Source/Runtime/Engine/Private/WorldPartition/HLOD/HLODActor.cpp b/Engine/Source/Runtime/Engine/Private/WorldPartition/HLOD/HLODActor.cpp
--- a/Engine/Source/Runtime/Engine/Private/WorldPartition/HLOD/HLODActor.cpp
+++ b/Engine/Source/Runtime/Engine/Private/WorldPartition/HLOD/HLODActor.cpp
@@ -80,7 +80,7 @@ void AWorldPartitionHLOD::OnSubActorUnloaded(AActor& Actor)
 void AWorldPartitionHLOD::SetupLoadedSubActors()
 {
 	UWorld* World = GetWorld();
-	if (World && !World->IsGameWorld())
+	if (World && World->IsEditorWorld() && !World->IsPlayInEditor())
 	{
 		LoadedSubActors.Empty();
 
diff --git a/Engine/Source/Runtime/Engine/Private/WorldPartition/WorldPartition.cpp b/Engine/Source/Runtime/Engine/Private/WorldPartition/WorldPartition.cpp
--- a/Engine/Source/Runtime/Engine/Private/WorldPartition/WorldPartition.cpp
+++ b/Engine/Source/Runtime/Engine/Private/WorldPartition/WorldPartition.cpp
@@ -139,9 +139,6 @@ UWorldPartition::UWorldPartition(const FObjectInitializer& ObjectInitializer)
 	, InitState(EWorldPartitionInitState::Uninitialized)
 	, InstanceTransform(FTransform::Identity)
 	, StreamingPolicy(nullptr)
-#if WITH_EDITOR
-	, bIsPreCooked(false)
-#endif
 {
 	static bool bRegisteredDelegate = false;
 	if (!bRegisteredDelegate)
@@ -303,9 +300,9 @@ void UWorldPartition::Initialize(UWorld* InWorld, const FTransform& InTransform)
 
 	InitState = EWorldPartitionInitState::Initialized;
 
-	UWorldPartitionSubsystem* WorldPartitionSubsystem = World->GetSubsystem<UWorldPartitionSubsystem>();
-	if (ensure(WorldPartitionSubsystem))
+	if (World->IsGameWorld())
 	{
+		UWorldPartitionSubsystem* WorldPartitionSubsystem = World->GetSubsystem<UWorldPartitionSubsystem>();
 		WorldPartitionSubsystem->RegisterWorldPartition(this);
 	}
 
@@ -376,9 +373,9 @@ void UWorldPartition::Uninitialize()
 		EditorHash = nullptr;
 #endif		
 
-		UWorldPartitionSubsystem* WorldPartitionSubsystem = World->GetSubsystem<UWorldPartitionSubsystem>();
-		if (ensure(WorldPartitionSubsystem))
+		if (World->IsGameWorld())
 		{
+			UWorldPartitionSubsystem* WorldPartitionSubsystem = World->GetSubsystem<UWorldPartitionSubsystem>();
 			WorldPartitionSubsystem->UnregisterWorldPartition(this);
 		}
 
@@ -676,6 +673,12 @@ bool UWorldPartition::UpdateEditorCells(TFunctionRef<bool(TArray<UWorldPartition
 
 bool UWorldPartition::ShouldActorBeLoaded(const FWorldPartitionActorDesc* ActorDesc) const
 {
+	// No filtering when running cook commandlet
+	if (IsRunningCookCommandlet())
+	{
+		return true;
+	}
+
 	if (const AWorldDataLayers* WorldDataLayers = AWorldDataLayers::Get(GetWorld()))
 	{
 		uint32 NumValidLayers = 0;
@@ -1062,42 +1065,42 @@ void UWorldPartition::PrepareForPIE()
 {
 	check(World->IsPlayInEditor());
 
-	if (!IsPreCooked())
+	if (!IsMainWorldPartition())
 	{
-		if (!IsMainWorldPartition())
-		{
-			GenerateStreaming(EWorldPartitionStreamingMode::PIE);
-		}
-
-		GetStreamingPolicy()->PrepareForPIE();
+		GenerateStreaming(EWorldPartitionStreamingMode::PIE);
 	}
+
+	GetStreamingPolicy()->PrepareForPIE();
 }
 
 void UWorldPartition::CleanupForPIE()
 {
 	check(World->IsPlayInEditor());
 
-	if (!IsPreCooked())
-	{
-		FlushStreaming();
-	}
+	FlushStreaming();
 }
 
-bool UWorldPartition::GenerateStreaming(EWorldPartitionStreamingMode Mode)
+bool UWorldPartition::GenerateStreaming(EWorldPartitionStreamingMode Mode, TArray<FString>* OutPackagesToGenerate)
 {
-	if (!IsPreCooked())
-	{
-		return RuntimeHash->GenerateStreaming(Mode, GetStreamingPolicy());
-	}
-	return false;
+	check(RuntimeHash);
+	return RuntimeHash->GenerateStreaming(Mode, GetStreamingPolicy(), OutPackagesToGenerate);
+}
+
+bool UWorldPartition::PopulateGeneratedPackageForCook(UPackage* InPackage, const FString& InPackageRelativePath, const FString& InPackageCookName)
+{
+	check(RuntimeHash);
+	return RuntimeHash->PopulateGeneratedPackageForCook(InPackage, InPackageRelativePath, InPackageCookName);
+}
+
+void UWorldPartition::FinalizeGeneratedPackageForCook()
+{
+	check(RuntimeHash);
+	RuntimeHash->FinalizeGeneratedPackageForCook();
 }
 
 void UWorldPartition::FlushStreaming()
 {
-	if (!IsPreCooked())
-	{
-		RuntimeHash->FlushStreaming();
-	}
+	RuntimeHash->FlushStreaming();
 }
 
 void UWorldPartition::GenerateHLOD()
diff --git a/Engine/Source/Runtime/Engine/Private/WorldPartition/WorldPartitionActorDesc.cpp b/Engine/Source/Runtime/Engine/Private/WorldPartition/WorldPartitionActorDesc.cpp
--- a/Engine/Source/Runtime/Engine/Private/WorldPartition/WorldPartitionActorDesc.cpp
+++ b/Engine/Source/Runtime/Engine/Private/WorldPartition/WorldPartitionActorDesc.cpp
@@ -173,7 +173,7 @@ AActor* FWorldPartitionActorDesc::GetActor() const
 	return FindObject<AActor>(nullptr, *ActorPath.ToString());
 }
 
-AActor* FWorldPartitionActorDesc::Load(const FLinkerInstancingContext* InstancingContext)
+AActor* FWorldPartitionActorDesc::Load(const FLinkerInstancingContext* InstancingContext) const
 {
 	UPackage* Package = nullptr;
 
diff --git a/Engine/Source/Runtime/Engine/Private/WorldPartition/WorldPartitionCookPackageSplitter.cpp b/Engine/Source/Runtime/Engine/Private/WorldPartition/WorldPartitionCookPackageSplitter.cpp
new file mode 100644
--- /dev/null
+++ b/Engine/Source/Runtime/Engine/Private/WorldPartition/WorldPartitionCookPackageSplitter.cpp
@@ -0,0 +1,85 @@
+// Copyright Epic Games, Inc. All Rights Reserved.
+
+#include "WorldPartition/WorldPartitionCookPackageSplitter.h"
+
+#if WITH_EDITOR
+
+#include "WorldPartition/WorldPartitionRuntimeLevelStreamingCell.h"
+#include "WorldPartition/WorldPartitionLevelStreamingDynamic.h"
+#include "WorldPartition/WorldPartitionRuntimeCell.h"
+#include "WorldPartition/WorldPartitionRuntimeHash.h"
+#include "WorldPartition/WorldPartitionSubsystem.h"
+#include "WorldPartition/WorldPartition.h"
+#include "WorldPartition/HLOD/HLODActor.h"
+#include "WorldPartition/HLOD/HLODActorDescFactory.h"
+#include "WorldPartition/LevelInstance/LevelInstanceActorDescFactory.h"
+#include "WorldPartition/Landscape/LandscapeActorDescFactory.h"
+#include "WorldPartition/ActorPartition/PartitionActorDescFactory.h"
+#include "LevelInstance/LevelInstanceActor.h"
+#include "LandscapeProxy.h"
+#include "ActorPartition/PartitionActor.h"
+
+// Register FWorldPartitionCookPackageSplitter for UWorld class
+REGISTER_COOKPACKAGE_SPLITTER(FWorldPartitionCookPackageSplitter, UWorld);
+
+static FHLODActorDescFactory GHLODActorDescFactory;
+static FLevelInstanceActorDescFactory GLevelInstanceActorDescFactory;
+static FLandscapeActorDescFactory GLandscapeActorDescFactory;
+static FPartitionActorDescFactory GPartitionActorDescFactory;
+
+bool FWorldPartitionCookPackageSplitter::ShouldSplit(UObject* SplitData)
+{
+	UWorld* World = Cast<UWorld>(SplitData);
+	return World && !!World->GetWorldPartition();
+}
+
+void FWorldPartitionCookPackageSplitter::SetDataObject(UObject* SplitData)
+{
+	PartitionedWorld = CastChecked<UWorld>(SplitData);
+	check(PartitionedWorld);
+	check(PartitionedWorld->PersistentLevel);
+	check(PartitionedWorld->PersistentLevel->GetWorldPartition());
+}
+
+TArray<ICookPackageSplitter::FGeneratedPackage> FWorldPartitionCookPackageSplitter::GetGenerateList()
+{
+	// World is not initialized
+	ensure(!PartitionedWorld->bIsWorldInitialized);
+
+	// Manually register ActorDesc factories.
+	// @todo_ow: Revisit how these factories are declared. This could be automatic.
+	UWorldPartition::RegisterActorDescFactory(AWorldPartitionHLOD::StaticClass(), &GHLODActorDescFactory);
+	UWorldPartition::RegisterActorDescFactory(ALevelInstance::StaticClass(), &GLevelInstanceActorDescFactory);
+	UWorldPartition::RegisterActorDescFactory(ALandscapeProxy::StaticClass(), &GLandscapeActorDescFactory);
+	UWorldPartition::RegisterActorDescFactory(APartitionActor::StaticClass(), &GPartitionActorDescFactory);
+
+	// Manually initialize WorldPartition
+	UWorldPartition* WorldPartition = PartitionedWorld->PersistentLevel->GetWorldPartition();
+	WorldPartition->Initialize(PartitionedWorld, FTransform::Identity);
+
+	TArray<FString> WorldPartitionGeneratedPackages;
+	WorldPartition->GenerateStreaming(EWorldPartitionStreamingMode::RuntimeStreamingCells, &WorldPartitionGeneratedPackages);
+
+	TArray<ICookPackageSplitter::FGeneratedPackage> PackagesToGenerate;
+	PackagesToGenerate.Reserve(WorldPartitionGeneratedPackages.Num());
+	for (const FString& PackageName : WorldPartitionGeneratedPackages)
+	{
+		PackagesToGenerate.Emplace_GetRef().RelativePath = PackageName;
+		// @todo_ow: Set dependencies once we get iterative cooking working
+	}
+	return PackagesToGenerate;
+}
+
+bool FWorldPartitionCookPackageSplitter::TryPopulatePackage(UPackage* GeneratedPackage, const FStringView& RelativePath, const FStringView& GeneratedPackageCookName)
+{
+	UWorldPartition* WorldPartition = PartitionedWorld->PersistentLevel->GetWorldPartition();
+	return WorldPartition->PopulateGeneratedPackageForCook(GeneratedPackage, FString(RelativePath), FString(GeneratedPackageCookName));
+}
+
+void FWorldPartitionCookPackageSplitter::FinalizeGeneratorPackage()
+{
+	UWorldPartition* WorldPartition = PartitionedWorld->PersistentLevel->GetWorldPartition();
+	WorldPartition->FinalizeGeneratedPackageForCook();
+}
+
+#endif
\ No newline at end of file
diff --git a/Engine/Source/Runtime/Engine/Private/WorldPartition/WorldPartitionCookPackageSplitter.h b/Engine/Source/Runtime/Engine/Private/WorldPartition/WorldPartitionCookPackageSplitter.h
new file mode 100644
--- /dev/null
+++ b/Engine/Source/Runtime/Engine/Private/WorldPartition/WorldPartitionCookPackageSplitter.h
@@ -0,0 +1,26 @@
+// Copyright Epic Games, Inc. All Rights Reserved.
+
+#pragma once
+
+#if WITH_EDITOR
+
+#include "CookPackageSplitter.h"
+#include "Engine/World.h"
+
+class FWorldPartitionCookPackageSplitter : public ICookPackageSplitter
+{
+public:
+	//~ Begin of ICookPackageSplitter
+	static bool ShouldSplit(UObject* SplitData);
+	virtual ~FWorldPartitionCookPackageSplitter() {}
+	virtual void SetDataObject(UObject* SplitData) override;
+	virtual TArray<ICookPackageSplitter::FGeneratedPackage> GetGenerateList() override;
+	virtual bool TryPopulatePackage(UPackage* GeneratedPackage, const FStringView& RelativePath, const FStringView& GeneratedPackageCookName) override;
+	virtual void FinalizeGeneratorPackage() override;
+	//~ End of ICookPackageSplitter
+
+private:
+	UWorld* PartitionedWorld;
+};
+
+#endif
\ No newline at end of file
diff --git a/Engine/Source/Runtime/Engine/Private/WorldPartition/WorldPartitionLevelHelper.cpp b/Engine/Source/Runtime/Engine/Private/WorldPartition/WorldPartitionLevelHelper.cpp
--- a/Engine/Source/Runtime/Engine/Private/WorldPartition/WorldPartitionLevelHelper.cpp
+++ b/Engine/Source/Runtime/Engine/Private/WorldPartition/WorldPartitionLevelHelper.cpp
@@ -30,17 +30,19 @@ UWorld::InitializationValues FWorldPartitionLevelHelper::GetWorldInitializationV
 }
 
 /**
- * Creates, populates and saves a Level used in World Partition
+ * Creates and populates a Level used in World Partition
  */
-bool FWorldPartitionLevelHelper::CreateAndSaveLevelForRuntimeCell(const UWorld* InWorld, const FString& InWorldAssetName, const TArray<FWorldPartitionRuntimeCellObjectMapping>& InChildPackages)
+bool FWorldPartitionLevelHelper::CreateAndFillLevelForRuntimeCell(const UWorld* InWorld, const FString& InWorldAssetName, UPackage* InPackage, const TArray<FWorldPartitionRuntimeCellObjectMapping>& InChildPackages)
 {
 	check(IsRunningCommandlet());
 	check(InWorld && !InWorld->IsGameWorld());
+	check(InPackage);
 
 	// Create streaming cell Level package
-	ULevel* NewLevel = FWorldPartitionLevelHelper::CreateEmptyLevelForRuntimeCell(InWorld, InWorldAssetName);
-	UPackage* NewPackage = NewLevel->GetPackage();
-	UWorld* NewWorld = UWorld::FindWorldInPackage(NewPackage);
+	ULevel* NewLevel = FWorldPartitionLevelHelper::CreateEmptyLevelForRuntimeCell(InWorld, InWorldAssetName, InPackage);
+	UPackage* NewLevelPackage = NewLevel->GetPackage();
+	check(NewLevelPackage == InPackage);
+	UWorld* NewWorld = UWorld::FindWorldInPackage(NewLevelPackage);
 
 	// Move all actors to Cell level
 	for (const FWorldPartitionRuntimeCellObjectMapping& PackageObjectMapping : InChildPackages)
@@ -50,52 +52,52 @@ bool FWorldPartitionLevelHelper::CreateAndSaveLevelForRuntimeCell(const UWorld*
 		{
 			Actor->SetPackageExternal(false, false);
 			Actor->Rename(nullptr, NewLevel);
-			check(Actor->GetPackage() == NewPackage);
+			check(Actor->GetPackage() == NewLevelPackage);
 		}
 		else
 		{
 			UE_LOG(LogEngine, Warning, TEXT("Can't find actor %s."), *PackageObjectMapping.Path.ToString());
 		}
 	}
-
-	// Save
-	const FString PackageName = NewPackage->GetName();
-	const bool bSuccess = UEditorLoadingAndSavingUtils::SaveMap(NewWorld, PackageName);
-	UE_CLOG(!bSuccess, LogEngine, Error, TEXT("Error saving map %s."), *PackageName);
-
-	// No need of this Level anymore
-	NewLevel->MarkPendingKill();
-	NewWorld->RemoveFromRoot();
-	NewWorld->MarkPendingKill();
-
-	return bSuccess;
+	return true;
 }
 
 /**
  * Creates an empty Level used in World Partition
  */
-ULevel* FWorldPartitionLevelHelper::CreateEmptyLevelForRuntimeCell(const UWorld* InWorld, const FString& InWorldAssetName)
+ULevel* FWorldPartitionLevelHelper::CreateEmptyLevelForRuntimeCell(const UWorld* InWorld, const FString& InWorldAssetName, UPackage* InPackage)
 {
-	// Create Level package
-	FString PackageName = FPackageName::ObjectPathToPackageName(InWorldAssetName);
-	check(!FindObject<UPackage>(nullptr, *PackageName));
-	UPackage* NewPackage = CreatePackage(*PackageName);
-	NewPackage->SetPackageFlags(PKG_NewlyCreated);
+	// Create or use given package
+	UPackage* CellPackage = nullptr;
+	if (InPackage)
+	{
+		check(FindObject<UPackage>(nullptr, *InPackage->GetName()));
+		CellPackage = InPackage;
+	}
+	else
+	{
+		FString PackageName = FPackageName::ObjectPathToPackageName(InWorldAssetName);
+		check(!FindObject<UPackage>(nullptr, *PackageName));
+		CellPackage = CreatePackage(*PackageName);
+		CellPackage->SetPackageFlags(PKG_NewlyCreated);
+	}
+
 	if (InWorld->IsPlayInEditor())
 	{
-		NewPackage->SetPackageFlags(PKG_PlayInEditor);
-		NewPackage->PIEInstanceID = InWorld->GetPackage()->PIEInstanceID;
-		check(NewPackage);
+		check(!InPackage);
+		CellPackage->SetPackageFlags(PKG_PlayInEditor);
+		CellPackage->PIEInstanceID = InWorld->GetPackage()->PIEInstanceID;
 	}
 
 	// Create World & Persistent Level
 	UWorld::InitializationValues IVS = FWorldPartitionLevelHelper::GetWorldInitializationValues();
 	const FName WorldName = FName(FPackageName::ObjectPathToObjectName(InWorldAssetName));
-	UWorld* NewWorld = UWorld::CreateWorld(InWorld->WorldType, /*bInformEngineOfWorld*/false, WorldName, NewPackage, /*bAddToRoot*/false, InWorld->FeatureLevel, &IVS, /*bInSkipInitWorld*/true);
-	check(UWorld::FindWorldInPackage(NewPackage) == NewWorld);
+	UWorld* NewWorld = UWorld::CreateWorld(InWorld->WorldType, /*bInformEngineOfWorld*/false, WorldName, CellPackage, /*bAddToRoot*/false, InWorld->FeatureLevel, &IVS, /*bInSkipInitWorld*/true);
 	check(NewWorld);
+	NewWorld->SetFlags(RF_Public | RF_Standalone);
 	check(NewWorld->GetWorldSettings());
-	check(NewWorld->GetPathName() == InWorldAssetName);
+	check(UWorld::FindWorldInPackage(CellPackage) == NewWorld);
+	check(InPackage || (NewWorld->GetPathName() == InWorldAssetName));
 	
 	// Setup of streaming cell Runtime Level
 	ULevel* NewLevel = NewWorld->PersistentLevel;
diff --git a/Engine/Source/Runtime/Engine/Private/WorldPartition/WorldPartitionLevelStreamingDynamic.cpp b/Engine/Source/Runtime/Engine/Private/WorldPartition/WorldPartitionLevelStreamingDynamic.cpp
--- a/Engine/Source/Runtime/Engine/Private/WorldPartition/WorldPartitionLevelStreamingDynamic.cpp
+++ b/Engine/Source/Runtime/Engine/Private/WorldPartition/WorldPartitionLevelStreamingDynamic.cpp
@@ -97,12 +97,6 @@ void UWorldPartitionLevelStreamingDynamic::CreateRuntimeLevel()
  */
 bool UWorldPartitionLevelStreamingDynamic::RequestLevel(UWorld* InPersistentWorld, bool bInAllowLevelLoadRequests, EReqLevelBlock InBlockPolicy)
 {
-	// When running with pre-cooked data, behave like regular LevelStreaming
-	if (InPersistentWorld->GetWorldPartition()->IsPreCooked())
-	{
-		return Super::RequestLevel(InPersistentWorld, bInAllowLevelLoadRequests, InBlockPolicy);
-	}
-
 	// Quit early in case load request already issued
 	if (GetCurrentState() == ECurrentState::Loading)
 	{
diff --git a/Engine/Source/Runtime/Engine/Private/WorldPartition/WorldPartitionLevelStreamingPolicy.cpp b/Engine/Source/Runtime/Engine/Private/WorldPartition/WorldPartitionLevelStreamingPolicy.cpp
--- a/Engine/Source/Runtime/Engine/Private/WorldPartition/WorldPartitionLevelStreamingPolicy.cpp
+++ b/Engine/Source/Runtime/Engine/Private/WorldPartition/WorldPartitionLevelStreamingPolicy.cpp
@@ -133,8 +133,7 @@ FString UWorldPartitionLevelStreamingPolicy::GetCellPackagePath(const FName& InC
 	}
 	else
 	{
-		const FString WorldPackageFileName = InWorld->GetPackage()->FileName.ToString();
-		return FString::Printf(TEXT("%s/%s/%s_%s"), *WorldPackageFileName, FWorldPartitionLevelHelper::GetSavedLevelOutputSubFolderName(), *FPaths::GetBaseFilename(*WorldPackageFileName), *InCellName.ToString());
+		return FString::Printf(TEXT("/%s"), *InCellName.ToString());
 	}
 }
 
@@ -145,8 +144,6 @@ TSubclassOf<UWorldPartitionRuntimeCell> UWorldPartitionLevelStreamingPolicy::Get
 
 void UWorldPartitionLevelStreamingPolicy::PrepareForPIE()
 {
-	check(!WorldPartition->IsPreCooked());
-
 	TSet<const UWorldPartitionRuntimeCell*> AlwaysLoadedCells;
 	TSet<const UWorldPartitionRuntimeCell*> StreamingCells;
 	WorldPartition->RuntimeHash->GetAllStreamingCells(StreamingCells, /*bIncludeDataLayers*/ true);
diff --git a/Engine/Source/Runtime/Engine/Private/WorldPartition/WorldPartitionRuntimeLevelStreamingCell.cpp b/Engine/Source/Runtime/Engine/Private/WorldPartition/WorldPartitionRuntimeLevelStreamingCell.cpp
--- a/Engine/Source/Runtime/Engine/Private/WorldPartition/WorldPartitionRuntimeLevelStreamingCell.cpp
+++ b/Engine/Source/Runtime/Engine/Private/WorldPartition/WorldPartitionRuntimeLevelStreamingCell.cpp
@@ -49,12 +49,16 @@ void UWorldPartitionRuntimeLevelStreamingCell::SetIsAlwaysLoaded(bool bInIsAlway
 
 #if WITH_EDITOR
 
-void UWorldPartitionRuntimeLevelStreamingCell::AddActorToCell(FName Package, FName Path)
+void UWorldPartitionRuntimeLevelStreamingCell::AddActorToCell(const FGuid& ActorGuid, FName Package, FName Path)
 {
 	Packages.Emplace(Package, Path);
+	if (!GetOuterUWorldPartition()->GetActorDesc(ActorGuid)->GetActorIsEditorOnly())
+	{
+		ActorsToLoadForCook.Add(ActorGuid);
+	}
 }
 
-ULevelStreaming* UWorldPartitionRuntimeLevelStreamingCell::CreateLevelStreaming() const
+ULevelStreaming* UWorldPartitionRuntimeLevelStreamingCell::CreateLevelStreaming(const FString& InPackageName) const
 {
 	if (GetActorCount() > 0)
 	{
@@ -65,11 +69,12 @@ ULevelStreaming* UWorldPartitionRuntimeLevelStreamingCell::CreateLevelStreaming(
 		const FName LevelStreamingName = FName(*FString::Printf(TEXT("WorldPartitionLevelStreaming_%s"), *GetName()));
 		const UClass* LevelStreamingClass = (IsRunningCommandlet() && IsAlwaysLoaded() && (OuterWorld == OwningWorld)) ? ULevelStreamingAlwaysLoaded::StaticClass() : UWorldPartitionLevelStreamingDynamic::StaticClass();
 		
-		// When called by Commandlet (CreateCellForCook), LevelStreaming's outer is set to Cell/WorldPartition's outer to prevent warnings when saving Cell Levels (Warning: Obj in another map). 
+		// When called by Commandlet (PopulateGeneratedPackageForCook), LevelStreaming's outer is set to Cell/WorldPartition's outer to prevent warnings when saving Cell Levels (Warning: Obj in another map). 
 		// At runtime, LevelStreaming's outer will be properly set to the main world (see UWorldPartitionRuntimeLevelStreamingCell::Activate).
 		UWorld* LevelStreamingOuterWorld = IsRunningCommandlet() ? OuterWorld : OwningWorld;
 		ULevelStreaming* NewLevelStreaming = NewObject<ULevelStreaming>(LevelStreamingOuterWorld, LevelStreamingClass, LevelStreamingName, RF_NoFlags, NULL);
-		TSoftObjectPtr<UWorld> WorldAsset(FSoftObjectPath(FString::Printf(TEXT("%s.%s"), *UWorldPartitionLevelStreamingPolicy::GetCellPackagePath(GetFName(), OuterWorld), *OuterWorld->GetName())));
+		FString PackageName = !InPackageName.IsEmpty() ? InPackageName : UWorldPartitionLevelStreamingPolicy::GetCellPackagePath(GetFName(), OuterWorld);
+		TSoftObjectPtr<UWorld> WorldAsset(FSoftObjectPath(FString::Printf(TEXT("%s.%s"), *PackageName, *OuterWorld->GetName())));
 		NewLevelStreaming->SetWorldAsset(WorldAsset);
 		// Transfer WorldPartition's transform to Level
 		NewLevelStreaming->LevelTransform = WorldPartition->GetInstanceTransform();
@@ -91,14 +96,35 @@ ULevelStreaming* UWorldPartitionRuntimeLevelStreamingCell::CreateLevelStreaming(
 	return nullptr;
 }
 
-bool UWorldPartitionRuntimeLevelStreamingCell::CreateCellForCook()
+bool UWorldPartitionRuntimeLevelStreamingCell::PopulateGeneratedPackageForCook(UPackage* InPackage, const FString& InPackageCookName)
 {
-	if (ULevelStreaming* NewLevelStreaming = CreateLevelStreaming())
+	if (!InPackage || InPackageCookName.IsEmpty())
 	{
-		LevelStreaming = Cast<UWorldPartitionLevelStreamingDynamic>(NewLevelStreaming);
-		const UWorldPartition* WorldPartition = GetOuterUWorldPartition();
+		return false;
+	}
+
+	if (ULevelStreaming* NewLevelStreaming = CreateLevelStreaming(InPackageCookName))
+	{
+		UWorldPartition* WorldPartition = GetOuterUWorldPartition();
 		UWorld* OuterWorld = WorldPartition->GetTypedOuter<UWorld>();
-		if (!FWorldPartitionLevelHelper::CreateAndSaveLevelForRuntimeCell(OuterWorld, NewLevelStreaming->GetWorldAsset().ToString(), Packages))
+
+		// Load cell Actors
+		for (const FGuid& ActorGuid : ActorsToLoadForCook)
+		{
+			const FWorldPartitionActorDesc* ActorDesc = WorldPartition->GetActorDesc(ActorGuid);
+			AActor* Actor = ActorDesc->GetActor();
+			if (!Actor)
+			{
+				Actor = ActorDesc->Load();
+			}
+			if (ensure(Actor))
+			{
+				LoadedActorsForCook.Add(Actor);
+			}
+		}
+
+		LevelStreaming = Cast<UWorldPartitionLevelStreamingDynamic>(NewLevelStreaming);
+		if (!FWorldPartitionLevelHelper::CreateAndFillLevelForRuntimeCell(OuterWorld, NewLevelStreaming->GetWorldAsset().ToString(), InPackage, Packages))
 		{
 			return false;
 		}
@@ -107,67 +133,80 @@ bool UWorldPartitionRuntimeLevelStreamingCell::CreateCellForCook()
 		if (NewLevelStreaming->IsA<ULevelStreamingAlwaysLoaded>())
 		{
 			OuterWorld->AddStreamingLevel(NewLevelStreaming);
+			NewLevelStreaming->bForceIsValidStreamingLevel = true;
 		}
 	}
 	return true;
 }
 
+void UWorldPartitionRuntimeLevelStreamingCell::FinalizeGeneratedPackageForCook()
+{
+#if DO_CHECK
+	// Make sure actors are not external anymore
+	for(AActor * Actor : LoadedActorsForCook)
+	{
+		check(!Actor->IsPackageExternal());
+	}
+#endif
+}
+
 int32 UWorldPartitionRuntimeLevelStreamingCell::GetActorCount() const
 {
 	return Packages.Num();
 }
 
+FString UWorldPartitionRuntimeLevelStreamingCell::GetPackageNameToCreate() const
+{
+	const UWorldPartition* WorldPartition = GetOuterUWorldPartition();
+	UWorld* OuterWorld = WorldPartition->GetTypedOuter<UWorld>();
+	return UWorldPartitionLevelStreamingPolicy::GetCellPackagePath(GetFName(), OuterWorld);
+}
+
 #endif
 
 void UWorldPartitionRuntimeLevelStreamingCell::Activate() const
 {
 #if WITH_EDITOR
-	if (!GetOuterUWorldPartition()->IsPreCooked())
+	if (GetActorCount() == 0)
 	{
-		if (GetActorCount() == 0)
-		{
-			return;
-		}
+		return;
+	}
 
-		if (!LevelStreaming)
-		{
-			LevelStreaming = Cast<UWorldPartitionLevelStreamingDynamic>(CreateLevelStreaming());
-		}
-		check(LevelStreaming);
+	if (!LevelStreaming)
+	{
+		LevelStreaming = Cast<UWorldPartitionLevelStreamingDynamic>(CreateLevelStreaming());
 	}
+	check(LevelStreaming);
 #else
 	// In Runtime, always loaded cell level is handled by World directly
 	check(LevelStreaming || IsAlwaysLoaded());
 #endif
 
-#if WITH_EDITOR
-	if (GetOuterUWorldPartition()->IsPreCooked())
-#endif
+#if !WITH_EDITOR
+	// In Runtime, prepare LevelStreaming for activation
+	if (LevelStreaming)
 	{
-		// in Runtime (or PreCooked), prepare LevelStreaming for activation
-		if (LevelStreaming)
+		// Setup pre-created LevelStreaming's outer to the WorldPartition owning world
+		const UWorldPartition* WorldPartition = GetOuterUWorldPartition();
+		UWorld* OwningWorld = WorldPartition->GetWorld();
+		if (LevelStreaming->GetWorld() != OwningWorld)
 		{
-			// Setup pre-created LevelStreaming's outer to the WorldPartition owning world
-			const UWorldPartition* WorldPartition = GetOuterUWorldPartition();
-			UWorld* OwningWorld = WorldPartition->GetWorld();
-			if (LevelStreaming->GetWorld() != OwningWorld)
-			{
-				LevelStreaming->Rename(nullptr, OwningWorld);
-			}
+			LevelStreaming->Rename(nullptr, OwningWorld);
+		}
 
-			// Transfer WorldPartition's transform to LevelStreaming
-			LevelStreaming->LevelTransform = WorldPartition->GetInstanceTransform();
+		// Transfer WorldPartition's transform to LevelStreaming
+		LevelStreaming->LevelTransform = WorldPartition->GetInstanceTransform();
 
-			// When Partition outer level is an instance, make sure to also generate unique cell level instance name
-			ULevel* PartitionLevel = WorldPartition->GetTypedOuter<ULevel>();
-			if (PartitionLevel->IsInstancedLevel())
-			{
-				FString PackageShortName = FPackageName::GetShortName(PartitionLevel->GetPackage());
-				FString InstancedLevelPackageName = FString::Printf(TEXT("%s_InstanceOf_%s"), *LevelStreaming->PackageNameToLoad.ToString(), *PackageShortName);
-				LevelStreaming->SetWorldAssetByPackageName(FName(InstancedLevelPackageName));
-			}
+		// When Partition outer level is an instance, make sure to also generate unique cell level instance name
+		ULevel* PartitionLevel = WorldPartition->GetTypedOuter<ULevel>();
+		if (PartitionLevel->IsInstancedLevel())
+		{
+			FString PackageShortName = FPackageName::GetShortName(PartitionLevel->GetPackage());
+			FString InstancedLevelPackageName = FString::Printf(TEXT("%s_InstanceOf_%s"), *LevelStreaming->PackageNameToLoad.ToString(), *PackageShortName);
+			LevelStreaming->SetWorldAssetByPackageName(FName(InstancedLevelPackageName));
 		}
 	}
+#endif
 
 	if (LevelStreaming)
 	{
@@ -184,14 +223,11 @@ void UWorldPartitionRuntimeLevelStreamingCell::Activate() const
 void UWorldPartitionRuntimeLevelStreamingCell::Deactivate() const
 {
 #if WITH_EDITOR
-	if (!GetOuterUWorldPartition()->IsPreCooked())
+	if (GetActorCount() == 0)
 	{
-		if (GetActorCount() == 0)
-		{
-			return;
-		}
-		check(LevelStreaming);
+		return;
 	}
+	check(LevelStreaming);
 #else
 	// In Runtime, always loaded cell level is handled by World directly
 	check(LevelStreaming || IsAlwaysLoaded());
diff --git a/Engine/Source/Runtime/Engine/Private/WorldPartition/WorldPartitionRuntimeSpatialHash.cpp b/Engine/Source/Runtime/Engine/Private/WorldPartition/WorldPartitionRuntimeSpatialHash.cpp
--- a/Engine/Source/Runtime/Engine/Private/WorldPartition/WorldPartitionRuntimeSpatialHash.cpp
+++ b/Engine/Source/Runtime/Engine/Private/WorldPartition/WorldPartitionRuntimeSpatialHash.cpp
@@ -46,44 +46,6 @@ extern UNREALED_API class UEditorEngine* GEditor;
 
 DEFINE_LOG_CATEGORY_STATIC(LogWorldPartitionRuntimeSpatialHash, Log, All);
 
-#if WITH_EDITOR
-class FScopedLoadActorsHelper
-{
-public:
-	FScopedLoadActorsHelper(UWorldPartition* InWorldPartition, const TArray<FGuid>& InActors, bool bSkipEditorOnly)
-		: WorldPartition(InWorldPartition)
-	{
-		LoadedActors.Reserve(InActors.Num());
-		for (const FGuid& ActorGuid : InActors)
-		{
-			FWorldPartitionActorDesc* ActorDesc = WorldPartition->GetActorDesc(ActorGuid);
-			if (!ActorDesc->GetActor() && (!bSkipEditorOnly || !ActorDesc->GetActorIsEditorOnly()))
-			{
-				AActor* Actor = ActorDesc->Load();
-				if (ensure(Actor))
-				{
-					LoadedActors.Add(Actor);
-				}
-			}
-		}
-	}
-
-	~FScopedLoadActorsHelper()
-	{
-		for (AActor* Actor : LoadedActors)
-		{
-			check(!Actor->IsPackageExternal());
-			TGuardValue<ITransaction*> SuppressTransaction(GUndo, nullptr);
-			WorldPartition->GetWorld()->DestroyActor(Actor, false, false);
-		}
-		CollectGarbage(GARBAGE_COLLECTION_KEEPFLAGS, true);
-	}
-private:
-	UWorldPartition* WorldPartition;
-	TArray<AActor*> LoadedActors;
-};
-#endif //WITH_EDITOR
-
 static int32 GShowRuntimeSpatialHashGridLevel = 0;
 static FAutoConsoleVariableRef CVarShowRuntimeSpatialHashGridLevel(
 	TEXT("wp.Runtime.ShowRuntimeSpatialHashGridLevel"),
@@ -1268,11 +1230,10 @@ void UWorldPartitionRuntimeSpatialHash::ImportFromWorldComposition(UWorldComposi
 	}
 }
 
-bool UWorldPartitionRuntimeSpatialHash::GenerateStreaming(EWorldPartitionStreamingMode Mode, UWorldPartitionStreamingPolicy* StreamingPolicy)
+bool UWorldPartitionRuntimeSpatialHash::GenerateStreaming(EWorldPartitionStreamingMode Mode, UWorldPartitionStreamingPolicy* StreamingPolicy, TArray<FString>* OutPackagesToGenerate)
 {
 	TRACE_CPUPROFILER_EVENT_SCOPE(UWorldPartitionRuntimeSpatialHash::GenerateStreaming);
 	UWorldPartition* WorldPartition = GetOuterUWorldPartition();
-	check(!WorldPartition->IsPreCooked());
 
 	UE_SCOPED_TIMER(TEXT("GenerateStreaming"), LogWorldPartitionRuntimeSpatialHash, Log);
 	
@@ -1328,7 +1289,7 @@ bool UWorldPartitionRuntimeSpatialHash::GenerateStreaming(EWorldPartitionStreami
 	{
 		const FSpatialHashRuntimeGrid& Grid = AllGrids[GridIndex];
 		const FSquare2DGridHelper PartionedActors = GetPartitionedActors(WorldPartition, WorldBounds, Grid, GridActors[GridIndex]);
-		if (!CreateStreamingGrid(Grid, PartionedActors, Mode, StreamingPolicy))
+		if (!CreateStreamingGrid(Grid, PartionedActors, Mode, StreamingPolicy, OutPackagesToGenerate))
 		{
 			return false;
 		}
@@ -1366,7 +1327,7 @@ void UWorldPartitionRuntimeSpatialHash::CacheHLODParents()
 	}
 }
 
-bool UWorldPartitionRuntimeSpatialHash::CreateStreamingGrid(const FSpatialHashRuntimeGrid& RuntimeGrid, const FSquare2DGridHelper& PartionedActors, EWorldPartitionStreamingMode Mode, UWorldPartitionStreamingPolicy* StreamingPolicy)
+bool UWorldPartitionRuntimeSpatialHash::CreateStreamingGrid(const FSpatialHashRuntimeGrid& RuntimeGrid, const FSquare2DGridHelper& PartionedActors, EWorldPartitionStreamingMode Mode, UWorldPartitionStreamingPolicy* StreamingPolicy, TArray<FString>* OutPackagesToGenerate)
 {
 	TRACE_CPUPROFILER_EVENT_SCOPE(UWorldPartitionRuntimeSpatialHash::CreateStreamingGrid);
 
@@ -1449,7 +1410,7 @@ bool UWorldPartitionRuntimeSpatialHash::CreateStreamingGrid(const FSpatialHashRu
 					{
 						ReferencedHLODActors.Add(ParentHLOD);
 					}
-					StreamingCell->AddActorToCell(ActorDesc->GetActorPackage(), ActorDesc->GetActorPath());
+					StreamingCell->AddActorToCell(ActorGuid, ActorDesc->GetActorPackage(), ActorDesc->GetActorPath());
 					UE_LOG(LogWorldPartitionRuntimeSpatialHash, Verbose, TEXT("  Actor : %s (%s) Origin(%s)"), *(ActorDesc->GetActorPath().ToString()), *ActorGuid.ToString(EGuidFormats::UniqueObjectGuid), *FVector2D(ActorDesc->GetOrigin()).ToString());
 				}
 
@@ -1463,12 +1424,22 @@ bool UWorldPartitionRuntimeSpatialHash::CreateStreamingGrid(const FSpatialHashRu
 
 				if (Mode == EWorldPartitionStreamingMode::RuntimeStreamingCells)
 				{
-					FScopedLoadActorsHelper LoadCellActors(WorldPartition, FilteredActors, /*bSkipEditorOnly*/true);
 					UE_LOG(LogWorldPartitionRuntimeSpatialHash, Log, TEXT("Creating runtime streaming cells %s."), *StreamingCell->GetName());
-					if (!StreamingCell->CreateCellForCook())
+
+					if (StreamingCell->GetActorCount())
 					{
-						UE_LOG(LogWorldPartitionRuntimeSpatialHash, Error, TEXT("Error creating runtime streaming cells for cook."));
-						return false;
+						if (!OutPackagesToGenerate)
+						{
+							UE_LOG(LogWorldPartitionRuntimeSpatialHash, Error, TEXT("Error creating runtime streaming cells for cook, OutPackagesToGenerate is null."));
+							return false;
+						}
+
+						const FString PackageRelativePath = StreamingCell->GetPackageNameToCreate();
+						check(!PackageRelativePath.IsEmpty());
+						OutPackagesToGenerate->Add(PackageRelativePath);
+
+						// Map relative package to StreamingCell for PopulateGeneratedPackageForCook/FinalizeGeneratedPackageForCook
+						PackagesToGenerateForCook.Add(PackageRelativePath, StreamingCell);
 					}
 				}
 			}
@@ -1478,9 +1449,33 @@ bool UWorldPartitionRuntimeSpatialHash::CreateStreamingGrid(const FSpatialHashRu
 	return true;
 }
 
+bool UWorldPartitionRuntimeSpatialHash::PopulateGeneratedPackageForCook(UPackage* InPackage, const FString& InPackageRelativePath, const FString& InPackageCookName)
+{
+	if (UWorldPartitionRuntimeCell** MatchingCell = PackagesToGenerateForCook.Find(InPackageRelativePath))
+	{
+		UWorldPartitionRuntimeCell* Cell = *MatchingCell;
+		if (ensure(Cell))
+		{
+			return Cell->PopulateGeneratedPackageForCook(InPackage, InPackageCookName);
+		}
+	}
+	return false;
+}
+
+void UWorldPartitionRuntimeSpatialHash::FinalizeGeneratedPackageForCook()
+{
+	for (const auto& Package : PackagesToGenerateForCook)
+	{
+		UWorldPartitionRuntimeCell* Cell = Package.Value;
+		if (ensure(Cell))
+		{
+			Cell->FinalizeGeneratedPackageForCook();
+		}
+	}
+}
+
 void UWorldPartitionRuntimeSpatialHash::FlushStreaming()
 {
-	check(!GetOuterUWorldPartition()->IsPreCooked());
 	StreamingGrids.Empty();
 }
 
@@ -2176,4 +2171,4 @@ void UWorldPartitionRuntimeSpatialHash::Draw3D(const TArray<FWorldPartitionStrea
 			StreamingGrid.Draw3D(World, Sources);
 		}
 	}
-}
+}
\ No newline at end of file
diff --git a/Engine/Source/Runtime/Engine/Private/WorldPartition/WorldPartitionSubsystem.cpp b/Engine/Source/Runtime/Engine/Private/WorldPartition/WorldPartitionSubsystem.cpp
--- a/Engine/Source/Runtime/Engine/Private/WorldPartition/WorldPartitionSubsystem.cpp
+++ b/Engine/Source/Runtime/Engine/Private/WorldPartition/WorldPartitionSubsystem.cpp
@@ -211,15 +211,4 @@ void UWorldPartitionSubsystem::RegisterActorDescFactory(TSubclassOf<AActor> Clas
 {
 	UWorldPartition::RegisterActorDescFactory(Class, Factory);
 }
-
-FBox UWorldPartitionSubsystem::GetWorldBounds()
-{
-	if (UWorldPartition* WorldPartition = GetMainWorldPartition())
-	{
-		return WorldPartition->GetWorldBounds();
-	}
-
-	return FBox(ForceInit);
-}
-
 #endif // WITH_EDITOR
\ No newline at end of file
diff --git a/Engine/Source/Runtime/Engine/Public/WorldPartition/WorldPartition.h b/Engine/Source/Runtime/Engine/Public/WorldPartition/WorldPartition.h
--- a/Engine/Source/Runtime/Engine/Public/WorldPartition/WorldPartition.h
+++ b/Engine/Source/Runtime/Engine/Public/WorldPartition/WorldPartition.h
@@ -90,23 +90,13 @@ private:
 #if WITH_EDITOR
 	// UWorldPartitionSubsystem interface+
 	friend class UWorldPartitionSubsystem;
-
-	// UActorPartitionSubsystem interface+
 	void ForEachIntersectingActorDesc(const FBox& Box, TSubclassOf<AActor> ActorClass, TFunctionRef<bool(const FWorldPartitionActorDesc*)> Predicate) const;
 	void ForEachActorDesc(TSubclassOf<AActor> ActorClass, TFunctionRef<bool(const FWorldPartitionActorDesc*)> Predicate) const;
 	// UActorPartitionSubsystem interface-
-
+public:
 	static void RegisterActorDescFactory(TSubclassOf<AActor> Class, FWorldPartitionActorDescFactory* Factory);
-	// UWorldPartitionSubsystem interface-
 
-public:
 	FName GetWorldPartitionEditorName();
-	
-	// PIE Methods
-	void PrepareForPIE();
-	void CleanupForPIE();
-	void OnPreFixupForPIE(int32 InPIEInstanceID, FSoftObjectPath& ObjectPath);
-
 	bool IsSimulating() const;
 
 	FWorldPartitionActorDesc* GetActorDesc(const FGuid& Guid);
@@ -118,12 +108,19 @@ public:
 	void UnloadEditorCells(const FBox& Box);
 	bool RefreshLoadedEditorCells();
 
-	bool IsPreCooked() const { return bIsPreCooked; }
-	void SetIsPreCooked(bool bInIsPreCooked) { bIsPreCooked = bInIsPreCooked; }
+	// PIE Methods
+	void PrepareForPIE();
+	void CleanupForPIE();
+	void OnPreFixupForPIE(int32 InPIEInstanceID, FSoftObjectPath& ObjectPath);
 
-	bool GenerateStreaming(EWorldPartitionStreamingMode Mode);
-	FBox GetWorldBounds() const;
+	// PIE/Cook Methods
+	bool GenerateStreaming(EWorldPartitionStreamingMode Mode, TArray<FString>* OutPackagesToGenerate = nullptr);
+
+	// Cook Methods
+	bool PopulateGeneratedPackageForCook(UPackage* InPackage, const FString& InPackageRelativePath, const FString& InPackageCookName);
+	void FinalizeGeneratedPackageForCook();
 
+	FBox GetWorldBounds() const;
 	void GenerateHLOD();
 	void GenerateNavigationData();
 
@@ -187,9 +184,6 @@ private:
 	mutable UWorldPartitionStreamingPolicy* StreamingPolicy;
 
 #if WITH_EDITORONLY_DATA
-	UPROPERTY()
-	bool bIsPreCooked;
-
 	FLinkerInstancingContext InstancingContext;
 #endif
 
diff --git a/Engine/Source/Runtime/Engine/Public/WorldPartition/WorldPartitionActorDesc.h b/Engine/Source/Runtime/Engine/Public/WorldPartition/WorldPartitionActorDesc.h
--- a/Engine/Source/Runtime/Engine/Public/WorldPartition/WorldPartitionActorDesc.h
+++ b/Engine/Source/Runtime/Engine/Public/WorldPartition/WorldPartitionActorDesc.h
@@ -70,7 +70,7 @@ public:
 	FString ToString() const;
 
 	AActor* GetActor() const;
-	AActor* Load(const FLinkerInstancingContext* InstancingContext = nullptr);
+	AActor* Load(const FLinkerInstancingContext* InstancingContext = nullptr) const;
 	void Unload();
 
 	virtual void Init(const AActor* InActor);
diff --git a/Engine/Source/Runtime/Engine/Public/WorldPartition/WorldPartitionLevelHelper.h b/Engine/Source/Runtime/Engine/Public/WorldPartition/WorldPartitionLevelHelper.h
--- a/Engine/Source/Runtime/Engine/Public/WorldPartition/WorldPartitionLevelHelper.h
+++ b/Engine/Source/Runtime/Engine/Public/WorldPartition/WorldPartitionLevelHelper.h
@@ -17,9 +17,8 @@
 class FWorldPartitionLevelHelper
 {
 public:
-	static const TCHAR* GetSavedLevelOutputSubFolderName() { return TEXT("WorldPartitionLevels"); }
-	static bool CreateAndSaveLevelForRuntimeCell(const UWorld* InWorld, const FString& InWorldAssetName, const TArray<FWorldPartitionRuntimeCellObjectMapping>& InChildPackages);
-	static ULevel* CreateEmptyLevelForRuntimeCell(const UWorld* InWorld, const FString& InWorldAssetName);
+	static bool CreateAndFillLevelForRuntimeCell(const UWorld* InWorld, const FString& InWorldAssetName, UPackage* InPackage, const TArray<FWorldPartitionRuntimeCellObjectMapping>& InChildPackages);
+	static ULevel* CreateEmptyLevelForRuntimeCell(const UWorld* InWorld, const FString& InWorldAssetName, UPackage* DestPackage = nullptr);
 private:
 	static UWorld::InitializationValues GetWorldInitializationValues();
 };
diff --git a/Engine/Source/Runtime/Engine/Public/WorldPartition/WorldPartitionRuntimeCell.h b/Engine/Source/Runtime/Engine/Public/WorldPartition/WorldPartitionRuntimeCell.h
--- a/Engine/Source/Runtime/Engine/Public/WorldPartition/WorldPartitionRuntimeCell.h
+++ b/Engine/Source/Runtime/Engine/Public/WorldPartition/WorldPartitionRuntimeCell.h
@@ -59,11 +59,14 @@ class UWorldPartitionRuntimeCell : public UObject
 
 #if WITH_EDITOR
 	void SetDataLayers(const TArray<const UDataLayer*> InDataLayers);
-	virtual void AddActorToCell(FName Package, FName Path) PURE_VIRTUAL(UWorldPartitionRuntimeCell::AddActorToCell,);
-	virtual bool CreateCellForCook() PURE_VIRTUAL(UWorldPartitionRuntimeCell::CreateCellForCook, return false;);
+	void AddCellData(const UWorldPartitionRuntimeCellData* InCellData);
+	virtual void AddActorToCell(const FGuid& ActorGuid, FName Package, FName Path) PURE_VIRTUAL(UWorldPartitionRuntimeCell::AddActorToCell,);
 	virtual int32 GetActorCount() const PURE_VIRTUAL(UWorldPartitionRuntimeCell::GetActorCount, return 0;);
 
-	void AddCellData(const UWorldPartitionRuntimeCellData* InCellData);
+	// Cook methods
+	virtual bool PopulateGeneratedPackageForCook(UPackage* InPackage, const FString& InPackageCookName) PURE_VIRTUAL(UWorldPartitionRuntimeCell::PopulateGeneratedPackageForCook, return false;);
+	virtual void FinalizeGeneratedPackageForCook() {}
+	virtual FString GetPackageNameToCreate() const PURE_VIRTUAL(UWorldPartitionRuntimeCell::GetPackageNameToCreate, return FString(""););
 #endif
 
 	const UWorldPartitionRuntimeCellData* GetCellData(const TSubclassOf<UWorldPartitionRuntimeCellData> InCellDataClass) const;
diff --git a/Engine/Source/Runtime/Engine/Public/WorldPartition/WorldPartitionRuntimeHash.h b/Engine/Source/Runtime/Engine/Public/WorldPartition/WorldPartitionRuntimeHash.h
--- a/Engine/Source/Runtime/Engine/Public/WorldPartition/WorldPartitionRuntimeHash.h
+++ b/Engine/Source/Runtime/Engine/Public/WorldPartition/WorldPartitionRuntimeHash.h
@@ -18,7 +18,9 @@ class ENGINE_API UWorldPartitionRuntimeHash : public UObject
 #if WITH_EDITOR
 	virtual void SetDefaultValues() {}
 	virtual void ImportFromWorldComposition(class UWorldComposition* WorldComposition) {}
-	virtual bool GenerateStreaming(EWorldPartitionStreamingMode Mode, class UWorldPartitionStreamingPolicy* Policy) { return false; }
+	virtual bool GenerateStreaming(EWorldPartitionStreamingMode Mode, class UWorldPartitionStreamingPolicy* Policy, TArray<FString>* OutPackagesToGenerate = nullptr) { return false; }
+	virtual bool PopulateGeneratedPackageForCook(UPackage* InPackage, const FString& InPackageRelativePath, const FString& InPackageCookName) { return false; }
+	virtual void FinalizeGeneratedPackageForCook() {}
 	virtual void FlushStreaming() {}
 	virtual bool GenerateHLOD() { return false; }
 	virtual bool GenerateNavigationData() { return false; }
diff --git a/Engine/Source/Runtime/Engine/Public/WorldPartition/WorldPartitionRuntimeLevelStreamingCell.h b/Engine/Source/Runtime/Engine/Public/WorldPartition/WorldPartitionRuntimeLevelStreamingCell.h
--- a/Engine/Source/Runtime/Engine/Public/WorldPartition/WorldPartitionRuntimeLevelStreamingCell.h
+++ b/Engine/Source/Runtime/Engine/Public/WorldPartition/WorldPartitionRuntimeLevelStreamingCell.h
@@ -17,10 +17,14 @@ class UWorldPartitionRuntimeLevelStreamingCell : public UWorldPartitionRuntimeSp
 	virtual FLinearColor GetDebugColor() const override;
 
 #if WITH_EDITOR
-	virtual void AddActorToCell(FName Package, FName Path) override;
-	virtual bool CreateCellForCook() override;
+	virtual void AddActorToCell(const FGuid& ActorGuid, FName Package, FName Path) override;
 	virtual int32 GetActorCount() const override;
 	const TArray<FWorldPartitionRuntimeCellObjectMapping>& GetPackages() const { return Packages; }
+
+	// Cook methods
+	virtual bool PopulateGeneratedPackageForCook(UPackage* InPackage, const FString& InPackageCookName) override;
+	virtual void FinalizeGeneratedPackageForCook() override;
+	virtual FString GetPackageNameToCreate() const override;
 #endif
 
 	virtual void SetIsAlwaysLoaded(bool bInIsAlwaysLoaded) override;
@@ -35,12 +39,19 @@ private:
 	EStreamingStatus GetLevelStreamingStatus() const;
 
 #if WITH_EDITOR
-	class ULevelStreaming* CreateLevelStreaming() const;
+	class ULevelStreaming* CreateLevelStreaming(const FString& InPackageName = FString()) const;
 #endif
 
 #if WITH_EDITORONLY_DATA
 	UPROPERTY()
 	TArray<FWorldPartitionRuntimeCellObjectMapping> Packages;
+
+	// Actors to be loaded for cook
+	TArray<FGuid> ActorsToLoadForCook;
+
+	// Loaded actors for cook
+	UPROPERTY(Transient)
+	TArray<AActor*> LoadedActorsForCook;
 #endif
 
 	UPROPERTY()
diff --git a/Engine/Source/Runtime/Engine/Public/WorldPartition/WorldPartitionRuntimeSpatialHash.h b/Engine/Source/Runtime/Engine/Public/WorldPartition/WorldPartitionRuntimeSpatialHash.h
--- a/Engine/Source/Runtime/Engine/Public/WorldPartition/WorldPartitionRuntimeSpatialHash.h
+++ b/Engine/Source/Runtime/Engine/Public/WorldPartition/WorldPartitionRuntimeSpatialHash.h
@@ -213,7 +213,9 @@ public:
 #if WITH_EDITOR
 	virtual void SetDefaultValues() override;
 	virtual void ImportFromWorldComposition(class UWorldComposition* WorldComposition) override;
-	virtual bool GenerateStreaming(EWorldPartitionStreamingMode Mode, class UWorldPartitionStreamingPolicy* StreamingPolicy) override;
+	virtual bool GenerateStreaming(EWorldPartitionStreamingMode Mode, class UWorldPartitionStreamingPolicy* StreamingPolicy, TArray<FString>* OutPackagesToGenerate = nullptr) override;
+	virtual bool PopulateGeneratedPackageForCook(UPackage* InPackage, const FString& InPackageRelativePath, const FString& InPackageCookName) override;
+	virtual void FinalizeGeneratedPackageForCook() override;
 	virtual void FlushStreaming() override;
 	virtual bool GenerateHLOD() override;
 	virtual bool GenerateNavigationData() override;
@@ -244,6 +246,7 @@ private:
 	UPROPERTY(EditAnywhere, Config, Category=RuntimeSettings)
 	TArray<FSpatialHashRuntimeGrid> Grids;
 
+	TMap<FString, UWorldPartitionRuntimeCell*> PackagesToGenerateForCook;
 	TMap<FGuid, FGuid> CachedHLODParents;
 #endif
 
@@ -261,7 +264,7 @@ private:
 	void GetAlwaysLoadedStreamingCells(const FSpatialHashStreamingGrid& StreamingGrid, TSet<const UWorldPartitionRuntimeCell*>& Cells) const;
 	void GetStreamingCells(const FVector& Position, const FSpatialHashStreamingGrid& StreamingGrid, TSet<const UWorldPartitionRuntimeCell*>& Cells) const;
 #if WITH_EDITOR
-	bool CreateStreamingGrid(const FSpatialHashRuntimeGrid& RuntimeGrid, const FSquare2DGridHelper& PartionedActors, EWorldPartitionStreamingMode Mode, UWorldPartitionStreamingPolicy* StreamingPolicy);
+	bool CreateStreamingGrid(const FSpatialHashRuntimeGrid& RuntimeGrid, const FSquare2DGridHelper& PartionedActors, EWorldPartitionStreamingMode Mode, UWorldPartitionStreamingPolicy* StreamingPolicy, TArray<FString>* OutPackagesToGenerate = nullptr);
 #endif
 
 	friend class UWorldPartitionSubsystem;
diff --git a/Engine/Source/Runtime/Engine/Public/WorldPartition/WorldPartitionSubsystem.h b/Engine/Source/Runtime/Engine/Public/WorldPartition/WorldPartitionSubsystem.h
--- a/Engine/Source/Runtime/Engine/Public/WorldPartition/WorldPartitionSubsystem.h
+++ b/Engine/Source/Runtime/Engine/Public/WorldPartition/WorldPartitionSubsystem.h
@@ -45,12 +45,8 @@ public:
 #if WITH_EDITOR
 	void ForEachIntersectingActorDesc(const FBox& Box, TSubclassOf<AActor> ActorClass, TFunctionRef<bool(const FWorldPartitionActorDesc*)> Predicate) const;
 	void ForEachActorDesc(TSubclassOf<AActor> ActorClass, TFunctionRef<bool(const FWorldPartitionActorDesc*)> Predicate) const;
-
 	void RegisterActorDescFactory(TSubclassOf<AActor> Class, FWorldPartitionActorDescFactory* Factory);
-
-	FBox GetWorldBounds();
 #endif
-
 	void ToggleDrawRuntimeHash2D();
 
 private:
diff --git a/Engine/Source/Runtime/RenderCore/Private/RenderResource.cpp b/Engine/Source/Runtime/RenderCore/Private/RenderResource.cpp
--- a/Engine/Source/Runtime/RenderCore/Private/RenderResource.cpp
+++ b/Engine/Source/Runtime/RenderCore/Private/RenderResource.cpp
@@ -843,10 +843,8 @@ bool IsRayTracingEnabled()
 
 #if DO_CHECK && WITH_EDITOR
 	{
-		FString Commandline = FCommandLine::Get();
-		bool bIsCookCommandlet = IsRunningCommandlet() && Commandline.Contains(TEXT("run=cook"));
 		// This function must not be called while cooking
-		if (bIsCookCommandlet)
+		if (IsRunningCookCommandlet())
 		{
 			return false;
 		}
diff --git a/Engine/Source/Runtime/Renderer/Private/HairStrands/HairStrandsInterface.cpp b/Engine/Source/Runtime/Renderer/Private/HairStrands/HairStrandsInterface.cpp
--- a/Engine/Source/Runtime/Renderer/Private/HairStrands/HairStrandsInterface.cpp
+++ b/Engine/Source/Runtime/Renderer/Private/HairStrands/HairStrandsInterface.cpp
@@ -113,10 +113,7 @@ FRDGBufferUAVRef RegisterAsUAV(FRDGBuilder& GraphBuilder, const FRDGExternalBuff
 
 bool IsHairRayTracingEnabled()
 {
-	FString Commandline = FCommandLine::Get();
-	bool bIsCookCommandlet = IsRunningCommandlet() && Commandline.Contains(TEXT("run=cook"));
-
-	if (GIsRHIInitialized && !bIsCookCommandlet)
+	if (GIsRHIInitialized && !IsRunningCookCommandlet())
 	{
 		return IsRayTracingEnabled() && GHairStrandsRaytracingEnable;
 	}
