To protect the skeletalmesh when using async operations like build or re-import we deprecate all public member variables of USkeletalMesh class and create the necessary getter and setter. Those member

diff --git a/Engine/Source/Editor/UnrealEd/Private/Animation/DebugSkelMeshComponent.cpp b/Engine/Source/Editor/UnrealEd/Private/Animation/DebugSkelMeshComponent.cpp
--- a/Engine/Source/Editor/UnrealEd/Private/Animation/DebugSkelMeshComponent.cpp
+++ b/Engine/Source/Editor/UnrealEd/Private/Animation/DebugSkelMeshComponent.cpp
@@ -266,7 +266,7 @@ void UDebugSkelMeshComponent::InitAnim(bool bForceReinit)
 	// then we need to clear it up to avoid an issue
 	if ( PreviewInstance && PreviewInstance->GetCurrentAsset() && SkeletalMesh )
 	{
-		if ( PreviewInstance->GetCurrentAsset()->GetSkeleton() != SkeletalMesh->Skeleton )
+		if ( PreviewInstance->GetCurrentAsset()->GetSkeleton() != SkeletalMesh->GetSkeleton() )
 		{
 			// if it doesn't match, just clear it
 			PreviewInstance->SetAnimationAsset(NULL);
@@ -681,7 +681,7 @@ void UDebugSkelMeshComponent::ToggleClothSectionsVisibility(bool bShowOnlyClothS
 void UDebugSkelMeshComponent::RestoreClothSectionsVisibility()
 {
 	// if this skeletal mesh doesn't have any clothing assets, just return
-	if (!SkeletalMesh || SkeletalMesh->MeshClothingAssets.Num() == 0)
+	if (!SkeletalMesh || SkeletalMesh->GetMeshClothingAssets().Num() == 0)
 	{
 		return;
 	}
@@ -827,7 +827,7 @@ void UDebugSkelMeshComponent::RefreshSelectedClothingSkinnedPositions()
 {
 	if(SkeletalMesh && SelectedClothingGuidForPainting.IsValid())
 	{
-		UClothingAssetBase** Asset = SkeletalMesh->MeshClothingAssets.FindByPredicate([&](UClothingAssetBase* Item)
+		UClothingAssetBase** Asset = SkeletalMesh->GetMeshClothingAssets().FindByPredicate([&](UClothingAssetBase* Item)
 		{
 			return Item && SelectedClothingGuidForPainting == Item->GetAssetGuid();
 		});
@@ -1040,10 +1040,10 @@ FDebugSkelMeshDynamicData::FDebugSkelMeshDynamicData(UDebugSkelMeshComponent* In
 
 		if(USkeletalMesh* Mesh = InComponent->SkeletalMesh)
 		{
-			const int32 NumClothingAssets = Mesh->MeshClothingAssets.Num();
+			const int32 NumClothingAssets = Mesh->GetMeshClothingAssets().Num();
 			for(int32 ClothingAssetIndex = 0; ClothingAssetIndex < NumClothingAssets; ++ClothingAssetIndex)
 			{
-				UClothingAssetBase* BaseAsset = Mesh->MeshClothingAssets[ClothingAssetIndex];
+				UClothingAssetBase* BaseAsset = Mesh->GetMeshClothingAssets()[ClothingAssetIndex];
 				if(BaseAsset && BaseAsset->GetAssetGuid() == InComponent->SelectedClothingGuidForPainting)
 				{
 					ClothingSimDataIndexWhenPainting = ClothingAssetIndex;
diff --git a/Engine/Source/Editor/UnrealEd/Private/AnimationEditorUtils.cpp b/Engine/Source/Editor/UnrealEd/Private/AnimationEditorUtils.cpp
--- a/Engine/Source/Editor/UnrealEd/Private/AnimationEditorUtils.cpp
+++ b/Engine/Source/Editor/UnrealEd/Private/AnimationEditorUtils.cpp
@@ -468,7 +468,7 @@ namespace AnimationEditorUtils
 			if (Skeleton == nullptr)
 			{
 				SkeletalMesh = CastChecked<USkeletalMesh>(SkelIt->Get());
-				Skeleton = SkeletalMesh->Skeleton;
+				Skeleton = SkeletalMesh->GetSkeleton();
 			}
 
 			if(Skeleton)
@@ -543,7 +543,7 @@ namespace AnimationEditorUtils
 			if (Skeleton == nullptr)
 			{
 				SkeletalMesh = CastChecked<USkeletalMesh>(SkeletonsOrSkeletalMeshes[0].Get());
-				Skeleton = SkeletalMesh->Skeleton;
+				Skeleton = SkeletalMesh->GetSkeleton();
 			}
 
 			if (Skeleton)
@@ -596,7 +596,7 @@ namespace AnimationEditorUtils
 				if (Skeleton == nullptr)
 				{
 					SkeletalMesh = CastChecked<USkeletalMesh>(ObjIt->Get());
-					Skeleton = SkeletalMesh->Skeleton;
+					Skeleton = SkeletalMesh->GetSkeleton();
 				}
 
 				if(Skeleton)
diff --git a/Engine/Source/Editor/UnrealEd/Private/Factories/EditorFactories.cpp b/Engine/Source/Editor/UnrealEd/Private/Factories/EditorFactories.cpp
--- a/Engine/Source/Editor/UnrealEd/Private/Factories/EditorFactories.cpp
+++ b/Engine/Source/Editor/UnrealEd/Private/Factories/EditorFactories.cpp
@@ -5975,19 +5975,19 @@ bool UReimportFbxSkeletalMeshFactory::CanReimport( UObject* Obj, TArray<FString>
 	USkeletalMesh* SkeletalMesh = Cast<USkeletalMesh>(Obj);
 	if (SkeletalMesh && !SkeletalMesh->HasCustomActorReimportFactory())
 	{
-		if (SkeletalMesh->AssetImportData)
+		if (SkeletalMesh->GetAssetImportData())
 		{
-			UFbxAssetImportData *FbxAssetImportData = Cast<UFbxAssetImportData>(SkeletalMesh->AssetImportData);
+			UFbxAssetImportData *FbxAssetImportData = Cast<UFbxAssetImportData>(SkeletalMesh->GetAssetImportData());
 			if (FbxAssetImportData != nullptr && FbxAssetImportData->bImportAsScene)
 			{
 				//This skeletal mesh was import with a scene import, we cannot reimport it here
 				return false;
 			}
-			else if (FPaths::GetExtension(SkeletalMesh->AssetImportData->GetFirstFilename()) == TEXT("abc"))
+			else if (FPaths::GetExtension(SkeletalMesh->GetAssetImportData()->GetFirstFilename()) == TEXT("abc"))
 			{
 				return false;
 			}
-			SkeletalMesh->AssetImportData->ExtractFilenames(OutFilenames);
+			SkeletalMesh->GetAssetImportData()->ExtractFilenames(OutFilenames);
 		}
 		else
 		{
@@ -6043,7 +6043,7 @@ EReimportResult::Type UReimportFbxSkeletalMeshFactory::Reimport( UObject* Obj, i
 	//Clean up the options
 	UnFbx::FBXImportOptions::ResetOptions(ImportOptions);
 
-	UFbxSkeletalMeshImportData* ImportData = Cast<UFbxSkeletalMeshImportData>(SkeletalMesh->AssetImportData);
+	UFbxSkeletalMeshImportData* ImportData = Cast<UFbxSkeletalMeshImportData>(SkeletalMesh->GetAssetImportData());
 	
 	// Prepare the import options
 	UFbxImportUI* ReimportUI;
@@ -6062,9 +6062,9 @@ EReimportResult::Type UReimportFbxSkeletalMeshFactory::Reimport( UObject* Obj, i
 	}
 
 	ReimportUI->MeshTypeToImport = FBXIT_SkeletalMesh;
-	ReimportUI->Skeleton = SkeletalMesh->Skeleton;
+	ReimportUI->Skeleton = SkeletalMesh->GetSkeleton();
 	ReimportUI->bCreatePhysicsAsset = false;
-	ReimportUI->PhysicsAsset = SkeletalMesh->PhysicsAsset;
+	ReimportUI->PhysicsAsset = SkeletalMesh->GetPhysicsAsset();
 	ReimportUI->bImportAnimations = false;
 	ReimportUI->OverrideAnimationName = TEXT("");
 	ReimportUI->bImportRigidMesh = false;
@@ -6083,7 +6083,7 @@ EReimportResult::Type UReimportFbxSkeletalMeshFactory::Reimport( UObject* Obj, i
 	{
 		// An existing import data object was not found, make one here and show the options dialog
 		ImportData = UFbxSkeletalMeshImportData::GetImportDataForSkeletalMesh(SkeletalMesh, ImportUI->SkeletalMeshImportData);
-		SkeletalMesh->AssetImportData = ImportData;
+		SkeletalMesh->SetAssetImportData(ImportData);
 	}
 	check(ImportData != nullptr);
 
@@ -6239,9 +6239,9 @@ EReimportResult::Type UReimportFbxSkeletalMeshFactory::Reimport( UObject* Obj, i
 		bool bIsObjFormat = false;
 		bool bIsAutomated = false;
 		// @hack to make sure skeleton is set before opening the dialog
-		ImportOptions->SkeletonForAnimation = SkeletalMesh->Skeleton;
+		ImportOptions->SkeletonForAnimation = SkeletalMesh->GetSkeleton();
 		ImportOptions->bCreatePhysicsAsset = false;
-		ImportOptions->PhysicsAsset = SkeletalMesh->PhysicsAsset;
+		ImportOptions->PhysicsAsset = SkeletalMesh->GetPhysicsAsset();
 		
 		EFBXImportContentType BeforeUIContentType = ReimportUI->SkeletalMeshImportData->ImportContentType;
 		ImportOptions = GetImportOptions( FFbxImporter, ReimportUI, bShowOptionDialog, bIsAutomated, Obj->GetPathName(), bOperationCanceled, bOutImportAll, bIsObjFormat, Filename, bForceImportType, FBXIT_SkeletalMesh);
@@ -6250,15 +6250,15 @@ EReimportResult::Type UReimportFbxSkeletalMeshFactory::Reimport( UObject* Obj, i
 		if (bShowOptionDialog && ReimportUI->SkeletalMeshImportData->ImportContentType != BeforeUIContentType && ReimportUI->SkeletalMeshImportData->ImportContentType != EFBXImportContentType::FBXICT_All)
 		{
 			TArray<FString> ExtractedFilenames;
-			SkeletalMesh->AssetImportData->ExtractFilenames(ExtractedFilenames);
+			SkeletalMesh->GetAssetImportData()->ExtractFilenames(ExtractedFilenames);
 			//By default add the original file
 			if (!ExtractedFilenames.IsValidIndex(1))
 			{
-				SkeletalMesh->AssetImportData->AddFileName(ExtractedFilenames[0], 1, NSSkeletalMeshSourceFileLabels::GeometryText().ToString());
+				SkeletalMesh->GetAssetImportData()->AddFileName(ExtractedFilenames[0], 1, NSSkeletalMeshSourceFileLabels::GeometryText().ToString());
 			}
 			if (!ExtractedFilenames.IsValidIndex(2))
 			{
-				SkeletalMesh->AssetImportData->AddFileName(ExtractedFilenames[0], 2, NSSkeletalMeshSourceFileLabels::SkinningText().ToString());
+				SkeletalMesh->GetAssetImportData()->AddFileName(ExtractedFilenames[0], 2, NSSkeletalMeshSourceFileLabels::SkinningText().ToString());
 			}
 		}
 
diff --git a/Engine/Source/Editor/UnrealEd/Private/Factories/SkeletalMeshImport.cpp b/Engine/Source/Editor/UnrealEd/Private/Factories/SkeletalMeshImport.cpp
--- a/Engine/Source/Editor/UnrealEd/Private/Factories/SkeletalMeshImport.cpp
+++ b/Engine/Source/Editor/UnrealEd/Private/Factories/SkeletalMeshImport.cpp
@@ -326,7 +326,7 @@ void SkeletalMeshHelperImpl::SaveSkeletalMeshLODModelSections(USkeletalMesh* Sou
 
 void SkeletalMeshHelperImpl::SaveSkeletalMeshMaterialNameWorkflowData(TSharedPtr<FExistingSkelMeshData>& ExistingMeshDataPtr, const USkeletalMesh* SourceSkeletalMesh)
 {
-	const UFbxSkeletalMeshImportData* ImportData = Cast<UFbxSkeletalMeshImportData>(SourceSkeletalMesh->AssetImportData);
+	const UFbxSkeletalMeshImportData* ImportData = Cast<UFbxSkeletalMeshImportData>(SourceSkeletalMesh->GetAssetImportData());
 	if (!ImportData)
 	{
 		return;
@@ -386,18 +386,19 @@ TSharedPtr<FExistingSkelMeshData> SkeletalMeshHelper::SaveExistingSkelMeshData(U
 
 	//Save the package UMetaData
 	ExistingMeshDataPtr->ExistingUMetaDataTagValues = UMetaData::GetMapForObject(SourceSkeletalMesh);
-	ExistingMeshDataPtr->UseMaterialNameSlotWorkflow = SkeletalMeshIsUsingMaterialSlotNameWorkflow(SourceSkeletalMesh->AssetImportData);
-	ExistingMeshDataPtr->MinLOD = SourceSkeletalMesh->MinLod;
-	ExistingMeshDataPtr->DisableBelowMinLodStripping = SourceSkeletalMesh->DisableBelowMinLodStripping;
-	ExistingMeshDataPtr->bOverrideLODStreamingSettings = SourceSkeletalMesh->bOverrideLODStreamingSettings;
-	ExistingMeshDataPtr->bSupportLODStreaming = SourceSkeletalMesh->bSupportLODStreaming;
-	ExistingMeshDataPtr->MaxNumStreamedLODs = SourceSkeletalMesh->MaxNumStreamedLODs;
-	ExistingMeshDataPtr->MaxNumOptionalLODs = SourceSkeletalMesh->MaxNumOptionalLODs;
-
+	ExistingMeshDataPtr->UseMaterialNameSlotWorkflow = SkeletalMeshIsUsingMaterialSlotNameWorkflow(SourceSkeletalMesh->GetAssetImportData());
+	ExistingMeshDataPtr->MinLOD = SourceSkeletalMesh->GetMinLod();
+	ExistingMeshDataPtr->DisableBelowMinLodStripping = SourceSkeletalMesh->GetDisableBelowMinLodStripping();
+	ExistingMeshDataPtr->bOverrideLODStreamingSettings = SourceSkeletalMesh->GetOverrideLODStreamingSettings();
+	ExistingMeshDataPtr->bSupportLODStreaming = SourceSkeletalMesh->GetSupportLODStreaming();
+	ExistingMeshDataPtr->MaxNumStreamedLODs = SourceSkeletalMesh->GetMaxNumStreamedLODs();
+	ExistingMeshDataPtr->MaxNumOptionalLODs = SourceSkeletalMesh->GetMaxNumOptionalLODs();
+
+	const TArray<FSkeletalMaterial>& SourceMaterials = SourceSkeletalMesh->GetMaterials();
 	//Add the existing Material slot name data
-	for (int32 MaterialIndex = 0; MaterialIndex < SourceSkeletalMesh->Materials.Num(); ++MaterialIndex)
+	for (int32 MaterialIndex = 0; MaterialIndex < SourceMaterials.Num(); ++MaterialIndex)
 	{
-		ExistingMeshDataPtr->ExistingImportMaterialOriginalNameData.Add(SourceSkeletalMesh->Materials[MaterialIndex].ImportedMaterialSlotName);
+		ExistingMeshDataPtr->ExistingImportMaterialOriginalNameData.Add(SourceMaterials[MaterialIndex].ImportedMaterialSlotName);
 	}
 
 	FSkeletalMeshModel* SourceMeshModel = SourceSkeletalMesh->GetImportedModel();
@@ -411,9 +412,9 @@ TSharedPtr<FExistingSkelMeshData> SkeletalMeshHelper::SaveExistingSkelMeshData(U
 	ExistingMeshDataPtr->bSaveRestoreMaterials = bSaveMaterials;
 	if (ExistingMeshDataPtr->bSaveRestoreMaterials)
 	{
-		ExistingMeshDataPtr->ExistingMaterials = SourceSkeletalMesh->Materials;
+		ExistingMeshDataPtr->ExistingMaterials = SourceMaterials;
 	}
-	ExistingMeshDataPtr->ExistingRetargetBasePose = SourceSkeletalMesh->RetargetBasePose;
+	ExistingMeshDataPtr->ExistingRetargetBasePose = SourceSkeletalMesh->GetRetargetBasePose();
 
 	if (SourceMeshModel->LODModels.Num() > 0 &&
 		SourceSkeletalMesh->GetLODNum() == SourceMeshModel->LODModels.Num())
@@ -443,31 +444,31 @@ TSharedPtr<FExistingSkelMeshData> SkeletalMeshHelper::SaveExistingSkelMeshData(U
 		check(ExistingMeshDataPtr->ExistingLODModels.Num() == SourceMeshModel->LODModels.Num());
 
 		ExistingMeshDataPtr->ExistingLODInfo = SourceSkeletalMesh->GetLODInfoArray();
-		ExistingMeshDataPtr->ExistingRefSkeleton = SourceSkeletalMesh->RefSkeleton;
+		ExistingMeshDataPtr->ExistingRefSkeleton = SourceSkeletalMesh->GetRefSkeleton();
 	}
 
 	// First asset should be the one that the skeletal mesh should point too
 	ExistingMeshDataPtr->ExistingPhysicsAssets.Empty();
-	ExistingMeshDataPtr->ExistingPhysicsAssets.Add(SourceSkeletalMesh->PhysicsAsset);
+	ExistingMeshDataPtr->ExistingPhysicsAssets.Add(SourceSkeletalMesh->GetPhysicsAsset());
 	for (TObjectIterator<UPhysicsAsset> It; It; ++It)
 	{
 		UPhysicsAsset* PhysicsAsset = *It;
-		if (PhysicsAsset->PreviewSkeletalMesh == SourceSkeletalMesh && SourceSkeletalMesh->PhysicsAsset != PhysicsAsset)
+		if (PhysicsAsset->PreviewSkeletalMesh == SourceSkeletalMesh && SourceSkeletalMesh->GetPhysicsAsset() != PhysicsAsset)
 		{
 			ExistingMeshDataPtr->ExistingPhysicsAssets.Add(PhysicsAsset);
 		}
 	}
 
-	ExistingMeshDataPtr->ExistingShadowPhysicsAsset = SourceSkeletalMesh->ShadowPhysicsAsset;
-	ExistingMeshDataPtr->ExistingSkeleton = SourceSkeletalMesh->Skeleton;
+	ExistingMeshDataPtr->ExistingShadowPhysicsAsset = SourceSkeletalMesh->GetShadowPhysicsAsset();
+	ExistingMeshDataPtr->ExistingSkeleton = SourceSkeletalMesh->GetSkeleton();
 	// since copying back original skeleton, this should be safe to do
-	ExistingMeshDataPtr->ExistingPostProcessAnimBlueprint = SourceSkeletalMesh->PostProcessAnimBlueprint;
-	ExistingMeshDataPtr->ExistingLODSettings = SourceSkeletalMesh->LODSettings;
+	ExistingMeshDataPtr->ExistingPostProcessAnimBlueprint = SourceSkeletalMesh->GetPostProcessAnimBlueprint();
+	ExistingMeshDataPtr->ExistingLODSettings = SourceSkeletalMesh->GetLODSettings();
 	SourceSkeletalMesh->ExportMirrorTable(ExistingMeshDataPtr->ExistingMirrorTable);
-	ExistingMeshDataPtr->ExistingMorphTargets = SourceSkeletalMesh->MorphTargets;
-	ExistingMeshDataPtr->ExistingAssetImportData = SourceSkeletalMesh->AssetImportData;
-	ExistingMeshDataPtr->ExistingThumbnailInfo = SourceSkeletalMesh->ThumbnailInfo;
-	ExistingMeshDataPtr->ExistingClothingAssets = SourceSkeletalMesh->MeshClothingAssets;
+	ExistingMeshDataPtr->ExistingMorphTargets = SourceSkeletalMesh->GetMorphTargets();
+	ExistingMeshDataPtr->ExistingAssetImportData = SourceSkeletalMesh->GetAssetImportData();
+	ExistingMeshDataPtr->ExistingThumbnailInfo = SourceSkeletalMesh->GetThumbnailInfo();
+	ExistingMeshDataPtr->ExistingClothingAssets = SourceSkeletalMesh->GetMeshClothingAssets();
 	ExistingMeshDataPtr->ExistingSamplingInfo = SourceSkeletalMesh->GetSamplingInfo();
 
 	if (ExistingMeshDataPtr->UseMaterialNameSlotWorkflow)
@@ -663,7 +664,7 @@ void SkeletalMeshHelper::ApplySkinning(USkeletalMesh* SkeletalMesh, FSkeletalMes
 		//Set the required/active bones
 		DestLODModel.RequiredBones = SrcLODModel.RequiredBones;
 		DestLODModel.RequiredBones.Sort();
-		SkeletalMesh->RefSkeleton.EnsureParentsExistAndSort(DestLODModel.ActiveBoneIndices);
+		SkeletalMesh->GetRefSkeleton().EnsureParentsExistAndSort(DestLODModel.ActiveBoneIndices);
 	}
 }
 
@@ -685,34 +686,34 @@ void SkeletalMeshHelper::RestoreExistingSkelMeshData(TSharedPtr<const FExistingS
 	}
 
 	int32 SafeReimportLODIndex = ReimportLODIndex < 0 ? 0 : ReimportLODIndex;
-	SkeletalMesh->MinLod = MeshData->MinLOD;
-	SkeletalMesh->DisableBelowMinLodStripping = MeshData->DisableBelowMinLodStripping;
-	SkeletalMesh->bOverrideLODStreamingSettings = MeshData->bOverrideLODStreamingSettings;
-	SkeletalMesh->bSupportLODStreaming = MeshData->bSupportLODStreaming;
-	SkeletalMesh->MaxNumStreamedLODs = MeshData->MaxNumStreamedLODs;
-	SkeletalMesh->MaxNumOptionalLODs = MeshData->MaxNumOptionalLODs;
+	SkeletalMesh->SetMinLod(MeshData->MinLOD);
+	SkeletalMesh->SetDisableBelowMinLodStripping(MeshData->DisableBelowMinLodStripping);
+	SkeletalMesh->SetOverrideLODStreamingSettings(MeshData->bOverrideLODStreamingSettings);
+	SkeletalMesh->SetSupportLODStreaming(MeshData->bSupportLODStreaming);
+	SkeletalMesh->SetMaxNumStreamedLODs(MeshData->MaxNumStreamedLODs);
+	SkeletalMesh->SetMaxNumOptionalLODs(MeshData->MaxNumOptionalLODs);
 
 	FSkeletalMeshModel* SkeletalMeshImportedModel = SkeletalMesh->GetImportedModel();
 
 	//Create a remap material Index use to find the matching section later
 	TArray<int32> RemapMaterial;
-	RemapMaterial.AddZeroed(SkeletalMesh->Materials.Num());
+	RemapMaterial.AddZeroed(SkeletalMesh->GetMaterials().Num());
 	TArray<FName> RemapMaterialName;
-	RemapMaterialName.AddZeroed(SkeletalMesh->Materials.Num());
+	RemapMaterialName.AddZeroed(SkeletalMesh->GetMaterials().Num());
 
 	bool bMaterialReset = false;
 	if (MeshData->bSaveRestoreMaterials)
 	{
 		UnFbx::EFBXReimportDialogReturnOption ReturnOption;
 		//Ask the user to match the materials conflict
-		UnFbx::FFbxImporter::PrepareAndShowMaterialConflictDialog<FSkeletalMaterial>(MeshData->ExistingMaterials, SkeletalMesh->Materials, RemapMaterial, RemapMaterialName, bCanShowDialog, false, bForceMaterialReset, ReturnOption);
+		UnFbx::FFbxImporter::PrepareAndShowMaterialConflictDialog<FSkeletalMaterial>(MeshData->ExistingMaterials, SkeletalMesh->GetMaterials(), RemapMaterial, RemapMaterialName, bCanShowDialog, false, bForceMaterialReset, ReturnOption);
 
 		if (ReturnOption != UnFbx::EFBXReimportDialogReturnOption::FBXRDRO_ResetToFbx)
 		{
 			//Build a ordered material list that try to keep intact the existing material list
 			TArray<FSkeletalMaterial> MaterialOrdered;
 			TArray<bool> MatchedNewMaterial;
-			MatchedNewMaterial.AddZeroed(SkeletalMesh->Materials.Num());
+			MatchedNewMaterial.AddZeroed(SkeletalMesh->GetMaterials().Num());
 			for (int32 ExistMaterialIndex = 0; ExistMaterialIndex < MeshData->ExistingMaterials.Num(); ++ExistMaterialIndex)
 			{
 				int32 MaterialIndexOrdered = MaterialOrdered.Add(MeshData->ExistingMaterials[ExistMaterialIndex]);
@@ -722,7 +723,7 @@ void SkeletalMeshHelper::RestoreExistingSkelMeshData(TSharedPtr<const FExistingS
 				{
 					MatchedNewMaterial[NewMaterialIndex] = true;
 					RemapMaterial[NewMaterialIndex] = MaterialIndexOrdered;
-					OrderedMaterial.ImportedMaterialSlotName = SkeletalMesh->Materials[NewMaterialIndex].ImportedMaterialSlotName;
+					OrderedMaterial.ImportedMaterialSlotName = SkeletalMesh->GetMaterials()[NewMaterialIndex].ImportedMaterialSlotName;
 				}
 				else
 				{
@@ -735,7 +736,7 @@ void SkeletalMeshHelper::RestoreExistingSkelMeshData(TSharedPtr<const FExistingS
 			{
 				if (MatchedNewMaterial[NewMaterialIndex] == false)
 				{
-					int32 NewMeshIndex = MaterialOrdered.Add(SkeletalMesh->Materials[NewMaterialIndex]);
+					int32 NewMeshIndex = MaterialOrdered.Add(SkeletalMesh->GetMaterials()[NewMaterialIndex]);
 					RemapMaterial[NewMaterialIndex] = NewMeshIndex;
 				}
 			}
@@ -751,7 +752,7 @@ void SkeletalMeshHelper::RestoreExistingSkelMeshData(TSharedPtr<const FExistingS
 			}
 
 			//Copy the re ordered materials (this ensure the material array do not change when we re-import)
-			SkeletalMesh->Materials = MaterialOrdered;
+			SkeletalMesh->SetMaterials(MaterialOrdered);
 		}
 		else
 		{
@@ -759,11 +760,11 @@ void SkeletalMeshHelper::RestoreExistingSkelMeshData(TSharedPtr<const FExistingS
 		}
 	}
 
-	SkeletalMesh->LODSettings = MeshData->ExistingLODSettings;
+	SkeletalMesh->SetLODSettings(MeshData->ExistingLODSettings);
 	// ensure LOD 0 contains correct setting 
-	if (SkeletalMesh->LODSettings && SkeletalMesh->GetLODInfoArray().Num() > 0)
+	if (SkeletalMesh->GetLODSettings() && SkeletalMesh->GetLODInfoArray().Num() > 0)
 	{
-		SkeletalMesh->LODSettings->SetLODSettingsToMesh(SkeletalMesh, 0);
+		SkeletalMesh->GetLODSettings()->SetLODSettingsToMesh(SkeletalMesh, 0);
 	}
 
 	//Do everything we need for base LOD re-import
@@ -771,16 +772,16 @@ void SkeletalMeshHelper::RestoreExistingSkelMeshData(TSharedPtr<const FExistingS
 	{
 		// this is not ideal. Ideally we'll have to save only diff with indicating which joints, 
 		// but for now, we allow them to keep the previous pose IF the element count is same
-		if (MeshData->ExistingRetargetBasePose.Num() == SkeletalMesh->RefSkeleton.GetRawBoneNum())
+		if (MeshData->ExistingRetargetBasePose.Num() == SkeletalMesh->GetRefSkeleton().GetRawBoneNum())
 		{
-			SkeletalMesh->RetargetBasePose = MeshData->ExistingRetargetBasePose;
+			SkeletalMesh->SetRetargetBasePose(MeshData->ExistingRetargetBasePose);
 		}
 
 		// Assign sockets from old version of this SkeletalMesh.
 		// Only copy ones for bones that exist in the new mesh.
 		for (int32 i = 0; i < MeshData->ExistingSockets.Num(); i++)
 		{
-			const int32 BoneIndex = SkeletalMesh->RefSkeleton.FindBoneIndex(MeshData->ExistingSockets[i]->BoneName);
+			const int32 BoneIndex = SkeletalMesh->GetRefSkeleton().FindBoneIndex(MeshData->ExistingSockets[i]->BoneName);
 			if (BoneIndex != INDEX_NONE)
 			{
 				SkeletalMesh->GetMeshOnlySocketList().Add(MeshData->ExistingSockets[i]);
@@ -811,14 +812,14 @@ void SkeletalMeshHelper::RestoreExistingSkelMeshData(TSharedPtr<const FExistingS
 				SkeletalMeshImportedModel->OriginalReductionSourceMeshData.Add(ReductionLODData);
 			};
 
-			if (SkeletonsAreCompatible(SkeletalMesh->RefSkeleton, MeshData->ExistingRefSkeleton, bImportSkinningOnly))
+			if (SkeletonsAreCompatible(SkeletalMesh->GetRefSkeleton(), MeshData->ExistingRefSkeleton, bImportSkinningOnly))
 			{
 				// First create mapping table from old skeleton to new skeleton.
 				TArray<int32> OldToNewMap;
 				OldToNewMap.AddUninitialized(MeshData->ExistingRefSkeleton.GetRawBoneNum());
 				for (int32 i = 0; i < MeshData->ExistingRefSkeleton.GetRawBoneNum(); i++)
 				{
-					OldToNewMap[i] = SkeletalMesh->RefSkeleton.FindBoneIndex(MeshData->ExistingRefSkeleton.GetBoneName(i));
+					OldToNewMap[i] = SkeletalMesh->GetRefSkeleton().FindBoneIndex(MeshData->ExistingRefSkeleton.GetBoneName(i));
 				}
 
 				//Starting at index 1 because we only need to add LOD models of LOD 1 and higher.
@@ -882,7 +883,7 @@ void SkeletalMeshHelper::RestoreExistingSkelMeshData(TSharedPtr<const FExistingS
 
 					// Sort ascending for parent child relationship
 					LODModelCopy->RequiredBones.Sort();
-					SkeletalMesh->RefSkeleton.EnsureParentsExistAndSort(LODModelCopy->ActiveBoneIndices);
+					SkeletalMesh->GetRefSkeleton().EnsureParentsExistAndSort(LODModelCopy->ActiveBoneIndices);
 
 					// Fix the sections' BoneMaps.
 					for (int32 SectionIndex = 0; SectionIndex < LODModelCopy->Sections.Num(); SectionIndex++)
@@ -940,7 +941,7 @@ void SkeletalMeshHelper::RestoreExistingSkelMeshData(TSharedPtr<const FExistingS
 			if (AssetIndex == 0)
 			{
 				// First asset is the one that the skeletal mesh should point too
-				SkeletalMesh->PhysicsAsset = PhysicsAsset;
+				SkeletalMesh->SetPhysicsAsset(PhysicsAsset);
 			}
 			// No need to mark as modified here, because the asset hasn't actually changed
 			if (PhysicsAsset)
@@ -949,20 +950,20 @@ void SkeletalMeshHelper::RestoreExistingSkelMeshData(TSharedPtr<const FExistingS
 			}
 		}
 
-		SkeletalMesh->ShadowPhysicsAsset = MeshData->ExistingShadowPhysicsAsset;
-		SkeletalMesh->Skeleton = MeshData->ExistingSkeleton;
-		SkeletalMesh->PostProcessAnimBlueprint = MeshData->ExistingPostProcessAnimBlueprint;
+		SkeletalMesh->SetShadowPhysicsAsset(MeshData->ExistingShadowPhysicsAsset);
+		SkeletalMesh->SetSkeleton(MeshData->ExistingSkeleton);
+		SkeletalMesh->SetPostProcessAnimBlueprint(MeshData->ExistingPostProcessAnimBlueprint);
 
 		// Copy mirror table.
 		SkeletalMesh->ImportMirrorTable(MeshData->ExistingMirrorTable);
-		SkeletalMesh->MorphTargets.Empty(MeshData->ExistingMorphTargets.Num());
-		SkeletalMesh->MorphTargets.Append(MeshData->ExistingMorphTargets);
+		SkeletalMesh->GetMorphTargets().Empty(MeshData->ExistingMorphTargets.Num());
+		SkeletalMesh->GetMorphTargets().Append(MeshData->ExistingMorphTargets);
 		SkeletalMesh->InitMorphTargets();
-		SkeletalMesh->AssetImportData = MeshData->ExistingAssetImportData.Get();
-		SkeletalMesh->ThumbnailInfo = MeshData->ExistingThumbnailInfo.Get();
-		SkeletalMesh->MeshClothingAssets = MeshData->ExistingClothingAssets;
+		SkeletalMesh->SetAssetImportData(MeshData->ExistingAssetImportData.Get());
+		SkeletalMesh->SetThumbnailInfo(MeshData->ExistingThumbnailInfo.Get());
+		SkeletalMesh->SetMeshClothingAssets(MeshData->ExistingClothingAssets);
 
-		for (UClothingAssetBase* ClothingAsset : SkeletalMesh->MeshClothingAssets)
+		for (UClothingAssetBase* ClothingAsset : SkeletalMesh->GetMeshClothingAssets())
 		{
 			if (ClothingAsset)
 			{
@@ -1048,13 +1049,14 @@ void SkeletalMeshHelperImpl::RestoreMaterialNameWorkflowSection(const TSharedPtr
 	TBitArray<> MatchedExistSectionIndex;
 	MatchedExistSectionIndex.Init(false, MaxExistSectionNumber);
 
+	const TArray<FSkeletalMaterial>& SkeletalMeshMaterials = SkeletalMesh->GetMaterials();
 	//Restore the section changes from the old import data
 	for (int32 SectionIndex = 0; SectionIndex < SkeletalMeshLodModel.Sections.Num(); SectionIndex++)
 	{
 		//Find the import section material index by using the RemapMaterial array. Fallback on the imported index if the remap entry is not valid
 		FSkelMeshSection& NewSection = SkeletalMeshLodModel.Sections[SectionIndex];
 		int32 RemapMaterialIndex = RemapMaterial.IsValidIndex(NewSection.MaterialIndex) ? RemapMaterial[NewSection.MaterialIndex] : NewSection.MaterialIndex;
-		if (!SkeletalMesh->Materials.IsValidIndex(RemapMaterialIndex))
+		if (!SkeletalMeshMaterials.IsValidIndex(RemapMaterialIndex))
 		{
 			//We have an invalid material section, in this case we set the material index to 0
 			NewSection.MaterialIndex = 0;
@@ -1069,7 +1071,7 @@ void SkeletalMeshHelperImpl::RestoreMaterialNameWorkflowSection(const TSharedPtr
 			continue;
 		}
 		//Get the RemapMaterial section Imported material slot name. We need it to match the saved existing section, so we can put back the saved existing section data
-		FName CurrentSectionImportedMaterialName = SkeletalMesh->Materials[RemapMaterialIndex].ImportedMaterialSlotName;
+		FName CurrentSectionImportedMaterialName = SkeletalMeshMaterials[RemapMaterialIndex].ImportedMaterialSlotName;
 		for (int32 ExistSectionIndex = 0; ExistSectionIndex < MaxExistSectionNumber; ++ExistSectionIndex)
 		{
 			//Skip already matched exist section
@@ -1116,9 +1118,9 @@ void SkeletalMeshHelperImpl::RestoreMaterialNameWorkflowSection(const TSharedPtr
 			//find the corresponding current slot name in the skeletal mesh materials list to remap properly the material index, in case the user have change it before re-importing
 			FName ExistMeshSectionSlotName = MeshData->ExistingImportMeshLodSectionMaterialData[LodIndex][ExistSectionIndex].ImportedMaterialSlotName;
 			{
-				for (int32 SkelMeshMaterialIndex = 0; SkelMeshMaterialIndex < SkeletalMesh->Materials.Num(); ++SkelMeshMaterialIndex)
+				for (int32 SkelMeshMaterialIndex = 0; SkelMeshMaterialIndex < SkeletalMeshMaterials.Num(); ++SkelMeshMaterialIndex)
 				{
-					const FSkeletalMaterial &NewSectionMaterial = SkeletalMesh->Materials[SkelMeshMaterialIndex];
+					const FSkeletalMaterial &NewSectionMaterial = SkeletalMeshMaterials[SkelMeshMaterialIndex];
 					if (NewSectionMaterial.ImportedMaterialSlotName == ExistMeshSectionSlotName)
 					{
 						if (ExistMeshSectionSlotName != OriginalImportMeshSectionSlotName)
diff --git a/Engine/Source/Editor/UnrealEd/Private/Factories/SkeletonFactory.cpp b/Engine/Source/Editor/UnrealEd/Private/Factories/SkeletonFactory.cpp
--- a/Engine/Source/Editor/UnrealEd/Private/Factories/SkeletonFactory.cpp
+++ b/Engine/Source/Editor/UnrealEd/Private/Factories/SkeletonFactory.cpp
@@ -92,9 +92,9 @@ UObject* USkeletonFactory::FactoryCreateNew(UClass* Class, UObject* InParent, FN
 			NewAsset = NULL;
 		}
 
-		if (TargetSkeletalMesh->Skeleton != NewAsset)
+		if (TargetSkeletalMesh->GetSkeleton() != NewAsset)
 		{
-			TargetSkeletalMesh->Skeleton = NewAsset;
+			TargetSkeletalMesh->SetSkeleton(NewAsset);
 			TargetSkeletalMesh->MarkPackageDirty();
 		}
 	}
diff --git a/Engine/Source/Editor/UnrealEd/Private/Fbx/FbxAnimationExport.cpp b/Engine/Source/Editor/UnrealEd/Private/Fbx/FbxAnimationExport.cpp
--- a/Engine/Source/Editor/UnrealEd/Private/Fbx/FbxAnimationExport.cpp
+++ b/Engine/Source/Editor/UnrealEd/Private/Fbx/FbxAnimationExport.cpp
@@ -706,7 +706,7 @@ void FFbxExporter::ExportAnimTrack(IAnimTrackAdapter& AnimTrackAdapter, AActor*
 		// Add the animation data to the bone nodes
 		for(int32 BoneIndex = 0; BoneIndex < BoneNodes.Num(); ++BoneIndex)
 		{
-			FName BoneName = SkeletalMeshComponent->SkeletalMesh->RefSkeleton.GetBoneName(BoneIndex);
+			FName BoneName = SkeletalMeshComponent->SkeletalMesh->GetRefSkeleton().GetBoneName(BoneIndex);
 			FbxNode* CurrentBoneNode = BoneNodes[BoneIndex];
 
 			// Create the AnimCurves
diff --git a/Engine/Source/Editor/UnrealEd/Private/Fbx/FbxFactory.cpp b/Engine/Source/Editor/UnrealEd/Private/Fbx/FbxFactory.cpp
--- a/Engine/Source/Editor/UnrealEd/Private/Fbx/FbxFactory.cpp
+++ b/Engine/Source/Editor/UnrealEd/Private/Fbx/FbxFactory.cpp
@@ -1701,10 +1701,11 @@ namespace ImportCompareHelper
 		if (!bImportSkinningOnly)
 		{
 			//Fill the currrent asset data
-			ImportUI->MaterialCompareData.CurrentAsset.Reserve(SkeletalMesh->Materials.Num());
-			for (int32 MaterialIndex = 0; MaterialIndex < SkeletalMesh->Materials.Num(); ++MaterialIndex)
+			const TArray<FSkeletalMaterial>& SkeletalMeshMaterials = SkeletalMesh->GetMaterials();
+			ImportUI->MaterialCompareData.CurrentAsset.Reserve(SkeletalMeshMaterials.Num());
+			for (int32 MaterialIndex = 0; MaterialIndex < SkeletalMeshMaterials.Num(); ++MaterialIndex)
 			{
-				const FSkeletalMaterial& Material = SkeletalMesh->Materials[MaterialIndex];
+				const FSkeletalMaterial& Material = SkeletalMeshMaterials[MaterialIndex];
 				FMaterialData MaterialData;
 				MaterialData.MaterialIndex = MaterialIndex;
 				MaterialData.ImportedMaterialSlotName = Material.ImportedMaterialSlotName;
@@ -1729,7 +1730,7 @@ namespace ImportCompareHelper
 		if (!bImportGeoOnly)
 		{
 			//Fill the currrent asset data
-			if (ImportUI->Skeleton && SkeletalMesh->Skeleton != ImportUI->Skeleton)
+			if (ImportUI->Skeleton && SkeletalMesh->GetSkeleton() != ImportUI->Skeleton)
 			{
 				//In this case we can't use 
 				const FReferenceSkeleton& ReferenceSkeleton = ImportUI->Skeleton->GetReferenceSkeleton();
@@ -1737,7 +1738,7 @@ namespace ImportCompareHelper
 			}
 			else
 			{
-				FillRecursivelySkeleton(SkeletalMesh->RefSkeleton, 0, ImportUI->SkeletonCompareData.CurrentAssetRoot);
+				FillRecursivelySkeleton(SkeletalMesh->GetRefSkeleton(), 0, ImportUI->SkeletonCompareData.CurrentAssetRoot);
 			}
 			
 			//Fill the result fbx data
diff --git a/Engine/Source/Editor/UnrealEd/Private/Fbx/FbxReimportDialogs.cpp b/Engine/Source/Editor/UnrealEd/Private/Fbx/FbxReimportDialogs.cpp
--- a/Engine/Source/Editor/UnrealEd/Private/Fbx/FbxReimportDialogs.cpp
+++ b/Engine/Source/Editor/UnrealEd/Private/Fbx/FbxReimportDialogs.cpp
@@ -114,7 +114,7 @@ void FFbxImporter::ShowFbxSkeletonConflictWindow(USkeletalMesh* SkeletalMesh, US
 	
 	if (Skeleton == nullptr)
 	{
-		Skeleton = SkeletalMesh->Skeleton;
+		Skeleton = SkeletalMesh->GetSkeleton();
 	}
 
 	FCompMesh SourceData;
diff --git a/Engine/Source/Editor/UnrealEd/Private/Fbx/FbxSkeletalMeshExport.cpp b/Engine/Source/Editor/UnrealEd/Private/Fbx/FbxSkeletalMeshExport.cpp
--- a/Engine/Source/Editor/UnrealEd/Private/Fbx/FbxSkeletalMeshExport.cpp
+++ b/Engine/Source/Editor/UnrealEd/Private/Fbx/FbxSkeletalMeshExport.cpp
@@ -25,7 +25,7 @@ namespace UnFbx
  */
 FbxNode* FFbxExporter::CreateSkeleton(const USkeletalMesh* SkelMesh, TArray<FbxNode*>& BoneNodes)
 {
-	const FReferenceSkeleton& RefSkeleton= SkelMesh->RefSkeleton;
+	const FReferenceSkeleton& RefSkeleton= SkelMesh->GetRefSkeleton();
 
 	if(RefSkeleton.GetRawBoneNum() == 0)
 	{
@@ -253,19 +253,19 @@ FbxNode* FFbxExporter::CreateMesh(const USkeletalMesh* SkelMesh, const TCHAR* Me
 	}
 
 
-	if (GetExportOptions()->bExportMorphTargets && SkelMesh->Skeleton) //The skeleton can be null if this is a destructible mesh.
+	if (GetExportOptions()->bExportMorphTargets && SkelMesh->GetSkeleton()) //The skeleton can be null if this is a destructible mesh.
 	{
-		const FSmartNameMapping* SmartNameMapping = SkelMesh->Skeleton->GetSmartNameContainer(USkeleton::AnimCurveMappingName);
+		const FSmartNameMapping* SmartNameMapping = SkelMesh->GetSkeleton()->GetSmartNameContainer(USkeleton::AnimCurveMappingName);
 		TMap<FName, FbxAnimCurve*> BlendShapeCurvesMap;
 
-		if (SkelMesh->MorphTargets.Num())
+		if (SkelMesh->GetMorphTargets().Num())
 		{
 			// The original BlendShape Name was not saved during import, so we need to come up with a new one.
 			const FString BlendShapeName(SkelMesh->GetName() + TEXT("_blendShapes"));
 			FbxBlendShape* BlendShape = FbxBlendShape::Create(Mesh, TCHAR_TO_UTF8(*BlendShapeName));
 			bool bHasBadMorphTarget = false;
 
-			for (UMorphTarget* MorphTarget : SkelMesh->MorphTargets)
+			for (UMorphTarget* MorphTarget : SkelMesh->GetMorphTargets())
 			{
 				int32 DeformerIndex = Mesh->AddDeformer(BlendShape);
 				FbxBlendShapeChannel* BlendShapeChannel = FbxBlendShapeChannel::Create(BlendShape, TCHAR_TO_UTF8(*MorphTarget->GetName()));
@@ -342,7 +342,8 @@ FbxNode* FFbxExporter::CreateMesh(const USkeletalMesh* SkelMesh, const TCHAR* Me
 
 
 	// Add the materials for the mesh
-	int32 MaterialCount = SkelMesh->Materials.Num();
+	const TArray<FSkeletalMaterial>& SkelMeshMaterials = SkelMesh->GetMaterials();
+	int32 MaterialCount = SkelMeshMaterials.Num();
 
 	for(int32 MaterialIndex = 0; MaterialIndex < MaterialCount; ++MaterialIndex)
 	{
@@ -354,7 +355,7 @@ FbxNode* FFbxExporter::CreateMesh(const USkeletalMesh* SkelMesh, const TCHAR* Me
 		}
 		else
 		{
-			MatInterface = SkelMesh->Materials[MaterialIndex].MaterialInterface;
+			MatInterface = SkelMeshMaterials[MaterialIndex].MaterialInterface;
 		}
 
 		FbxSurfaceMaterial* FbxMaterial = NULL;
@@ -751,7 +752,7 @@ FbxNode* FFbxExporter::ExportSkeletalMeshToFbx(const USkeletalMesh* SkeletalMesh
 			ActorRootNode->AddChild(SkeletonRootNode);
 		}
 
-		ExportObjectMetadataToBones(SkeletalMesh->Skeleton, BoneNodes);
+		ExportObjectMetadataToBones(SkeletalMesh->GetSkeleton(), BoneNodes);
 
 		if (MeshRootNode)
 		{
diff --git a/Engine/Source/Editor/UnrealEd/Private/Fbx/FbxSkeletalMeshImport.cpp b/Engine/Source/Editor/UnrealEd/Private/Fbx/FbxSkeletalMeshImport.cpp
--- a/Engine/Source/Editor/UnrealEd/Private/Fbx/FbxSkeletalMeshImport.cpp
+++ b/Engine/Source/Editor/UnrealEd/Private/Fbx/FbxSkeletalMeshImport.cpp
@@ -166,7 +166,7 @@ typedef TOctree2<FSoftSkinVertex, FSkeletalMeshVertPosOctreeSemantics> TSKCVertP
 void RemapSkeletalMeshVertexColorToImportData(const USkeletalMesh* SkeletalMesh, const int32 LODIndex, FSkeletalMeshImportData* SkelMeshImportData)
 {
 	//Make sure we have all the source data we need to do the remap
-	if (!SkeletalMesh->GetImportedModel() || !SkeletalMesh->GetImportedModel()->LODModels.IsValidIndex(LODIndex) || !SkeletalMesh->bHasVertexColors)
+	if (!SkeletalMesh->GetImportedModel() || !SkeletalMesh->GetImportedModel()->LODModels.IsValidIndex(LODIndex) || !SkeletalMesh->GetHasVertexColors())
 	{
 		return;
 	}
@@ -1516,7 +1516,7 @@ void UnFbx::FFbxImporter::FillLastImportMaterialNames(TArray<FName> &LastImporte
 	if (OrderedMaterialNames == nullptr && BaseSkelMesh)
 	{
 		int32 NoneNameCount = 0;
-		for (const FSkeletalMaterial &Material : BaseSkelMesh->Materials)
+		for (const FSkeletalMaterial &Material : BaseSkelMesh->GetMaterials())
 		{
 			if (Material.ImportedMaterialSlotName == NAME_None)
 			{
@@ -1724,7 +1724,7 @@ USkeletalMesh* UnFbx::FFbxImporter::ImportSkeletalMesh(FImportSkeletalMeshArgs &
 		}
 
 		//Reuse the vertex color in case we are not re-importing weights only and there is some vertexcolor in the existing skeleMesh
-		if (ExistingSkelMesh->bHasVertexColors && ImportOptions->VertexColorImportOption == EVertexColorImportOption::Ignore)
+		if (ExistingSkelMesh->GetHasVertexColors() && ImportOptions->VertexColorImportOption == EVertexColorImportOption::Ignore)
 		{
 			RemapSkeletalMeshVertexColorToImportData(ExistingSkelMesh, SafeLODIndex, SkelMeshImportDataPtr);
 		}
@@ -1802,12 +1802,12 @@ USkeletalMesh* UnFbx::FFbxImporter::ImportSkeletalMesh(FImportSkeletalMeshArgs &
 	FSkeletalMeshLODModel& LODModel = ImportedResource->LODModels[ImportLODModelIndex];
 
 	// process materials from import data
-	SkeletalMeshHelper::ProcessImportMeshMaterials(SkeletalMesh->Materials, *SkelMeshImportDataPtr);
+	SkeletalMeshHelper::ProcessImportMeshMaterials(SkeletalMesh->GetMaterials(), *SkelMeshImportDataPtr);
 
 	// process reference skeleton from import data
 	int32 SkeletalDepth = 0;
 	USkeleton* ExistingSkeleton = ExistSkelMeshDataPtr ? ExistSkelMeshDataPtr->ExistingSkeleton : ImportOptions->SkeletonForAnimation;
-	if (!SkeletalMeshHelper::ProcessImportMeshSkeleton(ExistingSkeleton, SkeletalMesh->RefSkeleton, SkeletalDepth, *SkelMeshImportDataPtr))
+	if (!SkeletalMeshHelper::ProcessImportMeshSkeleton(ExistingSkeleton, SkeletalMesh->GetRefSkeleton(), SkeletalDepth, *SkelMeshImportDataPtr))
 	{
 		EARLY_RETURN_ON_CANCEL(true, FailureCleanup);
 	}
@@ -1815,7 +1815,7 @@ USkeletalMesh* UnFbx::FFbxImporter::ImportSkeletalMesh(FImportSkeletalMeshArgs &
 	EARLY_RETURN_ON_CANCEL(false, CancelCleanup);
 
 	if (!GIsAutomationTesting)
-		UE_LOG(LogFbx, Warning, TEXT("Bones digested - %i  Depth of hierarchy - %i"), SkeletalMesh->RefSkeleton.GetNum(), SkeletalDepth);
+		UE_LOG(LogFbx, Warning, TEXT("Bones digested - %i  Depth of hierarchy - %i"), SkeletalMesh->GetRefSkeleton().GetNum(), SkeletalDepth);
 
 	// process bone influences from import data
 	SkeletalMeshHelper::ProcessImportMeshInfluences(*SkelMeshImportDataPtr, SkeletalMesh->GetPathName());
@@ -1845,8 +1845,8 @@ USkeletalMesh* UnFbx::FFbxImporter::ImportSkeletalMesh(FImportSkeletalMeshArgs &
 	SkeletalMesh->SetImportedBounds(FBoxSphereBounds(BoundingBox));
 
 	// Store whether or not this mesh has vertex colors
-	SkeletalMesh->bHasVertexColors = SkelMeshImportDataPtr->bHasVertexColors;
-	SkeletalMesh->VertexColorGuid = SkeletalMesh->bHasVertexColors ? FGuid::NewGuid() : FGuid();
+	SkeletalMesh->SetHasVertexColors(SkelMeshImportDataPtr->bHasVertexColors);
+	SkeletalMesh->SetVertexColorGuid(SkeletalMesh->GetHasVertexColors() ? FGuid::NewGuid() : FGuid());
 
 	// Pass the number of texture coordinate sets to the LODModel.  Ensure there is at least one UV coord
 	LODModel.NumTexCoords = FMath::Max<uint32>(1, SkelMeshImportDataPtr->NumTexCoords);
@@ -1872,7 +1872,7 @@ USkeletalMesh* UnFbx::FFbxImporter::ImportSkeletalMesh(FImportSkeletalMeshArgs &
 			TArray<FText> WarningMessages;
 			TArray<FName> WarningNames;
 			// Create actual rendering data.
-			bBuildSuccess = MeshUtilities.BuildSkeletalMesh(ImportedResource->LODModels[ImportLODModelIndex], SkeletalMesh->GetPathName(), SkeletalMesh->RefSkeleton, LODInfluences, LODWedges, LODFaces, LODPoints, LODPointToRawMap, LegacyBuildOptions, &WarningMessages, &WarningNames);
+			bBuildSuccess = MeshUtilities.BuildSkeletalMesh(ImportedResource->LODModels[ImportLODModelIndex], SkeletalMesh->GetPathName(), SkeletalMesh->GetRefSkeleton(), LODInfluences, LODWedges, LODFaces, LODPoints, LODPointToRawMap, LegacyBuildOptions, &WarningMessages, &WarningNames);
 
 			// temporary hack of message/names, should be one token or a struct
 			if (WarningMessages.Num() > 0 && WarningNames.Num() == WarningMessages.Num())
@@ -1909,13 +1909,14 @@ USkeletalMesh* UnFbx::FFbxImporter::ImportSkeletalMesh(FImportSkeletalMeshArgs &
 		//Get the last fbx file data need for reimport
 		if (ImportSkeletalMeshArgs.ImportMaterialOriginalNameData)
 		{
-			for (FSkeletalMaterial SkeletalMaterial : SkeletalMesh->Materials)
+			for (FSkeletalMaterial SkeletalMaterial : SkeletalMesh->GetMaterials())
 			{
 				ImportSkeletalMeshArgs.ImportMaterialOriginalNameData->Add(SkeletalMaterial.ImportedMaterialSlotName);
 			}
 		}
 		if (ImportSkeletalMeshArgs.ImportMeshSectionsData)
 		{
+			const TArray<FSkeletalMaterial>& Materials = SkeletalMesh->GetMaterials();
 			for (int32 SectionIndex = 0; SectionIndex < ImportedResource->LODModels[ImportLODModelIndex].Sections.Num(); ++SectionIndex)
 			{
 				const FSkelMeshSection &SkelMeshSection = ImportedResource->LODModels[ImportLODModelIndex].Sections[SectionIndex];
@@ -1924,7 +1925,11 @@ USkeletalMesh* UnFbx::FFbxImporter::ImportSkeletalMesh(FImportSkeletalMeshArgs &
 				{
 					MaterialIndex = SkeletalMesh->GetLODInfo(0)->LODMaterialMap[SectionIndex];
 				}
-				ImportSkeletalMeshArgs.ImportMeshSectionsData->SectionOriginalMaterialName.Add(SkeletalMesh->Materials[MaterialIndex].ImportedMaterialSlotName);
+
+				if (ensure(Materials.IsValidIndex(MaterialIndex)))
+				{
+					ImportSkeletalMeshArgs.ImportMeshSectionsData->SectionOriginalMaterialName.Add(Materials[MaterialIndex].ImportedMaterialSlotName);
+				}
 			}
 		}
 
@@ -1934,19 +1939,19 @@ USkeletalMesh* UnFbx::FFbxImporter::ImportSkeletalMesh(FImportSkeletalMeshArgs &
 		if (ExistingSkelMesh)
 		{
 			//In case of a re-import we update only the type we re-import
-			SkeletalMesh->AssetImportData->AddFileName(UFactory::GetCurrentFilename(), ImportOptions->bImportAsSkeletalGeometry ? 1 : ImportOptions->bImportAsSkeletalSkinning ? 2 : 0);
+			SkeletalMesh->GetAssetImportData()->AddFileName(UFactory::GetCurrentFilename(), ImportOptions->bImportAsSkeletalGeometry ? 1 : ImportOptions->bImportAsSkeletalSkinning ? 2 : 0);
 		}
 		else
 		{
 			//When we create a new skeletal mesh asset we create 1 or 3 source files. 1 in case we import all the content, 3 in case we import geo or skinning
 			int32 SourceFileIndex = ImportOptions->bImportAsSkeletalGeometry ? 1 : ImportOptions->bImportAsSkeletalSkinning ? 2 : 0;
 			//Always add the base sourcefile
-			SkeletalMesh->AssetImportData->AddFileName(UFactory::GetCurrentFilename(), 0, NSSkeletalMeshSourceFileLabels::GeoAndSkinningText().ToString());
+			SkeletalMesh->GetAssetImportData()->AddFileName(UFactory::GetCurrentFilename(), 0, NSSkeletalMeshSourceFileLabels::GeoAndSkinningText().ToString());
 			if (SourceFileIndex != 0)
 			{
 				//If the user import geo or skinning, we add both entries to allow reimport of them
-				SkeletalMesh->AssetImportData->AddFileName(UFactory::GetCurrentFilename(), 1, NSSkeletalMeshSourceFileLabels::GeometryText().ToString());
-				SkeletalMesh->AssetImportData->AddFileName(UFactory::GetCurrentFilename(), 2, NSSkeletalMeshSourceFileLabels::SkinningText().ToString());
+				SkeletalMesh->GetAssetImportData()->AddFileName(UFactory::GetCurrentFilename(), 1, NSSkeletalMeshSourceFileLabels::GeometryText().ToString());
+				SkeletalMesh->GetAssetImportData()->AddFileName(UFactory::GetCurrentFilename(), 2, NSSkeletalMeshSourceFileLabels::SkinningText().ToString());
 			}
 		}
 
@@ -2068,16 +2073,16 @@ USkeletalMesh* UnFbx::FFbxImporter::ImportSkeletalMesh(FImportSkeletalMeshArgs &
 
 		EARLY_RETURN_ON_CANCEL(false, CancelCleanup);
 
-		if (SkeletalMesh->Skeleton != Skeleton)
+		if (SkeletalMesh->GetSkeleton() != Skeleton)
 		{
-			SkeletalMesh->Skeleton = Skeleton;
+			SkeletalMesh->SetSkeleton(Skeleton);
 			SkeletalMesh->MarkPackageDirty();
 		}
 		// Create PhysicsAsset if requested and if physics asset is null
 		// We create the physic asset after we create the skeleton since we need the skeleton to correctly build it
 		if (ImportOptions->bCreatePhysicsAsset)
 		{
-			if (SkeletalMesh->PhysicsAsset == NULL)
+			if (SkeletalMesh->GetPhysicsAsset() == NULL)
 			{
 				FString ObjectName = FString::Printf(TEXT("%s_PhysicsAsset"), *SkeletalMesh->GetName());
 				UPhysicsAsset * NewPhysicsAsset = CreateAsset<UPhysicsAsset>(ImportSkeletalMeshArgs.InParent->GetName(), ObjectName, true);
@@ -2105,7 +2110,7 @@ USkeletalMesh* UnFbx::FFbxImporter::ImportSkeletalMesh(FImportSkeletalMeshArgs &
 		// if physics asset is selected
 		else if (ImportOptions->PhysicsAsset)
 		{
-			SkeletalMesh->PhysicsAsset = ImportOptions->PhysicsAsset;
+			SkeletalMesh->SetPhysicsAsset(ImportOptions->PhysicsAsset);
 		}
 	}
 
@@ -2118,7 +2123,7 @@ USkeletalMesh* UnFbx::FFbxImporter::ImportSkeletalMesh(FImportSkeletalMeshArgs &
 	// Import bone metadata to Skeleton
 	for (FbxNode* SkeletonNode : ImportedSkeletonLinkNodes)
 	{
-		ImportNodeCustomProperties(SkeletalMesh->Skeleton, SkeletonNode, true);
+		ImportNodeCustomProperties(SkeletalMesh->GetSkeleton(), SkeletonNode, true);
 	}
 
 	return SkeletalMesh;
@@ -2132,7 +2137,7 @@ void UnFbx::FFbxImporter::UpdateSkeletalMeshImportData(USkeletalMesh *SkeletalMe
 		return;
 	}
 
-	UFbxSkeletalMeshImportData* ImportData = Cast<UFbxSkeletalMeshImportData>(SkeletalMesh->AssetImportData);
+	UFbxSkeletalMeshImportData* ImportData = Cast<UFbxSkeletalMeshImportData>(SkeletalMesh->GetAssetImportData());
 	if (!ImportData && SkeletalMeshImportData)
 	{
 		ImportData = UFbxSkeletalMeshImportData::GetImportDataForSkeletalMesh(SkeletalMesh, SkeletalMeshImportData);
@@ -2203,7 +2208,7 @@ void UnFbx::FFbxImporter::UpdateSkeletalMeshImportData(USkeletalMesh *SkeletalMe
 			//In this case we update from the skeletal mesh import
 			ImportData->ImportMaterialOriginalNameData.Empty();
 			ImportData->ImportMeshLodData.Empty();
-			for (const FSkeletalMaterial &Material : SkeletalMesh->Materials)
+			for (const FSkeletalMaterial &Material : SkeletalMesh->GetMaterials())
 			{
 				ImportData->ImportMaterialOriginalNameData.Add(Material.ImportedMaterialSlotName);
 			}
@@ -2282,14 +2287,14 @@ UObject* UnFbx::FFbxImporter::CreateAssetOfClass(UClass* AssetClass, FString Par
 
 void UnFbx::FFbxImporter::SetupAnimationDataFromMesh(USkeletalMesh* SkeletalMesh, UObject* InParent, TArray<FbxNode*>& NodeArray, UFbxAnimSequenceImportData* TemplateImportData, const FString& Name)
 {
-	USkeleton* Skeleton = SkeletalMesh->Skeleton;
+	USkeleton* Skeleton = SkeletalMesh->GetSkeleton();
 
 	if (Scene->GetSrcObjectCount<FbxAnimStack>() > 0)
 	{
 		if ( ensure ( Skeleton != NULL ) )
 		{
 			TArray<FbxNode*> FBXMeshNodeArray;
-			FbxNode* SkeletonRoot = FindFBXMeshesByBone(SkeletalMesh->Skeleton->GetReferenceSkeleton().GetBoneName(0), true, FBXMeshNodeArray);
+			FbxNode* SkeletonRoot = FindFBXMeshesByBone(SkeletalMesh->GetSkeleton()->GetReferenceSkeleton().GetBoneName(0), true, FBXMeshNodeArray);
 
 			if (SkeletonRoot != nullptr)
 			{
@@ -2346,7 +2351,7 @@ USkeletalMesh* UnFbx::FFbxImporter::ReimportSkeletalMesh(USkeletalMesh* Mesh, UF
 	// support to update rigid animation mesh
 	ImportOptions->bImportRigidMesh = true;
 
-	UFbxSkeletalMeshImportData* SKImportData = Cast<UFbxSkeletalMeshImportData>(Mesh->AssetImportData);
+	UFbxSkeletalMeshImportData* SKImportData = Cast<UFbxSkeletalMeshImportData>(Mesh->GetAssetImportData());
 	if (SKImportData)
 	{
 		FSkeletalMeshLODInfo* LODInfo = Mesh->GetLODInfo(0);
@@ -2476,7 +2481,7 @@ USkeletalMesh* UnFbx::FFbxImporter::ReimportSkeletalMesh(USkeletalMesh* Mesh, UF
 		}
 		//Empty the morph target before re-importing, it will prevent to have old data that can point on random vertex
 		//Do not empty the morph if we import weights only
-		if (Mesh->MorphTargets.Num() > 0 && !ImportOptions->bImportAsSkeletalSkinning)
+		if (Mesh->GetMorphTargets().Num() > 0 && !ImportOptions->bImportAsSkeletalSkinning)
 		{
 			Mesh->UnregisterAllMorphTarget();
 		}
@@ -3750,7 +3755,7 @@ void UnFbx::FFbxImporter::InsertNewLODToBaseSkeletalMesh(USkeletalMesh* InSkelet
 
 				//Get the current skelmesh section slot import name
 				int32 ExistRemapMaterialIndex = ExistingSkelMeshLodInfo.LODMaterialMap[SectionIndex];
-				if (!BaseSkeletalMesh->Materials.IsValidIndex(ExistRemapMaterialIndex))
+				if (!BaseSkeletalMesh->GetMaterials().IsValidIndex(ExistRemapMaterialIndex))
 				{
 					//If we have an INDEX_NONE it mean we have to use the section material index. If we have a bad value we will use the section material index if possible
 					if (DestImportedResource->LODModels.IsValidIndex(DesiredLOD) && DestImportedResource->LODModels[DesiredLOD].Sections.IsValidIndex(SectionIndex))
@@ -3764,9 +3769,9 @@ void UnFbx::FFbxImporter::InsertNewLODToBaseSkeletalMesh(USkeletalMesh* InSkelet
 						ExistRemapMaterialIndex = 0;
 					}
 				}
-				check(BaseSkeletalMesh->Materials.IsValidIndex(ExistRemapMaterialIndex));
+				check(BaseSkeletalMesh->GetMaterials().IsValidIndex(ExistRemapMaterialIndex));
 
-				ExistingMeshSectionSlotNames.Add(BaseSkeletalMesh->Materials[ExistRemapMaterialIndex].ImportedMaterialSlotName);
+				ExistingMeshSectionSlotNames.Add(BaseSkeletalMesh->GetMaterials()[ExistRemapMaterialIndex].ImportedMaterialSlotName);
 
 				//Get the Last imported skelmesh section slot import name
 				OriginalImportMeshSectionSlotNames.Add(OriginalImportMeshLodSectionsData.SectionOriginalMaterialName[SectionIndex]);
@@ -3781,7 +3786,7 @@ void UnFbx::FFbxImporter::InsertNewLODToBaseSkeletalMesh(USkeletalMesh* InSkelet
 					int32 MatIdx = NewSection.MaterialIndex;
 					for (int32 ExistSectionIndex = 0; ExistSectionIndex < ExistLODModel.Sections.Num(); ExistSectionIndex++)
 					{
-						if (OriginalImportMeshSectionSlotNames.IsValidIndex(ExistSectionIndex) && OriginalImportMeshSectionSlotNames[ExistSectionIndex] == InSkeletalMesh->Materials[MatIdx].ImportedMaterialSlotName)
+						if (OriginalImportMeshSectionSlotNames.IsValidIndex(ExistSectionIndex) && OriginalImportMeshSectionSlotNames[ExistSectionIndex] == InSkeletalMesh->GetMaterials()[MatIdx].ImportedMaterialSlotName)
 						{
 							const FSkelMeshSection& ExistSection = ExistLODModel.Sections[ExistSectionIndex];
 							//We found a match, restore the data
@@ -3831,6 +3836,7 @@ void UnFbx::FFbxImporter::InsertNewLODToBaseSkeletalMesh(USkeletalMesh* InSkelet
 		LODInfo.BuildSettings = ImportedLODInfo.BuildSettings;
 	}
 	
+	TArray<FSkeletalMaterial>& BaseMaterials = BaseSkeletalMesh->GetMaterials();
 	LODInfo.LODMaterialMap.Empty();
 	// Now set up the material mapping array.
 	for (int32 SectionIndex = 0; SectionIndex < NewLODModel.Sections.Num(); SectionIndex++)
@@ -3839,10 +3845,10 @@ void UnFbx::FFbxImporter::InsertNewLODToBaseSkeletalMesh(USkeletalMesh* InSkelet
 		// Try and find the auto-assigned material in the array.
 		int32 LODMatIndex = INDEX_NONE;
 		//First try to match by name
-		for (int32 BaseMaterialIndex = 0; BaseMaterialIndex < BaseSkeletalMesh->Materials.Num(); ++BaseMaterialIndex)
+		for (int32 BaseMaterialIndex = 0; BaseMaterialIndex < BaseMaterials.Num(); ++BaseMaterialIndex)
 		{
-			const FSkeletalMaterial& SkeletalMaterial = BaseSkeletalMesh->Materials[BaseMaterialIndex];
-			if (SkeletalMaterial.ImportedMaterialSlotName != NAME_None && SkeletalMaterial.ImportedMaterialSlotName == InSkeletalMesh->Materials[MatIdx].ImportedMaterialSlotName)
+			const FSkeletalMaterial& SkeletalMaterial = BaseMaterials[BaseMaterialIndex];
+			if (SkeletalMaterial.ImportedMaterialSlotName != NAME_None && SkeletalMaterial.ImportedMaterialSlotName == InSkeletalMesh->GetMaterials()[MatIdx].ImportedMaterialSlotName)
 			{
 				LODMatIndex = BaseMaterialIndex;
 				break;
@@ -3852,7 +3858,7 @@ void UnFbx::FFbxImporter::InsertNewLODToBaseSkeletalMesh(USkeletalMesh* InSkelet
 		// If we dont have a match, add a new entry to the material list.
 		if (LODMatIndex == INDEX_NONE)
 		{
-			LODMatIndex = BaseSkeletalMesh->Materials.Add(InSkeletalMesh->Materials[MatIdx]);
+			LODMatIndex = BaseMaterials.Add(InSkeletalMesh->GetMaterials()[MatIdx]);
 		}
 
 		LODInfo.LODMaterialMap.Add(LODMatIndex);
@@ -3869,7 +3875,7 @@ void UnFbx::FFbxImporter::InsertNewLODToBaseSkeletalMesh(USkeletalMesh* InSkelet
 
 	// If this LOD had been generated previously by automatic mesh reduction, clear that flag.
 	LODInfo.bHasBeenSimplified = false;
-	if (BaseSkeletalMesh->LODSettings == nullptr || !BaseSkeletalMesh->LODSettings->HasValidSettings() || BaseSkeletalMesh->LODSettings->GetNumberOfSettings() <= DesiredLOD)
+	if (BaseSkeletalMesh->GetLODSettings() == nullptr || !BaseSkeletalMesh->GetLODSettings()->HasValidSettings() || BaseSkeletalMesh->GetLODSettings()->GetNumberOfSettings() <= DesiredLOD)
 	{
 		//Make sure any custom LOD have correct settings (no reduce)
 		LODInfo.ReductionSettings.NumOfTrianglesPercentage = 1.0f;
@@ -3892,17 +3898,17 @@ void UnFbx::FFbxImporter::InsertNewLODToBaseSkeletalMesh(USkeletalMesh* InSkelet
 			FName ExistMeshSectionSlotName = ExistingMeshSectionSlotNames[SectionIndex];
 			//Get the new skelmesh section slot import name
 			int32 NewRemapMaterialIndex = NewSkelMeshLodInfo.LODMaterialMap[SectionIndex];
-			if (!BaseSkeletalMesh->Materials.IsValidIndex(NewRemapMaterialIndex))
+			if (!BaseMaterials.IsValidIndex(NewRemapMaterialIndex))
 			{
 				int32 RematchMaterialIndex = NewSkelMeshLodInfo.LODMaterialMap[SectionIndex];
 				if (RematchMaterialIndex == INDEX_NONE)
 				{
 					RematchMaterialIndex = NewLODModel.Sections[SectionIndex].MaterialIndex;
 				}
-				NewSkelMeshLodInfo.LODMaterialMap[SectionIndex] = FMath::Clamp<int32>(RematchMaterialIndex, 0, BaseSkeletalMesh->Materials.Num()-1);
+				NewSkelMeshLodInfo.LODMaterialMap[SectionIndex] = FMath::Clamp<int32>(RematchMaterialIndex, 0, BaseMaterials.Num()-1);
 				continue;
 			}
-			FName NewMeshSectionSlotName = BaseSkeletalMesh->Materials[NewRemapMaterialIndex].ImportedMaterialSlotName;
+			FName NewMeshSectionSlotName = BaseMaterials[NewRemapMaterialIndex].ImportedMaterialSlotName;
 			//Get the Last imported skelmesh section slot import name
 			FName OriginalImportMeshSectionSlotName = OriginalImportMeshSectionSlotNames[SectionIndex];
 
@@ -3911,9 +3917,9 @@ void UnFbx::FFbxImporter::InsertNewLODToBaseSkeletalMesh(USkeletalMesh* InSkelet
 				//The last import slot name match the New import slot name, but the Exist slot name is different then the last import slot name.
 				//This mean the user has change the section assign slot and the fbx file did not change it
 				//Override the new section material index to use the one that the user set
-				for (int32 RemapMaterialIndex = 0; RemapMaterialIndex < BaseSkeletalMesh->Materials.Num(); ++RemapMaterialIndex)
+				for (int32 RemapMaterialIndex = 0; RemapMaterialIndex < BaseMaterials.Num(); ++RemapMaterialIndex)
 				{
-					const FSkeletalMaterial &NewSectionMaterial = BaseSkeletalMesh->Materials[RemapMaterialIndex];
+					const FSkeletalMaterial &NewSectionMaterial = BaseMaterials[RemapMaterialIndex];
 					if (NewSectionMaterial.ImportedMaterialSlotName == ExistMeshSectionSlotName)
 					{
 						NewSkelMeshLodInfo.LODMaterialMap[SectionIndex] = RemapMaterialIndex;
@@ -3938,20 +3944,20 @@ bool UnFbx::FFbxImporter::ImportSkeletalMeshLOD(USkeletalMesh* InSkeletalMesh, U
 
 	// Names of root bones must match.
 	// If the names of root bones don't match, the LOD Mesh does not share skeleton with base Mesh. 
-	if (InSkeletalMesh->RefSkeleton.GetBoneName(0) != BaseSkeletalMesh->RefSkeleton.GetBoneName(0) && !ImportOptions->bImportAsSkeletalGeometry)
+	if (InSkeletalMesh->GetRefSkeleton().GetBoneName(0) != BaseSkeletalMesh->GetRefSkeleton().GetBoneName(0) && !ImportOptions->bImportAsSkeletalGeometry)
 	{
 		AddTokenizedErrorMessage(FTokenizedMessage::Create(EMessageSeverity::Error, FText::Format(LOCTEXT("LODRootNameIncorrect", "Root bone in LOD is '{0}' instead of '{1}'.\nImport failed."),
-			FText::FromName(InSkeletalMesh->RefSkeleton.GetBoneName(0)), FText::FromName(BaseSkeletalMesh->RefSkeleton.GetBoneName(0)))), FFbxErrors::SkeletalMesh_LOD_RootNameIncorrect);
+			FText::FromName(InSkeletalMesh->GetRefSkeleton().GetBoneName(0)), FText::FromName(BaseSkeletalMesh->GetRefSkeleton().GetBoneName(0)))), FFbxErrors::SkeletalMesh_LOD_RootNameIncorrect);
 
 		return false;
 	}
 	bool bApplyBaseSkinning = false;
 	// We do some checking here that for every bone in the mesh we just imported, it's in our base ref skeleton, and the parent is the same.
-	for (int32 i = 0; i < InSkeletalMesh->RefSkeleton.GetRawBoneNum(); i++)
+	for (int32 i = 0; i < InSkeletalMesh->GetRefSkeleton().GetRawBoneNum(); i++)
 	{
 		int32 LODBoneIndex = i;
-		FName LODBoneName = InSkeletalMesh->RefSkeleton.GetBoneName(LODBoneIndex);
-		int32 BaseBoneIndex = BaseSkeletalMesh->RefSkeleton.FindBoneIndex(LODBoneName);
+		FName LODBoneName = InSkeletalMesh->GetRefSkeleton().GetBoneName(LODBoneIndex);
+		int32 BaseBoneIndex = BaseSkeletalMesh->GetRefSkeleton().FindBoneIndex(LODBoneName);
 		if (BaseBoneIndex == INDEX_NONE)
 		{
 			if (!ImportOptions->bImportAsSkeletalGeometry)
@@ -3968,11 +3974,11 @@ bool UnFbx::FFbxImporter::ImportSkeletalMeshLOD(USkeletalMesh* InSkeletalMesh, U
 
 		if (i > 0)
 		{
-			int32 LODParentIndex = InSkeletalMesh->RefSkeleton.GetParentIndex(LODBoneIndex);
-			FName LODParentName = InSkeletalMesh->RefSkeleton.GetBoneName(LODParentIndex);
+			int32 LODParentIndex = InSkeletalMesh->GetRefSkeleton().GetParentIndex(LODBoneIndex);
+			FName LODParentName = InSkeletalMesh->GetRefSkeleton().GetBoneName(LODParentIndex);
 
-			int32 BaseParentIndex = BaseSkeletalMesh->RefSkeleton.GetParentIndex(BaseBoneIndex);
-			FName BaseParentName = BaseSkeletalMesh->RefSkeleton.GetBoneName(BaseParentIndex);
+			int32 BaseParentIndex = BaseSkeletalMesh->GetRefSkeleton().GetParentIndex(BaseBoneIndex);
+			FName BaseParentName = BaseSkeletalMesh->GetRefSkeleton().GetBoneName(BaseParentIndex);
 
 			if (LODParentName != BaseParentName)
 			{
@@ -4031,7 +4037,7 @@ bool UnFbx::FFbxImporter::ImportSkeletalMeshLOD(USkeletalMesh* InSkeletalMesh, U
 		{
 			// Find bone index the socket is attached to.
 			USkeletalMeshSocket* Socket = Sockets[i];
-			int32 SocketBoneIndex = InSkeletalMesh->RefSkeleton.FindBoneIndex(Socket->BoneName);
+			int32 SocketBoneIndex = InSkeletalMesh->GetRefSkeleton().FindBoneIndex(Socket->BoneName);
 
 			// If this LOD does not contain the socket bone, abort import.
 			if (SocketBoneIndex == INDEX_NONE)
@@ -4059,13 +4065,13 @@ bool UnFbx::FFbxImporter::ImportSkeletalMeshLOD(USkeletalMesh* InSkeletalMesh, U
 			ImportDataBoneRemap.AddZeroed(LODImportDataBoneNumber);
 			//We generate a new RefBonesBinary array to replace the existing one
 			TArray<SkeletalMeshImportData::FBone> RemapedRefBonesBinary;
-			RemapedRefBonesBinary.AddZeroed(BaseSkeletalMesh->RefSkeleton.GetNum());
+			RemapedRefBonesBinary.AddZeroed(BaseSkeletalMesh->GetRefSkeleton().GetNum());
 			for (int32 ImportBoneIndex = 0; ImportBoneIndex < LODImportDataBoneNumber; ++ImportBoneIndex)
 			{
 				SkeletalMeshImportData::FBone& ImportedBone = LODImportData.RefBonesBinary[ImportBoneIndex];
 				int32 LODBoneIndex = ImportBoneIndex;
 				FName LODBoneName = FName(*ImportedBone.Name);
-				int32 BaseBoneIndex = BaseSkeletalMesh->RefSkeleton.FindBoneIndex(LODBoneName);
+				int32 BaseBoneIndex = BaseSkeletalMesh->GetRefSkeleton().FindBoneIndex(LODBoneName);
 				ImportDataBoneRemap[ImportBoneIndex] = BaseBoneIndex;
 				if (BaseBoneIndex != INDEX_NONE)
 				{
@@ -4107,8 +4113,8 @@ bool UnFbx::FFbxImporter::ImportSkeletalMeshLOD(USkeletalMesh* InSkeletalMesh, U
 		for (int32 ActiveIndex = 0; ActiveIndex < NewLODModel.ActiveBoneIndices.Num(); ActiveIndex++)
 		{
 			int32 LODBoneIndex = NewLODModel.ActiveBoneIndices[ActiveIndex];
-			FName LODBoneName = InSkeletalMesh->RefSkeleton.GetBoneName(LODBoneIndex);
-			int32 BaseBoneIndex = BaseSkeletalMesh->RefSkeleton.FindBoneIndex(LODBoneName);
+			FName LODBoneName = InSkeletalMesh->GetRefSkeleton().GetBoneName(LODBoneIndex);
+			int32 BaseBoneIndex = BaseSkeletalMesh->GetRefSkeleton().FindBoneIndex(LODBoneName);
 			NewLODModel.ActiveBoneIndices[ActiveIndex] = BaseBoneIndex;
 		}
 
@@ -4119,8 +4125,8 @@ bool UnFbx::FFbxImporter::ImportSkeletalMeshLOD(USkeletalMesh* InSkeletalMesh, U
 			for (int32 BoneMapIndex = 0; BoneMapIndex < Section.BoneMap.Num(); BoneMapIndex++)
 			{
 				int32 LODBoneIndex = Section.BoneMap[BoneMapIndex];
-				FName LODBoneName = InSkeletalMesh->RefSkeleton.GetBoneName(LODBoneIndex);
-				int32 BaseBoneIndex = BaseSkeletalMesh->RefSkeleton.FindBoneIndex(LODBoneName);
+				FName LODBoneName = InSkeletalMesh->GetRefSkeleton().GetBoneName(LODBoneIndex);
+				int32 BaseBoneIndex = BaseSkeletalMesh->GetRefSkeleton().FindBoneIndex(LODBoneName);
 				Section.BoneMap[BoneMapIndex] = BaseBoneIndex;
 			}
 		}
@@ -4128,8 +4134,8 @@ bool UnFbx::FFbxImporter::ImportSkeletalMeshLOD(USkeletalMesh* InSkeletalMesh, U
 		// Create the RequiredBones array in the LODModel from the ref skeleton.
 		for (int32 RequiredBoneIndex = 0; RequiredBoneIndex < NewLODModel.RequiredBones.Num(); RequiredBoneIndex++)
 		{
-			FName LODBoneName = InSkeletalMesh->RefSkeleton.GetBoneName(NewLODModel.RequiredBones[RequiredBoneIndex]);
-			int32 BaseBoneIndex = BaseSkeletalMesh->RefSkeleton.FindBoneIndex(LODBoneName);
+			FName LODBoneName = InSkeletalMesh->GetRefSkeleton().GetBoneName(NewLODModel.RequiredBones[RequiredBoneIndex]);
+			int32 BaseBoneIndex = BaseSkeletalMesh->GetRefSkeleton().FindBoneIndex(LODBoneName);
 			if (BaseBoneIndex != INDEX_NONE)
 			{
 				NewLODModel.RequiredBones[RequiredBoneIndex] = BaseBoneIndex;
@@ -4142,7 +4148,7 @@ bool UnFbx::FFbxImporter::ImportSkeletalMeshLOD(USkeletalMesh* InSkeletalMesh, U
 
 		// Also sort the RequiredBones array to be strictly increasing.
 		NewLODModel.RequiredBones.Sort();
-		BaseSkeletalMesh->RefSkeleton.EnsureParentsExistAndSort(NewLODModel.ActiveBoneIndices);
+		BaseSkeletalMesh->GetRefSkeleton().EnsureParentsExistAndSort(NewLODModel.ActiveBoneIndices);
 	}
 	// To be extra-nice, we apply the difference between the root transform of the meshes to the verts.
 	FMatrix LODToBaseTransform = InSkeletalMesh->GetRefPoseMatrix(0).InverseFast() * BaseSkeletalMesh->GetRefPoseMatrix(0);
@@ -4172,20 +4178,21 @@ bool UnFbx::FFbxImporter::ImportSkeletalMeshLOD(USkeletalMesh* InSkeletalMesh, U
 		
 		auto GetImportMaterialSlotName = [](const USkeletalMesh* SkelMesh, const FSkelMeshSection& Section, int32 SectionIndex, const FSkeletalMeshLODInfo& Info, int32& OutMaterialIndex)->FName
 		{
-			check(SkelMesh->Materials.Num() > 0);
+			const TArray<FSkeletalMaterial>& MeshMaterials = SkelMesh->GetMaterials();
+			check(MeshMaterials.Num() > 0);
 			OutMaterialIndex = Section.MaterialIndex;
 			if (Info.LODMaterialMap.IsValidIndex(SectionIndex) && Info.LODMaterialMap[SectionIndex] != INDEX_NONE)
 			{
 				OutMaterialIndex = Info.LODMaterialMap[SectionIndex];
 			}
 			FName ImportedMaterialSlotName = NAME_None;
-			if (SkelMesh->Materials.IsValidIndex(OutMaterialIndex))
+			if (MeshMaterials.IsValidIndex(OutMaterialIndex))
 			{
-				ImportedMaterialSlotName = SkelMesh->Materials[OutMaterialIndex].ImportedMaterialSlotName;
+				ImportedMaterialSlotName = MeshMaterials[OutMaterialIndex].ImportedMaterialSlotName;
 			}
 			else
 			{
-				ImportedMaterialSlotName = SkelMesh->Materials[0].ImportedMaterialSlotName;
+				ImportedMaterialSlotName = MeshMaterials[0].ImportedMaterialSlotName;
 				OutMaterialIndex = 0;
 			}
 			return ImportedMaterialSlotName;
@@ -4213,7 +4220,7 @@ bool UnFbx::FFbxImporter::ImportSkeletalMeshLOD(USkeletalMesh* InSkeletalMesh, U
 						break;
 					}
 				}
-				else if(InSkeletalMesh->Materials[ImportedMaterialIndex] == BaseSkeletalMesh->Materials[ExistingMaterialIndex]) //Use material slot compare to match in case the name is none
+				else if(InSkeletalMesh->GetMaterials()[ImportedMaterialIndex] == BaseSkeletalMesh->GetMaterials()[ExistingMaterialIndex]) //Use material slot compare to match in case the name is none
 				{
 					//Set the value and exit
 					ImportedSection.bCastShadow = ExistingSection.bCastShadow;
diff --git a/Engine/Source/Editor/UnrealEd/Private/Fbx/FbxSkeletalMeshImportData.cpp b/Engine/Source/Editor/UnrealEd/Private/Fbx/FbxSkeletalMeshImportData.cpp
--- a/Engine/Source/Editor/UnrealEd/Private/Fbx/FbxSkeletalMeshImportData.cpp
+++ b/Engine/Source/Editor/UnrealEd/Private/Fbx/FbxSkeletalMeshImportData.cpp
@@ -32,18 +32,18 @@ UFbxSkeletalMeshImportData* UFbxSkeletalMeshImportData::GetImportDataForSkeletal
 {
 	check(SkeletalMesh);
 	
-	UFbxSkeletalMeshImportData* ImportData = Cast<UFbxSkeletalMeshImportData>(SkeletalMesh->AssetImportData);
+	UFbxSkeletalMeshImportData* ImportData = Cast<UFbxSkeletalMeshImportData>(SkeletalMesh->GetAssetImportData());
 	if ( !ImportData )
 	{
 		ImportData = NewObject<UFbxSkeletalMeshImportData>(SkeletalMesh, NAME_None, RF_NoFlags, TemplateForCreation);
 
 		// Try to preserve the source file data if possible
-		if ( SkeletalMesh->AssetImportData != NULL )
+		if ( SkeletalMesh->GetAssetImportData() != NULL )
 		{
-			ImportData->SourceData = SkeletalMesh->AssetImportData->SourceData;
+			ImportData->SourceData = SkeletalMesh->GetAssetImportData()->SourceData;
 		}
 
-		SkeletalMesh->AssetImportData = ImportData;
+		SkeletalMesh->SetAssetImportData(ImportData);
 	}
 
 	return ImportData;
diff --git a/Engine/Source/Editor/UnrealEd/Private/Fbx/ReimportFbxSceneFactory.cpp b/Engine/Source/Editor/UnrealEd/Private/Fbx/ReimportFbxSceneFactory.cpp
--- a/Engine/Source/Editor/UnrealEd/Private/Fbx/ReimportFbxSceneFactory.cpp
+++ b/Engine/Source/Editor/UnrealEd/Private/Fbx/ReimportFbxSceneFactory.cpp
@@ -87,9 +87,9 @@ UFbxSceneImportData *GetFbxSceneImportData(UObject *Obj)
 		{
 			//Reimport from one of the static mesh
 			USkeletalMesh* SkeletalMesh = Cast<USkeletalMesh>(Obj);
-			if (SkeletalMesh != nullptr && SkeletalMesh->AssetImportData != nullptr)
+			if (SkeletalMesh != nullptr && SkeletalMesh->GetAssetImportData() != nullptr)
 			{
-				ImportData = Cast<UFbxAssetImportData>(SkeletalMesh->AssetImportData);
+				ImportData = Cast<UFbxAssetImportData>(SkeletalMesh->GetAssetImportData());
 			}
 		}
 		else if (Obj->IsA(UAnimSequence::StaticClass()))
@@ -1285,7 +1285,7 @@ EReimportResult::Type UReimportFbxSceneFactory::ReimportSkeletalMesh(void* VoidF
 	EReimportResult::Type ReimportResult = EReimportResult::Succeeded;
 	if (FbxImporter->ReimportSkeletalMesh(Mesh, SkeletalMeshImportData, MeshInfo->UniqueId, &OutSkeletalMeshArray))
 	{
-		Mesh->AssetImportData->Update(FbxImportFileName);
+		Mesh->GetAssetImportData()->Update(FbxImportFileName);
 
 		// Copy user data to newly created mesh
 		for (auto Kvp : UserDataCopy)
@@ -1316,10 +1316,10 @@ EReimportResult::Type UReimportFbxSceneFactory::ReimportSkeletalMesh(void* VoidF
 		//1. Store all anim sequence reference that was originally import, for every skeletal mesh
 		//2. On reimport match the existing one
 		//3. Reimport matching animation
-		if (GlobalImportSettings->bImportAnimations)
+		if (GlobalImportSettings->bImportAnimations && ensure(Mesh->GetSkeleton()))
 		{
 			TArray<FbxNode*> FBXMeshNodeArray;
-			FbxNode* SkeletonRoot = FbxImporter->FindFBXMeshesByBone(Mesh->Skeleton->GetReferenceSkeleton().GetBoneName(0), true, FBXMeshNodeArray);
+			FbxNode* SkeletonRoot = FbxImporter->FindFBXMeshesByBone(Mesh->GetSkeleton()->GetReferenceSkeleton().GetBoneName(0), true, FBXMeshNodeArray);
 
 			FString AnimName = FbxImporter->MakeNameForMesh(FBXMeshNodeArray[0]->GetName(), FBXMeshNodeArray[0]).ToString();
 			AnimName = (GlobalImportSettings->AnimationName != "") ? GlobalImportSettings->AnimationName : AnimName + TEXT("_Anim");
@@ -1439,18 +1439,18 @@ EReimportResult::Type UReimportFbxSceneFactory::ReimportSkeletalMesh(void* VoidF
 							{
 								DestSeq->CleanAnimSequenceForImport();
 							}
-							DestSeq->SetSkeleton(Mesh->Skeleton);
+							DestSeq->SetSkeleton(Mesh->GetSkeleton());
 							// since to know full path, reimport will need to do same
 							UFbxAnimSequenceImportData* ImportData = UFbxAnimSequenceImportData::GetImportDataForAnimSequence(DestSeq, AnimSequenceImportData);
 							ImportData->Update(UFactory::CurrentFilename);
-							FbxImporter->ImportAnimation(Mesh->Skeleton, DestSeq, CurrentFilename, SortedLinks, FBXMeshNodeArray, CurAnimStack, ResampleRate, AnimTimeSpan);
+							FbxImporter->ImportAnimation(Mesh->GetSkeleton(), DestSeq, CurrentFilename, SortedLinks, FBXMeshNodeArray, CurAnimStack, ResampleRate, AnimTimeSpan);
 						}
 						else
 						{
 							//Reimport in a existing sequence
 							if (FbxImporter->ValidateAnimStack(SortedLinks, FBXMeshNodeArray, CurAnimStack, ResampleRate, true, AnimTimeSpan))
 							{
-								FbxImporter->ImportAnimation(Mesh->Skeleton, DestSeq, CurrentFilename, SortedLinks, FBXMeshNodeArray, CurAnimStack, ResampleRate, AnimTimeSpan);
+								FbxImporter->ImportAnimation(Mesh->GetSkeleton(), DestSeq, CurrentFilename, SortedLinks, FBXMeshNodeArray, CurAnimStack, ResampleRate, AnimTimeSpan);
 							}
 						}
 					}
diff --git a/Engine/Source/Editor/UnrealEd/Private/FbxMeshUtils.cpp b/Engine/Source/Editor/UnrealEd/Private/FbxMeshUtils.cpp
--- a/Engine/Source/Editor/UnrealEd/Private/FbxMeshUtils.cpp
+++ b/Engine/Source/Editor/UnrealEd/Private/FbxMeshUtils.cpp
@@ -320,7 +320,7 @@ namespace FbxMeshUtils
 		//Set the skeletal mesh import data from the base mesh, this make sure the import rotation transform is use when importing a LOD
 		UFbxSkeletalMeshImportData* TempAssetImportData = NULL;
 
-		UFbxAssetImportData *FbxAssetImportData = Cast<UFbxAssetImportData>(SelectedSkelMesh->AssetImportData);
+		UFbxAssetImportData *FbxAssetImportData = Cast<UFbxAssetImportData>(SelectedSkelMesh->GetAssetImportData());
 		if (FbxAssetImportData != nullptr)
 		{
 			UFbxSkeletalMeshImportData* ImportData = Cast<UFbxSkeletalMeshImportData>(FbxAssetImportData);
@@ -331,8 +331,8 @@ namespace FbxMeshUtils
 				// Prepare the import options
 				UFbxImportUI* ReimportUI = NewObject<UFbxImportUI>();
 				ReimportUI->MeshTypeToImport = FBXIT_SkeletalMesh;
-				ReimportUI->Skeleton = SelectedSkelMesh->Skeleton;
-				ReimportUI->PhysicsAsset = SelectedSkelMesh->PhysicsAsset;
+				ReimportUI->Skeleton = SelectedSkelMesh->GetSkeleton();
+				ReimportUI->PhysicsAsset = SelectedSkelMesh->GetPhysicsAsset();
 				// Import data already exists, apply it to the fbx import options
 				ReimportUI->SkeletalMeshImportData = ImportData;
 				//Some options not supported with skeletal mesh
@@ -471,7 +471,7 @@ namespace FbxMeshUtils
 				TArray<FName> OrderedMaterialNames;
 				{
 					int32 NoneNameCount = 0;
-					for (const FSkeletalMaterial &Material : SelectedSkelMesh->Materials)
+					for (const FSkeletalMaterial &Material : SelectedSkelMesh->GetMaterials())
 					{
 						if (Material.ImportedMaterialSlotName == NAME_None)
 							NoneNameCount++;
diff --git a/Engine/Source/Editor/UnrealEd/Private/LevelEditorViewport.cpp b/Engine/Source/Editor/UnrealEd/Private/LevelEditorViewport.cpp
--- a/Engine/Source/Editor/UnrealEd/Private/LevelEditorViewport.cpp
+++ b/Engine/Source/Editor/UnrealEd/Private/LevelEditorViewport.cpp
@@ -690,7 +690,7 @@ static bool AttemptApplyObjToComponent(UObject* ObjToUse, USceneComponent* Compo
 						SkeletalMeshComponent->Modify();
 
 						// If the component doesn't have a mesh or the anim blueprint's skeleton isn't compatible with the existing mesh's skeleton, the mesh should change
-						const bool bShouldChangeMesh = !SkeletalMeshComponent->SkeletalMesh || !AnimBPSkeleton->IsCompatible(SkeletalMeshComponent->SkeletalMesh->Skeleton);
+						const bool bShouldChangeMesh = !SkeletalMeshComponent->SkeletalMesh || !AnimBPSkeleton->IsCompatible(SkeletalMeshComponent->SkeletalMesh->GetSkeleton());
 
 						if (bShouldChangeMesh)
 						{
@@ -698,7 +698,7 @@ static bool AttemptApplyObjToComponent(UObject* ObjToUse, USceneComponent* Compo
 						}
 
 						// Verify that the skeletons are compatible before changing the anim BP
-						if (SkeletalMeshComponent->SkeletalMesh && AnimBPSkeleton->IsCompatible(SkeletalMeshComponent->SkeletalMesh->Skeleton))
+						if (SkeletalMeshComponent->SkeletalMesh && AnimBPSkeleton->IsCompatible(SkeletalMeshComponent->SkeletalMesh->GetSkeleton()))
 						{
 							SkeletalMeshComponent->SetAnimInstanceClass(DroppedObjAsAnimBlueprint->GeneratedClass);
 							bResult = true;
@@ -725,7 +725,7 @@ static bool AttemptApplyObjToComponent(UObject* ObjToUse, USceneComponent* Compo
 						SkeletalMeshComponent->Modify();
 
 						// If the component doesn't have a mesh or the anim blueprint's skeleton isn't compatible with the existing mesh's skeleton, the mesh should change
-						const bool bShouldChangeMesh = !SkeletalMeshComponent->SkeletalMesh || !AnimSkeleton->IsCompatible(SkeletalMeshComponent->SkeletalMesh->Skeleton);
+						const bool bShouldChangeMesh = !SkeletalMeshComponent->SkeletalMesh || !AnimSkeleton->IsCompatible(SkeletalMeshComponent->SkeletalMesh->GetSkeleton());
 
 						if (bShouldChangeMesh)
 						{
diff --git a/Engine/Source/Editor/UnrealEd/Private/PhysicsAssetUtils.cpp b/Engine/Source/Editor/UnrealEd/Private/PhysicsAssetUtils.cpp
--- a/Engine/Source/Editor/UnrealEd/Private/PhysicsAssetUtils.cpp
+++ b/Engine/Source/Editor/UnrealEd/Private/PhysicsAssetUtils.cpp
@@ -31,9 +31,9 @@ static int32 GetChildIndex(int32 BoneIndex, USkeletalMesh* SkelMesh, const TArra
 {
 	int32 ChildIndex = INDEX_NONE;
 
-	for(int32 i=0; i<SkelMesh->RefSkeleton.GetRawBoneNum(); i++)
+	for(int32 i=0; i<SkelMesh->GetRefSkeleton().GetRawBoneNum(); i++)
 	{
-		int32 ParentIndex = SkelMesh->RefSkeleton.GetParentIndex(i);
+		int32 ParentIndex = SkelMesh->GetRefSkeleton().GetParentIndex(i);
 
 		if (ParentIndex == BoneIndex && Infos[i].Positions.Num() > 0)
 		{
@@ -76,20 +76,20 @@ static float CalcBoneInfoLength(const FBoneVertInfo& Info)
  */
 static float GetMaximalMinSizeBelow(int32 BoneIndex, USkeletalMesh* SkelMesh, const TArray<FBoneVertInfo>& Infos)
 {
-	check( Infos.Num() == SkelMesh->RefSkeleton.GetRawBoneNum() );
+	check( Infos.Num() == SkelMesh->GetRefSkeleton().GetRawBoneNum() );
 
 	UE_LOG(LogPhysics, Log, TEXT("-------------------------------------------------"));
 
 	float MaximalMinBoxSize = 0.f;
 
 	// For all bones that are children of the supplied one...
-	for(int32 i=BoneIndex; i<SkelMesh->RefSkeleton.GetRawBoneNum(); i++)
+	for(int32 i=BoneIndex; i<SkelMesh->GetRefSkeleton().GetRawBoneNum(); i++)
 	{
-		if( SkelMesh->RefSkeleton.BoneIsChildOf(i, BoneIndex) )
+		if( SkelMesh->GetRefSkeleton().BoneIsChildOf(i, BoneIndex) )
 		{
 			float MinBoneDim = CalcBoneInfoLength( Infos[i] );
 			
-			UE_LOG(LogPhysics, Log,  TEXT("Parent: %s Bone: %s Size: %f"), *SkelMesh->RefSkeleton.GetBoneName(BoneIndex).ToString(), *SkelMesh->RefSkeleton.GetBoneName(i).ToString(), MinBoneDim );
+			UE_LOG(LogPhysics, Log,  TEXT("Parent: %s Bone: %s Size: %f"), *SkelMesh->GetRefSkeleton().GetBoneName(BoneIndex).ToString(), *SkelMesh->GetRefSkeleton().GetBoneName(i).ToString(), MinBoneDim );
 
 			MaximalMinBoxSize = FMath::Max(MaximalMinBoxSize, MinBoneDim);
 		}
@@ -121,7 +121,7 @@ bool CreateFromSkeletalMeshInternal(UPhysicsAsset* PhysicsAsset, USkeletalMesh*
 	// For each bone, get the vertices most firmly attached to it.
 	TArray<FBoneVertInfo> Infos;
 	MeshUtilities.CalcBoneVertInfos(SkelMesh, Infos, (Params.VertWeight == EVW_DominantWeight));
-	check(Infos.Num() == SkelMesh->RefSkeleton.GetRawBoneNum());
+	check(Infos.Num() == SkelMesh->GetRefSkeleton().GetRawBoneNum());
 
 	PhysicsAsset->CollisionDisableTable.Empty();
 
@@ -130,7 +130,7 @@ bool CreateFromSkeletalMeshInternal(UPhysicsAsset* PhysicsAsset, USkeletalMesh*
 	//If bone is big enough, make a body
 	//If not, add bone to parent for possible merge
 
-	const TArray<FTransform> LocalPose = SkelMesh->RefSkeleton.GetRefBonePose();
+	const TArray<FTransform> LocalPose = SkelMesh->GetRefSkeleton().GetRefBonePose();
 	TMap<int32, FBoneVertInfo> BoneToMergedBones;
 	TMap<FName, TArray<FName>> BoneNameToMergedParent;
 	const int32 NumBones = Infos.Num();
@@ -144,7 +144,7 @@ bool CreateFromSkeletalMeshInternal(UPhysicsAsset* PhysicsAsset, USkeletalMesh*
 		if(MyMergedSize < Params.MinBoneSize && MyMergedSize >= Params.MinWeldSize)
 		{
 			//Too small to make a body for, so let's merge with parent bone. TODO: use a merge threshold
-			const int32 ParentIndex = SkelMesh->RefSkeleton.GetParentIndex(BoneIdx);
+			const int32 ParentIndex = SkelMesh->GetRefSkeleton().GetParentIndex(BoneIdx);
 			if(ParentIndex != INDEX_NONE)
 			{
 				MergedSizes[ParentIndex] += MyMergedSize;
@@ -170,7 +170,7 @@ bool CreateFromSkeletalMeshInternal(UPhysicsAsset* PhysicsAsset, USkeletalMesh*
 	{
 		if (MergedSizes[BoneIndex] > Params.MinBoneSize)
 		{
-			const int32 ParentBoneIndex = SkelMesh->RefSkeleton.GetParentIndex(BoneIndex);
+			const int32 ParentBoneIndex = SkelMesh->GetRefSkeleton().GetParentIndex(BoneIndex);
 			if(ParentBoneIndex == INDEX_NONE)
 			{
 				break;	//We already have a single root body, so don't worry about it
@@ -215,7 +215,7 @@ bool CreateFromSkeletalMeshInternal(UPhysicsAsset* PhysicsAsset, USkeletalMesh*
 		if (bMakeBone)
 		{
 			// Go ahead and make this bone physical.
-			FName BoneName = SkelMesh->RefSkeleton.GetBoneName(BoneIndex);
+			FName BoneName = SkelMesh->GetRefSkeleton().GetBoneName(BoneIndex);
 
 			SlowTask.EnterProgressFrame(1.0f, FText::Format(NSLOCTEXT("PhysicsAssetEditor", "ResetCollsionStepInfo", "Generating collision for {0}"), FText::FromName(BoneName)));
 
@@ -252,10 +252,10 @@ bool CreateFromSkeletalMeshInternal(UPhysicsAsset* PhysicsAsset, USkeletalMesh*
 						RelTM = RelTM * LocalPose[ParentIndex];
 
 						//Travel up the hierarchy to find a parent which has a valid body
-						ParentIndex = SkelMesh->RefSkeleton.GetParentIndex(ParentIndex);
+						ParentIndex = SkelMesh->GetRefSkeleton().GetParentIndex(ParentIndex);
 						if(ParentIndex != INDEX_NONE)
 						{
-							ParentName = SkelMesh->RefSkeleton.GetBoneName(ParentIndex);
+							ParentName = SkelMesh->GetRefSkeleton().GetBoneName(ParentIndex);
 							ParentBodyIndex = PhysicsAsset->FindBodyIndex(ParentName);
 						}
 						else
@@ -371,7 +371,7 @@ bool CreateFromSkeletalMesh(UPhysicsAsset* PhysicsAsset, USkeletalMesh* SkelMesh
 
 	if(bSuccess && bSetToMesh)
 	{
-		SkelMesh->PhysicsAsset = PhysicsAsset;
+		SkelMesh->SetPhysicsAsset(PhysicsAsset);
 		SkelMesh->MarkPackageDirty();
 	}
 
@@ -683,13 +683,13 @@ void WeldBodies(UPhysicsAsset* PhysAsset, int32 BaseBodyIndex, int32 AddBodyInde
 	}
 
 	UBodySetup* Body1 = PhysAsset->SkeletalBodySetups[BaseBodyIndex];
-	int32 Bone1Index = SkelComp->SkeletalMesh->RefSkeleton.FindBoneIndex(Body1->BoneName);
+	int32 Bone1Index = SkelComp->SkeletalMesh->GetRefSkeleton().FindBoneIndex(Body1->BoneName);
 	check(Bone1Index != INDEX_NONE);
 	FTransform Bone1TM = SkelComp->GetBoneTransform(Bone1Index);
 	Bone1TM.RemoveScaling();
 
 	UBodySetup* Body2 = PhysAsset->SkeletalBodySetups[AddBodyIndex];
-	int32 Bone2Index = SkelComp->SkeletalMesh->RefSkeleton.FindBoneIndex(Body2->BoneName);
+	int32 Bone2Index = SkelComp->SkeletalMesh->GetRefSkeleton().FindBoneIndex(Body2->BoneName);
 	check(Bone2Index != INDEX_NONE);
 	FTransform Bone2TM = SkelComp->GetBoneTransform(Bone2Index);
 	Bone2TM.RemoveScaling();
diff --git a/Engine/Source/Editor/UnrealEd/Private/SSkeletonWidget.cpp b/Engine/Source/Editor/UnrealEd/Private/SSkeletonWidget.cpp
--- a/Engine/Source/Editor/UnrealEd/Private/SSkeletonWidget.cpp
+++ b/Engine/Source/Editor/UnrealEd/Private/SSkeletonWidget.cpp
@@ -309,9 +309,9 @@ void SSkeletonSelectorWindow::ConstructWindowFromMesh(USkeletalMesh* InSkeletalM
 {
 	TArray<FName>  BoneNames;
 
-	for (int32 I=0; I<InSkeletalMesh->RefSkeleton.GetRawBoneNum(); ++I)
+	for (int32 I=0; I<InSkeletalMesh->GetRefSkeleton().GetRawBoneNum(); ++I)
 	{
-		BoneNames.Add(InSkeletalMesh->RefSkeleton.GetBoneName(I));
+		BoneNames.Add(InSkeletalMesh->GetRefSkeleton().GetBoneName(I));
 	}
 
 	TSharedRef<SVerticalBox> ContentBox = SNew(SVerticalBox)
diff --git a/Engine/Source/Editor/UnrealEd/Private/SkinWeightsUtilities.cpp b/Engine/Source/Editor/UnrealEd/Private/SkinWeightsUtilities.cpp
--- a/Engine/Source/Editor/UnrealEd/Private/SkinWeightsUtilities.cpp
+++ b/Engine/Source/Editor/UnrealEd/Private/SkinWeightsUtilities.cpp
@@ -166,7 +166,7 @@ bool FSkinWeightsUtilities::ImportAlternateSkinWeight(USkeletalMesh* SkeletalMes
 	FbxFactory->RemoveFromRoot();
 
 	USkeletalMesh* TmpSkeletalMesh = Cast<USkeletalMesh>(ImportedObject);
-	if (TmpSkeletalMesh == nullptr || TmpSkeletalMesh->Skeleton == nullptr)
+	if (TmpSkeletalMesh == nullptr || TmpSkeletalMesh->GetSkeleton() == nullptr)
 	{
 		UE_LOG(LogSkinWeightsUtilities, Error, TEXT("Failed to import Skin Weight Profile from provided FBX file (%s)."), *Path);
 		DeletePathAssets();
@@ -342,10 +342,10 @@ bool FSkinWeightsUtilities::RemoveSkinnedWeightProfileData(USkeletalMesh* Skelet
 	TArray<FName> WarningNames;
 
 	//BaseLOD need to make sure the source data fit with the skeletalmesh materials array before using meshutilities.BuildSkeletalMesh
-	FLODUtilities::AdjustImportDataFaceMaterialIndex(SkeletalMesh->Materials, ImportDataDest.Materials, LODFacesDest, LODIndex);
+	FLODUtilities::AdjustImportDataFaceMaterialIndex(SkeletalMesh->GetMaterials(), ImportDataDest.Materials, LODFacesDest, LODIndex);
 
 	//Build the destination mesh with the Alternate influences, so the chunking is done properly.
-	const bool bBuildSuccess = MeshUtilities.BuildSkeletalMesh(LODModelDest, SkeletalMesh->GetPathName(), SkeletalMesh->RefSkeleton, LODInfluencesDest, LODWedgesDest, LODFacesDest, LODPointsDest, LODPointToRawMapDest, BuildOptions, &WarningMessages, &WarningNames);
+	const bool bBuildSuccess = MeshUtilities.BuildSkeletalMesh(LODModelDest, SkeletalMesh->GetPathName(), SkeletalMesh->GetRefSkeleton(), LODInfluencesDest, LODWedgesDest, LODFacesDest, LODPointsDest, LODPointToRawMapDest, BuildOptions, &WarningMessages, &WarningNames);
 	FLODUtilities::RegenerateAllImportSkinWeightProfileData(LODModelDest);
 
 	return bBuildSuccess;
@@ -365,7 +365,7 @@ FString FSkinWeightsUtilities::PickSkinWeightFBXPath(int32 LODIndex, USkeletalMe
 	if (DesktopPlatform)
 	{
 		// Try and retrieve the path containing the original skeletal mesh source data, and set it as default path for the file dialog
-		UFbxSkeletalMeshImportData* ImportData = SkeletalMesh ? Cast<UFbxSkeletalMeshImportData>(SkeletalMesh->AssetImportData) : nullptr;
+		UFbxSkeletalMeshImportData* ImportData = SkeletalMesh ? Cast<UFbxSkeletalMeshImportData>(SkeletalMesh->GetAssetImportData()) : nullptr;
 		FString DefaultPath;
 		FString TempString;
 		if (ImportData)
diff --git a/Engine/Source/Editor/UnrealEd/Private/SkinnedBoneTriangleCache.cpp b/Engine/Source/Editor/UnrealEd/Private/SkinnedBoneTriangleCache.cpp
--- a/Engine/Source/Editor/UnrealEd/Private/SkinnedBoneTriangleCache.cpp
+++ b/Engine/Source/Editor/UnrealEd/Private/SkinnedBoneTriangleCache.cpp
@@ -164,5 +164,5 @@ FVector FSkinnedBoneTriangleCache::VertexPosition(const FSkinnedVertexIndex Vert
 
 FMatrix FSkinnedBoneTriangleCache::BoneTransformMatrix(FBoneIndex BoneIndex) const
 {
-	return SkeletalMesh.RefBasesInvMatrix[BoneIndex];
+	return SkeletalMesh.GetRefBasesInvMatrix()[BoneIndex];
 }
diff --git a/Engine/Source/Editor/UnrealEd/Private/ThumbnailHelpers.cpp b/Engine/Source/Editor/UnrealEd/Private/ThumbnailHelpers.cpp
--- a/Engine/Source/Editor/UnrealEd/Private/ThumbnailHelpers.cpp
+++ b/Engine/Source/Editor/UnrealEd/Private/ThumbnailHelpers.cpp
@@ -451,7 +451,7 @@ void FSkeletalMeshThumbnailScene::GetViewMatrixParameters(const float InFOVDegre
 	const float BoundsZOffset = GetBoundsZOffset(PreviewActor->GetSkeletalMeshComponent()->Bounds);
 	const float TargetDistance = HalfMeshSize / FMath::Tan(HalfFOVRadians);
 
-	USceneThumbnailInfo* ThumbnailInfo = Cast<USceneThumbnailInfo>(PreviewActor->GetSkeletalMeshComponent()->SkeletalMesh->ThumbnailInfo);
+	USceneThumbnailInfo* ThumbnailInfo = Cast<USceneThumbnailInfo>(PreviewActor->GetSkeletalMeshComponent()->SkeletalMesh->GetThumbnailInfo());
 	if ( ThumbnailInfo )
 	{
 		if ( TargetDistance + ThumbnailInfo->OrbitZoom < 0 )
@@ -1000,9 +1000,9 @@ void FPhysicsAssetThumbnailScene::GetViewMatrixParameters(const float InFOVDegre
 	const float TargetDistance = HalfMeshSize / FMath::Tan(HalfFOVRadians);
 
 	USceneThumbnailInfo* ThumbnailInfo = USceneThumbnailInfo::StaticClass()->GetDefaultObject<USceneThumbnailInfo>();
-	if(PreviewActor->GetSkeletalMeshComponent()->SkeletalMesh && PreviewActor->GetSkeletalMeshComponent()->SkeletalMesh->PhysicsAsset)
+	if(PreviewActor->GetSkeletalMeshComponent()->SkeletalMesh && PreviewActor->GetSkeletalMeshComponent()->SkeletalMesh->GetPhysicsAsset())
 	{
-		if ( USceneThumbnailInfo* InteralThumbnailInfo = Cast<USceneThumbnailInfo>(PreviewActor->GetSkeletalMeshComponent()->SkeletalMesh->PhysicsAsset->ThumbnailInfo) )
+		if ( USceneThumbnailInfo* InteralThumbnailInfo = Cast<USceneThumbnailInfo>(PreviewActor->GetSkeletalMeshComponent()->SkeletalMesh->GetPhysicsAsset()->ThumbnailInfo) )
 		{
 			ThumbnailInfo = InteralThumbnailInfo;
 			if ( TargetDistance + InteralThumbnailInfo->OrbitZoom < 0 )
diff --git a/Engine/Source/Editor/UnrealEd/Public/AnimationEditorUtils.h b/Engine/Source/Editor/UnrealEd/Public/AnimationEditorUtils.h
--- a/Engine/Source/Editor/UnrealEd/Public/AnimationEditorUtils.h
+++ b/Engine/Source/Editor/UnrealEd/Public/AnimationEditorUtils.h
@@ -153,7 +153,7 @@ namespace AnimationEditorUtils
 		if (Skeleton == nullptr)
 		{
 			SkeletalMesh = CastChecked<USkeletalMesh>(SkeletonOrSkeletalMesh);
-			Skeleton = SkeletalMesh->Skeleton;
+			Skeleton = SkeletalMesh->GetSkeleton();
 		}
 
 		if (Skeleton)
@@ -210,7 +210,7 @@ namespace AnimationEditorUtils
 			if (Skeleton == nullptr)
 			{
 				SkeletalMesh = CastChecked<USkeletalMesh>(SkeletonsOrSkeletalMeshes[0].Get());
-				Skeleton = SkeletalMesh->Skeleton;
+				Skeleton = SkeletalMesh->GetSkeleton();
 			}
 
 			if(Skeleton)
diff --git a/Engine/Source/Runtime/AnimGraphRuntime/Private/AnimCustomInstanceHelper.cpp b/Engine/Source/Runtime/AnimGraphRuntime/Private/AnimCustomInstanceHelper.cpp
--- a/Engine/Source/Runtime/AnimGraphRuntime/Private/AnimCustomInstanceHelper.cpp
+++ b/Engine/Source/Runtime/AnimGraphRuntime/Private/AnimCustomInstanceHelper.cpp
@@ -15,5 +15,5 @@ bool FAnimCustomInstanceHelper::ShouldCreateCustomInstancePlayer(const USkeletal
 	UAnimInstance* AnimInstance = SkeletalMeshComponent->GetAnimInstance();
 
 	return (AnimInstance == nullptr || SkeletalMeshComponent->GetAnimationMode() != EAnimationMode::AnimationBlueprint ||
-		AnimInstance->GetClass() != SkeletalMeshComponent->AnimClass || !SkeletalMesh->Skeleton->IsCompatible(AnimInstance->CurrentSkeleton));
+		AnimInstance->GetClass() != SkeletalMeshComponent->AnimClass || !SkeletalMesh->GetSkeleton()->IsCompatible(AnimInstance->CurrentSkeleton));
 }
\ No newline at end of file
diff --git a/Engine/Source/Runtime/AnimGraphRuntime/Private/AnimNodes/AnimNode_CopyPoseFromMesh.cpp b/Engine/Source/Runtime/AnimGraphRuntime/Private/AnimNodes/AnimNode_CopyPoseFromMesh.cpp
--- a/Engine/Source/Runtime/AnimGraphRuntime/Private/AnimNodes/AnimNode_CopyPoseFromMesh.cpp
+++ b/Engine/Source/Runtime/AnimGraphRuntime/Private/AnimNodes/AnimNode_CopyPoseFromMesh.cpp
@@ -202,7 +202,7 @@ void FAnimNode_CopyPoseFromMesh::Evaluate_AnyThread(FPoseContext& Output)
 				if (Value && SourceMeshTransformArray.IsValidIndex(*Value))
 				{
 					const int32 SourceBoneIndex = *Value;
-					const int32 ParentIndex = CurrentMesh->RefSkeleton.GetParentIndex(SourceBoneIndex);
+					const int32 ParentIndex = CurrentMesh->GetRefSkeleton().GetParentIndex(SourceBoneIndex);
 					const FCompactPoseBoneIndex MyParentIndex = RequiredBones.GetParentBoneIndex(PoseBoneIndex);
 					// only apply if I also have parent, otherwise, it should apply the space bases
 					if (SourceMeshTransformArray.IsValidIndex(ParentIndex) && MyParentIndex != INDEX_NONE)
@@ -272,29 +272,29 @@ void FAnimNode_CopyPoseFromMesh::ReinitializeMeshComponent(USkeletalMeshComponen
 
 			if (SourceSkelMesh == TargetSkelMesh)
 			{
-				for(int32 ComponentSpaceBoneId = 0; ComponentSpaceBoneId < SourceSkelMesh->RefSkeleton.GetNum(); ++ComponentSpaceBoneId)
+				for(int32 ComponentSpaceBoneId = 0; ComponentSpaceBoneId < SourceSkelMesh->GetRefSkeleton().GetNum(); ++ComponentSpaceBoneId)
 				{
 					BoneMapToSource.Add(ComponentSpaceBoneId, ComponentSpaceBoneId);
 				}
 			}
 			else
 			{
-				const int32 SplitBoneIndex = (RootBoneToCopy != NAME_Name)? TargetSkelMesh->RefSkeleton.FindBoneIndex(RootBoneToCopy) : INDEX_NONE;
-				for (int32 ComponentSpaceBoneId = 0; ComponentSpaceBoneId < TargetSkelMesh->RefSkeleton.GetNum(); ++ComponentSpaceBoneId)
+				const int32 SplitBoneIndex = (RootBoneToCopy != NAME_Name)? TargetSkelMesh->GetRefSkeleton().FindBoneIndex(RootBoneToCopy) : INDEX_NONE;
+				for (int32 ComponentSpaceBoneId = 0; ComponentSpaceBoneId < TargetSkelMesh->GetRefSkeleton().GetNum(); ++ComponentSpaceBoneId)
 				{
 					if (SplitBoneIndex == INDEX_NONE || ComponentSpaceBoneId == SplitBoneIndex
-						|| TargetSkelMesh->RefSkeleton.BoneIsChildOf(ComponentSpaceBoneId, SplitBoneIndex))
+						|| TargetSkelMesh->GetRefSkeleton().BoneIsChildOf(ComponentSpaceBoneId, SplitBoneIndex))
 					{
-						FName BoneName = TargetSkelMesh->RefSkeleton.GetBoneName(ComponentSpaceBoneId);
-						BoneMapToSource.Add(ComponentSpaceBoneId, SourceSkelMesh->RefSkeleton.FindBoneIndex(BoneName));
+						FName BoneName = TargetSkelMesh->GetRefSkeleton().GetBoneName(ComponentSpaceBoneId);
+						BoneMapToSource.Add(ComponentSpaceBoneId, SourceSkelMesh->GetRefSkeleton().FindBoneIndex(BoneName));
 					}
 				}
 			}
 		
 			if (bCopyCurves)
 			{
-				USkeleton* SourceSkeleton = SourceSkelMesh->Skeleton;
-				USkeleton* TargetSkeleton = TargetSkelMesh->Skeleton;
+				USkeleton* SourceSkeleton = SourceSkelMesh->GetSkeleton();
+				USkeleton* TargetSkeleton = TargetSkelMesh->GetSkeleton();
 
 				// you shouldn't be here if this happened
 				if (ensureMsgf(SourceSkeleton, TEXT("Invalid null source skeleton : %s"), *GetNameSafe(SourceSkelMesh))
diff --git a/Engine/Source/Runtime/AnimGraphRuntime/Private/AnimNodes/AnimNode_PoseSnapshot.cpp b/Engine/Source/Runtime/AnimGraphRuntime/Private/AnimNodes/AnimNode_PoseSnapshot.cpp
--- a/Engine/Source/Runtime/AnimGraphRuntime/Private/AnimNodes/AnimNode_PoseSnapshot.cpp
+++ b/Engine/Source/Runtime/AnimGraphRuntime/Private/AnimNodes/AnimNode_PoseSnapshot.cpp
@@ -31,11 +31,11 @@ void FAnimNode_PoseSnapshot::PreUpdate(const UAnimInstance* InAnimInstance)
 		{
 			// cache bone names for the target mesh
 			TargetBoneNames.Reset();
-			TargetBoneNames.AddDefaulted(CurrentSkeletalMesh->RefSkeleton.GetNum());
+			TargetBoneNames.AddDefaulted(CurrentSkeletalMesh->GetRefSkeleton().GetNum());
 
-			for (int32 BoneIndex = 0; BoneIndex < CurrentSkeletalMesh->RefSkeleton.GetNum(); ++BoneIndex)
+			for (int32 BoneIndex = 0; BoneIndex < CurrentSkeletalMesh->GetRefSkeleton().GetNum(); ++BoneIndex)
 			{
-				TargetBoneNames[BoneIndex] = CurrentSkeletalMesh->RefSkeleton.GetBoneName(BoneIndex);
+				TargetBoneNames[BoneIndex] = CurrentSkeletalMesh->GetRefSkeleton().GetBoneName(BoneIndex);
 			}
 		}
 
diff --git a/Engine/Source/Runtime/AnimGraphRuntime/Private/BoneControllers/AnimNode_RigidBody.cpp b/Engine/Source/Runtime/AnimGraphRuntime/Private/BoneControllers/AnimNode_RigidBody.cpp
--- a/Engine/Source/Runtime/AnimGraphRuntime/Private/BoneControllers/AnimNode_RigidBody.cpp
+++ b/Engine/Source/Runtime/AnimGraphRuntime/Private/BoneControllers/AnimNode_RigidBody.cpp
@@ -873,7 +873,7 @@ void ComputeBodyInsertionOrder(TArray<FBoneIndexType>& InsertionOrder, const USk
 		SKC.ComputeRequiredBones(RequiredBones0, ComponentSpaceTMs0, 0, /*bIgnorePhysicsAsset=*/ true);
 
 		TArray<bool> InSortedOrder;
-		InSortedOrder.AddZeroed(SKC.SkeletalMesh->RefSkeleton.GetNum());
+		InSortedOrder.AddZeroed(SKC.SkeletalMesh->GetRefSkeleton().GetNum());
 
 		auto MergeIndices = [&InsertionOrder, &InSortedOrder](const TArray<FBoneIndexType>& RequiredBones) -> void
 		{
@@ -911,11 +911,11 @@ void FAnimNode_RigidBody::InitPhysics(const UAnimInstance* InAnimInstance)
 	const USkeletalMeshComponent* SkeletalMeshComp = InAnimInstance->GetSkelMeshComponent();
 	const USkeletalMesh* SkeletalMeshAsset = SkeletalMeshComp->SkeletalMesh;
 
-	const FReferenceSkeleton& SkelMeshRefSkel = SkeletalMeshAsset->RefSkeleton;
+	const FReferenceSkeleton& SkelMeshRefSkel = SkeletalMeshAsset->GetRefSkeleton();
 	UsePhysicsAsset = OverridePhysicsAsset ? ToRawPtr(OverridePhysicsAsset) : InAnimInstance->GetSkelMeshComponent()->GetPhysicsAsset();
 		
 	USkeleton* SkeletonAsset = InAnimInstance->CurrentSkeleton;
-	ensure(SkeletonAsset == SkeletalMeshAsset->Skeleton);
+	ensure(SkeletonAsset == SkeletalMeshAsset->GetSkeleton());
 
 	const int32 SkelMeshLinkupIndex = SkeletonAsset->GetMeshLinkupIndex(SkeletalMeshAsset);
 	ensure(SkelMeshLinkupIndex != INDEX_NONE);
diff --git a/Engine/Source/Runtime/AnimGraphRuntime/Private/BoneControllers/AnimNode_SplineIK.cpp b/Engine/Source/Runtime/AnimGraphRuntime/Private/BoneControllers/AnimNode_SplineIK.cpp
--- a/Engine/Source/Runtime/AnimGraphRuntime/Private/BoneControllers/AnimNode_SplineIK.cpp
+++ b/Engine/Source/Runtime/AnimGraphRuntime/Private/BoneControllers/AnimNode_SplineIK.cpp
@@ -37,7 +37,7 @@ void FAnimNode_SplineIK::OnInitializeAnimInstance(const FAnimInstanceProxy* InPr
 {
 	if (InAnimInstance->GetSkelMeshComponent() && InAnimInstance->GetSkelMeshComponent()->SkeletalMesh)
 	{
-		GatherBoneReferences(InAnimInstance->GetSkelMeshComponent()->SkeletalMesh->RefSkeleton);
+		GatherBoneReferences(InAnimInstance->GetSkelMeshComponent()->SkeletalMesh->GetRefSkeleton());
 	}
 }
 
diff --git a/Engine/Source/Runtime/ClothingSystemRuntimeCommon/Private/ClothingAsset.cpp b/Engine/Source/Runtime/ClothingSystemRuntimeCommon/Private/ClothingAsset.cpp
--- a/Engine/Source/Runtime/ClothingSystemRuntimeCommon/Private/ClothingAsset.cpp
+++ b/Engine/Source/Runtime/ClothingSystemRuntimeCommon/Private/ClothingAsset.cpp
@@ -342,7 +342,7 @@ bool UClothingAssetCommon::BindToSkeletalMesh(
 	TArray<FBoneIndexType> TempBoneMap = OriginalSection.BoneMap;
 	for(FName& BoneName : UsedBoneNames)
 	{
-		const int32 BoneIndex = InSkelMesh->RefSkeleton.FindBoneIndex(BoneName);
+		const int32 BoneIndex = InSkelMesh->GetRefSkeleton().FindBoneIndex(BoneName);
 		if(BoneIndex != INDEX_NONE)
 		{
 			TempBoneMap.AddUnique(BoneIndex);
@@ -378,7 +378,7 @@ bool UClothingAssetCommon::BindToSkeletalMesh(
 
 	// Set the asset index, used during rendering to pick the correct sim mesh buffer
 	int32 AssetIndex = INDEX_NONE;
-	check(InSkelMesh->MeshClothingAssets.Find(this, AssetIndex));
+	check(InSkelMesh->GetMeshClothingAssets().Find(this, AssetIndex));
 	OriginalSection.CorrespondClothAssetIndex = AssetIndex;
 
 	// sim properties
@@ -392,7 +392,7 @@ bool UClothingAssetCommon::BindToSkeletalMesh(
 		if(!SkelLod.RequiredBones.Contains(BoneIndex))
 		{
 			bRequireBoneChange = true;
-			if(InSkelMesh->RefSkeleton.IsValidIndex(BoneIndex))
+			if(InSkelMesh->GetRefSkeleton().IsValidIndex(BoneIndex))
 			{
 				SkelLod.RequiredBones.Add(BoneIndex);
 				SkelLod.ActiveBoneIndices.AddUnique(BoneIndex);
@@ -402,7 +402,7 @@ bool UClothingAssetCommon::BindToSkeletalMesh(
 	if(bRequireBoneChange)
 	{
 		SkelLod.RequiredBones.Sort();
-		InSkelMesh->RefSkeleton.EnsureParentsExistAndSort(SkelLod.ActiveBoneIndices);		
+		InSkelMesh->GetRefSkeleton().EnsureParentsExistAndSort(SkelLod.ActiveBoneIndices);
 	}
 
 	if(CustomData)
@@ -639,7 +639,7 @@ void UClothingAssetCommon::RefreshBoneMapping(USkeletalMesh* InSkelMesh)
 	for(int32 BoneNameIndex = 0; BoneNameIndex < UsedBoneNames.Num(); ++BoneNameIndex)
 	{
 		UsedBoneIndices[BoneNameIndex] = 
-			InSkelMesh->RefSkeleton.FindBoneIndex(
+			InSkelMesh->GetRefSkeleton().FindBoneIndex(
 				UsedBoneNames[BoneNameIndex]);
 	}
 }
@@ -696,7 +696,7 @@ void UClothingAssetCommon::CalculateReferenceBoneIndex()
 		PathsToRoot.Reserve(NumWeightedBones);
 		
 		// Compute paths to the root bone
-		const FReferenceSkeleton& RefSkel = OwnerMesh->RefSkeleton;
+		const FReferenceSkeleton& RefSkel = OwnerMesh->GetRefSkeleton();
 		for(int32 WeightedBoneIndex = 0; WeightedBoneIndex < NumWeightedBones; ++WeightedBoneIndex)
 		{
 			PathsToRoot.AddDefaulted();
@@ -1023,7 +1023,7 @@ void UClothingAssetCommon::PropagateSharedConfigs(bool bMigrateSharedConfigToCon
 	// Update this asset's shared config when the asset belongs to a skeletal mesh
 	if (USkeletalMesh* const SkeletalMesh = Cast<USkeletalMesh>(GetOuter()))
 	{
-		const TArray<UClothingAssetBase*>& ClothingAssets = SkeletalMesh->MeshClothingAssets;
+		const TArray<UClothingAssetBase*>& ClothingAssets = SkeletalMesh->GetMeshClothingAssets();
  
 		// Collect all shared configs found in the other assets
 		TMap<FName, UClothConfigBase*> ClothSharedConfigs;
diff --git a/Engine/Source/Runtime/ClothingSystemRuntimeCommon/Private/ClothingSimulation.cpp b/Engine/Source/Runtime/ClothingSystemRuntimeCommon/Private/ClothingSimulation.cpp
--- a/Engine/Source/Runtime/ClothingSystemRuntimeCommon/Private/ClothingSimulation.cpp
+++ b/Engine/Source/Runtime/ClothingSystemRuntimeCommon/Private/ClothingSimulation.cpp
@@ -65,7 +65,7 @@ void FClothingSimulationContextCommon::FillBoneTransforms(const USkeletalMeshCom
 			if (SkeletalMesh)
 			{
 				// This case indicates an invalid master pose component (e.g. no skeletal mesh)
-				NumBones = SkeletalMesh->RefSkeleton.GetNum();
+				NumBones = SkeletalMesh->GetRefSkeleton().GetNum();
 
 				BoneTransforms.Empty(NumBones);
 				BoneTransforms.AddDefaulted(NumBones);
@@ -92,12 +92,12 @@ void FClothingSimulationContextCommon::FillBoneTransforms(const USkeletalMeshCom
 
 				if (!bFoundMaster && SkeletalMesh)
 				{
-					const int32 ParentIndex = SkeletalMesh->RefSkeleton.GetParentIndex(BoneIndex);
+					const int32 ParentIndex = SkeletalMesh->GetRefSkeleton().GetParentIndex(BoneIndex);
 
 					BoneTransforms[BoneIndex] =
 						BoneTransforms.IsValidIndex(ParentIndex) && ParentIndex < BoneIndex ?
-						BoneTransforms[ParentIndex] * SkeletalMesh->RefSkeleton.GetRefBonePose()[BoneIndex] :
-						SkeletalMesh->RefSkeleton.GetRefBonePose()[BoneIndex];
+						BoneTransforms[ParentIndex] * SkeletalMesh->GetRefSkeleton().GetRefBonePose()[BoneIndex] :
+						SkeletalMesh->GetRefSkeleton().GetRefBonePose()[BoneIndex];
 				}
 			}
 		}
