Implemented initial support for element based viewport interaction

There are two parts to this:
 - UTypedElementWorldInterface, which provides the lower-level access to things like element transforms.
 - UTypedElementViewportInteraction, which provides the higher-level access for things like gizmo manipulation.
   - The default implementation is implemented purely in terms of get/set transform via UTypedElementWorldInterface.
   - The level editor overrides this to continue to go through the current code path for gizmo manipulation.

diff --git a/Engine/Source/Editor/EditorFramework/EditorFramework.Build.cs b/Engine/Source/Editor/EditorFramework/EditorFramework.Build.cs
--- a/Engine/Source/Editor/EditorFramework/EditorFramework.Build.cs
+++ b/Engine/Source/Editor/EditorFramework/EditorFramework.Build.cs
@@ -31,6 +31,7 @@ public class EditorFramework : ModuleRules
 
 		PublicDependencyModuleNames.AddRange(
 			new string[] {
+				"InteractiveToolsFramework",
 				"TypedElementFramework",
 				"TypedElementInterfaces",
 			}
diff --git a/Engine/Source/Editor/EditorFramework/Private/Elements/Framework/TypedElementViewportInteraction.cpp b/Engine/Source/Editor/EditorFramework/Private/Elements/Framework/TypedElementViewportInteraction.cpp
new file mode 100644
--- /dev/null
+++ b/Engine/Source/Editor/EditorFramework/Private/Elements/Framework/TypedElementViewportInteraction.cpp
@@ -0,0 +1,193 @@
+// Copyright Epic Games, Inc. All Rights Reserved.
+
+#include "Elements/Framework/TypedElementViewportInteraction.h"
+#include "Elements/Framework/TypedElementSelectionSet.h"
+#include "Elements/Framework/TypedElementRegistry.h"
+
+void FTypedElementAssetEditorViewportInteractionCustomization::GetElementsToMove(const TTypedElement<UTypedElementWorldInterface>& InElementWorldHandle, const ETypedElementViewportInteractionWorldType InWorldType, const UTypedElementSelectionSet* InSelectionSet, UTypedElementList* OutElementsToMove)
+{
+	OutElementsToMove->Add(InElementWorldHandle);
+}
+
+bool FTypedElementAssetEditorViewportInteractionCustomization::GetGizmoPivotLocation(const TTypedElement<UTypedElementWorldInterface>& InElementWorldHandle, const UE::Widget::EWidgetMode InWidgetMode, FVector& OutPivotLocation)
+{
+	FTransform ElementWorldTransform;
+	if (InElementWorldHandle.GetWorldTransform(ElementWorldTransform))
+	{
+		OutPivotLocation = ElementWorldTransform.GetTranslation();
+		return true;
+	}
+
+	return false;
+}
+
+void FTypedElementAssetEditorViewportInteractionCustomization::PreGizmoManipulationStarted(TArrayView<const TTypedElement<UTypedElementWorldInterface>> InElementWorldHandles, const UE::Widget::EWidgetMode InWidgetMode)
+{
+}
+
+void FTypedElementAssetEditorViewportInteractionCustomization::GizmoManipulationStarted(const TTypedElement<UTypedElementWorldInterface>& InElementWorldHandle, const UE::Widget::EWidgetMode InWidgetMode)
+{
+}
+
+void FTypedElementAssetEditorViewportInteractionCustomization::GizmoManipulationDeltaUpdate(const TTypedElement<UTypedElementWorldInterface>& InElementWorldHandle, const UE::Widget::EWidgetMode InWidgetMode, const EAxisList::Type InDragAxis, const FInputDeviceState& InInputState, const FTransform& InDeltaTransform, const FVector& InPivotLocation)
+{
+	FTransform ElementWorldTransform;
+	if (InElementWorldHandle.GetWorldTransform(ElementWorldTransform))
+	{
+		// Apply delta rotation around the pivot location
+		{
+			const FQuat DeltaRotation = InDeltaTransform.GetRotation();
+			if (!DeltaRotation.Rotator().IsZero())
+			{
+				ElementWorldTransform.SetRotation(ElementWorldTransform.GetRotation() * DeltaRotation);
+
+				FVector ElementLocation = ElementWorldTransform.GetTranslation();
+				ElementLocation -= InPivotLocation;
+				ElementLocation = FRotationMatrix::Make(DeltaRotation).TransformPosition(ElementLocation);
+				ElementLocation += InPivotLocation;
+				ElementWorldTransform.SetTranslation(ElementLocation);
+			}
+		}
+
+		// Apply delta translation
+		{
+			const FVector DeltaTranslation = InDeltaTransform.GetTranslation();
+			ElementWorldTransform.SetTranslation(ElementWorldTransform.GetTranslation() + DeltaTranslation);
+		}
+
+		// Apply delta scaling around the pivot location
+		{
+			const FVector DeltaScale3D = InDeltaTransform.GetScale3D();
+			if (!DeltaScale3D.IsNearlyZero(0.000001f))
+			{
+				ElementWorldTransform.SetScale3D(ElementWorldTransform.GetScale3D() + DeltaScale3D);
+
+				FVector ElementLocation = ElementWorldTransform.GetTranslation();
+				ElementLocation -= InPivotLocation;
+				ElementLocation += FScaleMatrix::Make(DeltaScale3D).TransformPosition(ElementLocation);
+				ElementLocation += InPivotLocation;
+				ElementWorldTransform.SetTranslation(ElementLocation);
+			}
+		}
+
+		InElementWorldHandle.SetWorldTransform(ElementWorldTransform);
+	}
+}
+
+void FTypedElementAssetEditorViewportInteractionCustomization::GizmoManipulationStopped(const TTypedElement<UTypedElementWorldInterface>& InElementWorldHandle, const UE::Widget::EWidgetMode InWidgetMode)
+{
+}
+
+void FTypedElementAssetEditorViewportInteractionCustomization::PostGizmoManipulationStopped(TArrayView<const TTypedElement<UTypedElementWorldInterface>> InElementWorldHandles, const UE::Widget::EWidgetMode InWidgetMode)
+{
+}
+
+
+void UTypedElementViewportInteraction::GetSelectedElementsToMove(const UTypedElementSelectionSet* InSelectionSet, const ETypedElementViewportInteractionWorldType InWorldType, UTypedElementList* OutElementsToMove) const
+{
+	OutElementsToMove->Reset();
+	InSelectionSet->ForEachSelectedElement<UTypedElementWorldInterface>([this, InSelectionSet, InWorldType, OutElementsToMove](const TTypedElement<UTypedElementWorldInterface>& InElementWorldHandle)
+	{
+		if (InElementWorldHandle.CanEditElement())
+		{
+			FTypedElementViewportInteractionElement ViewportInteractionElement(InElementWorldHandle, GetAssetEditorCustomizationByTypeId(InElementWorldHandle.GetId().GetTypeId()));
+			check(ViewportInteractionElement.IsSet());
+			ViewportInteractionElement.GetElementsToMove(InWorldType, InSelectionSet, OutElementsToMove);
+		}
+		return true;
+	});
+}
+
+void UTypedElementViewportInteraction::BeginGizmoManipulation(const UTypedElementList* InElementsToMove, const UE::Widget::EWidgetMode InWidgetMode)
+{
+	{
+		TMap<FTypedHandleTypeId, TArray<TTypedElement<UTypedElementWorldInterface>>> ElementsToMoveByType;
+		BatchElementsByType(InElementsToMove, ElementsToMoveByType);
+
+		for (const auto& ElementsByTypePair : ElementsToMoveByType)
+		{
+			FTypedElementAssetEditorViewportInteractionCustomization* AssetEditorViewportInteractionCustomization = GetAssetEditorCustomizationByTypeId(ElementsByTypePair.Key);
+			check(AssetEditorViewportInteractionCustomization);
+			AssetEditorViewportInteractionCustomization->PreGizmoManipulationStarted(ElementsByTypePair.Value, InWidgetMode);
+		}
+	}
+
+	InElementsToMove->ForEachElementHandle([this, InWidgetMode](const FTypedElementHandle& InElementToMove)
+	{
+		FTypedElementViewportInteractionElement ViewportInteractionElement = ResolveViewportInteractionElement(InElementToMove);
+		if (ViewportInteractionElement)
+		{
+			ViewportInteractionElement.GizmoManipulationStarted(InWidgetMode);
+		}
+		return true;
+	});
+}
+
+void UTypedElementViewportInteraction::UpdateGizmoManipulation(const UTypedElementList* InElementsToMove, const UE::Widget::EWidgetMode InWidgetMode, const EAxisList::Type InDragAxis, const FInputDeviceState& InInputState, const FTransform& InDeltaTransform)
+{
+	InElementsToMove->ForEachElementHandle([this, InWidgetMode, InDragAxis, &InInputState, &InDeltaTransform](const FTypedElementHandle& InElementToMove)
+	{
+		FTypedElementViewportInteractionElement ViewportInteractionElement = ResolveViewportInteractionElement(InElementToMove);
+		if (ViewportInteractionElement)
+		{
+			FVector PivotLocation = FVector::ZeroVector;
+			ViewportInteractionElement.GetGizmoPivotLocation(InWidgetMode, PivotLocation);
+			ViewportInteractionElement.GizmoManipulationDeltaUpdate(InWidgetMode, InDragAxis, InInputState, InDeltaTransform, PivotLocation);
+		}
+		return true;
+	});
+}
+
+void UTypedElementViewportInteraction::EndGizmoManipulation(const UTypedElementList* InElementsToMove, const UE::Widget::EWidgetMode InWidgetMode)
+{
+	InElementsToMove->ForEachElementHandle([this, InWidgetMode](const FTypedElementHandle& InElementToMove)
+	{
+		FTypedElementViewportInteractionElement ViewportInteractionElement = ResolveViewportInteractionElement(InElementToMove);
+		if (ViewportInteractionElement)
+		{
+			ViewportInteractionElement.GizmoManipulationStopped(InWidgetMode);
+		}
+		return true;
+	});
+	
+	{
+		TMap<FTypedHandleTypeId, TArray<TTypedElement<UTypedElementWorldInterface>>> ElementsToMoveByType;
+		BatchElementsByType(InElementsToMove, ElementsToMoveByType);
+
+		for (const auto& ElementsByTypePair : ElementsToMoveByType)
+		{
+			FTypedElementAssetEditorViewportInteractionCustomization* AssetEditorViewportInteractionCustomization = GetAssetEditorCustomizationByTypeId(ElementsByTypePair.Key);
+			check(AssetEditorViewportInteractionCustomization);
+			AssetEditorViewportInteractionCustomization->PostGizmoManipulationStopped(ElementsByTypePair.Value, InWidgetMode);
+		}
+	}
+}
+
+void UTypedElementViewportInteraction::ApplyDeltaToElement(const FTypedElementHandle& InElementHandle, const UE::Widget::EWidgetMode InWidgetMode, const EAxisList::Type InDragAxis, const FInputDeviceState& InInputState, const FTransform& InDeltaTransform)
+{
+	FTypedElementViewportInteractionElement ViewportInteractionElement = ResolveViewportInteractionElement(InElementHandle);
+	if (ViewportInteractionElement)
+	{
+		FVector PivotLocation = FVector::ZeroVector;
+		ViewportInteractionElement.GetGizmoPivotLocation(InWidgetMode, PivotLocation);
+		ViewportInteractionElement.GizmoManipulationDeltaUpdate(InWidgetMode, InDragAxis, InInputState, InDeltaTransform, PivotLocation);
+	}
+}
+
+FTypedElementViewportInteractionElement UTypedElementViewportInteraction::ResolveViewportInteractionElement(const FTypedElementHandle& InElementHandle) const
+{
+	return InElementHandle
+		? FTypedElementViewportInteractionElement(UTypedElementRegistry::GetInstance()->GetElement<UTypedElementWorldInterface>(InElementHandle), GetAssetEditorCustomizationByTypeId(InElementHandle.GetId().GetTypeId()))
+		: FTypedElementViewportInteractionElement();
+}
+
+void UTypedElementViewportInteraction::BatchElementsByType(const UTypedElementList* InElementsToMove, TMap<FTypedHandleTypeId, TArray<TTypedElement<UTypedElementWorldInterface>>>& OutElementsByType)
+{
+	OutElementsByType.Reset();
+	InElementsToMove->ForEachElement<UTypedElementWorldInterface>([&OutElementsByType](const TTypedElement<UTypedElementWorldInterface>& InElementWorldHandle)
+	{
+		TArray<TTypedElement<UTypedElementWorldInterface>>& ElementsForType = OutElementsByType.FindOrAdd(InElementWorldHandle.GetId().GetTypeId());
+		ElementsForType.Add(InElementWorldHandle);
+		return true;
+	});
+}
diff --git a/Engine/Source/Editor/EditorFramework/Public/Elements/Framework/TypedElementViewportInteraction.h b/Engine/Source/Editor/EditorFramework/Public/Elements/Framework/TypedElementViewportInteraction.h
new file mode 100644
--- /dev/null
+++ b/Engine/Source/Editor/EditorFramework/Public/Elements/Framework/TypedElementViewportInteraction.h
@@ -0,0 +1,127 @@
+// Copyright Epic Games, Inc. All Rights Reserved.
+
+#pragma once
+
+#include "InputState.h"
+#include "UnrealWidgetFwd.h"
+#include "Elements/Framework/TypedElementAssetEditorCustomization.h"
+#include "Elements/Interfaces/TypedElementWorldInterface.h"
+#include "TypedElementViewportInteraction.generated.h"
+
+class UTypedElementList;
+class UTypedElementSelectionSet;
+
+enum class ETypedElementViewportInteractionWorldType : uint8
+{
+	Editor,
+	PlayInEditor,
+};
+
+/**
+ * Customization used to allow asset editors (such as the level editor) to override the base behavior of viewport interaction.
+ */
+class EDITORFRAMEWORK_API FTypedElementAssetEditorViewportInteractionCustomization
+{
+public:
+	virtual ~FTypedElementAssetEditorViewportInteractionCustomization() = default;
+
+	//~ See UTypedElementViewportInteraction for API docs
+	virtual void GetElementsToMove(const TTypedElement<UTypedElementWorldInterface>& InElementWorldHandle, const ETypedElementViewportInteractionWorldType InWorldType, const UTypedElementSelectionSet* InSelectionSet, UTypedElementList* OutElementsToMove);
+	virtual bool GetGizmoPivotLocation(const TTypedElement<UTypedElementWorldInterface>& InElementWorldHandle, const UE::Widget::EWidgetMode InWidgetMode, FVector& OutPivotLocation);
+	virtual void PreGizmoManipulationStarted(TArrayView<const TTypedElement<UTypedElementWorldInterface>> InElementWorldHandles, const UE::Widget::EWidgetMode InWidgetMode);
+	virtual void GizmoManipulationStarted(const TTypedElement<UTypedElementWorldInterface>& InElementWorldHandle, const UE::Widget::EWidgetMode InWidgetMode);
+	virtual void GizmoManipulationDeltaUpdate(const TTypedElement<UTypedElementWorldInterface>& InElementWorldHandle, const UE::Widget::EWidgetMode InWidgetMode, const EAxisList::Type InDragAxis, const FInputDeviceState& InInputState, const FTransform& InDeltaTransform, const FVector& InPivotLocation);
+	virtual void GizmoManipulationStopped(const TTypedElement<UTypedElementWorldInterface>& InElementWorldHandle, const UE::Widget::EWidgetMode InWidgetMode);
+	virtual void PostGizmoManipulationStopped(TArrayView<const TTypedElement<UTypedElementWorldInterface>> InElementWorldHandles, const UE::Widget::EWidgetMode InWidgetMode);
+};
+
+/**
+ * Utility to hold a typed element handle and its associated world interface and asset editor viewport interaction customization.
+ */
+struct EDITORFRAMEWORK_API FTypedElementViewportInteractionElement
+{
+public:
+	FTypedElementViewportInteractionElement() = default;
+
+	FTypedElementViewportInteractionElement(TTypedElement<UTypedElementWorldInterface> InElementWorldHandle, FTypedElementAssetEditorViewportInteractionCustomization* InAssetEditorViewportInteractionCustomization)
+		: ElementWorldHandle(MoveTemp(InElementWorldHandle))
+		, AssetEditorViewportInteractionCustomization(InAssetEditorViewportInteractionCustomization)
+	{
+	}
+
+	FTypedElementViewportInteractionElement(const FTypedElementViewportInteractionElement&) = default;
+	FTypedElementViewportInteractionElement& operator=(const FTypedElementViewportInteractionElement&) = default;
+
+	FTypedElementViewportInteractionElement(FTypedElementViewportInteractionElement&&) = default;
+	FTypedElementViewportInteractionElement& operator=(FTypedElementViewportInteractionElement&&) = default;
+
+	FORCEINLINE explicit operator bool() const
+	{
+		return IsSet();
+	}
+
+	FORCEINLINE bool IsSet() const
+	{
+		return ElementWorldHandle.IsSet()
+			&& AssetEditorViewportInteractionCustomization;
+	}
+
+	//~ See UTypedElementViewportInteraction for API docs
+	void GetElementsToMove(const ETypedElementViewportInteractionWorldType InWorldType, const UTypedElementSelectionSet* InSelectionSet, UTypedElementList* OutElementsToMove) { return AssetEditorViewportInteractionCustomization->GetElementsToMove(ElementWorldHandle, InWorldType, InSelectionSet, OutElementsToMove); }
+	bool GetGizmoPivotLocation(const UE::Widget::EWidgetMode InWidgetMode, FVector& OutPivotLocation) const { return AssetEditorViewportInteractionCustomization->GetGizmoPivotLocation(ElementWorldHandle, InWidgetMode, OutPivotLocation); }
+	void GizmoManipulationStarted(const UE::Widget::EWidgetMode InWidgetMode) const { return AssetEditorViewportInteractionCustomization->GizmoManipulationStarted(ElementWorldHandle, InWidgetMode); }
+	void GizmoManipulationDeltaUpdate(const UE::Widget::EWidgetMode InWidgetMode, const EAxisList::Type InDragAxis, const FInputDeviceState& InInputState, const FTransform& InDeltaTransform, const FVector& InPivotLocation) const { return AssetEditorViewportInteractionCustomization->GizmoManipulationDeltaUpdate(ElementWorldHandle, InWidgetMode, InDragAxis, InInputState, InDeltaTransform, InPivotLocation); }
+	void GizmoManipulationStopped(const UE::Widget::EWidgetMode InWidgetMode) const { return AssetEditorViewportInteractionCustomization->GizmoManipulationStopped(ElementWorldHandle, InWidgetMode); }
+
+private:
+	TTypedElement<UTypedElementWorldInterface> ElementWorldHandle;
+	FTypedElementAssetEditorViewportInteractionCustomization* AssetEditorViewportInteractionCustomization = nullptr;
+};
+
+/**
+ * A utility to handle higher-level viewport interactions, but default via UTypedElementWorldInterface,
+ * but asset editors can customize this behavior via FTypedElementAssetEditorViewportInteractionCustomization.
+ */
+UCLASS(Transient)
+class EDITORFRAMEWORK_API UTypedElementViewportInteraction : public UObject, public TTypedElementAssetEditorCustomizationRegistry<FTypedElementAssetEditorViewportInteractionCustomization>
+{
+	GENERATED_BODY()
+
+public:
+	/**
+	 * Get the elements from the given selection set that can be moved (eg, by a gizmo).
+	 */
+	void GetSelectedElementsToMove(const UTypedElementSelectionSet* InSelectionSet, const ETypedElementViewportInteractionWorldType InWorldType, UTypedElementList* OutElementsToMove) const;
+
+	/**
+	 * Notify that the gizmo is potentially about to start manipulating the transform of the given set of objects (calculated from calling GetSelectedElementsToMove).
+	 */
+	void BeginGizmoManipulation(const UTypedElementList* InElementsToMove, const UE::Widget::EWidgetMode InWidgetMode);
+
+	/**
+	 * Notify that the gizmo has manipulated the transform of the given set of objects (calculated from calling GetSelectedElementsToMove) by the given delta.
+	 */
+	void UpdateGizmoManipulation(const UTypedElementList* InElementsToMove, const UE::Widget::EWidgetMode InWidgetMode, const EAxisList::Type InDragAxis, const FInputDeviceState& InInputState, const FTransform& InDeltaTransform);
+	
+	/**
+	 * Notify that the gizmo has finished manipulating the transform of the given set of objects (calculated from calling GetSelectedElementsToMove).
+	 */
+	void EndGizmoManipulation(const UTypedElementList* InElementsToMove, const UE::Widget::EWidgetMode InWidgetMode);
+
+	/**
+	 * Apply the given delta to the specified element.
+	 * This performs a similar operation to an in-flight gizmo manipulation, but without any pre/post-change notification.
+	 */
+	void ApplyDeltaToElement(const FTypedElementHandle& InElementHandle, const UE::Widget::EWidgetMode InWidgetMode, const EAxisList::Type InDragAxis, const FInputDeviceState& InInputState, const FTransform& InDeltaTransform);
+
+private:
+	/**
+	 * Attempt to resolve the selection interface and asset editor viewport interaction customization for the given element, if any.
+	 */
+	FTypedElementViewportInteractionElement ResolveViewportInteractionElement(const FTypedElementHandle& InElementHandle) const;
+
+	/**
+	 * Batch the given elements by their type.
+	 */
+	static void BatchElementsByType(const UTypedElementList* InElementsToMove, TMap<FTypedHandleTypeId, TArray<TTypedElement<UTypedElementWorldInterface>>>& OutElementsByType);
+};
diff --git a/Engine/Source/Editor/EditorFramework/Public/UnrealWidgetFwd.h b/Engine/Source/Editor/EditorFramework/Public/UnrealWidgetFwd.h
new file mode 100644
--- /dev/null
+++ b/Engine/Source/Editor/EditorFramework/Public/UnrealWidgetFwd.h
@@ -0,0 +1,29 @@
+// Copyright Epic Games, Inc. All Rights Reserved.
+
+#pragma once
+
+class FWidget;
+
+/** Coordinate system identifiers. */
+enum ECoordSystem
+{
+	COORD_None = -1,
+	COORD_World,
+	COORD_Local,
+	COORD_Max,
+};
+
+namespace UE {
+namespace Widget {
+	enum EWidgetMode
+	{
+		WM_None = -1,
+		WM_Translate,
+		WM_TranslateRotateZ,
+		WM_2D,
+		WM_Rotate,
+		WM_Scale,
+		WM_Max,
+	};
+}
+}
diff --git a/Engine/Source/Editor/LevelEditor/Private/Elements/Actor/ActorElementLevelEditorViewportInteractionCustomization.cpp b/Engine/Source/Editor/LevelEditor/Private/Elements/Actor/ActorElementLevelEditorViewportInteractionCustomization.cpp
new file mode 100644
--- /dev/null
+++ b/Engine/Source/Editor/LevelEditor/Private/Elements/Actor/ActorElementLevelEditorViewportInteractionCustomization.cpp
@@ -0,0 +1,179 @@
+// Copyright Epic Games, Inc. All Rights Reserved.
+
+#include "Elements/Actor/ActorElementLevelEditorViewportInteractionCustomization.h"
+#include "Elements/Actor/ActorElementData.h"
+#include "GameFramework/Actor.h"
+
+#include "Elements/Framework/EngineElementsLibrary.h"
+#include "Elements/Framework/TypedElementSelectionSet.h"
+#include "Elements/Component/ComponentElementSelectionInterface.h"
+#include "Elements/Component/ComponentElementLevelEditorViewportInteractionCustomization.h"
+
+#include "Editor.h"
+#include "Engine/Brush.h"
+#include "Editor/GroupActor.h"
+#include "ActorGroupingUtils.h"
+#include "LevelEditorViewport.h"
+
+FActorElementLevelEditorViewportInteractionCustomization::FActorElementLevelEditorViewportInteractionCustomization(FLevelEditorViewportClient* InLevelEditorViewportClient)
+	: LevelEditorViewportClient(InLevelEditorViewportClient)
+{
+}
+
+void FActorElementLevelEditorViewportInteractionCustomization::GetElementsToMove(const TTypedElement<UTypedElementWorldInterface>& InElementWorldHandle, const ETypedElementViewportInteractionWorldType InWorldType, const UTypedElementSelectionSet* InSelectionSet, UTypedElementList* OutElementsToMove)
+{
+	AActor* Actor = InElementWorldHandle.GetDataChecked<FActorElementData>().Actor;
+
+	if (UComponentElementSelectionInterface::HasSelectedComponents(InSelectionSet->GetElementList()))
+	{
+		// If we have components selected then we will move those rather than the actors
+		// The component may still choose to move its owner actor rather than itself
+		return;
+	}
+
+	if (CanMoveActorInViewport(Actor, InWorldType))
+	{
+		AppendActorsToMove(Actor, InSelectionSet, OutElementsToMove);
+	}
+}
+
+void FActorElementLevelEditorViewportInteractionCustomization::GizmoManipulationStarted(const TTypedElement<UTypedElementWorldInterface>& InElementWorldHandle, const UE::Widget::EWidgetMode InWidgetMode)
+{
+	AActor* Actor = InElementWorldHandle.GetDataChecked<FActorElementData>().Actor;
+
+	// Notify that this actor is beginning to move
+	GEditor->BroadcastBeginObjectMovement(*Actor);
+
+	// Broadcast Pre Edit change notification, we can't call PreEditChange directly on Actor or ActorComponent from here since it will unregister the components until PostEditChange
+	if (FProperty* TransformProperty = FComponentElementLevelEditorViewportInteractionCustomization::GetEditTransformProperty(InWidgetMode))
+	{
+		FEditPropertyChain PropertyChain;
+		PropertyChain.AddHead(TransformProperty);
+		FCoreUObjectDelegates::OnPreObjectPropertyChanged.Broadcast(Actor, PropertyChain);
+	}
+}
+
+void FActorElementLevelEditorViewportInteractionCustomization::GizmoManipulationDeltaUpdate(const TTypedElement<UTypedElementWorldInterface>& InElementWorldHandle, const UE::Widget::EWidgetMode InWidgetMode, const EAxisList::Type InDragAxis, const FInputDeviceState& InInputState, const FTransform& InDeltaTransform, const FVector& InPivotLocation)
+{
+	AActor* Actor = InElementWorldHandle.GetDataChecked<FActorElementData>().Actor;
+
+	GEditor->NoteActorMovement();
+
+	FTransform ModifiedDeltaTransform = InDeltaTransform;
+
+	{
+		FVector AdjustedScale = ModifiedDeltaTransform.GetScale3D();
+
+		// If we are scaling, we may need to change the scaling factor a bit to properly align to the grid
+		if (AdjustedScale.IsNearlyZero())
+		{
+			// We don't scale actors when we only have a very small scale change
+			AdjustedScale = FVector::ZeroVector;
+		}
+		else if (!GEditor->UsePercentageBasedScaling())
+		{
+			ModifyScale(Actor, InDragAxis, AdjustedScale, Actor->IsA<ABrush>());
+		}
+
+		ModifiedDeltaTransform.SetScale3D(AdjustedScale);
+	}
+
+	FActorElementEditorViewportInteractionCustomization::GizmoManipulationDeltaUpdate(InElementWorldHandle, InWidgetMode, InDragAxis, InInputState, ModifiedDeltaTransform, InPivotLocation);
+
+	// Update the cameras from their locked actor (if any) only if the viewport is real-time enabled
+	LevelEditorViewportClient->UpdateLockedActorViewports(Actor, true);
+}
+
+void FActorElementLevelEditorViewportInteractionCustomization::GizmoManipulationStopped(const TTypedElement<UTypedElementWorldInterface>& InElementWorldHandle, const UE::Widget::EWidgetMode InWidgetMode)
+{
+	AActor* Actor = InElementWorldHandle.GetDataChecked<FActorElementData>().Actor;
+
+	// Broadcast Post Edit change notification, we can't call PostEditChangeProperty directly on Actor or ActorComponent from here since it wasn't pair with a proper PreEditChange
+	if (FProperty* TransformProperty = FComponentElementLevelEditorViewportInteractionCustomization::GetEditTransformProperty(InWidgetMode))
+	{
+		FPropertyChangedEvent PropertyChangedEvent(TransformProperty, EPropertyChangeType::ValueSet);
+		FCoreUObjectDelegates::OnObjectPropertyChanged.Broadcast(Actor, PropertyChangedEvent);
+	}
+
+	Actor->PostEditMove(true);
+	GEditor->BroadcastEndObjectMovement(*Actor);
+}
+
+void FActorElementLevelEditorViewportInteractionCustomization::PostGizmoManipulationStopped(TArrayView<const TTypedElement<UTypedElementWorldInterface>> InElementWorldHandles, const UE::Widget::EWidgetMode InWidgetMode)
+{
+	TArray<AActor*> MovedActors;
+	MovedActors.Reserve(InElementWorldHandles.Num());
+	for (const TTypedElement<UTypedElementWorldInterface>& ElementWorldHandle : InElementWorldHandles)
+	{
+		MovedActors.Add(ElementWorldHandle.GetDataChecked<FActorElementData>().Actor);
+	}
+	GEditor->BroadcastActorsMoved(MovedActors);
+}
+
+void FActorElementLevelEditorViewportInteractionCustomization::ModifyScale(AActor* InActor, const EAxisList::Type InDragAxis, FVector& ScaleDelta, bool bCheckSmallExtent) const
+{
+	if (InActor->GetRootComponent())
+	{
+		const FVector CurrentScale = InActor->GetRootComponent()->GetRelativeScale3D();
+
+		const FBox LocalBox = InActor->GetComponentsBoundingBox(true);
+		const FVector ScaledExtents = LocalBox.GetExtent() * CurrentScale;
+		const FTransform PreDragTransform = LevelEditorViewportClient->CachePreDragActorTransform(InActor);
+
+		FComponentElementLevelEditorViewportInteractionCustomization::ValidateScale(PreDragTransform.GetScale3D(), InDragAxis, CurrentScale, ScaledExtents, ScaleDelta, bCheckSmallExtent);
+
+		if (ScaleDelta.IsNearlyZero())
+		{
+			ScaleDelta = FVector::ZeroVector;
+		}
+	}
+}
+
+bool FActorElementLevelEditorViewportInteractionCustomization::CanMoveActorInViewport(const AActor* InActor, const ETypedElementViewportInteractionWorldType InWorldType)
+{
+	if (!GEditor || !InActor)
+	{
+		return false;
+	}
+
+	// The actor cannot be location locked
+	if (InActor->bLockLocation)
+	{
+		return false;
+	}
+
+	// The actor needs to be in the current viewport world
+	if (GEditor->PlayWorld)
+	{
+		const UWorld* CurrentWorld = InActor->GetWorld();
+		const UWorld* RequiredWorld = InWorldType == ETypedElementViewportInteractionWorldType::PlayInEditor ? GEditor->PlayWorld : GEditor->EditorWorld;
+		if (CurrentWorld != RequiredWorld)
+		{
+			return false;
+		}
+	}
+
+	return true;
+}
+
+void FActorElementLevelEditorViewportInteractionCustomization::AppendActorsToMove(AActor* InActor, const UTypedElementSelectionSet* InSelectionSet, UTypedElementList* OutElementsToMove)
+{
+	auto AddActorElement = [&OutElementsToMove](AActor* InActorToAdd)
+	{
+		FTypedElementHandle ActorElementHandle = UEngineElementsLibrary::AcquireEditorActorElementHandle(InActorToAdd);
+		if (ActorElementHandle)
+		{
+			OutElementsToMove->Add(MoveTemp(ActorElementHandle));
+		}
+	};
+
+	AGroupActor* ParentGroup = AGroupActor::GetRootForActor(InActor, true, true);
+	if (ParentGroup && UActorGroupingUtils::IsGroupingActive())
+	{
+		ParentGroup->ForEachMovableActorInGroup(InSelectionSet, AddActorElement);
+	}
+	else
+	{
+		AddActorElement(InActor);
+	}
+}
diff --git a/Engine/Source/Editor/LevelEditor/Private/Elements/Component/ComponentElementLevelEditorViewportInteractionCustomization.cpp b/Engine/Source/Editor/LevelEditor/Private/Elements/Component/ComponentElementLevelEditorViewportInteractionCustomization.cpp
new file mode 100644
--- /dev/null
+++ b/Engine/Source/Editor/LevelEditor/Private/Elements/Component/ComponentElementLevelEditorViewportInteractionCustomization.cpp
@@ -0,0 +1,355 @@
+// Copyright Epic Games, Inc. All Rights Reserved.
+
+#include "Elements/Component/ComponentElementLevelEditorViewportInteractionCustomization.h"
+#include "Elements/Component/ComponentElementData.h"
+#include "Components/SceneComponent.h"
+
+#include "Elements/Framework/EngineElementsLibrary.h"
+#include "Elements/Framework/TypedElementSelectionSet.h"
+#include "Elements/Actor/ActorElementLevelEditorViewportInteractionCustomization.h"
+
+#include "Editor.h"
+#include "LevelEditorViewport.h"
+#include "Kismet2/ComponentEditorUtils.h"
+
+FComponentElementLevelEditorViewportInteractionCustomization::FComponentElementLevelEditorViewportInteractionCustomization(FLevelEditorViewportClient* InLevelEditorViewportClient)
+	: LevelEditorViewportClient(InLevelEditorViewportClient)
+{
+}
+
+void FComponentElementLevelEditorViewportInteractionCustomization::GetElementsToMove(const TTypedElement<UTypedElementWorldInterface>& InElementWorldHandle, const ETypedElementViewportInteractionWorldType InWorldType, const UTypedElementSelectionSet* InSelectionSet, UTypedElementList* OutElementsToMove)
+{
+	UActorComponent* Component = InElementWorldHandle.GetDataChecked<FComponentElementData>().Component;
+
+	if (USceneComponent* SceneComponent = Cast<USceneComponent>(Component))
+	{
+		if (AActor* ComponentOwner = SceneComponent->GetOwner())
+		{
+			if (FActorElementLevelEditorViewportInteractionCustomization::CanMoveActorInViewport(ComponentOwner, InWorldType))
+			{
+				if (ComponentOwner->GetRootComponent() == SceneComponent)
+				{
+					// If it is a root component, use the parent actor instead
+					FActorElementLevelEditorViewportInteractionCustomization::AppendActorsToMove(ComponentOwner, InSelectionSet, OutElementsToMove);
+				}
+				else
+				{
+					// Check to see if any parent is selected
+					bool bHasSelectedParent = false;
+					for (USceneComponent* Parent = SceneComponent->GetAttachParent(); Parent && !bHasSelectedParent; Parent = Parent->GetAttachParent())
+					{
+						FTypedElementHandle ParentElementHandle = UEngineElementsLibrary::AcquireEditorComponentElementHandle(Parent, /*bAllowCreate*/false);
+						bHasSelectedParent = ParentElementHandle && InSelectionSet->IsElementSelected(ParentElementHandle, FTypedElementIsSelectedOptions());
+					}
+
+					if (!bHasSelectedParent)
+					{
+						// If no parent of this component is also in the selection set, move it
+						OutElementsToMove->Add(InElementWorldHandle);
+					}
+				}
+			}
+		}
+	}
+}
+
+void FComponentElementLevelEditorViewportInteractionCustomization::GizmoManipulationStarted(const TTypedElement<UTypedElementWorldInterface>& InElementWorldHandle, const UE::Widget::EWidgetMode InWidgetMode)
+{
+	UActorComponent* Component = InElementWorldHandle.GetDataChecked<FComponentElementData>().Component;
+
+	if (USceneComponent* SceneComponent = Cast<USceneComponent>(Component))
+	{
+		// Notify that this component is beginning to move
+		GEditor->BroadcastBeginObjectMovement(*SceneComponent);
+
+		// Broadcast Pre Edit change notification, we can't call PreEditChange directly on Actor or ActorComponent from here since it will unregister the components until PostEditChange
+		if (FProperty* TransformProperty = FComponentElementLevelEditorViewportInteractionCustomization::GetEditTransformProperty(InWidgetMode))
+		{
+			FEditPropertyChain PropertyChain;
+			PropertyChain.AddHead(TransformProperty);
+			FCoreUObjectDelegates::OnPreObjectPropertyChanged.Broadcast(SceneComponent, PropertyChain);
+		}
+	}
+}
+
+void FComponentElementLevelEditorViewportInteractionCustomization::GizmoManipulationDeltaUpdate(const TTypedElement<UTypedElementWorldInterface>& InElementWorldHandle, const UE::Widget::EWidgetMode InWidgetMode, const EAxisList::Type InDragAxis, const FInputDeviceState& InInputState, const FTransform& InDeltaTransform, const FVector& InPivotLocation)
+{
+	UActorComponent* Component = InElementWorldHandle.GetDataChecked<FComponentElementData>().Component;
+
+	if (USceneComponent* SceneComponent = Cast<USceneComponent>(Component))
+	{
+		FTransform ModifiedDeltaTransform = InDeltaTransform;
+
+		{
+			FVector AdjustedDrag = ModifiedDeltaTransform.GetTranslation();
+			FRotator AdjustedRot = ModifiedDeltaTransform.Rotator();
+			FVector AdjustedScale = ModifiedDeltaTransform.GetScale3D();
+
+			FComponentEditorUtils::AdjustComponentDelta(SceneComponent, AdjustedDrag, AdjustedRot);
+
+			// If we are scaling, we need to change the scaling factor a bit to properly align to grid
+			if (AdjustedScale.IsNearlyZero())
+			{
+				// We don't scale components when we only have a very small scale change
+				AdjustedScale = FVector::ZeroVector;
+			}
+			else if (!GEditor->UsePercentageBasedScaling())
+			{
+				ModifyScale(SceneComponent, InDragAxis, AdjustedScale);
+			}
+
+			ModifiedDeltaTransform.SetTranslation(AdjustedDrag);
+			ModifiedDeltaTransform.SetRotation(AdjustedRot.Quaternion());
+			ModifiedDeltaTransform.SetScale3D(AdjustedScale);
+		}
+
+		FComponentElementEditorViewportInteractionCustomization::GizmoManipulationDeltaUpdate(InElementWorldHandle, InWidgetMode, InDragAxis, InInputState, ModifiedDeltaTransform, InPivotLocation);
+	}
+}
+
+void FComponentElementLevelEditorViewportInteractionCustomization::GizmoManipulationStopped(const TTypedElement<UTypedElementWorldInterface>& InElementWorldHandle, const UE::Widget::EWidgetMode InWidgetMode)
+{
+	UActorComponent* Component = InElementWorldHandle.GetDataChecked<FComponentElementData>().Component;
+
+	if (USceneComponent* SceneComponent = Cast<USceneComponent>(Component))
+	{
+		// Broadcast Post Edit change notification, we can't call PostEditChangeProperty directly on Actor or ActorComponent from here since it wasn't pair with a proper PreEditChange
+		if (FProperty* TransformProperty = FComponentElementLevelEditorViewportInteractionCustomization::GetEditTransformProperty(InWidgetMode))
+		{
+			FPropertyChangedEvent PropertyChangedEvent(TransformProperty, EPropertyChangeType::ValueSet);
+			FCoreUObjectDelegates::OnObjectPropertyChanged.Broadcast(SceneComponent, PropertyChangedEvent);
+		}
+
+		SceneComponent->PostEditComponentMove(true);
+		GEditor->BroadcastEndObjectMovement(*SceneComponent);
+	}
+}
+
+void FComponentElementLevelEditorViewportInteractionCustomization::ModifyScale(USceneComponent* InComponent, const EAxisList::Type InDragAxis, FVector& ScaleDelta) const
+{
+	AActor* Actor = InComponent->GetOwner();
+	const FTransform PreDragTransform = LevelEditorViewportClient->CachePreDragActorTransform(Actor);
+	const FBox LocalBox = Actor->GetComponentsBoundingBox(true);
+	const FVector ScaledExtents = LocalBox.GetExtent() * InComponent->GetRelativeScale3D();
+
+	FComponentElementLevelEditorViewportInteractionCustomization::ValidateScale(PreDragTransform.GetScale3D(), InDragAxis, InComponent->GetRelativeScale3D(), ScaledExtents, ScaleDelta, /*bCheckSmallExtent*/false);
+
+	if (ScaleDelta.IsNearlyZero())
+	{
+		ScaleDelta = FVector::ZeroVector;
+	}
+}
+
+void FComponentElementLevelEditorViewportInteractionCustomization::ValidateScale(const FVector& InOriginalPreDragScale, const EAxisList::Type InDragAxis, const FVector& InCurrentScale, const FVector& InBoxExtent, FVector& InOutScaleDelta, bool bInCheckSmallExtent)
+{
+	static const float MIN_ACTOR_BOUNDS_EXTENT = 1.0f;
+
+	/** Convert the active Dragging Axis to per-axis flags */
+	auto CheckActiveAxes = [InDragAxis](bool bActiveAxes[3])
+	{
+		bActiveAxes[0] = bActiveAxes[1] = bActiveAxes[2] = false;
+		switch (InDragAxis)
+		{
+		default:
+		case EAxisList::None:
+			break;
+		case EAxisList::X:
+			bActiveAxes[0] = true;
+			break;
+		case EAxisList::Y:
+			bActiveAxes[1] = true;
+			break;
+		case EAxisList::Z:
+			bActiveAxes[2] = true;
+			break;
+		case EAxisList::XYZ:
+		case EAxisList::All:
+		case EAxisList::Screen:
+			bActiveAxes[0] = bActiveAxes[1] = bActiveAxes[2] = true;
+			break;
+		case EAxisList::XY:
+			bActiveAxes[0] = bActiveAxes[1] = true;
+			break;
+		case EAxisList::XZ:
+			bActiveAxes[0] = bActiveAxes[2] = true;
+			break;
+		case EAxisList::YZ:
+			bActiveAxes[1] = bActiveAxes[2] = true;
+			break;
+		}
+	};
+
+	/** Check scale criteria to see if this is allowed, returns modified absolute scale*/
+	auto CheckScaleValue = [](float ScaleDeltaToCheck, float CurrentScaleFactor, float CurrentExtent, bool bCheckSmallExtent, bool bSnap) -> float
+	{
+		float AbsoluteScaleValue = ScaleDeltaToCheck + CurrentScaleFactor;
+		if (bSnap)
+		{
+			AbsoluteScaleValue = FMath::GridSnap(AbsoluteScaleValue, GEditor->GetScaleGridSize());
+		}
+		// In some situations CurrentExtent can be 0 (eg: when scaling a plane in Z), this causes a divide by 0 that we need to avoid.
+		if (FMath::Abs(CurrentExtent) < KINDA_SMALL_NUMBER) {
+			return AbsoluteScaleValue;
+		}
+		float UnscaledExtent = CurrentExtent / CurrentScaleFactor;
+		float ScaledExtent = UnscaledExtent * AbsoluteScaleValue;
+
+		if ((FMath::Square(ScaledExtent)) > BIG_NUMBER)	// cant get too big...
+		{
+			return CurrentScaleFactor;
+		}
+		else if (bCheckSmallExtent &&
+			(FMath::Abs(ScaledExtent) < MIN_ACTOR_BOUNDS_EXTENT * 0.5f ||		// ...or too small (apply sign in this case)...
+				(CurrentScaleFactor < 0.0f) != (AbsoluteScaleValue < 0.0f)))	// ...also cant cross the zero boundary
+		{
+			return ((MIN_ACTOR_BOUNDS_EXTENT * 0.5) / UnscaledExtent) * (CurrentScaleFactor < 0.0f ? -1.0f : 1.0f);
+		}
+
+		return AbsoluteScaleValue;
+	};
+
+	/**
+	* If the "PreserveNonUniformScale" setting is enabled, this function will appropriately re-scale the scale delta so that
+	* proportions are preserved also when snapping.
+	* This function will modify the scale delta sign so that scaling is apply in the good direction when using multiple axis in same time.
+	* The function will not transform the scale delta in the case the scale delta is not uniform
+	* @param    InOriginalPreDragScale		The object's original scale
+	* @param	bActiveAxes					The axes that are active when scaling interactively.
+	* @param	InOutScaleDelta				The scale delta we are potentially transforming.
+	* @return true if the axes should be snapped individually, according to the snap setting (i.e. this function had no effect)
+	*/
+	auto ApplyScalingOptions = [](const FVector& InOriginalPreDragScale, const bool bActiveAxes[3], FVector& InOutScaleDelta)
+	{
+		int ActiveAxisCount = 0;
+		bool CurrentValueSameSign = true;
+		bool FirstSignPositive = true;
+		float MaxComponentSum = -1.0f;
+		int32 MaxAxisIndex = -1;
+		const ULevelEditorViewportSettings* ViewportSettings = GetDefault<ULevelEditorViewportSettings>();
+		bool SnapScaleAfter = ViewportSettings->SnapScaleEnabled;
+
+		//Found the number of active axis
+		//Found if we have to swap some sign
+		for (int Axis = 0; Axis < 3; ++Axis)
+		{
+			if (bActiveAxes[Axis])
+			{
+				bool CurrentValueIsZero = FMath::IsNearlyZero(InOriginalPreDragScale[Axis], SMALL_NUMBER);
+				//when the current value is zero we assume it is positive
+				bool IsCurrentValueSignPositive = CurrentValueIsZero ? true : InOriginalPreDragScale[Axis] > 0.0f;
+				if (ActiveAxisCount == 0)
+				{
+					//Set the first value when we find the first active axis
+					FirstSignPositive = IsCurrentValueSignPositive;
+				}
+				else
+				{
+					if (FirstSignPositive != IsCurrentValueSignPositive)
+					{
+						CurrentValueSameSign = false;
+					}
+				}
+				ActiveAxisCount++;
+			}
+		}
+
+		//If we scale more then one axis and
+		//we have to swap some sign
+		if (ActiveAxisCount > 1 && !CurrentValueSameSign)
+		{
+			//Change the scale delta to reflect the sign of the value
+			for (int Axis = 0; Axis < 3; ++Axis)
+			{
+				if (bActiveAxes[Axis])
+				{
+					bool CurrentValueIsZero = FMath::IsNearlyZero(InOriginalPreDragScale[Axis], SMALL_NUMBER);
+					//when the current value is zero we assume it is positive
+					bool IsCurrentValueSignPositive = CurrentValueIsZero ? true : InOriginalPreDragScale[Axis] > 0.0f;
+					InOutScaleDelta[Axis] = IsCurrentValueSignPositive ? InOutScaleDelta[Axis] : -(InOutScaleDelta[Axis]);
+				}
+			}
+		}
+
+		if (ViewportSettings->PreserveNonUniformScale)
+		{
+			for (int Axis = 0; Axis < 3; ++Axis)
+			{
+				if (bActiveAxes[Axis])
+				{
+					const float AbsScale = FMath::Abs(InOutScaleDelta[Axis] + InOriginalPreDragScale[Axis]);
+					if (AbsScale > MaxComponentSum)
+					{
+						MaxAxisIndex = Axis;
+						MaxComponentSum = AbsScale;
+					}
+				}
+			}
+
+			check(MaxAxisIndex != -1);
+
+			float AbsoluteScaleValue = InOriginalPreDragScale[MaxAxisIndex] + InOutScaleDelta[MaxAxisIndex];
+			if (ViewportSettings->SnapScaleEnabled)
+			{
+				AbsoluteScaleValue = FMath::GridSnap(InOriginalPreDragScale[MaxAxisIndex] + InOutScaleDelta[MaxAxisIndex], GEditor->GetScaleGridSize());
+				SnapScaleAfter = false;
+			}
+
+			float ScaleRatioMax = AbsoluteScaleValue / InOriginalPreDragScale[MaxAxisIndex];
+			for (int Axis = 0; Axis < 3; ++Axis)
+			{
+				if (bActiveAxes[Axis])
+				{
+					InOutScaleDelta[Axis] = (InOriginalPreDragScale[Axis] * ScaleRatioMax) - InOriginalPreDragScale[Axis];
+				}
+			}
+		}
+
+		return SnapScaleAfter;
+	};
+
+	/** Helper function for ModifyScale - Check scale criteria to see if this is allowed */
+
+	// get the axes that are active in this operation
+	bool bActiveAxes[3];
+	CheckActiveAxes(bActiveAxes);
+
+	//When scaling with more then one active axis, We must make sure we apply the correct delta sign to each delta scale axis
+	//We also want to support the PreserveNonUniformScale option
+	bool bSnapAxes = ApplyScalingOptions(InOriginalPreDragScale, bActiveAxes, InOutScaleDelta);
+
+	// check each axis
+	for (int Axis = 0; Axis < 3; ++Axis)
+	{
+		if (bActiveAxes[Axis])
+		{
+			float ModifiedScaleAbsolute = CheckScaleValue(InOutScaleDelta[Axis], InCurrentScale[Axis], InBoxExtent[Axis], bInCheckSmallExtent, bSnapAxes);
+			InOutScaleDelta[Axis] = ModifiedScaleAbsolute - InCurrentScale[Axis];
+		}
+		else
+		{
+			InOutScaleDelta[Axis] = 0.0f;
+		}
+	}
+}
+
+FProperty* FComponentElementLevelEditorViewportInteractionCustomization::GetEditTransformProperty(const UE::Widget::EWidgetMode InWidgetMode)
+{
+	switch (InWidgetMode)
+	{
+	case UE::Widget::WM_Translate:
+	case UE::Widget::WM_TranslateRotateZ:
+	case UE::Widget::WM_2D:
+		return FindFProperty<FProperty>(USceneComponent::StaticClass(), USceneComponent::GetRelativeLocationPropertyName());
+
+	case UE::Widget::WM_Rotate:
+		return FindFProperty<FProperty>(USceneComponent::StaticClass(), USceneComponent::GetRelativeRotationPropertyName());
+
+	case UE::Widget::WM_Scale:
+		return FindFProperty<FProperty>(USceneComponent::StaticClass(), USceneComponent::GetRelativeScale3DPropertyName());
+
+	default:
+		break;
+	}
+
+	return nullptr;
+}
diff --git a/Engine/Source/Editor/LevelEditor/Public/Elements/Actor/ActorElementLevelEditorViewportInteractionCustomization.h b/Engine/Source/Editor/LevelEditor/Public/Elements/Actor/ActorElementLevelEditorViewportInteractionCustomization.h
new file mode 100644
--- /dev/null
+++ b/Engine/Source/Editor/LevelEditor/Public/Elements/Actor/ActorElementLevelEditorViewportInteractionCustomization.h
@@ -0,0 +1,27 @@
+// Copyright Epic Games, Inc. All Rights Reserved.
+
+#pragma once
+
+#include "Elements/Actor/ActorElementEditorViewportInteractionCustomization.h"
+
+class FLevelEditorViewportClient;
+
+class LEVELEDITOR_API FActorElementLevelEditorViewportInteractionCustomization : public FActorElementEditorViewportInteractionCustomization
+{
+public:
+	explicit FActorElementLevelEditorViewportInteractionCustomization(FLevelEditorViewportClient* InLevelEditorViewportClient);
+
+	virtual void GetElementsToMove(const TTypedElement<UTypedElementWorldInterface>& InElementWorldHandle, const ETypedElementViewportInteractionWorldType InWorldType, const UTypedElementSelectionSet* InSelectionSet, UTypedElementList* OutElementsToMove) override;
+	virtual void GizmoManipulationStarted(const TTypedElement<UTypedElementWorldInterface>& InElementWorldHandle, const UE::Widget::EWidgetMode InWidgetMode) override;
+	virtual void GizmoManipulationDeltaUpdate(const TTypedElement<UTypedElementWorldInterface>& InElementWorldHandle, const UE::Widget::EWidgetMode InWidgetMode, const EAxisList::Type InDragAxis, const FInputDeviceState& InInputState, const FTransform& InDeltaTransform, const FVector& InPivotLocation) override;
+	virtual void GizmoManipulationStopped(const TTypedElement<UTypedElementWorldInterface>& InElementWorldHandle, const UE::Widget::EWidgetMode InWidgetMode) override;
+	virtual void PostGizmoManipulationStopped(TArrayView<const TTypedElement<UTypedElementWorldInterface>> InElementWorldHandles, const UE::Widget::EWidgetMode InWidgetMode) override;
+
+	static bool CanMoveActorInViewport(const AActor* InActor, const ETypedElementViewportInteractionWorldType InWorldType);
+	static void AppendActorsToMove(AActor* InActor, const UTypedElementSelectionSet* InSelectionSet, UTypedElementList* OutElementsToMove);
+
+private:
+	void ModifyScale(AActor* InActor, const EAxisList::Type InDragAxis, FVector& ScaleDelta, bool bCheckSmallExtent) const;
+
+	FLevelEditorViewportClient* LevelEditorViewportClient = nullptr;
+};
diff --git a/Engine/Source/Editor/LevelEditor/Public/Elements/Component/ComponentElementLevelEditorViewportInteractionCustomization.h b/Engine/Source/Editor/LevelEditor/Public/Elements/Component/ComponentElementLevelEditorViewportInteractionCustomization.h
new file mode 100644
--- /dev/null
+++ b/Engine/Source/Editor/LevelEditor/Public/Elements/Component/ComponentElementLevelEditorViewportInteractionCustomization.h
@@ -0,0 +1,26 @@
+// Copyright Epic Games, Inc. All Rights Reserved.
+
+#pragma once
+
+#include "Elements/Component/ComponentElementEditorViewportInteractionCustomization.h"
+
+class FLevelEditorViewportClient;
+
+class LEVELEDITOR_API FComponentElementLevelEditorViewportInteractionCustomization : public FComponentElementEditorViewportInteractionCustomization
+{
+public:
+	explicit FComponentElementLevelEditorViewportInteractionCustomization(FLevelEditorViewportClient* InLevelEditorViewportClient);
+
+	virtual void GetElementsToMove(const TTypedElement<UTypedElementWorldInterface>& InElementWorldHandle, const ETypedElementViewportInteractionWorldType InWorldType, const UTypedElementSelectionSet* InSelectionSet, UTypedElementList* OutElementsToMove) override;
+	virtual void GizmoManipulationStarted(const TTypedElement<UTypedElementWorldInterface>& InElementWorldHandle, const UE::Widget::EWidgetMode InWidgetMode) override;
+	virtual void GizmoManipulationDeltaUpdate(const TTypedElement<UTypedElementWorldInterface>& InElementWorldHandle, const UE::Widget::EWidgetMode InWidgetMode, const EAxisList::Type InDragAxis, const FInputDeviceState& InInputState, const FTransform& InDeltaTransform, const FVector& InPivotLocation) override;
+	virtual void GizmoManipulationStopped(const TTypedElement<UTypedElementWorldInterface>& InElementWorldHandle, const UE::Widget::EWidgetMode InWidgetMode) override;
+
+	static void ValidateScale(const FVector& InOriginalPreDragScale, const EAxisList::Type InDragAxis, const FVector& InCurrentScale, const FVector& InBoxExtent, FVector& InOutScaleDelta, bool bInCheckSmallExtent);
+	static FProperty* GetEditTransformProperty(const UE::Widget::EWidgetMode InWidgetMode);
+
+private:
+	void ModifyScale(USceneComponent* InComponent, const EAxisList::Type InDragAxis, FVector& ScaleDelta) const;
+
+	FLevelEditorViewportClient* LevelEditorViewportClient = nullptr;
+};
diff --git a/Engine/Source/Editor/UnrealEd/Classes/Editor/EditorEngine.h b/Engine/Source/Editor/UnrealEd/Classes/Editor/EditorEngine.h
--- a/Engine/Source/Editor/UnrealEd/Classes/Editor/EditorEngine.h
+++ b/Engine/Source/Editor/UnrealEd/Classes/Editor/EditorEngine.h
@@ -2104,6 +2104,11 @@ public:
 	 */
 	void DisableDeltaModification(bool bDisable) { bDisableDeltaModification = bDisable; }
 
+	/**
+	 * Test whether actor/component modification during delta movement is currently enabled?
+	 */
+	bool IsDeltaModificationEnabled() const { return !bDisableDeltaModification; }
+
 	/**
 	 *	Game-specific function called by Map_Check BEFORE iterating over all actors.
 	 *
diff --git a/Engine/Source/Editor/UnrealEd/Classes/Editor/GroupActor.h b/Engine/Source/Editor/UnrealEd/Classes/Editor/GroupActor.h
--- a/Engine/Source/Editor/UnrealEd/Classes/Editor/GroupActor.h
+++ b/Engine/Source/Editor/UnrealEd/Classes/Editor/GroupActor.h
@@ -13,6 +13,7 @@
 
 class FLevelEditorViewportClient;
 class FViewport;
+class UTypedElementSelectionSet;
 
 UCLASS(hidedropdown,MinimalAPI, notplaceable, NotBlueprintable)
 class AGroupActor : public AActor
@@ -54,13 +55,11 @@ class AGroupActor : public AActor
 	virtual bool IsSelected() const;
 
 	/**
-	 * Apply given deltas to all actors and subgroups for this group.
-	 * @param	Viewport		The viewport to draw to apply our deltas
-	 * @param	InDrag			Delta Transition
-	 * @param	InRot			Delta Rotation
-	 * @param	InScale			Delta Scale
+	 * Enumerate all actors and subgroups for this group that should be moved alongside this actor (includes this).
+	 * @param InSelectionSet	The selection set to use, or null to use the global selection set
+	 * @param InCallback		The function to call for each movable actor
 	 */
-	void GroupApplyDelta(FLevelEditorViewportClient* Viewport, const FVector& InDrag, const FRotator& InRot, const FVector& InScale );
+	UNREALED_API void ForEachMovableActorInGroup(const UTypedElementSelectionSet* InSelectionSet, TFunctionRef<void(AActor*)> InCallback);
 	
 	/**
 	 * Apply given deltas to all actors and subgroups for this group.
@@ -70,7 +69,6 @@ class AGroupActor : public AActor
 	 */
 	UNREALED_API void GroupApplyDelta(const FVector& InDrag, const FRotator& InRot, const FVector& InScale );
 
-
 	/**
 	 * Draw brackets around all selected groups
 	 * @param	PDI				FPrimitiveDrawInterface used to draw lines in active viewports
diff --git a/Engine/Source/Editor/UnrealEd/Private/EditorEngine.cpp b/Engine/Source/Editor/UnrealEd/Private/EditorEngine.cpp
--- a/Engine/Source/Editor/UnrealEd/Private/EditorEngine.cpp
+++ b/Engine/Source/Editor/UnrealEd/Private/EditorEngine.cpp
@@ -131,6 +131,8 @@
 
 #include "Elements/Framework/TypedElementSelectionSet.h"
 #include "Elements/Object/ObjectElementEditorSelectionCustomization.h"
+#include "Elements/Actor/ActorElementEditorViewportInteractionCustomization.h"
+#include "Elements/Component/ComponentElementEditorViewportInteractionCustomization.h"
 
 #include "IDirectoryWatcher.h"
 #include "DirectoryWatcherModule.h"
@@ -2617,166 +2619,10 @@ void UEditorEngine::ApplyDeltaToActor(AActor* InActor,
 									  bool bShiftDown,
 									  bool bControlDown) const
 {
-	if(!bDisableDeltaModification)
-	{
-		InActor->Modify();
-	}
-	
-	FNavigationLockContext LockNavigationUpdates(InActor->GetWorld(), ENavigationLockReason::ContinuousEditorMove);
-
-	bool bTranslationOnly = true;
-
-	///////////////////
-	// Rotation
-
-	// Unfortunately this can't be moved into ABrush::EditorApplyRotation, as that would
-	// create a dependence in Engine on Editor.
-	if ( InRot )
-	{
-		const FRotator& InDeltaRot = *InRot;
-		const bool bRotatingActor = !bDelta || !InDeltaRot.IsZero();
-		if( bRotatingActor )
-		{
-			bTranslationOnly = false;
-
-			if ( bDelta )
-			{
-				if( InActor->GetRootComponent() != NULL )
-				{
-					const FRotator OriginalRotation = InActor->GetRootComponent()->GetComponentRotation();
-
-					InActor->EditorApplyRotation( InDeltaRot, bAltDown, bShiftDown, bControlDown );
-
-					// Check to see if we should transform the rigid body
-					UPrimitiveComponent* RootPrimitiveComponent = Cast< UPrimitiveComponent >( InActor->GetRootComponent() );
-					if( bIsSimulatingInEditor && GIsPlayInEditorWorld && RootPrimitiveComponent != NULL )
-					{
-						FRotator ActorRotWind, ActorRotRem;
-						OriginalRotation.GetWindingAndRemainder(ActorRotWind, ActorRotRem);
-
-						const FQuat ActorQ = ActorRotRem.Quaternion();
-						const FQuat DeltaQ = InDeltaRot.Quaternion();
-						const FQuat ResultQ = DeltaQ * ActorQ;
-
-						const FRotator NewActorRotRem = FRotator( ResultQ );
-						FRotator DeltaRot = NewActorRotRem - ActorRotRem;
-						DeltaRot.Normalize();
-
-						// @todo SIE: Not taking into account possible offset between root component and actor
-						RootPrimitiveComponent->SetWorldRotation( OriginalRotation + DeltaRot );
-					}
-				}
-
-				FVector NewActorLocation = InActor->GetActorLocation();
-				NewActorLocation -= GLevelEditorModeTools().PivotLocation;
-				NewActorLocation = FRotationMatrix(InDeltaRot).TransformPosition(NewActorLocation);
-				NewActorLocation += GLevelEditorModeTools().PivotLocation;
-				NewActorLocation -= InActor->GetActorLocation();
-				InActor->EditorApplyTranslation(NewActorLocation, bAltDown, bShiftDown, bControlDown);
-			}
-			else
-			{
-				InActor->SetActorRotation( InDeltaRot );
-			}
-		}
-	}
-
-	///////////////////
-	// Translation
-	if ( InTrans )
-	{
-		if ( bDelta )
-		{
-			if( InActor->GetRootComponent() != NULL )
-			{
-				const FVector OriginalLocation = InActor->GetRootComponent()->GetComponentLocation();
-
-				InActor->EditorApplyTranslation( *InTrans, bAltDown, bShiftDown, bControlDown );
-
-				// Check to see if we should transform the rigid body
-				UPrimitiveComponent* RootPrimitiveComponent = Cast< UPrimitiveComponent >( InActor->GetRootComponent() );
-				if( bIsSimulatingInEditor && GIsPlayInEditorWorld && RootPrimitiveComponent != NULL )
-				{
-					// @todo SIE: Not taking into account possible offset between root component and actor
-					RootPrimitiveComponent->SetWorldLocation( OriginalLocation + *InTrans );
-				}
-			}
-		}
-		else
-		{
-			InActor->SetActorLocation( *InTrans, false );
-		}
-	}
-
-	///////////////////
-	// Scaling
-	if ( InScale )
-	{
-		const FVector& InDeltaScale = *InScale;
-		const bool bScalingActor = !bDelta || !InDeltaScale.IsNearlyZero(0.000001f);
-		if( bScalingActor )
-		{
-			bTranslationOnly = false;
-
-			FVector ModifiedScale = InDeltaScale;
-
-			// Note: With the new additive scaling method, this is handled in FLevelEditorViewportClient::ModifyScale
-			if( UsePercentageBasedScaling() )
-			{
-				// Get actor box extents
-				const FBox BoundingBox = InActor->GetComponentsBoundingBox( true );
-				const FVector BoundsExtents = BoundingBox.GetExtent();
-
-				// Make sure scale on actors is clamped to a minimum and maximum size.
-				const float MinThreshold = 1.0f;
-
-				for (int32 Idx=0; Idx<3; Idx++)
-				{
-					if ( ( FMath::Pow(BoundsExtents[Idx], 2) ) > BIG_NUMBER)
-					{
-						ModifiedScale[Idx] = 0.0f;
-					}
-					else if (SMALL_NUMBER < BoundsExtents[Idx])
-					{
-						const bool bBelowAllowableScaleThreshold = ((InDeltaScale[Idx] + 1.0f) * BoundsExtents[Idx]) < MinThreshold;
-
-						if(bBelowAllowableScaleThreshold)
-						{
-							ModifiedScale[Idx] = (MinThreshold / BoundsExtents[Idx]) - 1.0f;
-						}
-					}
-				}
-			}
-
-			if ( bDelta )
-			{
-				// Flag actors to use old-style scaling or not
-				// @todo: Remove this hack once we have decided on the scaling method to use.
-				AActor::bUsePercentageBasedScaling = UsePercentageBasedScaling();
-
-				InActor->EditorApplyScale( 
-					ModifiedScale,
-					&GLevelEditorModeTools().PivotLocation,
-					bAltDown,
-					bShiftDown,
-					bControlDown
-					);
-
-			}
-			else if( InActor->GetRootComponent() != NULL )
-			{
-				InActor->GetRootComponent()->SetRelativeScale3D( InDeltaScale );
-			}
-		}
-	}
+	FInputDeviceState InputState;
+	InputState.SetModifierKeyStates(bShiftDown, bAltDown, bControlDown, false);
 
-	// Update the actor before leaving.
-	InActor->MarkPackageDirty();
-	if (!GIsDemoMode)
-	{
-		InActor->InvalidateLightingCacheDetailed(bTranslationOnly);
-	}
-	InActor->PostEditMove( false );
+	FActorElementEditorViewportInteractionCustomization::ApplyDeltaToActor(InActor, bDelta, InTrans, InRot, InScale, GLevelEditorModeTools().PivotLocation, InputState);
 }
 
 void UEditorEngine::ApplyDeltaToComponent(USceneComponent* InComponent,
@@ -2786,99 +2632,7 @@ void UEditorEngine::ApplyDeltaToComponent(USceneComponent* InComponent,
 	const FVector* InScale,
 	const FVector& PivotLocation ) const
 {
-	if(!bDisableDeltaModification)
-	{
-		InComponent->Modify();
-	}
-
-	///////////////////
-	// Rotation
-	if ( InRot )
-	{
-		const FRotator& InDeltaRot = *InRot;
-		const bool bRotatingComp = !bDelta || !InDeltaRot.IsZero();
-		if( bRotatingComp )
-		{
-			if ( bDelta )
-			{
-				const FRotator Rot = InComponent->GetRelativeRotation();
-				FRotator ActorRotWind, ActorRotRem;
-				Rot.GetWindingAndRemainder(ActorRotWind, ActorRotRem);
-				const FQuat ActorQ = ActorRotRem.Quaternion();
-				const FQuat DeltaQ = InDeltaRot.Quaternion();
-				const FQuat ResultQ = DeltaQ * ActorQ;
-
-				FRotator NewActorRotRem = FRotator(ResultQ);
-				ActorRotRem.SetClosestToMe(NewActorRotRem);
-				FRotator DeltaRot = NewActorRotRem - ActorRotRem;
-				DeltaRot.Normalize();
-				InComponent->SetRelativeRotationExact(Rot + DeltaRot);
-			}
-			else
-			{
-				InComponent->SetRelativeRotationExact( InDeltaRot );
-			}
-
-			if ( bDelta )
-			{
-				FVector NewCompLocation = InComponent->GetRelativeLocation();
-				NewCompLocation -= PivotLocation;
-				NewCompLocation = FRotationMatrix( InDeltaRot ).TransformPosition( NewCompLocation );
-				NewCompLocation += PivotLocation;
-				InComponent->SetRelativeLocation(NewCompLocation);
-			}
-		}
-	}
-
-	///////////////////
-	// Translation
-	if ( InTrans )
-	{
-		if ( bDelta )
-		{
-			InComponent->SetRelativeLocation(InComponent->GetRelativeLocation() + *InTrans);
-		}
-		else
-		{
-			InComponent->SetRelativeLocation( *InTrans );
-		}
-	}
-
-	///////////////////
-	// Scaling
-	if ( InScale )
-	{
-		const FVector& InDeltaScale = *InScale;
-		const bool bScalingComp = !bDelta || !InDeltaScale.IsNearlyZero(0.000001f);
-		if( bScalingComp )
-		{
-			if ( bDelta )
-			{
-				InComponent->SetRelativeScale3D(InComponent->GetRelativeScale3D() + InDeltaScale);
-
-				FVector NewCompLocation = InComponent->GetRelativeLocation();
-				NewCompLocation -= PivotLocation;
-				NewCompLocation += FScaleMatrix( InDeltaScale ).TransformPosition( NewCompLocation );
-				NewCompLocation += PivotLocation;
-				InComponent->SetRelativeLocation(NewCompLocation);
-			}
-			else
-			{
-				InComponent->SetRelativeScale3D(InDeltaScale);
-			}
-		}
-	}
-
-	// Update the actor before leaving.
-	InComponent->MarkPackageDirty();
-
-	{
-		InComponent->PostEditComponentMove(false);
-	}
-
-	// Fire callbacks
-	FEditorSupportDelegates::RefreshPropertyWindows.Broadcast();
-	FEditorSupportDelegates::UpdateUI.Broadcast();
+	FComponentElementEditorViewportInteractionCustomization::ApplyDeltaToComponent(InComponent, bDelta, InTrans, InRot, InScale, PivotLocation, FInputDeviceState());
 }
 
 
diff --git a/Engine/Source/Editor/UnrealEd/Private/EditorViewportClient.cpp b/Engine/Source/Editor/UnrealEd/Private/EditorViewportClient.cpp
--- a/Engine/Source/Editor/UnrealEd/Private/EditorViewportClient.cpp
+++ b/Engine/Source/Editor/UnrealEd/Private/EditorViewportClient.cpp
@@ -49,6 +49,7 @@
 #include "EditorBuildUtils.h"
 #include "AudioDevice.h"
 #include "Editor/EditorPerformanceSettings.h"
+#include "Elements/Framework/TypedElementViewportInteraction.h"
 #include "ImageWriteQueue.h"
 #include "DebugViewModeHelpers.h"
 #include "RayTracingDebugVisualizationMenuCommands.h"
@@ -437,6 +438,7 @@ FEditorViewportClient::FEditorViewportClient(FEditorModeTools* InModeTools, FPre
 	, bIsCameraMoving(false)
 	, bIsCameraMovingOnTick(false)
 	, EditorViewportWidget(InEditorViewportWidget)
+	, ViewportInteraction(NewObject<UTypedElementViewportInteraction>())
 	, PreviewScene(InPreviewScene)
 	, MovingPreviewLightSavedScreenPos(ForceInitToZero)
 	, MovingPreviewLightTimer(0.0f)
@@ -3487,6 +3489,8 @@ void FEditorViewportClient::OnChangeCameraSpeed( const struct FInputEventState&
 
 void FEditorViewportClient::AddReferencedObjects( FReferenceCollector& Collector )
 {
+	Collector.AddReferencedObject(ViewportInteraction);
+
 	if( PreviewScene )
 	{
 		PreviewScene->AddReferencedObjects( Collector );
diff --git a/Engine/Source/Editor/UnrealEd/Private/Elements/Actor/ActorElementEditorViewportInteractionCustomization.cpp b/Engine/Source/Editor/UnrealEd/Private/Elements/Actor/ActorElementEditorViewportInteractionCustomization.cpp
new file mode 100644
--- /dev/null
+++ b/Engine/Source/Editor/UnrealEd/Private/Elements/Actor/ActorElementEditorViewportInteractionCustomization.cpp
@@ -0,0 +1,192 @@
+// Copyright Epic Games, Inc. All Rights Reserved.
+
+#include "Elements/Actor/ActorElementEditorViewportInteractionCustomization.h"
+#include "Elements/Actor/ActorElementData.h"
+#include "GameFramework/Actor.h"
+
+#include "Editor.h"
+#include "EditorModeManager.h"
+#include "AI/NavigationSystemBase.h"
+
+bool FActorElementEditorViewportInteractionCustomization::GetGizmoPivotLocation(const TTypedElement<UTypedElementWorldInterface>& InElementWorldHandle, const UE::Widget::EWidgetMode InWidgetMode, FVector& OutPivotLocation)
+{
+	OutPivotLocation = GLevelEditorModeTools().PivotLocation;
+	return true;
+}
+
+void FActorElementEditorViewportInteractionCustomization::GizmoManipulationDeltaUpdate(const TTypedElement<UTypedElementWorldInterface>& InElementWorldHandle, const UE::Widget::EWidgetMode InWidgetMode, const EAxisList::Type InDragAxis, const FInputDeviceState& InInputState, const FTransform& InDeltaTransform, const FVector& InPivotLocation)
+{
+	AActor* Actor = InElementWorldHandle.GetDataChecked<FActorElementData>().Actor;
+
+	const FVector DeltaTranslation = InDeltaTransform.GetTranslation();
+	const FRotator DeltaRotation = InDeltaTransform.Rotator();
+	const FVector DeltaScale3D = InDeltaTransform.GetScale3D();
+
+	FActorElementEditorViewportInteractionCustomization::ApplyDeltaToActor(Actor, /*bDelta*/true, &DeltaTranslation, &DeltaRotation, &DeltaScale3D, InPivotLocation, InInputState);
+}
+
+void FActorElementEditorViewportInteractionCustomization::ApplyDeltaToActor(AActor* InActor, const bool InIsDelta, const FVector* InDeltaTranslationPtr, const FRotator* InDeltaRotationPtr, const FVector* InDeltaScalePtr, const FVector& InPivotLocation, const FInputDeviceState& InInputState)
+{
+	const bool bIsSimulatingInEditor = GEditor->IsSimulatingInEditor();
+
+	if (GEditor->IsDeltaModificationEnabled())
+	{
+		InActor->Modify();
+	}
+
+	FNavigationLockContext LockNavigationUpdates(InActor->GetWorld(), ENavigationLockReason::ContinuousEditorMove);
+
+	bool bTranslationOnly = true;
+
+	///////////////////
+	// Rotation
+
+	// Unfortunately this can't be moved into ABrush::EditorApplyRotation, as that would
+	// create a dependence in Engine on Editor.
+	if (InDeltaRotationPtr)
+	{
+		const FRotator& InDeltaRot = *InDeltaRotationPtr;
+		const bool bRotatingActor = !InIsDelta || !InDeltaRot.IsZero();
+		if (bRotatingActor)
+		{
+			bTranslationOnly = false;
+
+			if (InIsDelta)
+			{
+				if (InActor->GetRootComponent())
+				{
+					const FRotator OriginalRotation = InActor->GetRootComponent()->GetComponentRotation();
+
+					InActor->EditorApplyRotation(InDeltaRot, InInputState.bAltKeyDown, InInputState.bShiftKeyDown, InInputState.bCtrlKeyDown);
+
+					// Check to see if we should transform the rigid body
+					UPrimitiveComponent* RootPrimitiveComponent = Cast<UPrimitiveComponent>(InActor->GetRootComponent());
+					if (bIsSimulatingInEditor && GIsPlayInEditorWorld && RootPrimitiveComponent)
+					{
+						FRotator ActorRotWind, ActorRotRem;
+						OriginalRotation.GetWindingAndRemainder(ActorRotWind, ActorRotRem);
+
+						const FQuat ActorQ = ActorRotRem.Quaternion();
+						const FQuat DeltaQ = InDeltaRot.Quaternion();
+						const FQuat ResultQ = DeltaQ * ActorQ;
+
+						const FRotator NewActorRotRem = FRotator(ResultQ);
+						FRotator DeltaRot = NewActorRotRem - ActorRotRem;
+						DeltaRot.Normalize();
+
+						// @todo SIE: Not taking into account possible offset between root component and actor
+						RootPrimitiveComponent->SetWorldRotation(OriginalRotation + DeltaRot);
+					}
+				}
+
+				FVector NewActorLocation = InActor->GetActorLocation();
+				NewActorLocation -= InPivotLocation;
+				NewActorLocation = FRotationMatrix(InDeltaRot).TransformPosition(NewActorLocation);
+				NewActorLocation += InPivotLocation;
+				NewActorLocation -= InActor->GetActorLocation();
+				InActor->EditorApplyTranslation(NewActorLocation, InInputState.bAltKeyDown, InInputState.bShiftKeyDown, InInputState.bCtrlKeyDown);
+			}
+			else
+			{
+				InActor->SetActorRotation(InDeltaRot);
+			}
+		}
+	}
+
+	///////////////////
+	// Translation
+	if (InDeltaTranslationPtr)
+	{
+		if (InIsDelta)
+		{
+			if (InActor->GetRootComponent())
+			{
+				const FVector OriginalLocation = InActor->GetRootComponent()->GetComponentLocation();
+
+				InActor->EditorApplyTranslation(*InDeltaTranslationPtr, InInputState.bAltKeyDown, InInputState.bShiftKeyDown, InInputState.bCtrlKeyDown);
+
+				// Check to see if we should transform the rigid body
+				UPrimitiveComponent* RootPrimitiveComponent = Cast<UPrimitiveComponent>(InActor->GetRootComponent());
+				if (bIsSimulatingInEditor && GIsPlayInEditorWorld && RootPrimitiveComponent)
+				{
+					// @todo SIE: Not taking into account possible offset between root component and actor
+					RootPrimitiveComponent->SetWorldLocation(OriginalLocation + *InDeltaTranslationPtr);
+				}
+			}
+		}
+		else
+		{
+			InActor->SetActorLocation(*InDeltaTranslationPtr, false);
+		}
+	}
+
+	///////////////////
+	// Scaling
+	if (InDeltaScalePtr)
+	{
+		const FVector& InDeltaScale = *InDeltaScalePtr;
+		const bool bScalingActor = !InIsDelta || !InDeltaScale.IsNearlyZero(0.000001f);
+		if (bScalingActor)
+		{
+			bTranslationOnly = false;
+
+			FVector ModifiedScale = InDeltaScale;
+
+			// Note: With the new additive scaling method, this is handled in FLevelEditorViewportClient::ModifyScale
+			if (GEditor->UsePercentageBasedScaling())
+			{
+				// Get actor box extents
+				const FBox BoundingBox = InActor->GetComponentsBoundingBox(true);
+				const FVector BoundsExtents = BoundingBox.GetExtent();
+
+				// Make sure scale on actors is clamped to a minimum and maximum size.
+				const float MinThreshold = 1.0f;
+
+				for (int32 Idx = 0; Idx < 3; Idx++)
+				{
+					if ((FMath::Pow(BoundsExtents[Idx], 2)) > BIG_NUMBER)
+					{
+						ModifiedScale[Idx] = 0.0f;
+					}
+					else if (SMALL_NUMBER < BoundsExtents[Idx])
+					{
+						const bool bBelowAllowableScaleThreshold = ((InDeltaScale[Idx] + 1.0f) * BoundsExtents[Idx]) < MinThreshold;
+
+						if (bBelowAllowableScaleThreshold)
+						{
+							ModifiedScale[Idx] = (MinThreshold / BoundsExtents[Idx]) - 1.0f;
+						}
+					}
+				}
+			}
+
+			if (InIsDelta)
+			{
+				// Flag actors to use old-style scaling or not
+				// @todo: Remove this hack once we have decided on the scaling method to use.
+				AActor::bUsePercentageBasedScaling = GEditor->UsePercentageBasedScaling();
+
+				InActor->EditorApplyScale(
+					ModifiedScale,
+					&InPivotLocation,
+					InInputState.bAltKeyDown,
+					InInputState.bShiftKeyDown,
+					InInputState.bCtrlKeyDown
+					);
+
+			}
+			else if (InActor->GetRootComponent())
+			{
+				InActor->GetRootComponent()->SetRelativeScale3D(InDeltaScale);
+			}
+		}
+	}
+
+	// Update the actor before leaving.
+	InActor->MarkPackageDirty();
+	if (!GIsDemoMode)
+	{
+		InActor->InvalidateLightingCacheDetailed(bTranslationOnly);
+	}
+	InActor->PostEditMove(false);
+}
diff --git a/Engine/Source/Editor/UnrealEd/Private/Elements/Component/ComponentElementEditorViewportInteractionCustomization.cpp b/Engine/Source/Editor/UnrealEd/Private/Elements/Component/ComponentElementEditorViewportInteractionCustomization.cpp
new file mode 100644
--- /dev/null
+++ b/Engine/Source/Editor/UnrealEd/Private/Elements/Component/ComponentElementEditorViewportInteractionCustomization.cpp
@@ -0,0 +1,133 @@
+// Copyright Epic Games, Inc. All Rights Reserved.
+
+#include "Elements/Component/ComponentElementEditorViewportInteractionCustomization.h"
+#include "Elements/Component/ComponentElementData.h"
+#include "Components/SceneComponent.h"
+
+#include "Editor.h"
+#include "EditorModeManager.h"
+#include "EditorSupportDelegates.h"
+
+bool FComponentElementEditorViewportInteractionCustomization::GetGizmoPivotLocation(const TTypedElement<UTypedElementWorldInterface>& InElementWorldHandle, const UE::Widget::EWidgetMode InWidgetMode, FVector& OutPivotLocation)
+{
+	UActorComponent* Component = InElementWorldHandle.GetDataChecked<FComponentElementData>().Component;
+
+	if (USceneComponent* SceneComponent = Cast<USceneComponent>(Component))
+	{
+		// If necessary, transform the editor pivot location to be relative to the component's parent
+		const bool bIsRootComponent = SceneComponent->GetOwner()->GetRootComponent() == SceneComponent;
+		OutPivotLocation = bIsRootComponent || !SceneComponent->GetAttachParent() ? GEditor->GetPivotLocation() : SceneComponent->GetAttachParent()->GetComponentToWorld().Inverse().TransformPosition(GEditor->GetPivotLocation());
+		return true;
+	}
+
+	return false;
+}
+
+void FComponentElementEditorViewportInteractionCustomization::GizmoManipulationDeltaUpdate(const TTypedElement<UTypedElementWorldInterface>& InElementWorldHandle, const UE::Widget::EWidgetMode InWidgetMode, const EAxisList::Type InDragAxis, const FInputDeviceState& InInputState, const FTransform& InDeltaTransform, const FVector& InPivotLocation)
+{
+	UActorComponent* Component = InElementWorldHandle.GetDataChecked<FComponentElementData>().Component;
+
+	if (USceneComponent* SceneComponent = Cast<USceneComponent>(Component))
+	{
+		const FVector DeltaTranslation = InDeltaTransform.GetTranslation();
+		const FRotator DeltaRotation = InDeltaTransform.Rotator();
+		const FVector DeltaScale3D = InDeltaTransform.GetScale3D();
+
+		FComponentElementEditorViewportInteractionCustomization::ApplyDeltaToComponent(SceneComponent, /*bDelta*/true, &DeltaTranslation, &DeltaRotation, &DeltaScale3D, InPivotLocation, InInputState);
+	}
+}
+
+void FComponentElementEditorViewportInteractionCustomization::ApplyDeltaToComponent(USceneComponent* InComponent, const bool InIsDelta, const FVector* InDeltaTranslationPtr, const FRotator* InDeltaRotationPtr, const FVector* InDeltaScalePtr, const FVector& InPivotLocation, const FInputDeviceState& InInputState)
+{
+	if (GEditor->IsDeltaModificationEnabled())
+	{
+		InComponent->Modify();
+	}
+
+	///////////////////
+	// Rotation
+	if (InDeltaRotationPtr)
+	{
+		const FRotator& InDeltaRot = *InDeltaRotationPtr;
+		const bool bRotatingComp = !InIsDelta || !InDeltaRot.IsZero();
+		if (bRotatingComp)
+		{
+			if (InIsDelta)
+			{
+				const FRotator Rot = InComponent->GetRelativeRotation();
+				FRotator ActorRotWind, ActorRotRem;
+				Rot.GetWindingAndRemainder(ActorRotWind, ActorRotRem);
+				const FQuat ActorQ = ActorRotRem.Quaternion();
+				const FQuat DeltaQ = InDeltaRot.Quaternion();
+				const FQuat ResultQ = DeltaQ * ActorQ;
+
+				FRotator NewActorRotRem = FRotator(ResultQ);
+				ActorRotRem.SetClosestToMe(NewActorRotRem);
+				FRotator DeltaRot = NewActorRotRem - ActorRotRem;
+				DeltaRot.Normalize();
+				InComponent->SetRelativeRotationExact(Rot + DeltaRot);
+			}
+			else
+			{
+				InComponent->SetRelativeRotationExact(InDeltaRot);
+			}
+
+			if (InIsDelta)
+			{
+				FVector NewCompLocation = InComponent->GetRelativeLocation();
+				NewCompLocation -= InPivotLocation;
+				NewCompLocation = FRotationMatrix(InDeltaRot).TransformPosition(NewCompLocation);
+				NewCompLocation += InPivotLocation;
+				InComponent->SetRelativeLocation(NewCompLocation);
+			}
+		}
+	}
+
+	///////////////////
+	// Translation
+	if (InDeltaTranslationPtr)
+	{
+		if (InIsDelta)
+		{
+			InComponent->SetRelativeLocation(InComponent->GetRelativeLocation() + *InDeltaTranslationPtr);
+		}
+		else
+		{
+			InComponent->SetRelativeLocation(*InDeltaTranslationPtr);
+		}
+	}
+
+	///////////////////
+	// Scaling
+	if (InDeltaScalePtr)
+	{
+		const FVector& InDeltaScale = *InDeltaScalePtr;
+		const bool bScalingComp = !InIsDelta || !InDeltaScale.IsNearlyZero(0.000001f);
+		if (bScalingComp)
+		{
+			if (InIsDelta)
+			{
+				InComponent->SetRelativeScale3D(InComponent->GetRelativeScale3D() + InDeltaScale);
+
+				FVector NewCompLocation = InComponent->GetRelativeLocation();
+				NewCompLocation -= InPivotLocation;
+				NewCompLocation += FScaleMatrix(InDeltaScale).TransformPosition(NewCompLocation);
+				NewCompLocation += InPivotLocation;
+				InComponent->SetRelativeLocation(NewCompLocation);
+			}
+			else
+			{
+				InComponent->SetRelativeScale3D(InDeltaScale);
+			}
+		}
+	}
+
+	// Update the actor before leaving.
+	InComponent->MarkPackageDirty();
+
+	InComponent->PostEditComponentMove(false);
+
+	// Fire callbacks
+	FEditorSupportDelegates::RefreshPropertyWindows.Broadcast();
+	FEditorSupportDelegates::UpdateUI.Broadcast();
+}
diff --git a/Engine/Source/Editor/UnrealEd/Private/GroupActor.cpp b/Engine/Source/Editor/UnrealEd/Private/GroupActor.cpp
--- a/Engine/Source/Editor/UnrealEd/Private/GroupActor.cpp
+++ b/Engine/Source/Editor/UnrealEd/Private/GroupActor.cpp
@@ -12,6 +12,8 @@
 #include "LevelEditorViewport.h"
 #include "Layers/LayersSubsystem.h"
 #include "ActorGroupingUtils.h"
+#include "Elements/Framework/EngineElementsLibrary.h"
+#include "Elements/Framework/TypedElementSelectionSet.h"
 
 const FLinearColor BOXCOLOR_LOCKEDGROUPS( 0.0f, 1.0f, 0.0f );
 const FLinearColor BOXCOLOR_UNLOCKEDGROUPS( 1.0f, 0.0f, 0.0f );
@@ -101,19 +103,20 @@ bool AGroupActor::IsSelected() const
 	return (IsLocked() && HasSelectedActors()) || Super::IsSelected();
 }
 
+namespace GroupActorHelpers
+{
 
-bool ActorHasParentInGroup(const TArray<class AActor*> &GroupActors, const AActor* Actor)
+bool ActorHasParentInGroup(const TArray<AActor*>& GroupActors, const AActor* Actor)
 {
 	check(Actor);
 	// Check that we've not got a parent attachment within the group.
-	USceneComponent *Curr = Actor->GetRootComponent();
-	for(int32 OtherIndex=0; OtherIndex<GroupActors.Num(); ++OtherIndex)
+	USceneComponent* RootComponent = Actor->GetRootComponent();
+	for (const AActor* OtherActor : GroupActors)
 	{
-		const AActor* OtherActor = GroupActors[OtherIndex];
-		if( OtherActor != NULL && OtherActor != Actor )
+		if (OtherActor && OtherActor != Actor)
 		{
-			USceneComponent *Other = OtherActor->GetRootComponent();
-			if( Curr->IsAttachedTo( Other ) )
+			USceneComponent* OtherRootComponent = OtherActor->GetRootComponent();
+			if (RootComponent->IsAttachedTo(OtherRootComponent))
 			{
 				// We do have parent so don't apply the delta - our parent object will apply it instead.
 				return true;
@@ -123,69 +126,56 @@ bool ActorHasParentInGroup(const TArray<class AActor*> &GroupActors, const AActo
 	return false;
 }
 
-bool ActorHasParentInSelection(const AActor* Actor)
+bool ActorHasParentInSelection(const AActor* Actor, const UTypedElementSelectionSet* SelectionSet)
 {
 	check(Actor);
-	bool bHasParentInSelection = false;
-	AActor* ParentActor = Actor->GetAttachParentActor();
-	while (ParentActor != NULL && !bHasParentInSelection)
+	if (SelectionSet)
 	{
-		if (ParentActor->IsSelected())
+		for (const AActor* ParentActor = Actor->GetAttachParentActor(); ParentActor; ParentActor = ParentActor->GetAttachParentActor())
 		{
-			bHasParentInSelection = true;
+			FTypedElementHandle ParentActorElementHandle = UEngineElementsLibrary::AcquireEditorActorElementHandle(ParentActor, /*bAllowCreate*/false);
+			if (ParentActorElementHandle && SelectionSet->IsElementSelected(ParentActorElementHandle, FTypedElementIsSelectedOptions()))
+			{
+				return true;
+			}
 		}
-		ParentActor = ParentActor->GetAttachParentActor();
 	}
-	return bHasParentInSelection;
+	return false;
 }
 
-void AGroupActor::GroupApplyDelta(FLevelEditorViewportClient* Viewport, const FVector& InDrag, const FRotator& InRot, const FVector& InScale )
+} // namespace GroupActorHelpers
+
+void AGroupActor::ForEachMovableActorInGroup(const UTypedElementSelectionSet* InSelectionSet, TFunctionRef<void(AActor*)> InCallback)
 {
-	check(Viewport);
-	for(int32 ActorIndex=0; ActorIndex<GroupActors.Num(); ++ActorIndex)
+	const UTypedElementSelectionSet* SelectionSet = InSelectionSet ? InSelectionSet : GEditor->GetSelectedActors()->GetElementSelectionSet();
+	for (AActor* Actor : GroupActors)
 	{
-		if( GroupActors[ActorIndex] != NULL )
+		if (Actor)
 		{
 			// Check that we've not got a parent attachment within the group/selection
-			const bool bCanApplyDelta = !ActorHasParentInGroup(GroupActors, GroupActors[ActorIndex]) && !ActorHasParentInSelection(GroupActors[ActorIndex]);
-			if(bCanApplyDelta)
+			const bool bCanApplyDelta = !GroupActorHelpers::ActorHasParentInGroup(GroupActors, Actor) && !GroupActorHelpers::ActorHasParentInSelection(Actor, SelectionSet);
+			if (bCanApplyDelta)
 			{
-				Viewport->ApplyDeltaToActor(GroupActors[ActorIndex], InDrag, InRot, InScale);
+				InCallback(Actor);
 			}
 		}
 	}
-	for(int32 SubGroupIndex=0; SubGroupIndex<SubGroups.Num(); ++SubGroupIndex)
+	for (AGroupActor* SubGroup : SubGroups)
 	{
-		if( SubGroups[SubGroupIndex] != NULL )
+		if (SubGroup)
 		{
-			SubGroups[SubGroupIndex]->GroupApplyDelta(Viewport, InDrag, InRot, InScale);
+			SubGroup->ForEachMovableActorInGroup(SelectionSet, InCallback);
 		}
 	}
-	Viewport->ApplyDeltaToActor(this, InDrag, InRot, InScale);
+	InCallback(this);
 }
 
 void AGroupActor::GroupApplyDelta(const FVector& InDrag, const FRotator& InRot, const FVector& InScale )
 {
-	for(int32 ActorIndex=0; ActorIndex<GroupActors.Num(); ++ActorIndex)
+	ForEachMovableActorInGroup(nullptr, [&InDrag, &InRot, &InScale](AActor* InActor)
 	{
-		if (GroupActors[ActorIndex] != NULL)
-		{
-			// Check that we've not got a parent attachment within the group/selection
-			const bool bCanApplyDelta = !ActorHasParentInGroup(GroupActors, GroupActors[ActorIndex]) && !ActorHasParentInSelection(GroupActors[ActorIndex]);
-			if(bCanApplyDelta)
-			{
-				GEditor->ApplyDeltaToActor(GroupActors[ActorIndex], true, &InDrag, &InRot, &InScale);
-			}
-		}
-	}
-	for(int32 SubGroupIndex=0; SubGroupIndex<SubGroups.Num(); ++SubGroupIndex)
-	{
-		if( SubGroups[SubGroupIndex] != NULL )
-		{
-			SubGroups[SubGroupIndex]->GroupApplyDelta(InDrag, InRot, InScale);
-		}
-	}
-	GEditor->ApplyDeltaToActor(this, true, &InDrag, &InRot, &InScale);
+		GEditor->ApplyDeltaToActor(InActor, true, &InDrag, &InRot, &InScale);
+	});
 }
 
 bool AGroupActor::Modify(bool bAlwaysMarkDirty/*=true*/)
diff --git a/Engine/Source/Editor/UnrealEd/Private/LevelEditorViewport.cpp b/Engine/Source/Editor/UnrealEd/Private/LevelEditorViewport.cpp
--- a/Engine/Source/Editor/UnrealEd/Private/LevelEditorViewport.cpp
+++ b/Engine/Source/Editor/UnrealEd/Private/LevelEditorViewport.cpp
@@ -38,6 +38,12 @@
 #include "UnrealEdGlobals.h"
 #include "Materials/MaterialExpressionTextureSample.h"
 #include "EditorSupportDelegates.h"
+#include "Elements/Framework/TypedElementRegistry.h"
+#include "Elements/Framework/EngineElementsLibrary.h"
+#include "Elements/Framework/TypedElementListObjectUtil.h"
+#include "Elements/Framework/TypedElementViewportInteraction.h"
+#include "Elements/Actor/ActorElementLevelEditorViewportInteractionCustomization.h"
+#include "Elements/Component/ComponentElementLevelEditorViewportInteractionCustomization.h"
 #include "AudioDevice.h"
 #include "MouseDeltaTracker.h"
 #include "ScopedTransaction.h"
@@ -91,8 +97,6 @@ DEFINE_LOG_CATEGORY(LogEditorViewport);
 
 #define LOCTEXT_NAMESPACE "LevelEditorViewportClient"
 
-static const float MIN_ACTOR_BOUNDS_EXTENT	= 1.0f;
-
 const FLevelViewportActorLock FLevelViewportActorLock::None(nullptr);
 
 TArray< TWeakObjectPtr< AActor > > FLevelEditorViewportClient::DropPreviewActors;
@@ -183,35 +187,6 @@ static void NotifyAtmosphericLightHasMoved(UDirectionalLightComponent& SelectedA
 	}
 }
 
-namespace LevelEditorViewportClientHelper
-{
-	FProperty* GetEditTransformProperty(UE::Widget::EWidgetMode WidgetMode)
-	{
-		FProperty* ValueProperty = nullptr;
-		switch (WidgetMode)
-		{
-		case UE::Widget::WM_Translate:
-			ValueProperty = FindFProperty<FProperty>(USceneComponent::StaticClass(), USceneComponent::GetRelativeLocationPropertyName());
-			break;
-		case UE::Widget::WM_Rotate:
-			ValueProperty = FindFProperty<FProperty>(USceneComponent::StaticClass(), USceneComponent::GetRelativeRotationPropertyName());
-			break;
-		case UE::Widget::WM_Scale:
-			ValueProperty = FindFProperty<FProperty>(USceneComponent::StaticClass(), USceneComponent::GetRelativeScale3DPropertyName());
-			break;
-		case UE::Widget::WM_TranslateRotateZ:
-			ValueProperty = FindFProperty<FProperty>(USceneComponent::StaticClass(), USceneComponent::GetRelativeLocationPropertyName());
-			break;
-		case UE::Widget::WM_2D:
-			ValueProperty = FindFProperty<FProperty>(USceneComponent::StaticClass(), USceneComponent::GetRelativeLocationPropertyName());
-			break;
-		default:
-			break;
-		}
-		return ValueProperty;
-	}
-}
-
 TArray<AActor*> FLevelEditorViewportClient::TryPlacingActorFromObject( ULevel* InLevel, UObject* ObjToUse, bool bSelectActors, EObjectFlags ObjectFlags, UActorFactory* FactoryToUse, const FName Name, const FViewportCursorLocation* Cursor )
 {
 	TArray<AActor*> PlacedActors;
@@ -1732,6 +1707,7 @@ FLevelEditorViewportClient::FLevelEditorViewportClient(const TSharedPtr<SLevelVi
 	, bLockedCameraView(true)
 	, bReceivedFocusRecently(false)
 	, bAlwaysShowModeWidgetAfterSelectionChanges(true)
+	, ElementsToManipulate(UTypedElementRegistry::GetInstance()->CreateElementList())
 	, SpriteCategoryVisibility()
 	, World(nullptr)
 	, TrackingTransaction()
@@ -1742,6 +1718,9 @@ FLevelEditorViewportClient::FLevelEditorViewportClient(const TSharedPtr<SLevelVi
 	, bWasEditorCameraCut(false)
 	, bApplyCameraSpeedScaleByDistance(true)
 {
+	ViewportInteraction->RegisterAssetEditorCustomizationByTypeName(NAME_Actor, MakeUnique<FActorElementLevelEditorViewportInteractionCustomization>(this));
+	ViewportInteraction->RegisterAssetEditorCustomizationByTypeName(NAME_Components, MakeUnique<FComponentElementLevelEditorViewportInteractionCustomization>(this));
+
 	// By default a level editor viewport is pointed to the editor world
 	SetReferenceToWorldContext(GEditor->GetEditorWorldContext());
 
@@ -1966,7 +1945,7 @@ void FLevelEditorViewportClient::BeginCameraMovement(bool bHasMovement)
 			{
 				GEditor->BroadcastBeginCameraMovement(*ActorLock);
 				// consider modification from piloting as relative location changes
-				FProperty* TransformProperty = LevelEditorViewportClientHelper::GetEditTransformProperty(UE::Widget::WM_Translate);
+				FProperty* TransformProperty = FComponentElementLevelEditorViewportInteractionCustomization::GetEditTransformProperty(UE::Widget::WM_Translate);
 				if (TransformProperty)
 				{
 					// Create edit property event
@@ -1995,7 +1974,7 @@ void FLevelEditorViewportClient::EndCameraMovement()
 		{
 			GEditor->BroadcastEndCameraMovement(*ActorLock);
 			// Create post edit property change event, consider modification from piloting as relative location changes
-			FProperty* TransformProperty = LevelEditorViewportClientHelper::GetEditTransformProperty(UE::Widget::WM_Translate);
+			FProperty* TransformProperty = FComponentElementLevelEditorViewportInteractionCustomization::GetEditTransformProperty(UE::Widget::WM_Translate);
 			FPropertyChangedEvent PropertyChangedEvent(TransformProperty, EPropertyChangeType::ValueSet);
 
 			// Broadcast Post Edit change notification, we can't call PostEditChangeProperty directly on Actor or ActorComponent from here since it wasn't pair with a proper PreEditChange
@@ -2616,7 +2595,7 @@ bool FLevelEditorViewportClient::InputWidgetDelta(FViewport* InViewport, EAxisLi
 					}
 					else
 					{
-						ApplyDeltaToActors( Drag, Rot, Scale );
+						ApplyDeltaToSelectedElements(FTransform(Rot, Drag, Scale));
 					}
 
 					ApplyDeltaToRotateWidget( Rot );
@@ -2892,89 +2871,24 @@ void FLevelEditorViewportClient::TrackingStarted( const FInputEventState& InInpu
 
 	bOnlyMovedPivot = false;
 
-	const bool bIsDraggingComponents = GEditor->GetSelectedComponentCount() > 0;
 	PreDragActorTransforms.Empty();
+	CacheElementsToManipulate();
 
-
-	// Create edit property event
-	FEditPropertyChain PropertyChain;
-	FProperty* TransformProperty = LevelEditorViewportClientHelper::GetEditTransformProperty(GetWidgetMode());
-	if (TransformProperty)
-	{
-		PropertyChain.AddHead(TransformProperty);
-	}
-
-	if (bIsDraggingComponents)
+	// Track BSP changes
 	{
-		if (bIsDraggingWidget)
+		bIsTrackingBrushModification = false;
+		TypedElementListObjectUtil::ForEachObject<ABrush>(ElementsToManipulate, [this](ABrush* InBrush)
 		{
-			Widget->SetSnapEnabled(true);
-
-			for (FSelectedEditableComponentIterator It(GEditor->GetSelectedEditableComponentIterator()); It; ++It)
-			{
-				USceneComponent* SceneComponent = Cast<USceneComponent>(*It);
-				if (SceneComponent)
-				{
-					// Notify that this component is beginning to move
-					GEditor->BroadcastBeginObjectMovement(*SceneComponent);
-
-					// Broadcast Pre Edit change notification, we can't call PreEditChange directly on Actor or ActorComponent from here since it will unregister the components until PostEditChange
-					if (TransformProperty)
-					{
-						FCoreUObjectDelegates::OnPreObjectPropertyChanged.Broadcast(SceneComponent, PropertyChain);
-					}
-				}
-			}
-		}
+			bIsTrackingBrushModification = InBrush && !InBrush->IsVolumeBrush() && !FActorEditorUtils::IsABuilderBrush(InBrush);
+			return !bIsTrackingBrushModification;
+		});
 	}
-	else
-	{
-		for (FSelectionIterator It(GEditor->GetSelectedActorIterator()); It && !bIsTrackingBrushModification; ++It)
-		{
-			AActor* Actor = CastChecked<AActor>(*It);
-
-			if (bIsDraggingWidget)
-			{
-				// Notify that this actor is beginning to move
-				GEditor->BroadcastBeginObjectMovement(*Actor);
-
-				// Broadcast Pre Edit change notification, we can't call PreEditChange directly on Actor or ActorComponent from here since it will unregister the components until PostEditChange
-				if (TransformProperty)
-				{
-					FCoreUObjectDelegates::OnPreObjectPropertyChanged.Broadcast(Actor, PropertyChain);
-				}
-			}
 
-			Widget->SetSnapEnabled(true);
+	if (bIsDraggingWidget)
+	{
+		Widget->SetSnapEnabled(true);
 
-			// See if any brushes are about to be transformed via their Widget
-			TArray<AActor*> AttachedActors;
-			Actor->GetAttachedActors(AttachedActors);
-			const bool bExactClass = true;
-			// First, check for selected brush actors, check the actors attached actors for brush actors as well.  If a parent actor moves, the bsp needs to be rebuilt
-			ABrush* Brush = Cast< ABrush >(Actor);
-			if (Brush && ( !Brush->IsVolumeBrush() && !FActorEditorUtils::IsABuilderBrush(Actor) ))
-			{
-				bIsTrackingBrushModification = true;
-			}
-			else // Next, check for selected groups actors that contain brushes
-			{
-				AGroupActor* GroupActor = Cast<AGroupActor>(Actor);
-				if (GroupActor)
-				{
-					TArray<AActor*> GroupMembers;
-					GroupActor->GetAllChildren(GroupMembers, true);
-					for (int32 GroupMemberIdx = 0; GroupMemberIdx < GroupMembers.Num(); ++GroupMemberIdx)
-					{
-						Brush = Cast< ABrush >(GroupMembers[GroupMemberIdx]);
-						if (Brush && ( !Brush->IsVolumeBrush() && !FActorEditorUtils::IsABuilderBrush(Actor) ))
-						{
-							bIsTrackingBrushModification = true;
-						}
-					}
-				}
-			}
-		}
+		ViewportInteraction->BeginGizmoManipulation(ElementsToManipulate, GetWidgetMode());
 	}
 
 	// Start a transformation transaction if required
@@ -2984,30 +2898,28 @@ void FLevelEditorViewportClient::TrackingStarted( const FInputEventState& InInpu
 		{
 			TrackingTransaction.TransCount++;
 
-			FText ObjectTypeBeingTracked = bIsDraggingComponents ? LOCTEXT("TransactionFocus_Components", "Components") : LOCTEXT("TransactionFocus_Actors", "Actors");
 			FText TrackingDescription;
-
 			switch( GetWidgetMode() )
 			{
 			case UE::Widget::WM_Translate:
-				TrackingDescription = FText::Format(LOCTEXT("MoveTransaction", "Move {0}"), ObjectTypeBeingTracked);
+				TrackingDescription = LOCTEXT("MoveTransaction", "Move Elements");
 				break;
 			case UE::Widget::WM_Rotate:
-				TrackingDescription = FText::Format(LOCTEXT("RotateTransaction", "Rotate {0}"), ObjectTypeBeingTracked);
+				TrackingDescription = LOCTEXT("RotateTransaction", "Rotate Elements");
 				break;
 			case UE::Widget::WM_Scale:
-				TrackingDescription = FText::Format(LOCTEXT("ScaleTransaction", "Scale {0}"), ObjectTypeBeingTracked);
+				TrackingDescription = LOCTEXT("ScaleTransaction", "Scale Elements");
 				break;
 			case UE::Widget::WM_TranslateRotateZ:
-				TrackingDescription = FText::Format(LOCTEXT("TranslateRotateZTransaction", "Translate/RotateZ {0}"), ObjectTypeBeingTracked);
+				TrackingDescription = LOCTEXT("TranslateRotateZTransaction", "Translate/RotateZ Elements");
 				break;
 			case UE::Widget::WM_2D:
-				TrackingDescription = FText::Format(LOCTEXT("TranslateRotate2D", "Translate/Rotate2D {0}"), ObjectTypeBeingTracked);
+				TrackingDescription = LOCTEXT("TranslateRotate2D", "Translate/Rotate2D Elements");
 				break;
 			default:
 				if( bNudge )
 				{
-					TrackingDescription = FText::Format(LOCTEXT("NudgeTransaction", "Nudge {0}"), ObjectTypeBeingTracked);
+					TrackingDescription = LOCTEXT("NudgeTransaction", "Nudge Elements");
 				}
 			}
 
@@ -3042,6 +2954,7 @@ void FLevelEditorViewportClient::TrackingStarted( const FInputEventState& InInpu
 		}
 	}
 
+	ElementsToManipulate->Reset();
 }
 
 void FLevelEditorViewportClient::TrackingStopped()
@@ -3088,35 +3001,8 @@ void FLevelEditorViewportClient::TrackingStopped()
 	// Don't do this if AddDelta was never called.
 	if( bDidAnythingActuallyChange && MouseDeltaTracker->HasReceivedDelta() )
 	{
-		// Create post edit property change event
-		FProperty* TransformProperty = LevelEditorViewportClientHelper::GetEditTransformProperty(GetWidgetMode());
-		FPropertyChangedEvent PropertyChangedEvent(TransformProperty, EPropertyChangeType::ValueSet);
-
-		// Move components and actors
-		{
-			TArray<USceneComponent*> ComponentsToMove;
-			TArray<AActor*> ActorsToMove;
-			GetSelectedActorsAndComponentsForMove(ActorsToMove, ComponentsToMove);
-
-			for (USceneComponent* Component : ComponentsToMove)
-			{
-				// Broadcast Post Edit change notification, we can't call PostEditChangeProperty directly on Actor or ActorComponent from here since it wasn't pair with a proper PreEditChange
-				FCoreUObjectDelegates::OnObjectPropertyChanged.Broadcast(Component, PropertyChangedEvent);
-				
-				Component->PostEditComponentMove(true);
-				GEditor->BroadcastEndObjectMovement(*Component);
-			}
-
-			for (AActor* Actor : ActorsToMove)
-			{
-				// Broadcast Post Edit change notification, we can't call PostEditChangeProperty directly on Actor or ActorComponent from here since it wasn't pair with a proper PreEditChange
-				FCoreUObjectDelegates::OnObjectPropertyChanged.Broadcast(Actor, PropertyChangedEvent);
-				Actor->PostEditMove(true);
-				GEditor->BroadcastEndObjectMovement(*Actor);
-			}
-
-			GEditor->BroadcastActorsMoved(ActorsToMove);
-		}
+		CacheElementsToManipulate();
+		ViewportInteraction->EndGizmoManipulation(ElementsToManipulate, GetWidgetMode());
 
 		if (!GUnrealEd->IsPivotMovedIndependently())
 		{
@@ -3151,6 +3037,7 @@ void FLevelEditorViewportClient::TrackingStopped()
 	}
 
 	PreDragActorTransforms.Empty();
+	ElementsToManipulate->Reset();
 }
 
 void FLevelEditorViewportClient::AbortTracking()
@@ -3191,116 +3078,6 @@ void FLevelEditorViewportClient::OnActorMoved(AActor* InActor)
 	UpdateLockedActorViewport(InActor, false);
 }
 
-void FLevelEditorViewportClient::GetSelectedActorsAndComponentsForMove(TArray<AActor*>& OutActorsToMove, TArray<USceneComponent*>& OutComponentsToMove) const
-{
-	OutActorsToMove.Reset();
-	OutComponentsToMove.Reset();
-
-	// Get the list of parent-most component(s) that are selected
-	if (GEditor->GetSelectedComponentCount() > 0)
-	{
-		// Otherwise, if both a parent and child are selected and the delta is applied to both, the child will actually move 2x delta
-		for (FSelectedEditableComponentIterator EditableComponentIt(GEditor->GetSelectedEditableComponentIterator()); EditableComponentIt; ++EditableComponentIt)
-		{
-			USceneComponent* SceneComponent = Cast<USceneComponent>(*EditableComponentIt);
-			if (!SceneComponent)
-			{
-				continue;
-			}
-
-			// Check to see if any parent is selected
-			bool bParentAlsoSelected = false;
-			USceneComponent* Parent = SceneComponent->GetAttachParent();
-			while (Parent != nullptr)
-			{
-				if (Parent->IsSelected())
-				{
-					bParentAlsoSelected = true;
-					break;
-				}
-
-				Parent = Parent->GetAttachParent();
-			}
-
-			AActor* ComponentOwner = SceneComponent->GetOwner();
-			if (!CanMoveActorInViewport(ComponentOwner))
-			{
-				continue;
-			}
-
-			const bool bIsRootComponent = (ComponentOwner && (ComponentOwner->GetRootComponent() == SceneComponent));
-			if (bIsRootComponent)
-			{
-				// If it is a root component, use the parent actor instead
-				OutActorsToMove.Add(ComponentOwner);
-			}
-			else if (!bParentAlsoSelected)
-			{
-				// If no parent of this component is also in the selection set, move it
-				OutComponentsToMove.Add(SceneComponent);
-			}
-		}
-	}
-
-	// Skip gathering selected actors if we had a valid component selection
-	if (OutComponentsToMove.Num() || OutActorsToMove.Num())
-	{
-		return;
-	}
-	
-	for (FSelectionIterator It(GEditor->GetSelectedActorIterator()); It; ++It)
-	{
-		AActor* Actor = CastChecked<AActor>(*It);
-
-		// If the root component was selected, this actor is already accounted for
-		USceneComponent* RootComponent = Actor->GetRootComponent();
-		if (RootComponent && RootComponent->IsSelected())
-		{
-			continue;
-		}
-
-		if (!CanMoveActorInViewport(Actor))
-		{
-			continue;
-		}
-
-		OutActorsToMove.Add(Actor);
-	}
-}
-
-bool FLevelEditorViewportClient::CanMoveActorInViewport(const AActor* InActor) const
-{
-	if (!GEditor || !InActor)
-	{
-		return false;
-	}
-
-	// The actor cannot be location locked
-	if (InActor->bLockLocation)
-	{
-		return false;
-	}
-
-	// The actor needs to be in the current viewport world
-	if (GEditor->PlayWorld)
-	{
-		if (bIsSimulateInEditorViewport)
-		{
-			// If the Actor's outer (level) outer (world) is not the PlayWorld then it cannot be moved in this viewport.
-			if (!(GEditor->PlayWorld == InActor->GetOuter()->GetOuter()))
-			{
-				return false;
-			}
-		}
-		else if (!(GEditor->EditorWorld == InActor->GetOuter()->GetOuter()))
-		{
-			return false;
-		}
-	}
-
-	return true;
-}
-
 void FLevelEditorViewportClient::NudgeSelectedObjects( const struct FInputEventState& InputState )
 {
 	FViewport* InViewport = InputState.GetViewport();
@@ -3618,6 +3395,24 @@ bool FLevelEditorViewportClient::HaveSelectedObjectsBeenChanged() const
 	return (TrackingTransaction.TransCount > 0 || TrackingTransaction.IsActive()) && (MouseDeltaTracker->HasReceivedDelta() || MouseDeltaTracker->WasExternalMovement());
 }
 
+void FLevelEditorViewportClient::CacheElementsToManipulate()
+{
+	const UTypedElementSelectionSet* SelectionSet = GetSelectionSet();
+	ViewportInteraction->GetSelectedElementsToMove(SelectionSet, bIsSimulateInEditorViewport ? ETypedElementViewportInteractionWorldType::PlayInEditor : ETypedElementViewportInteractionWorldType::Editor, ElementsToManipulate);
+}
+
+const UTypedElementSelectionSet* FLevelEditorViewportClient::GetSelectionSet() const
+{
+	TSharedPtr<ILevelEditor> LevelEditor = ParentLevelEditor.Pin();
+	return LevelEditor ? LevelEditor->GetElementSelectionSet() : GEditor->GetSelectedActors()->GetElementSelectionSet();
+}
+
+UTypedElementSelectionSet* FLevelEditorViewportClient::GetMutableSelectionSet() const
+{
+	TSharedPtr<ILevelEditor> LevelEditor = ParentLevelEditor.Pin();
+	return LevelEditor ? LevelEditor->GetMutableElementSelectionSet() : GEditor->GetSelectedActors()->GetElementSelectionSet();
+}
+
 void FLevelEditorViewportClient::MoveCameraToLockedActor()
 {
 	// If turned on, move cameras location/rotation to the selected actors
@@ -3764,394 +3559,73 @@ void FLevelEditorViewportClient::UpdateLockedActorViewport(const AActor* InActor
 	}
 }
 
-
-void FLevelEditorViewportClient::ApplyDeltaToActors(const FVector& InDrag,
-													const FRotator& InRot,
-													const FVector& InScale)
+void FLevelEditorViewportClient::ApplyDeltaToActors(const FVector& InDrag, const FRotator& InRot, const FVector& InScale)
 {
-	if( (InDrag.IsZero() && InRot.IsZero() && InScale.IsZero()) )
-	{
-		return;
-	}
-
-	FVector ModifiedScale = InScale;
-	// If we are scaling, we need to change the scaling factor a bit to properly align to grid.
-
-	if( GEditor->UsePercentageBasedScaling() )
-	{
-		USelection* SelectedActors = GEditor->GetSelectedActors();
-		const bool bScalingActors = !InScale.IsNearlyZero();
-
-		if( bScalingActors )
-		{
-			/* @todo: May reenable this form of calculating scaling factors later on.
-			// Calculate a bounding box for the actors.
-			FBox ActorsBoundingBox( 0 );
-
-			for ( FSelectionIterator It( GEditor->GetSelectedActorIterator() ) ; It ; ++It )
-			{
-				AActor* Actor = static_cast<AActor*>( *It );
-				checkSlow( Actor->IsA(AActor::StaticClass()) );
-
-				const FBox ActorsBox = Actor->GetComponentsBoundingBox( true );
-				ActorsBoundingBox += ActorsBox;
-			}
-
-			const FVector BoxExtent = ActorsBoundingBox.GetExtent();
-
-			for (int32 Idx=0; Idx < 3; Idx++)
-			{
-				ModifiedScale[Idx] = InScale[Idx] / BoxExtent[Idx];
-			}
-			*/
-
-			ModifiedScale = InScale * ((GEditor->GetScaleGridSize() / 100.0f) / GEditor->GetGridSize());
-		}
-	}
-
-	// Transact the actors.
-	GEditor->NoteActorMovement();
-		
-	// Move Components and Actors
-	{
-		TArray<USceneComponent*> ComponentsToMove;
-		TArray<AActor*> ActorsToMove;
-		GetSelectedActorsAndComponentsForMove(ActorsToMove, ComponentsToMove);
-
-		for (USceneComponent* SceneComp : ComponentsToMove)
-		{
-			ApplyDeltaToComponent(SceneComp, InDrag, InRot, ModifiedScale);
-		}
-
-		TArray<AGroupActor*> ActorGroups;
-		for (AActor* Actor : ActorsToMove)
-		{
-			AGroupActor* ParentGroup = AGroupActor::GetRootForActor(Actor, true, true);
-			if (ParentGroup && UActorGroupingUtils::IsGroupingActive())
-			{
-				ActorGroups.AddUnique(ParentGroup);
-			}
-			else
-			{
-				ApplyDeltaToActor(Actor, InDrag, InRot, ModifiedScale);
-			}
-		}
-
-		AGroupActor::RemoveSubGroupsFromArray(ActorGroups);
-		for (AGroupActor* ActorGroup : ActorGroups)
-		{
-			ActorGroup->GroupApplyDelta(this, InDrag, InRot, ModifiedScale);
-		}
-	}
+	ApplyDeltaToSelectedElements(FTransform(InRot, InDrag, InScale));
 }
 
-void FLevelEditorViewportClient::ApplyDeltaToComponent(USceneComponent* InComponent, const FVector& InDeltaDrag, const FRotator& InDeltaRot, const FVector& InDeltaScale)
+void FLevelEditorViewportClient::ApplyDeltaToActor(AActor* InActor, const FVector& InDeltaDrag, const FRotator& InDeltaRot, const FVector& InDeltaScale)
 {
-	// If we are scaling, we need to change the scaling factor a bit to properly align to grid.
-	FVector ModifiedDeltaScale = InDeltaScale;
-
-	// we don't scale components when we only have a very small scale change
-	if (!InDeltaScale.IsNearlyZero())
+	if (FTypedElementHandle ActorElementHandle = UEngineElementsLibrary::AcquireEditorActorElementHandle(InActor))
 	{
-		if (!GEditor->UsePercentageBasedScaling())
-		{
-			ModifyScale(InComponent, ModifiedDeltaScale);
-		}
+		ApplyDeltaToElement(ActorElementHandle, FTransform(InDeltaRot, InDeltaDrag, InDeltaScale));
 	}
-	else
-	{
-		ModifiedDeltaScale = FVector::ZeroVector;
-	}
-
-	FVector AdjustedDrag = InDeltaDrag;
-	FRotator AdjustedRot = InDeltaRot;
-	FComponentEditorUtils::AdjustComponentDelta(InComponent, AdjustedDrag, AdjustedRot);
-
-	FVector EditorWorldPivotLocation = GEditor->GetPivotLocation();
-
-	// If necessary, transform the editor pivot location to be relative to the component's parent
-	const bool bIsRootComponent = InComponent->GetOwner()->GetRootComponent() == InComponent;
-	FVector RelativePivotLocation = bIsRootComponent || !InComponent->GetAttachParent() ? EditorWorldPivotLocation : InComponent->GetAttachParent()->GetComponentToWorld().Inverse().TransformPosition(EditorWorldPivotLocation);
-
-	GEditor->ApplyDeltaToComponent(
-		InComponent,
-		true,
-		&AdjustedDrag,
-		&AdjustedRot,
-		&ModifiedDeltaScale,
-		RelativePivotLocation);
 }
 
-/** Helper function for ModifyScale - Convert the active Dragging Axis to per-axis flags */
-static void CheckActiveAxes( EAxisList::Type DraggingAxis, bool bActiveAxes[3] )
+void FLevelEditorViewportClient::ApplyDeltaToComponent(USceneComponent* InComponent, const FVector& InDeltaDrag, const FRotator& InDeltaRot, const FVector& InDeltaScale)
 {
-	bActiveAxes[0] = bActiveAxes[1] = bActiveAxes[2] = false;
-	switch(DraggingAxis)
+	if (FTypedElementHandle ComponentElementHandle = UEngineElementsLibrary::AcquireEditorComponentElementHandle(InComponent))
 	{
-	default:
-	case EAxisList::None:
-		break;
-	case EAxisList::X:
-		bActiveAxes[0] = true;
-		break;
-	case EAxisList::Y:
-		bActiveAxes[1] = true;
-		break;
-	case EAxisList::Z:
-		bActiveAxes[2] = true;
-		break;
-	case EAxisList::XYZ:
-	case EAxisList::All:
-	case EAxisList::Screen:
-		bActiveAxes[0] = bActiveAxes[1] = bActiveAxes[2] = true;
-		break;
-	case EAxisList::XY:
-		bActiveAxes[0] = bActiveAxes[1] = true;
-		break;
-	case EAxisList::XZ:
-		bActiveAxes[0] = bActiveAxes[2] = true;
-		break;
-	case EAxisList::YZ:
-		bActiveAxes[1] = bActiveAxes[2] = true;
-		break;
+		ApplyDeltaToElement(ComponentElementHandle, FTransform(InDeltaRot, InDeltaDrag, InDeltaScale));
 	}
 }
 
-/** Helper function for ModifyScale - Check scale criteria to see if this is allowed, returns modified absolute scale*/
-static float CheckScaleValue( float ScaleDeltaToCheck, float CurrentScaleFactor, float CurrentExtent, bool bCheckSmallExtent, bool bSnap )
+void FLevelEditorViewportClient::ApplyDeltaToSelectedElements(const FTransform& InDeltaTransform)
 {
-	float AbsoluteScaleValue = ScaleDeltaToCheck + CurrentScaleFactor;
-	if( bSnap )
+	if (InDeltaTransform.GetTranslation().IsZero() && InDeltaTransform.Rotator().IsZero() && InDeltaTransform.GetScale3D().IsZero())
 	{
-		AbsoluteScaleValue = FMath::GridSnap( AbsoluteScaleValue, GEditor->GetScaleGridSize() );
-	}
-	// In some situations CurrentExtent can be 0 (eg: when scaling a plane in Z), this causes a divide by 0 that we need to avoid.
-	if(FMath::Abs(CurrentExtent) < KINDA_SMALL_NUMBER) {
-		return AbsoluteScaleValue;
-	}
-	float UnscaledExtent = CurrentExtent / CurrentScaleFactor;
-	float ScaledExtent = UnscaledExtent * AbsoluteScaleValue;
-
-	if( ( FMath::Square( ScaledExtent ) ) > BIG_NUMBER )	// cant get too big...
-	{
-		return CurrentScaleFactor;
-	}
-	else if( bCheckSmallExtent && 
-			(FMath::Abs( ScaledExtent ) < MIN_ACTOR_BOUNDS_EXTENT * 0.5f ||		// ...or too small (apply sign in this case)...
-			( CurrentScaleFactor < 0.0f ) != ( AbsoluteScaleValue < 0.0f )) )	// ...also cant cross the zero boundary
-	{
-		return ( ( MIN_ACTOR_BOUNDS_EXTENT * 0.5 ) / UnscaledExtent ) * ( CurrentScaleFactor < 0.0f ? -1.0f : 1.0f );
+		return;
 	}
 
-	return AbsoluteScaleValue;
-}
-
-/**
-* Helper function for ValidateScale().
-* If the "PreserveNonUniformScale" setting is enabled, this function will appropriately re-scale the scale delta so that
-* proportions are preserved also when snapping.
-* This function will modify the scale delta sign so that scaling is apply in the good direction when using multiple axis in same time.
-* The function will not transform the scale delta in the case the scale delta is not uniform
-* @param    InOriginalPreDragScale		The object's original scale
-* @param	bActiveAxes					The axes that are active when scaling interactively.
-* @param	InOutScaleDelta				The scale delta we are potentially transforming.
-* @return true if the axes should be snapped individually, according to the snap setting (i.e. this function had no effect)
-*/
-static bool ApplyScalingOptions(const FVector& InOriginalPreDragScale, const bool bActiveAxes[3], FVector& InOutScaleDelta)
-{
-	int ActiveAxisCount = 0;
-	bool CurrentValueSameSign = true;
-	bool FirstSignPositive = true;
-	float MaxComponentSum = -1.0f;
-	int32 MaxAxisIndex = -1;
-	const ULevelEditorViewportSettings* ViewportSettings = GetDefault<ULevelEditorViewportSettings>();
-	bool SnapScaleAfter = ViewportSettings->SnapScaleEnabled;
-
-	//Found the number of active axis
-	//Found if we have to swap some sign
-	for (int Axis = 0; Axis < 3; ++Axis)
-	{
-		if (bActiveAxes[Axis])
-		{
-			bool CurrentValueIsZero = FMath::IsNearlyZero(InOriginalPreDragScale[Axis], SMALL_NUMBER);
-			//when the current value is zero we assume it is positive
-			bool IsCurrentValueSignPositive = CurrentValueIsZero ? true : InOriginalPreDragScale[Axis] > 0.0f;
-			if (ActiveAxisCount == 0)
-			{
-				//Set the first value when we find the first active axis
-				FirstSignPositive = IsCurrentValueSignPositive;
-			}
-			else
-			{
-				if (FirstSignPositive != IsCurrentValueSignPositive)
-				{
-					CurrentValueSameSign = false;
-				}
-			}
-			ActiveAxisCount++;
-		}
-	}
+	CacheElementsToManipulate();
 
-	//If we scale more then one axis and
-	//we have to swap some sign
-	if (ActiveAxisCount > 1 && !CurrentValueSameSign)
-	{
-		//Change the scale delta to reflect the sign of the value
-		for (int Axis = 0; Axis < 3; ++Axis)
-		{
-			if (bActiveAxes[Axis])
-			{
-				bool CurrentValueIsZero = FMath::IsNearlyZero(InOriginalPreDragScale[Axis], SMALL_NUMBER);
-				//when the current value is zero we assume it is positive
-				bool IsCurrentValueSignPositive = CurrentValueIsZero ? true : InOriginalPreDragScale[Axis] > 0.0f;
-				InOutScaleDelta[Axis] = IsCurrentValueSignPositive ? InOutScaleDelta[Axis] : -(InOutScaleDelta[Axis]);
-			}
-		}
-	}
+	FTransform ModifiedDeltaTransform = InDeltaTransform;
 
-	if (ViewportSettings->PreserveNonUniformScale)
 	{
-		for (int Axis = 0; Axis < 3; ++Axis)
-		{
-			if (bActiveAxes[Axis])
-			{
-				const float AbsScale = FMath::Abs(InOutScaleDelta[Axis] + InOriginalPreDragScale[Axis]);
-				if (AbsScale > MaxComponentSum)
-				{
-					MaxAxisIndex = Axis;
-					MaxComponentSum = AbsScale;
-				}
-			}
-		}
+		FVector AdjustedScale = ModifiedDeltaTransform.GetScale3D();
 
-		check(MaxAxisIndex != -1);
-
-		float AbsoluteScaleValue = InOriginalPreDragScale[MaxAxisIndex] + InOutScaleDelta[MaxAxisIndex];
-		if (ViewportSettings->SnapScaleEnabled)
+		// If we are scaling, we need to change the scaling factor a bit to properly align to grid
+		if (GEditor->UsePercentageBasedScaling() && !AdjustedScale.IsNearlyZero())
 		{
-			AbsoluteScaleValue = FMath::GridSnap(InOriginalPreDragScale[MaxAxisIndex] + InOutScaleDelta[MaxAxisIndex], GEditor->GetScaleGridSize());
-			SnapScaleAfter = false;
+			AdjustedScale *= ((GEditor->GetScaleGridSize() / 100.0f) / GEditor->GetGridSize());
 		}
 
-		float ScaleRatioMax = AbsoluteScaleValue / InOriginalPreDragScale[MaxAxisIndex];
-		for (int Axis = 0; Axis < 3; ++Axis)
-		{
-			if (bActiveAxes[Axis])
-			{
-				InOutScaleDelta[Axis] = (InOriginalPreDragScale[Axis] * ScaleRatioMax) - InOriginalPreDragScale[Axis];
-			}
-		}
-	}
-
-	
-	return SnapScaleAfter;
-}
-
-
-/** Helper function for ModifyScale - Check scale criteria to see if this is allowed */
-void FLevelEditorViewportClient::ValidateScale(const FVector& InOriginalPreDragScale, const FVector& InCurrentScale, const FVector& InBoxExtent, FVector& InOutScaleDelta, bool bInCheckSmallExtent ) const
-{
-	// get the axes that are active in this operation
-	bool bActiveAxes[3];
-	CheckActiveAxes( Widget != NULL ? Widget->GetCurrentAxis() : EAxisList::None, bActiveAxes );
-
-	//When scaling with more then one active axis, We must make sure we apply the correct delta sign to each delta scale axis
-	//We also want to support the PreserveNonUniformScale option
-	bool bSnapAxes = ApplyScalingOptions(InOriginalPreDragScale, bActiveAxes, InOutScaleDelta);
-
-	// check each axis
-	for( int Axis = 0; Axis < 3; ++Axis )
-	{
-		if( bActiveAxes[Axis] ) 
-		{
-			float ModifiedScaleAbsolute = CheckScaleValue( InOutScaleDelta[Axis], InCurrentScale[Axis], InBoxExtent[Axis], bInCheckSmallExtent, bSnapAxes );
-			InOutScaleDelta[Axis] = ModifiedScaleAbsolute - InCurrentScale[Axis];
-		}
-		else
-		{
-			InOutScaleDelta[Axis] = 0.0f;
-		}
+		ModifiedDeltaTransform.SetScale3D(AdjustedScale);
 	}
-}
 
-void FLevelEditorViewportClient::ModifyScale( AActor* InActor, FVector& ScaleDelta, bool bCheckSmallExtent ) const
-{
-	if( InActor->GetRootComponent() )
-	{
-		const FVector CurrentScale = InActor->GetRootComponent()->GetRelativeScale3D();
+	FInputDeviceState InputState;
+	InputState.SetModifierKeyStates(IsShiftPressed(), IsAltPressed(), IsCtrlPressed(), IsCmdPressed());
 
-		const FBox LocalBox = InActor->GetComponentsBoundingBox( true );
-		const FVector ScaledExtents = LocalBox.GetExtent() * CurrentScale;
-		const FTransform* PreDragTransform = PreDragActorTransforms.Find(InActor);
-		//In scale mode we need the predrag transform before the first delta calculation
-		if (PreDragTransform == nullptr)
-		{
-			PreDragTransform = &PreDragActorTransforms.Add(InActor, InActor->GetTransform());
-		}
-		check(PreDragTransform);
-		ValidateScale(PreDragTransform->GetScale3D(), CurrentScale, ScaledExtents, ScaleDelta, bCheckSmallExtent);
+	ViewportInteraction->UpdateGizmoManipulation(ElementsToManipulate, GetWidgetMode(), Widget ? Widget->GetCurrentAxis() : EAxisList::None, InputState, ModifiedDeltaTransform);
 
-		if( ScaleDelta.IsNearlyZero() )
-		{
-			ScaleDelta = FVector::ZeroVector;
-		}
-	}
+	ElementsToManipulate->Reset();
 }
 
-void FLevelEditorViewportClient::ModifyScale( USceneComponent* InComponent, FVector& ScaleDelta ) const
+void FLevelEditorViewportClient::ApplyDeltaToElement(const FTypedElementHandle& InElementHandle, const FTransform& InDeltaTransform)
 {
-	AActor* Actor = InComponent->GetOwner();
-	const FTransform* PreDragTransform = PreDragActorTransforms.Find(Actor);
-	//In scale mode we need the predrag transform before the first delta calculation
-	if (PreDragTransform == nullptr)
-	{
-		PreDragTransform = &PreDragActorTransforms.Add(Actor, Actor->GetTransform());
-	}
-	check(PreDragTransform);
-	const FBox LocalBox = Actor->GetComponentsBoundingBox(true);
-	const FVector ScaledExtents = LocalBox.GetExtent() * InComponent->GetRelativeScale3D();
-	ValidateScale(PreDragTransform->GetScale3D(), InComponent->GetRelativeScale3D(), ScaledExtents, ScaleDelta);
+	FInputDeviceState InputState;
+	InputState.SetModifierKeyStates(IsShiftPressed(), IsAltPressed(), IsCtrlPressed(), IsCmdPressed());
 
-	if( ScaleDelta.IsNearlyZero() )
-	{
-		ScaleDelta = FVector::ZeroVector;
-	}
+	ViewportInteraction->ApplyDeltaToElement(InElementHandle, GetWidgetMode(), Widget ? Widget->GetCurrentAxis() : EAxisList::None, InputState, InDeltaTransform);
 }
 
-//
-//	FLevelEditorViewportClient::ApplyDeltaToActor
-//
-
-void FLevelEditorViewportClient::ApplyDeltaToActor( AActor* InActor, const FVector& InDeltaDrag, const FRotator& InDeltaRot, const FVector& InDeltaScale )
+FTransform FLevelEditorViewportClient::CachePreDragActorTransform(const AActor* InActor)
 {
-	// If we are scaling, we may need to change the scaling factor a bit to properly align to the grid.
-	FVector ModifiedDeltaScale = InDeltaScale;
-
-	// we dont scale actors when we only have a very small scale change
-	if( !InDeltaScale.IsNearlyZero() )
+	if (const FTransform* PreDragTransform = PreDragActorTransforms.Find(InActor))
 	{
-		if(!GEditor->UsePercentageBasedScaling())
-		{
-			ModifyScale( InActor, ModifiedDeltaScale, Cast< ABrush >( InActor ) != NULL );
-		}
-	}
-	else
-	{
-		ModifiedDeltaScale = FVector::ZeroVector;
+		return *PreDragTransform;
 	}
-
-	GEditor->ApplyDeltaToActor(
-		InActor,
-		true,
-		&InDeltaDrag,
-		&InDeltaRot,
-		&ModifiedDeltaScale,
-		IsAltPressed(),
-		IsShiftPressed(),
-		IsCtrlPressed() );
-
-	// Update the cameras from their locked actor (if any) only if the viewport is realtime enabled
-	UpdateLockedActorViewports(InActor, true);
+	return PreDragActorTransforms.Add(InActor, InActor->GetTransform());
 }
 
 EMouseCursor::Type FLevelEditorViewportClient::GetCursor(FViewport* InViewport,int32 X,int32 Y)
@@ -4880,6 +4354,8 @@ void FLevelEditorViewportClient::AddReferencedObjects( FReferenceCollector& Coll
 {
 	FEditorViewportClient::AddReferencedObjects( Collector );
 
+	Collector.AddReferencedObject(ElementsToManipulate);
+
 	for( TSet<FViewportHoverTarget>::TIterator It( FLevelEditorViewportClient::HoveredObjects ); It; ++It )
 	{
 		FViewportHoverTarget& CurHoverTarget = *It;
diff --git a/Engine/Source/Editor/UnrealEd/Public/EditorViewportClient.h b/Engine/Source/Editor/UnrealEd/Public/EditorViewportClient.h
--- a/Engine/Source/Editor/UnrealEd/Public/EditorViewportClient.h
+++ b/Engine/Source/Editor/UnrealEd/Public/EditorViewportClient.h
@@ -34,6 +34,7 @@ class FPreviewScene;
 class IMatineeBase;
 class SEditorViewport;
 class UActorFactory;
+class UTypedElementViewportInteraction;
 
 /** Delegate called by FEditorViewportClient to check its visibility */
 DECLARE_DELEGATE_RetVal( bool, FViewportStateGetter );
@@ -1721,6 +1722,9 @@ protected:
 	/** The editor viewport widget this client is attached to */
 	TWeakPtr<SEditorViewport> EditorViewportWidget;
 
+	/** The viewport interaction instance for this viewport */
+	UTypedElementViewportInteraction* ViewportInteraction;
+
 	/** The scene used for the viewport. Owned externally */
 	FPreviewScene* PreviewScene;
 
diff --git a/Engine/Source/Editor/UnrealEd/Public/Elements/Actor/ActorElementEditorViewportInteractionCustomization.h b/Engine/Source/Editor/UnrealEd/Public/Elements/Actor/ActorElementEditorViewportInteractionCustomization.h
new file mode 100644
--- /dev/null
+++ b/Engine/Source/Editor/UnrealEd/Public/Elements/Actor/ActorElementEditorViewportInteractionCustomization.h
@@ -0,0 +1,16 @@
+// Copyright Epic Games, Inc. All Rights Reserved.
+
+#pragma once
+
+#include "Elements/Framework/TypedElementViewportInteraction.h"
+
+class AActor;
+
+class UNREALED_API FActorElementEditorViewportInteractionCustomization : public FTypedElementAssetEditorViewportInteractionCustomization
+{
+public:
+	virtual bool GetGizmoPivotLocation(const TTypedElement<UTypedElementWorldInterface>& InElementWorldHandle, const UE::Widget::EWidgetMode InWidgetMode, FVector& OutPivotLocation) override;
+	virtual void GizmoManipulationDeltaUpdate(const TTypedElement<UTypedElementWorldInterface>& InElementWorldHandle, const UE::Widget::EWidgetMode InWidgetMode, const EAxisList::Type InDragAxis, const FInputDeviceState& InInputState, const FTransform& InDeltaTransform, const FVector& InPivotLocation) override;
+
+	static void ApplyDeltaToActor(AActor* InActor, const bool InIsDelta, const FVector* InDeltaTranslationPtr, const FRotator* InDeltaRotationPtr, const FVector* InDeltaScalePtr, const FVector& InPivotLocation, const FInputDeviceState& InInputState);
+};
diff --git a/Engine/Source/Editor/UnrealEd/Public/Elements/Component/ComponentElementEditorViewportInteractionCustomization.h b/Engine/Source/Editor/UnrealEd/Public/Elements/Component/ComponentElementEditorViewportInteractionCustomization.h
new file mode 100644
--- /dev/null
+++ b/Engine/Source/Editor/UnrealEd/Public/Elements/Component/ComponentElementEditorViewportInteractionCustomization.h
@@ -0,0 +1,17 @@
+// Copyright Epic Games, Inc. All Rights Reserved.
+
+#pragma once
+
+#include "Elements/Framework/TypedElementViewportInteraction.h"
+
+class FProperty;
+class USceneComponent;
+
+class UNREALED_API FComponentElementEditorViewportInteractionCustomization : public FTypedElementAssetEditorViewportInteractionCustomization
+{
+public:
+	virtual bool GetGizmoPivotLocation(const TTypedElement<UTypedElementWorldInterface>& InElementWorldHandle, const UE::Widget::EWidgetMode InWidgetMode, FVector& OutPivotLocation) override;
+	virtual void GizmoManipulationDeltaUpdate(const TTypedElement<UTypedElementWorldInterface>& InElementWorldHandle, const UE::Widget::EWidgetMode InWidgetMode, const EAxisList::Type InDragAxis, const FInputDeviceState& InInputState, const FTransform& InDeltaTransform, const FVector& InPivotLocation) override;
+
+	static void ApplyDeltaToComponent(USceneComponent* InComponent, const bool InIsDelta, const FVector* InDeltaTranslationPtr, const FRotator* InDeltaRotationPtr, const FVector* InDeltaScalePtr, const FVector& InPivotLocation, const FInputDeviceState& InInputState);
+};
diff --git a/Engine/Source/Editor/UnrealEd/Public/LevelEditorViewport.h b/Engine/Source/Editor/UnrealEd/Public/LevelEditorViewport.h
--- a/Engine/Source/Editor/UnrealEd/Public/LevelEditorViewport.h
+++ b/Engine/Source/Editor/UnrealEd/Public/LevelEditorViewport.h
@@ -21,7 +21,10 @@ class ILevelEditor;
 class SLevelViewport;
 class UActorFactory;
 class UModel;
+class UTypedElementList;
+class UTypedElementSelectionSet;
 struct FWorldContext;
+struct FTypedElementHandle;
 
 /** Describes an object that's currently hovered over in the level viewport */
 struct FViewportHoverTarget
@@ -156,6 +159,9 @@ struct UNREALED_API FLevelViewportActorLock
 /** */
 class UNREALED_API FLevelEditorViewportClient : public FEditorViewportClient
 {
+	friend class FActorElementLevelEditorViewportInteractionCustomization;
+	friend class FComponentElementLevelEditorViewportInteractionCustomization;
+
 public:
 
 	/** @return Returns the current global drop preview actor, or a NULL pointer if we don't currently have one */
@@ -307,6 +313,9 @@ public:
 	void ApplyDeltaToActors( const FVector& InDrag, const FRotator& InRot, const FVector& InScale );
 	void ApplyDeltaToActor( AActor* InActor, const FVector& InDeltaDrag, const FRotator& InDeltaRot, const FVector& InDeltaScale );
 	void ApplyDeltaToComponent(USceneComponent* InComponent, const FVector& InDeltaDrag, const FRotator& InDeltaRot, const FVector& InDeltaScale);
+	
+	void ApplyDeltaToSelectedElements(const FTransform& InDeltaTransform);
+	void ApplyDeltaToElement(const FTypedElementHandle& InElementHandle, const FTransform& InDeltaTransform);
 
 	virtual void SetIsSimulateInEditorViewport( bool bInIsSimulateInEditorViewport ) override;
 
@@ -451,13 +460,6 @@ public:
 	 */
 	static void ClearHoverFromObjects();
 
-
-	/**
-	 * Helper function for ApplyDeltaTo* functions - modifies scale based on grid settings.
-	 * Currently public so it can be re-used in FEdModeBlueprint.
-	 */
-	void ModifyScale( USceneComponent* InComponent, FVector& ScaleDelta ) const;
-
 	/** Set the global ptr to the current viewport */
 	void SetCurrentViewport();
 
@@ -663,31 +665,8 @@ protected:
 	/** Delegate handler for ActorMoved events */
 	void OnActorMoved(AActor* InActor);
 
-	/** FEditorViewportClient Interface*/
-
-	/**
-	 * Collects the set of components and actors on which to apply move operations during or after drag operations.
-	 */
-	void GetSelectedActorsAndComponentsForMove(TArray<AActor*>& OutActorsToMove, TArray<USceneComponent*>& OutComponentsToMove) const;
-
-	/**
-	 * Determines if it is valid to move an actor in this viewport.
-	 *
-	 * @param InActor - the actor that the viewport may be interested in moving.
-	 * @returns true if it is valid for this viewport to update the given actor's transform.
-	 */
-	bool CanMoveActorInViewport(const AActor* InActor) const;
-
-	/** Performs the legacy behavior for calling post edit move and updating transforms from ApplyDeltaToActors function. */
-	UE_DEPRECATED(4.26, "This functions is meant to be used for ease of rollback if too many post edit move calls degrade performance during drag operations. See ULevelEditorSettings::bUseLegacyPostEditBehavior to toggle legacy behavior.")
-	bool LegacyApplyDeltasForSelectedComponentsAndActors(const FVector& InDrag, const FRotator& InRot, const FVector& ModifiedScale);
-
-	/** Performs the legacy behavior for applying transforms and calling post edit move and property changed events from TrackingStopped function. */
-	UE_DEPRECATED(4.26, "This functions is meant to be used for ease of rollback if too many post edit move calls degrade performance during drag operations. See ULevelEditorSettings::bUseLegacyPostEditBehavior to toggle legacy behavior.")
-	bool LegacyTrackingStoppedForSelectedComponentsAndActors(FPropertyChangedEvent& PropertyChangedEvent);
-
 public:
-
+	/** FEditorViewportClient Interface*/
 	virtual void UpdateLinkedOrthoViewports(bool bInvalidate = false) override;
 	virtual ELevelViewportType GetViewportType() const override;
 	virtual void SetViewportType(ELevelViewportType InViewportType) override;
@@ -703,6 +682,8 @@ protected:
 	virtual void RedrawAllViewportsIntoThisScene() override;
 
 private:
+	FTransform CachePreDragActorTransform(const AActor* InActor);
+
 	/**
 	 * Checks to see the viewports locked actor need updating
 	 */
@@ -717,6 +698,13 @@ private:
 	/** @return	Returns true if the delta tracker was used to modify any selected actors or BSP.  Must be called before EndTracking(). */
 	bool HaveSelectedObjectsBeenChanged() const;
 
+	/** Cache the list of elements to manipulate based on the current selection set. */
+	void CacheElementsToManipulate();
+
+	/** Get the selection set that associated with our level editor. */
+	const UTypedElementSelectionSet* GetSelectionSet() const;
+	UTypedElementSelectionSet* GetMutableSelectionSet() const;
+
 	/**
 	 * Called when to attempt to apply an object to a BSP surface
 	 *
@@ -789,10 +777,6 @@ private:
 	 */
 	bool DropObjectsOnWidget(FSceneView* View, struct FViewportCursorLocation& Cursor, const TArray<UObject*>& DroppedObjects, bool bCreateDropPreview = false);
 
-	/** Helper functions for ApplyDeltaTo* functions - modifies scale based on grid settings */
-	void ModifyScale( AActor* InActor, FVector& ScaleDelta, bool bCheckSmallExtent = false ) const;
-	void ValidateScale( const FVector& InOriginalPreDragScale, const FVector& CurrentScale, const FVector& BoxExtent, FVector& ScaleDelta, bool bCheckSmallExtent = false ) const;
-
 	/** Project the specified actors into the world according to the current drag parameters */
 	void ProjectActorsIntoWorld(const TArray<AActor*>& Actors, FViewport* Viewport, const FVector& Drag, const FRotator& Rot);
 
@@ -875,7 +859,10 @@ private:
 	static bool bIsDroppingPreviewActor;
 
 	/** A map of actor locations before a drag operation */
-	mutable TMap<TWeakObjectPtr<AActor>, FTransform> PreDragActorTransforms;
+	mutable TMap<TWeakObjectPtr<const AActor>, FTransform> PreDragActorTransforms;
+
+	/** List of elements that are currently being manipulated (eg, via a drag) */
+	UTypedElementList* ElementsToManipulate;
 
 	/** Bit array representing the visibility of every sprite category in the current viewport */
 	TBitArray<>	SpriteCategoryVisibility;
diff --git a/Engine/Source/Editor/UnrealEd/Public/UnrealWidgetFwd.h b/Engine/Source/Editor/UnrealEd/Public/UnrealWidgetFwd.h
deleted file mode 100644
diff --git a/Engine/Source/Runtime/Engine/Private/Elements/Actor/ActorElementWorldInterface.cpp b/Engine/Source/Runtime/Engine/Private/Elements/Actor/ActorElementWorldInterface.cpp
new file mode 100644
--- /dev/null
+++ b/Engine/Source/Runtime/Engine/Private/Elements/Actor/ActorElementWorldInterface.cpp
@@ -0,0 +1,45 @@
+// Copyright Epic Games, Inc. All Rights Reserved.
+
+#include "Elements/Actor/ActorElementWorldInterface.h"
+#include "Elements/Actor/ActorElementData.h"
+#include "GameFramework/Actor.h"
+
+UWorld* UActorElementWorldInterface::GetOwnerWorld(const FTypedElementHandle& InElementHandle)
+{
+	const FActorElementData* ActorData = InElementHandle.GetData<FActorElementData>();
+	return ActorData ? ActorData->Actor->GetWorld() : nullptr;
+}
+
+bool UActorElementWorldInterface::GetWorldBounds(const FTypedElementHandle& InElementHandle, FBoxSphereBounds& OutBounds)
+{
+	if (const FActorElementData* ActorData = InElementHandle.GetData<FActorElementData>())
+	{
+		// TODO: This was taken from FActorOrComponent, but AActor has a function to calculate bounds too...
+		OutBounds = ActorData->Actor->GetRootComponent()->Bounds;
+		return true;
+	}
+
+	return false;
+}
+
+bool UActorElementWorldInterface::GetWorldTransform(const FTypedElementHandle& InElementHandle, FTransform& OutTransform)
+{
+	if (const FActorElementData* ActorData = InElementHandle.GetData<FActorElementData>())
+	{
+		OutTransform = ActorData->Actor->GetActorTransform();
+		return true;
+	}
+
+	return false;
+}
+
+bool UActorElementWorldInterface::SetWorldTransform(const FTypedElementHandle& InElementHandle, const FTransform& InTransform)
+{
+	if (const FActorElementData* ActorData = InElementHandle.GetData<FActorElementData>())
+	{
+		ActorData->Actor->Modify();
+		return ActorData->Actor->SetActorTransform(InTransform);
+	}
+
+	return false;
+}
diff --git a/Engine/Source/Runtime/Engine/Private/Elements/Component/ComponentElementWorldInterface.cpp b/Engine/Source/Runtime/Engine/Private/Elements/Component/ComponentElementWorldInterface.cpp
new file mode 100644
--- /dev/null
+++ b/Engine/Source/Runtime/Engine/Private/Elements/Component/ComponentElementWorldInterface.cpp
@@ -0,0 +1,61 @@
+// Copyright Epic Games, Inc. All Rights Reserved.
+
+#include "Elements/Component/ComponentElementWorldInterface.h"
+#include "Elements/Component/ComponentElementData.h"
+#include "Components/ActorComponent.h"
+#include "Components/SceneComponent.h"
+
+bool UComponentElementWorldInterface::CanEditElement(const FTypedElementHandle& InElementHandle)
+{
+	const FComponentElementData* ComponentData = InElementHandle.GetData<FComponentElementData>();
+	return ComponentData && ComponentData->Component->IsEditableWhenInherited();
+}
+
+UWorld* UComponentElementWorldInterface::GetOwnerWorld(const FTypedElementHandle& InElementHandle)
+{
+	const FComponentElementData* ComponentData = InElementHandle.GetData<FComponentElementData>();
+	return ComponentData ? ComponentData->Component->GetWorld() : nullptr;
+}
+
+bool UComponentElementWorldInterface::GetWorldBounds(const FTypedElementHandle& InElementHandle, FBoxSphereBounds& OutBounds)
+{
+	if (const FComponentElementData* ComponentData = InElementHandle.GetData<FComponentElementData>())
+	{
+		if (const USceneComponent* SceneComponent = Cast<USceneComponent>(ComponentData->Component))
+		{
+			OutBounds = SceneComponent->Bounds;
+			return true;
+		}
+	}
+
+	return false;
+}
+
+bool UComponentElementWorldInterface::GetWorldTransform(const FTypedElementHandle& InElementHandle, FTransform& OutTransform)
+{
+	if (const FComponentElementData* ComponentData = InElementHandle.GetData<FComponentElementData>())
+	{
+		if (const USceneComponent* SceneComponent = Cast<USceneComponent>(ComponentData->Component))
+		{
+			OutTransform = SceneComponent->GetComponentTransform();
+			return true;
+		}
+	}
+
+	return false;
+}
+
+bool UComponentElementWorldInterface::SetWorldTransform(const FTypedElementHandle& InElementHandle, const FTransform& InTransform)
+{
+	if (const FComponentElementData* ComponentData = InElementHandle.GetData<FComponentElementData>())
+	{
+		if (USceneComponent* SceneComponent = Cast<USceneComponent>(ComponentData->Component))
+		{
+			SceneComponent->Modify();
+			SceneComponent->SetWorldTransform(InTransform);
+			return true;
+		}
+	}
+	
+	return false;
+}
diff --git a/Engine/Source/Runtime/Engine/Private/Elements/Framework/EngineElements.cpp b/Engine/Source/Runtime/Engine/Private/Elements/Framework/EngineElements.cpp
--- a/Engine/Source/Runtime/Engine/Private/Elements/Framework/EngineElements.cpp
+++ b/Engine/Source/Runtime/Engine/Private/Elements/Framework/EngineElements.cpp
@@ -9,10 +9,12 @@
 
 #include "Elements/Actor/ActorElementData.h"
 #include "Elements/Actor/ActorElementObjectInterface.h"
+#include "Elements/Actor/ActorElementWorldInterface.h"
 #include "Elements/Actor/ActorElementSelectionInterface.h"
 
 #include "Elements/Component/ComponentElementData.h"
 #include "Elements/Component/ComponentElementObjectInterface.h"
+#include "Elements/Component/ComponentElementWorldInterface.h"
 #include "Elements/Component/ComponentElementSelectionInterface.h"
 
 #include "Modules/ModuleManager.h"
@@ -34,6 +36,7 @@ void RegisterEngineActorElements()
 
 	Registry->RegisterElementType<FActorElementData>(NAME_Actor);
 	Registry->RegisterElementInterface<UTypedElementObjectInterface>(NAME_Actor, NewObject<UActorElementObjectInterface>());
+	Registry->RegisterElementInterface<UTypedElementWorldInterface>(NAME_Actor, NewObject<UActorElementWorldInterface>());
 	Registry->RegisterElementInterface<UTypedElementSelectionInterface>(NAME_Actor, NewObject<UActorElementSelectionInterface>());
 }
 
@@ -43,6 +46,7 @@ void RegisterEngineComponentElements()
 
 	Registry->RegisterElementType<FComponentElementData>(NAME_Components);
 	Registry->RegisterElementInterface<UTypedElementObjectInterface>(NAME_Components, NewObject<UComponentElementObjectInterface>());
+	Registry->RegisterElementInterface<UTypedElementWorldInterface>(NAME_Components, NewObject<UComponentElementWorldInterface>());
 	Registry->RegisterElementInterface<UTypedElementSelectionInterface>(NAME_Components, NewObject<UComponentElementSelectionInterface>());
 }
 
diff --git a/Engine/Source/Runtime/Engine/Public/Elements/Actor/ActorElementWorldInterface.h b/Engine/Source/Runtime/Engine/Public/Elements/Actor/ActorElementWorldInterface.h
new file mode 100644
--- /dev/null
+++ b/Engine/Source/Runtime/Engine/Public/Elements/Actor/ActorElementWorldInterface.h
@@ -0,0 +1,18 @@
+// Copyright Epic Games, Inc. All Rights Reserved.
+
+#pragma once
+
+#include "Elements/Interfaces/TypedElementWorldInterface.h"
+#include "ActorElementWorldInterface.generated.h"
+
+UCLASS()
+class ENGINE_API UActorElementWorldInterface : public UTypedElementWorldInterface
+{
+	GENERATED_BODY()
+
+public:
+	virtual UWorld* GetOwnerWorld(const FTypedElementHandle& InElementHandle) override;
+	virtual bool GetWorldBounds(const FTypedElementHandle& InElementHandle, FBoxSphereBounds& OutBounds) override;
+	virtual bool GetWorldTransform(const FTypedElementHandle& InElementHandle, FTransform& OutTransform) override;
+	virtual bool SetWorldTransform(const FTypedElementHandle& InElementHandle, const FTransform& InTransform) override;
+};
diff --git a/Engine/Source/Runtime/Engine/Public/Elements/Component/ComponentElementWorldInterface.h b/Engine/Source/Runtime/Engine/Public/Elements/Component/ComponentElementWorldInterface.h
new file mode 100644
--- /dev/null
+++ b/Engine/Source/Runtime/Engine/Public/Elements/Component/ComponentElementWorldInterface.h
@@ -0,0 +1,19 @@
+// Copyright Epic Games, Inc. All Rights Reserved.
+
+#pragma once
+
+#include "Elements/Interfaces/TypedElementWorldInterface.h"
+#include "ComponentElementWorldInterface.generated.h"
+
+UCLASS()
+class ENGINE_API UComponentElementWorldInterface : public UTypedElementWorldInterface
+{
+	GENERATED_BODY()
+
+public:
+	virtual bool CanEditElement(const FTypedElementHandle& InElementHandle) override;
+	virtual UWorld* GetOwnerWorld(const FTypedElementHandle& InElementHandle) override;
+	virtual bool GetWorldBounds(const FTypedElementHandle& InElementHandle, FBoxSphereBounds& OutBounds) override;
+	virtual bool GetWorldTransform(const FTypedElementHandle& InElementHandle, FTransform& OutTransform) override;
+	virtual bool SetWorldTransform(const FTypedElementHandle& InElementHandle, const FTransform& InTransform) override;
+};
diff --git a/Engine/Source/Runtime/Engine/Public/Elements/Interfaces/TypedElementWorldInterface.h b/Engine/Source/Runtime/Engine/Public/Elements/Interfaces/TypedElementWorldInterface.h
new file mode 100644
--- /dev/null
+++ b/Engine/Source/Runtime/Engine/Public/Elements/Interfaces/TypedElementWorldInterface.h
@@ -0,0 +1,55 @@
+// Copyright Epic Games, Inc. All Rights Reserved.
+
+#pragma once
+
+#include "Elements/Framework/TypedElementHandle.h"
+#include "TypedElementWorldInterface.generated.h"
+
+class UWorld;
+
+UCLASS(Abstract)
+class ENGINE_API UTypedElementWorldInterface : public UTypedElementInterface
+{
+	GENERATED_BODY()
+
+public:
+	/**
+	 * Can this element actually be edited in the world?
+	 */
+	UFUNCTION(BlueprintPure, Category="TypedElementInterfaces|World")
+	virtual bool CanEditElement(const FTypedElementHandle& InElementHandle) { return true; }
+
+	/**
+	 * Get the owner world associated with this element, if any.
+	 */
+	UFUNCTION(BlueprintPure, Category="TypedElementInterfaces|World")
+	virtual UWorld* GetOwnerWorld(const FTypedElementHandle& InElementHandle) { return nullptr; }
+
+	/**
+	 * Get the bounds of this element within its owner world, if any.
+	 */
+	UFUNCTION(BlueprintPure, Category="TypedElementInterfaces|World")
+	virtual bool GetWorldBounds(const FTypedElementHandle& InElementHandle, FBoxSphereBounds& OutBounds) { return false; }
+
+	/**
+	 * Get the transform of this element within its owner world, if any.
+	 */
+	UFUNCTION(BlueprintPure, Category="TypedElementInterfaces|World")
+	virtual bool GetWorldTransform(const FTypedElementHandle& InElementHandle, FTransform& OutTransform) { return false; }
+	
+	/**
+	 * Attempt to set the transform of this element within its owner world.
+	 */
+	UFUNCTION(BlueprintCallable, Category="TypedElementInterfaces|World")
+	virtual bool SetWorldTransform(const FTypedElementHandle& InElementHandle, const FTransform& InTransform) { return false; }
+};
+
+template <>
+struct TTypedElement<UTypedElementWorldInterface> : public TTypedElementBase<UTypedElementWorldInterface>
+{
+	bool CanEditElement() const { return InterfacePtr->CanEditElement(*this); }
+	UWorld* GetOwnerWorld() const { return InterfacePtr->GetOwnerWorld(*this); }
+	bool GetWorldBounds(FBoxSphereBounds& OutBounds) const { return InterfacePtr->GetWorldBounds(*this, OutBounds); }
+	bool GetWorldTransform(FTransform& OutTransform) const { return InterfacePtr->GetWorldTransform(*this, OutTransform); }
+	bool SetWorldTransform(const FTransform& InTransform) const { return InterfacePtr->SetWorldTransform(*this, InTransform); }
+};
diff --git a/Engine/Source/Runtime/TypedElementInterfaces/Public/Elements/Interfaces/TypedElementObjectInterface.h b/Engine/Source/Runtime/TypedElementInterfaces/Public/Elements/Interfaces/TypedElementObjectInterface.h
--- a/Engine/Source/Runtime/TypedElementInterfaces/Public/Elements/Interfaces/TypedElementObjectInterface.h
+++ b/Engine/Source/Runtime/TypedElementInterfaces/Public/Elements/Interfaces/TypedElementObjectInterface.h
@@ -12,7 +12,7 @@ class TYPEDELEMENTINTERFACES_API UTypedElementObjectInterface : public UTypedEle
 
 public:
 	/**
-	 * Retreive the object instance that this handle represents, if any.
+	 * Get the object instance that this handle represents, if any.
 	 */
 	UFUNCTION(BlueprintPure, Category="TypedElementInterfaces|Object")
 	virtual UObject* GetObject(const FTypedElementHandle& InElementHandle) { return nullptr; }
