Fix content browser thumbnails saving with improper gamma causing them to be significantly darker on load.

diff --git a/Engine/Source/Editor/UnrealEd/Private/EditorEngine.cpp b/Engine/Source/Editor/UnrealEd/Private/EditorEngine.cpp
--- a/Engine/Source/Editor/UnrealEd/Private/EditorEngine.cpp
+++ b/Engine/Source/Editor/UnrealEd/Private/EditorEngine.cpp
@@ -2459,62 +2459,70 @@ void UEditorEngine::CloseEditedWorldAssets(UWorld* InWorld)
 
 UTextureRenderTarget2D* UEditorEngine::GetScratchRenderTarget( uint32 MinSize )
 {
-	auto NewFactory = NewObject<UTextureRenderTargetFactoryNew>();
-	UTextureRenderTarget2D* ScratchRenderTarget = NULL;
+	UTextureRenderTarget2D* ScratchRenderTarget = nullptr;
 
 	// We never allow render targets greater than 2048
-	check( MinSize <= 2048 );
+	check(MinSize <= 2048);
 
 	// 256x256
-	if( MinSize <= 256 )
+	if(MinSize <= 256)
 	{
-		if( ScratchRenderTarget256 == NULL )
+		if(ScratchRenderTarget256 == nullptr)
 		{
-			NewFactory->Width = 256;
-			NewFactory->Height = 256;
-			UObject* NewObj = NewFactory->FactoryCreateNew( UTextureRenderTarget2D::StaticClass(), GetTransientPackage(), NAME_None, RF_Transient, NULL, GWarn );
-			ScratchRenderTarget256 = CastChecked<UTextureRenderTarget2D>(NewObj);
+			ScratchRenderTarget256 = NewObject<UTextureRenderTarget2D>(GetTransientPackage(), NAME_None, RF_Transient);
+
+			ScratchRenderTarget256->TargetGamma = GEngine->DisplayGamma;
+			ScratchRenderTarget256->RenderTargetFormat = RTF_RGBA8;
+
+			ScratchRenderTarget256->InitAutoFormat(256, 256);
 		}
 		ScratchRenderTarget = ScratchRenderTarget256;
 	}
 	// 512x512
-	else if( MinSize <= 512 )
+	else if(MinSize <= 512)
 	{
-		if( ScratchRenderTarget512 == NULL )
+		if( ScratchRenderTarget512 == nullptr)
 		{
-			NewFactory->Width = 512;
-			NewFactory->Height = 512;
-			UObject* NewObj = NewFactory->FactoryCreateNew( UTextureRenderTarget2D::StaticClass(), GetTransientPackage(), NAME_None, RF_Transient, NULL, GWarn );
-			ScratchRenderTarget512 = CastChecked<UTextureRenderTarget2D>(NewObj);
+			ScratchRenderTarget512 = NewObject<UTextureRenderTarget2D>(GetTransientPackage(), NAME_None, RF_Transient);
+
+			ScratchRenderTarget512->TargetGamma = GEngine->DisplayGamma;
+			ScratchRenderTarget512->RenderTargetFormat = RTF_RGBA8;
+
+			ScratchRenderTarget512->InitAutoFormat(515, 512);
 		}
 		ScratchRenderTarget = ScratchRenderTarget512;
 	}
 	// 1024x1024
-	else if( MinSize <= 1024 )
+	else if(MinSize <= 1024)
 	{
-		if( ScratchRenderTarget1024 == NULL )
+		if( ScratchRenderTarget1024 == nullptr)
 		{
-			NewFactory->Width = 1024;
-			NewFactory->Height = 1024;
-			UObject* NewObj = NewFactory->FactoryCreateNew( UTextureRenderTarget2D::StaticClass(), GetTransientPackage(), NAME_None, RF_Transient, NULL, GWarn );
-			ScratchRenderTarget1024 = CastChecked<UTextureRenderTarget2D>(NewObj);
+			ScratchRenderTarget1024 = NewObject<UTextureRenderTarget2D>(GetTransientPackage(), NAME_None, RF_Transient);
+
+			ScratchRenderTarget1024->TargetGamma = GEngine->DisplayGamma;
+			ScratchRenderTarget1024->RenderTargetFormat = RTF_RGBA8;
+
+			ScratchRenderTarget1024->InitAutoFormat(1024, 1024);
 		}
 		ScratchRenderTarget = ScratchRenderTarget1024;
 	}
 	// 2048x2048
-	else if( MinSize <= 2048 )
+	else if(MinSize <= 2048)
 	{
-		if( ScratchRenderTarget2048 == NULL )
+		if(ScratchRenderTarget2048 == nullptr)
 		{
-			NewFactory->Width = 2048;
-			NewFactory->Height = 2048;
-			UObject* NewObj = NewFactory->FactoryCreateNew( UTextureRenderTarget2D::StaticClass(), GetTransientPackage(), NAME_None, RF_Transient, NULL, GWarn );
-			ScratchRenderTarget2048 = CastChecked<UTextureRenderTarget2D>(NewObj);
+			ScratchRenderTarget2048 = NewObject<UTextureRenderTarget2D>(GetTransientPackage(), NAME_None, RF_Transient);
+
+			ScratchRenderTarget2048->TargetGamma = GEngine->DisplayGamma;
+			ScratchRenderTarget2048->RenderTargetFormat = RTF_RGBA8;
+
+			ScratchRenderTarget2048->InitAutoFormat(2048, 2048);
 		}
 		ScratchRenderTarget = ScratchRenderTarget2048;
 	}
 
-	check( ScratchRenderTarget != NULL );
+	check(ScratchRenderTarget != nullptr);
+
 	return ScratchRenderTarget;
 }
 
