Cook: Modified CookOnTheFlyServer's ConditionalSplitPackage to respect maximum allowed memory usage (HasExceededMaxMemory) while generator package creates new packages.

diff --git a/Engine/Source/Editor/UnrealEd/Classes/CookOnTheSide/CookOnTheFlyServer.h b/Engine/Source/Editor/UnrealEd/Classes/CookOnTheSide/CookOnTheFlyServer.h
--- a/Engine/Source/Editor/UnrealEd/Classes/CookOnTheSide/CookOnTheFlyServer.h
+++ b/Engine/Source/Editor/UnrealEd/Classes/CookOnTheSide/CookOnTheFlyServer.h
@@ -116,6 +116,7 @@ namespace Cook
 	struct FCookerTimer;
 	class FExternalRequests;
 	struct FPackageData;
+	struct FGeneratorPackage;
 	struct FPackageDatas;
 	struct FPackageTracker;
 	struct FPendingCookedPlatformData;
@@ -872,9 +873,10 @@ private:
 	 * for this UObject.
 	 *
 	 * @param PackageData the PackageData used to gather all uobjects from
+	 * @param bIsPreCaching true if called for precaching
 	 * @return false if time slice was reached, true if all objects have had BeginCacheForCookedPlatformData called
 	 */
-	bool BeginPrepareSave(UE::Cook::FPackageData& PackageData, UE::Cook::FCookerTimer& Timer);
+	bool BeginPrepareSave(UE::Cook::FPackageData& PackageData, UE::Cook::FCookerTimer& Timer, bool bIsPreCaching);
 	
 	/**
 	 * Returns true when all objects in package have all their cooked platform data loaded
@@ -1067,13 +1069,32 @@ private:
 	const FPackageNameCache& GetPackageNameCache() const;
 
 	/**
-	* In case the given package object is a Package splitter, this function will split the package and generate other packages.
+	* Try to find a registered CookPackageSplitter to split the given package.
 	*
 	* @param PackageData			Package to be potentially split
-	* @param Obj					Package's object to potentially handle splitting
-	* @param bHasError				If function fails, bHasError is set to true, else false.
+	* @param OutSplitDataObject		Return package's object to handle splitting
+	* @param bOutHasError			Return True if function fails
 	*/
-	void ConditionalSplitPackage(UE::Cook::FPackageData& PackageData, UObject* Obj, bool& bHasError);
+	UE::Cook::Private::FRegisteredCookPackageSplitter* TryGetRegisteredCookPackageSplitter(UE::Cook::FPackageData& PackageData, UObject*& OutSplitDataObject, bool& bOutError);
+
+	/**
+	* Generate a creator package helper object that will be used to split package.
+	*
+	* @param PackageData			Package to split
+	* @param SplitDataObject		Package's object to handle splitting
+	* @param Splitter				Registered CookPackageSplitter to be used for splitting the package
+	*/
+	UE::Cook::FGeneratorPackage* CreateGeneratorPackage(UE::Cook::FPackageData& PackageData, UObject* SplitDataObject, UE::Cook::Private::FRegisteredCookPackageSplitter* Splitter);
+
+	/**
+	* Split the package and generate other packages.
+	*
+	* @param PackageData			Package to be split
+	* @param SplitDataObject		Package's object to handle splitting
+	* @param bOutCompleted			Return true is function is done splitting
+	* @param bOutError				Return true if function fails
+	*/
+	void SplitPackage(UE::Cook::FGeneratorPackage* Generator, bool& bCompleted, bool& bOutError);
 
 	uint32 FullLoadAndSave(uint32& CookedPackageCount);
 
diff --git a/Engine/Source/Editor/UnrealEd/Private/CookOnTheFlyServer.cpp b/Engine/Source/Editor/UnrealEd/Private/CookOnTheFlyServer.cpp
--- a/Engine/Source/Editor/UnrealEd/Private/CookOnTheFlyServer.cpp
+++ b/Engine/Source/Editor/UnrealEd/Private/CookOnTheFlyServer.cpp
@@ -220,6 +220,13 @@ DECLARE_CYCLE_STAT(TEXT("Tick cooking"), STAT_TickCooker, STATGROUP_Cooking);
 
 constexpr uint32 ExpectedMaxNumPlatforms = 32;
 
+namespace UE
+{
+namespace Cook
+{
+static const TCHAR* GeneratedPackageSubPath = TEXT("_Generated_");
+}
+}
 
 /* helper structs functions
  *****************************************************************************/
@@ -1966,74 +1973,75 @@ void UCookOnTheFlyServer::TickNetwork()
 	}
 }
 
-void UCookOnTheFlyServer::ConditionalSplitPackage(UE::Cook::FPackageData& PackageData, UObject* Obj, bool& bHasError)
+UE::Cook::Private::FRegisteredCookPackageSplitter* UCookOnTheFlyServer::TryGetRegisteredCookPackageSplitter(UE::Cook::FPackageData& PackageData, UObject*& OutSplitDataObject, bool& bOutError)
 {
-	bHasError = false;
+	bOutError = false;
+	OutSplitDataObject = nullptr;
+
+	UE::Cook::Private::FRegisteredCookPackageSplitter* FoundSplitter = nullptr;
+	UObject* FoundSplitDataObject = nullptr;
 
-	// Check if Package is a Generator Package that needs to be split
-	UE::Cook::Private::FRegisteredCookPackageSplitter* RegisteredSplitter = nullptr;
 	TArray<UE::Cook::Private::FRegisteredCookPackageSplitter*> FoundRegisteredSplitters;
-	RegisteredSplitDataClasses.MultiFind(Obj->GetClass(), FoundRegisteredSplitters);
-	for (UE::Cook::Private::FRegisteredCookPackageSplitter* FoundRegisteredSplitter : FoundRegisteredSplitters)
+
+	for (FWeakObjectPtr& WeakObj : PackageData.GetCachedObjectsInOuter())
 	{
-		if (FoundRegisteredSplitter && FoundRegisteredSplitter->ShouldSplitPackage(Obj))
+		UObject* Obj = WeakObj.Get();
+		if (!Obj)
 		{
-			if (RegisteredSplitter)
+			continue;
+		}
+	
+		FoundRegisteredSplitters.Reset();
+		RegisteredSplitDataClasses.MultiFind(Obj->GetClass(), FoundRegisteredSplitters);
+
+		for (UE::Cook::Private::FRegisteredCookPackageSplitter* Splitter : FoundRegisteredSplitters)
+		{
+			if (Splitter && Splitter->ShouldSplitPackage(Obj))
 			{
-				UE_LOG(LogCook, Error, TEXT("Found more than one registered Cook Package Splitter for package %s."), *PackageData.GetPackageName().ToString());
-				bHasError = true;
-				return;
-			}
+				if (!Obj->HasAnyFlags(RF_Public))
+				{
+					UE_LOG(LogCook, Error, TEXT("SplitterData object %s must be publicly referenceable so we can keep them from being garbage collected"), *Obj->GetFullName());
+					bOutError = true;
+					return nullptr;
+				}
+
+				if (FoundSplitter)
+				{
+					UE_LOG(LogCook, Error, TEXT("Found more than one registered Cook Package Splitter for package %s."), *PackageData.GetPackageName().ToString());
+					bOutError = true;
+					return nullptr;
+				}
 
-			RegisteredSplitter = FoundRegisteredSplitter;
+				FoundSplitter = Splitter;
+				FoundSplitDataObject = Obj;
+			}
 		}
 	}
-	
-	// No splitter found
-	if (!RegisteredSplitter)
-	{
-		return;
-	}
+
+	OutSplitDataObject = FoundSplitDataObject;
+	return FoundSplitter;
+}
+
+UE::Cook::FGeneratorPackage* UCookOnTheFlyServer::CreateGeneratorPackage(UE::Cook::FPackageData& PackageData, UObject* SplitDataObject, UE::Cook::Private::FRegisteredCookPackageSplitter* Splitter)
+{
+	check(Splitter);
+	check(!PackageData.GetGeneratorPackage());
 
 	// TODO: Add support for cooking in the editor
 	if (IsCookingInEditor())
 	{
 		UE_LOG(LogCook, Error, TEXT("Cooking in editor doesn't support Cook Package Splitters."));
-		bHasError = true;
-		return;
-	}
-	
-	// Package was already marked as a Generator package, only accept one registered splitter acting on one object of this package.
-	if (PackageData.GetIsGeneratorPackage())
-	{
-		UE_LOG(LogCook, Error, TEXT("Generator package %s already split by another split data object."), *PackageData.GetPackageName().ToString());
-		bHasError = true;
-		return;
-	}
-	PackageData.SetIsGeneratorPackage(true);
-
-	UE_LOG(LogCook, Display, TEXT("Splitting Package %s with class %s acting on object %s."), *PackageData.GetPackageName().ToString(), *RegisteredSplitter->GetSplitDataClass()->GetName(), *Obj->GetFullName());
-
-	// Create instance of CookPackageSplitter class
-	TUniquePtr<ICookPackageSplitter> CookPackageSplitterInstance(RegisteredSplitter->CreateInstance(Obj));
-	if (!CookPackageSplitterInstance)
-	{
-		UE_LOG(LogCook, Error, TEXT("Error instantiating Cook Package Splitter for object %s."), *Obj->GetFullName());
-		bHasError = true;
-		return;
+		return nullptr;
 	}
 
-	// Initialize the instance with the data object
-	CookPackageSplitterInstance->SetDataObject(Obj);
+	UE_LOG(LogCook, Display, TEXT("Splitting Package %s with class %s acting on object %s."), *PackageData.GetPackageName().ToString(), *Splitter->GetSplitDataClass()->GetName(), *SplitDataObject->GetFullName());
 
 	// Prepare necessary paths
-	const TCHAR* GeneratedKeyWord = TEXT("_Generated_");
 	const UPackage* GeneratorPackage = PackageData.GetPackage();
 	const FString GeneratorPackagePath = FPackageName::GetLongPackagePath(GeneratorPackage->GetPathName());
 	const FString GeneratorPackageShortName = FPackageName::GetShortName(GeneratorPackage->GetName());
-	const FString GeneratedCookedRootPath = FPaths::RemoveDuplicateSlashes(FString::Printf(TEXT("/%s/%s/%s/"), *GeneratorPackagePath, *GeneratorPackageShortName, GeneratedKeyWord));
-	const FString GeneratedUncookedRootPath = FPaths::RemoveDuplicateSlashes(FString::Printf(TEXT("/%s%s/"), *GeneratorPackageShortName, GeneratedKeyWord));
-	const FString GeneratedUncookedContentPath = FPaths::RemoveDuplicateSlashes(FString::Printf(TEXT("%s/Cooked/%s/%s/"), *FPaths::ProjectIntermediateDir(), *GeneratorPackageShortName, GeneratedKeyWord));
+	const FString GeneratedUncookedRootPath = FPaths::RemoveDuplicateSlashes(FString::Printf(TEXT("/%s%s/"), *GeneratorPackageShortName, UE::Cook::GeneratedPackageSubPath));
+	const FString GeneratedUncookedContentPath = FPaths::RemoveDuplicateSlashes(FString::Printf(TEXT("%s/Cooked/%s/%s/"), *FPaths::ProjectIntermediateDir(), *GeneratorPackageShortName, UE::Cook::GeneratedPackageSubPath));
 
 	// Empty the uncooked intermediate directory
 	const FString GeneratedDirectoryToDelete = FPaths::ConvertRelativePathToFull(FPaths::CreateStandardFilename(GeneratedUncookedContentPath));
@@ -2045,25 +2053,50 @@ void UCookOnTheFlyServer::ConditionalSplitPackage(UE::Cook::FPackageData& Packag
 		FPackageName::RegisterMountPoint(GeneratedUncookedRootPath, GeneratedUncookedContentPath);
 	}
 
-	// Get packages to generate
-	TArray<ICookPackageSplitter::FGeneratedPackage> PackagesToGenerate = CookPackageSplitterInstance->GetGenerateList();
+	// Create instance of CookPackageSplitter class
+	ICookPackageSplitter* SplitterInstance = Splitter->CreateInstance(SplitDataObject);
+	if (!SplitterInstance)
+	{
+		UE_LOG(LogCook, Error, TEXT("Error instantiating Cook Package Splitter for object %s."), *SplitDataObject->GetFullName());
+		return nullptr;
+	}
+
+	// Initialize the instance with the data object
+	SplitterInstance->SetDataObject(SplitDataObject);
+
+	// Create a FGeneratorPackage helper object using this CookPackageSplitter instance and return it
+	PackageData.CreateGeneratorPackage(SplitDataObject, SplitterInstance, GeneratedUncookedRootPath);
+
+	return PackageData.GetGeneratorPackage();
+}
+
+void UCookOnTheFlyServer::SplitPackage(UE::Cook::FGeneratorPackage* Generator, bool& bOutCompleted, bool& bOutError)
+{
+	bOutCompleted = false;
+	bOutError = false;
+	check(Generator);
+
+	const UPackage* GeneratorPackage = Generator->GetOwner().GetPackage();
+	const FString GeneratorPackagePath = FPackageName::GetLongPackagePath(GeneratorPackage->GetPathName());
+	const FString GeneratorPackageShortName = FPackageName::GetShortName(GeneratorPackage->GetName());
+	const FString GeneratedCookedRootPath = FPaths::RemoveDuplicateSlashes(FString::Printf(TEXT("/%s/%s/%s/"), *GeneratorPackagePath, *GeneratorPackageShortName, UE::Cook::GeneratedPackageSubPath));
 
 	// Start Splitting
-	for (const ICookPackageSplitter::FGeneratedPackage& PackageToGenerate : PackagesToGenerate)
+	while (const ICookPackageSplitter::FGeneratedPackage* PackageToGenerate = Generator->GetNextPackageToGenerate())
 	{
 		// Make package name to fit Mounted GeneratedUncookedRootPath
-		const FString GeneratedPackageName = FPaths::RemoveDuplicateSlashes(FString::Printf(TEXT("/%s/%s"), *GeneratedUncookedRootPath, *PackageToGenerate.RelativePath));
+		const FString GeneratedPackageName = FPaths::RemoveDuplicateSlashes(FString::Printf(TEXT("/%s/%s"), *Generator->GetGeneratedUncookedRootPath(), *PackageToGenerate->RelativePath));
 
 		// Create package
 		check(!FindObject<UPackage>(nullptr, *GeneratedPackageName));
 		UPackage* GeneratedPackage = CreatePackage(*GeneratedPackageName);
 
 		// Populate package using CookPackageSplitterInstance and pass GeneratedPackage's cooked name for it to properly setup any internal reference to this package (SoftObjectPaths or others)
-		const FString GeneratedPackageCookedName = FPaths::RemoveDuplicateSlashes(FString::Printf(TEXT("/%s/%s"), *GeneratedCookedRootPath, *PackageToGenerate.RelativePath));
-		if (!CookPackageSplitterInstance->TryPopulatePackage(GeneratedPackage, PackageToGenerate.RelativePath, GeneratedPackageCookedName))
+		const FString GeneratedPackageCookedName = FPaths::RemoveDuplicateSlashes(FString::Printf(TEXT("/%s/%s"), *GeneratedCookedRootPath, *PackageToGenerate->RelativePath));
+		if (!Generator->GetCookPackageSplitterInstance()->TryPopulatePackage(GeneratedPackage, PackageToGenerate->RelativePath, GeneratedPackageCookedName))
 		{
-			UE_LOG(LogCook, Error, TEXT("Error populating generated package %s for object %s."), *PackageToGenerate.RelativePath, *Obj->GetFullName());
-			bHasError = true;
+			UE_LOG(LogCook, Error, TEXT("Error populating generated package %s for object %s."), *PackageToGenerate->RelativePath, *Generator->GetSplitDataObject()->GetFullName());
+			bOutError = true;
 			return;
 		}
 
@@ -2072,8 +2105,8 @@ void UCookOnTheFlyServer::ConditionalSplitPackage(UE::Cook::FPackageData& Packag
 		const FString GeneratedPackageUncookedFileName = FPaths::RemoveDuplicateSlashes(FPackageName::LongPackageNameToFilename(GeneratedPackage->GetName(), GeneratedPackageExtension));
 		if (!UPackage::SavePackage(GeneratedPackage, nullptr, RF_Standalone, *GeneratedPackageUncookedFileName, GError, nullptr, false, true, SAVE_None))
 		{
-			UE_LOG(LogCook, Error, TEXT("Error saving generated package %s for object %s."), *GeneratedPackageUncookedFileName, *Obj->GetFullName());
-			bHasError = true;
+			UE_LOG(LogCook, Error, TEXT("Error saving generated package %s for object %s."), *GeneratedPackageUncookedFileName, *Generator->GetSplitDataObject()->GetFullName());
+			bOutError = true;
 			return;
 		}
 
@@ -2083,8 +2116,8 @@ void UCookOnTheFlyServer::ConditionalSplitPackage(UE::Cook::FPackageData& Packag
 		// Make sure PackageNameCache successfully caches the generated package
 		if (GetPackageNameCache().GetCachedStandardFileName(GeneratedPackage).IsNone())
 		{
-			UE_LOG(LogCook, Error, TEXT("Error caching generated package name %s for object %s."), *GeneratedPackageUncookedFileName, *Obj->GetFullName());
-			bHasError = true;
+			UE_LOG(LogCook, Error, TEXT("Error caching generated package name %s for object %s."), *GeneratedPackageUncookedFileName, *Generator->GetSplitDataObject()->GetFullName());
+			bOutError = true;
 			return;
 		}
 
@@ -2092,24 +2125,31 @@ void UCookOnTheFlyServer::ConditionalSplitPackage(UE::Cook::FPackageData& Packag
 		UE::Cook::FPackageData* GeneratedPackageData = PackageDatas->TryAddPackageDataByFileName(FName(GeneratedPackageUncookedFileName));
 		if (!GeneratedPackageData)
 		{
-			UE_LOG(LogCook, Error, TEXT("Error adding generated package data %s for object %s."), *GeneratedPackageUncookedFileName, *Obj->GetFullName());
-			bHasError = true;
+			UE_LOG(LogCook, Error, TEXT("Error adding generated package data %s for object %s."), *GeneratedPackageUncookedFileName, *Generator->GetSplitDataObject()->GetFullName());
+			bOutError = true;
 			return;
 		}
+
 		// Set override cooked FileName to be used when saving the cooked generated package
 		const FString GeneratePackageCookedFileName = FPaths::RemoveDuplicateSlashes(FString::Printf(TEXT("%s/%s/%s/%s%s"), 
-													  *FPaths::GetPath(PackageData.GetFileName().ToString()), *GeneratorPackageShortName, GeneratedKeyWord, *PackageToGenerate.RelativePath, *GeneratedPackageExtension));
+													  *FPaths::GetPath(Generator->GetOwner().GetFileName().ToString()), *GeneratorPackageShortName, UE::Cook::GeneratedPackageSubPath, *PackageToGenerate->RelativePath, *GeneratedPackageExtension));
 		GeneratedPackageData->SetCookedFileName(FPackageNameCache::GetStandardFileName(GeneratePackageCookedFileName));
 
 		// Queue the generated package
 		QueueDiscoveredPackage(GeneratedPackage, /*bIsGeneratedPackage*/true);
 		check(GeneratedPackageData->IsInProgress());
+
+		if (HasExceededMaxMemory())
+		{
+			return;
+		}
 	}
 
-	CookPackageSplitterInstance->FinalizeGeneratorPackage();
+	Generator->Finalize();
+	bOutCompleted = true;
 }
 
-bool UCookOnTheFlyServer::BeginPrepareSave(UE::Cook::FPackageData& PackageData, UE::Cook::FCookerTimer& Timer)
+bool UCookOnTheFlyServer::BeginPrepareSave(UE::Cook::FPackageData& PackageData, UE::Cook::FCookerTimer& Timer, bool bIsPreCaching)
 {
 	if (PackageData.GetCookedPlatformDataCalled())
 	{
@@ -2142,29 +2182,50 @@ bool UCookOnTheFlyServer::BeginPrepareSave(UE::Cook::FPackageData& PackageData,
 	check(PackageData.GetState() == UE::Cook::EPackageState::Save);
 	PackageData.CreateObjectCache();
 
-	TArray<FWeakObjectPtr>& CachedObjectsInOuter = PackageData.GetCachedObjectsInOuter();
 	int32& CookedPlatformDataNextIndex = PackageData.GetCookedPlatformDataNextIndex();
-
 	if (CookedPlatformDataNextIndex == -1)
 	{
 		// Check for Splitting the package; this needs to happen before we make any of the BeginCacheForCookedPlatformData calls in the package
-		for (FWeakObjectPtr& WeakObj : CachedObjectsInOuter)
+		bool bError = false;
+
+		UE::Cook::FGeneratorPackage* Generator = PackageData.GetGeneratorPackage();
+		if (!Generator)
 		{
-			UObject* Obj = WeakObj.Get();
-			if (!Obj)
+			UObject* SplitDataObject = nullptr;
+			if (UE::Cook::Private::FRegisteredCookPackageSplitter* Splitter = TryGetRegisteredCookPackageSplitter(PackageData, SplitDataObject, bError))
 			{
-				continue;
+				// Found a splitter, create generator package
+				Generator = CreateGeneratorPackage(PackageData, SplitDataObject, Splitter);
+				bError = !Generator;
 			}
-
-			bool bHasError = false;
-			ConditionalSplitPackage(PackageData, Obj, bHasError);
-			if (bHasError)
+			if (bError)
 			{
-				// ConditionalSplitPackage failure marks the package data for PumpSaves to handle this as an error and to put it in idle state.
+				// Any error marks the package data for PumpSaves to handle this as an error and to put it in idle state
 				PackageData.SetHasBeginPrepareSaveFailed(true);
 				return false;
 			}
 		}
+		if (Generator)
+		{
+			// Don't split a generator package when precaching
+			if (bIsPreCaching)
+			{
+				return false;
+			}
+			bool bCompleted = false;
+			SplitPackage(Generator, bCompleted, bError);
+			if (bError)
+			{
+				// SplitPackage failure marks the package data for PumpSaves to handle this as an error and to put it in idle state
+				PackageData.SetHasBeginPrepareSaveFailed(true);
+				return false;
+			}
+			if (!bCompleted)
+			{
+				return false;
+			}
+			PackageData.DestroyGeneratorPackage();
+		}
 		CookedPlatformDataNextIndex = 0;
 	}
 
@@ -2173,6 +2234,7 @@ bool UCookOnTheFlyServer::BeginPrepareSave(UE::Cook::FPackageData& PackageData,
 	// while BeginPrepareSave is active.
 	// Currently this does not cause significant cost since saving new platforms with some platforms already saved is a rare operation.
 
+	TArray<FWeakObjectPtr>& CachedObjectsInOuter = PackageData.GetCachedObjectsInOuter();
 	FWeakObjectPtr* CachedObjectsInOuterData = CachedObjectsInOuter.GetData();
 	const TArray<const ITargetPlatform*>& TargetPlatforms = PackageData.GetRequestedPlatforms();
 	int NumPlatforms = TargetPlatforms.Num();
@@ -2251,7 +2313,7 @@ bool UCookOnTheFlyServer::FinishPrepareSave(UE::Cook::FPackageData& PackageData,
 
 	if (!PackageData.GetCookedPlatformDataCalled())
 	{
-		if (!BeginPrepareSave(PackageData, Timer))
+		if (!BeginPrepareSave(PackageData, Timer, /*bIsPreCaching*/ false))
 		{
 			return false;
 		}
@@ -2613,6 +2675,15 @@ void UCookOnTheFlyServer::PumpSaves(UE::Cook::FTickStackData& StackData, uint32
 				++OutNumPushed;
 				continue;
 			}
+
+			// GC is required
+			if (PackageData.GeneratorPackageRequiresGC())
+			{
+				StackData.ResultFlags |= COSR_RequiresGC;
+				SaveQueue.AddFront(&PackageData);
+				return;
+			}
+
 			// Can we postpone?
 			if (!PackageData.GetIsUrgent())
 			{
@@ -2663,7 +2734,7 @@ void UCookOnTheFlyServer::PumpSaves(UE::Cook::FTickStackData& StackData, uint32
 					break;
 				}
 				--LeftToPrecache;
-				BeginPrepareSave(*NextData, StackData.Timer);
+				BeginPrepareSave(*NextData, StackData.Timer, /*bIsPreCaching*/ true);
 			}
 
 			// If we're in RealTimeMode, check whether the precaching overflowed our timer and if so exit before we do the potentially expensive SavePackage
@@ -3358,6 +3429,17 @@ void UCookOnTheFlyServer::PreGarbageCollect()
 		GCKeepObjects.Add(SavingPackageData->GetPackage());
 	}
 
+	TArray<UPackage*> GCKeepPackages;
+	for (FPackageData* PackageData : PackageDatas->GetSaveQueue())
+	{
+		// When a package is in the process of generating other packages,
+		// we need to keep it loaded (and all objects outered to it)
+		if (PackageData->GetGeneratorPackage())
+		{
+			GCKeepPackages.Add(PackageData->GetPackage());
+		}
+	}
+
 	const bool bPartialGC = IsCookFlagSet(ECookInitializationFlags::EnablePartialGC);
 	if (bPartialGC)
 	{
@@ -3382,20 +3464,13 @@ void UCookOnTheFlyServer::PreGarbageCollect()
 		}
 
 		TSet<FName> LoadedPackages;
-		TArray<UObject*> ObjectsWithOuter;
 		for (UPackage* Package : PackageTracker->LoadedPackages)
 		{
 			const FName& PackageName = Package->GetFName();
 			if (KeepPackages.Contains(PackageName))
 			{
 				LoadedPackages.Add(PackageName);
-				GCKeepObjects.Add(Package);
-				ObjectsWithOuter.Reset();
-				GetObjectsWithOuter(Package, ObjectsWithOuter);
-				for (UObject* Obj : ObjectsWithOuter)
-				{
-					GCKeepObjects.Add(Obj);
-				}
+				GCKeepPackages.Add(Package);
 			}
 		}
 
@@ -3424,6 +3499,22 @@ void UCookOnTheFlyServer::PreGarbageCollect()
 			RequestQueue.AddRequest(Request); // Urgent requests will still be moved to the front of the RequestQueue by AddRequest
 		}
 	}
+
+	// Add packages and all RF_Public objects outered to them to GCKeepObjects
+	TArray<UObject*> ObjectsWithOuter;
+	for (UPackage* Package : GCKeepPackages)
+	{
+		GCKeepObjects.Add(Package);
+		ObjectsWithOuter.Reset();
+		GetObjectsWithOuter(Package, ObjectsWithOuter);
+		for (UObject* Obj : ObjectsWithOuter)
+		{
+			if (Obj->HasAnyFlags(RF_Public))
+			{
+				GCKeepObjects.Add(Obj);
+			}
+		}
+	}
 }
 
 void UCookOnTheFlyServer::CookerAddReferencedObjects(FReferenceCollector& Collector)
@@ -3766,7 +3857,7 @@ void UCookOnTheFlyServer::SaveCookedPackage(UE::Cook::FPackageData& PackageData,
 			UE_LOG(LogCook, Fatal, TEXT("Package %s marked as reloading for cook by was requested to save"), *Package->GetName());
 		}
 
-		// If valid, override cooked FileName. Currently used for Generated Packages (see ConditionalSplitPackage)
+		// If valid, override cooked FileName. Currently used for Generated Packages (see SplitPackage)
 		if (PackageData.GetCookedFileName() != NAME_None)
 		{
 			Filename = PackageData.GetCookedFileName().ToString();
@@ -4139,7 +4230,7 @@ void UCookOnTheFlyServer::Initialize( ECookMode::Type DesiredCookMode, ECookInit
 		UE_LOG(LogCook, Warning, TEXT("Cooking with Event Driven Loader disabled. Loading code will use deprecated path which will be removed in future release."));
 	}
 
-	// Prepare a map SplitDataClass to FRegisteredCookPackageSplitter* for ConditionalSplitPackage to use
+	// Prepare a map SplitDataClass to FRegisteredCookPackageSplitter* for TryGetRegisteredCookPackageSplitter to use
 	RegisteredSplitDataClasses.Reset();
 	UE::Cook::Private::FRegisteredCookPackageSplitter::ForEach([this](UE::Cook::Private::FRegisteredCookPackageSplitter* RegisteredCookPackageSplitter)
 	{
diff --git a/Engine/Source/Editor/UnrealEd/Private/Cooker/CookPackageData.cpp b/Engine/Source/Editor/UnrealEd/Private/Cooker/CookPackageData.cpp
--- a/Engine/Source/Editor/UnrealEd/Private/Cooker/CookPackageData.cpp
+++ b/Engine/Source/Editor/UnrealEd/Private/Cooker/CookPackageData.cpp
@@ -27,7 +27,7 @@ namespace Cook
 
 	FPackageData::FPackageData(FPackageDatas& PackageDatas, const FName& InPackageName, const FName& InFileName)
 		: PackageName(InPackageName), FileName(InFileName), PackageDatas(PackageDatas), PreloadableFileFormat(EPackageFormat::Binary)
-		, bIsUrgent(0), bIsVisited(0), bIsPreloadAttempted(0), bIsPreloaded(0), bIsGeneratorPackage(0), bHasSaveCache(0), bHasBeginPrepareSaveFailed(0), bCookedPlatformDataStarted(0), bCookedPlatformDataCalled(0), bCookedPlatformDataComplete(0), bMonitorIsCooked(0)
+		, bIsUrgent(0), bIsVisited(0), bIsPreloadAttempted(0), bIsPreloaded(0), bHasSaveCache(0), bHasBeginPrepareSaveFailed(0), bCookedPlatformDataStarted(0), bCookedPlatformDataCalled(0), bCookedPlatformDataComplete(0), bMonitorIsCooked(0)
 	{
 		SetState(EPackageState::Idle);
 		SendToState(EPackageState::Idle, ESendFlags::QueueAdd);
@@ -594,8 +594,8 @@ namespace Cook
 	{
 		check(GetPackage() != nullptr && GetPackage()->IsFullyLoaded());
 
-		SetHasBeginPrepareSaveFailed(false);
-		SetIsGeneratorPackage(false);
+		check(!GetHasBeginPrepareSaveFailed())
+		check(!GetGeneratorPackage())
 		CheckObjectCacheEmpty();
 		CheckCookedPlatformDataEmpty();
 	}
@@ -605,7 +605,7 @@ namespace Cook
 		PackageDatas.GetCookOnTheFlyServer().ReleaseCookedPlatformData(*this);
 		ClearObjectCache();
 		SetHasBeginPrepareSaveFailed(false);
-		SetIsGeneratorPackage(false);
+		DestroyGeneratorPackage();
 	}
 
 	void FPackageData::OnEnterInProgress()
@@ -893,6 +893,55 @@ namespace Cook
 				});
 	}
 
+	bool FPackageData::GeneratorPackageRequiresGC() const
+	{
+		// We consider that if a FPackageData has valid GeneratorPackage helper object,
+		// this means that COTFS's process of generating packages was not completed 
+		// either due to an error or because it has exceeded a maximum memory threshold. 
+		return GetGeneratorPackage() && !GetHasBeginPrepareSaveFailed();
+	}
+
+	UE::Cook::FGeneratorPackage* FPackageData::CreateGeneratorPackage(const UObject* InSplitDataObject, ICookPackageSplitter* InCookPackageSplitterInstance, const FString& InGeneratedUncookedRootPath)
+	{
+		check(!GetGeneratorPackage());
+		GeneratorPackage.Reset(new UE::Cook::FGeneratorPackage(*this, InSplitDataObject, InCookPackageSplitterInstance, InGeneratedUncookedRootPath));
+		return GetGeneratorPackage();
+	}
+	
+	//////////////////////////////////////////////////////////////////////////
+	// FGeneratorPackage
+
+	FGeneratorPackage::FGeneratorPackage(UE::Cook::FPackageData& InOwner, const UObject* InSplitDataObject, ICookPackageSplitter* InCookPackageSplitterInstance, const FString& InGeneratedUncookedRootPath)
+	: bIsFinalized(false)
+	, Owner(InOwner)
+	, SplitDataObject(InSplitDataObject)
+	, GeneratedUncookedRootPath(InGeneratedUncookedRootPath)
+	{
+		check(PackageToGenerateNextIndex == 0);
+		check(InCookPackageSplitterInstance);
+		CookPackageSplitterInstance.Reset(InCookPackageSplitterInstance);
+		PackagesToGenerate = CookPackageSplitterInstance->GetGenerateList();
+	}
+
+	const ICookPackageSplitter::FGeneratedPackage* FGeneratorPackage::GetNextPackageToGenerate()
+	{
+		if (PackagesToGenerate.IsValidIndex(PackageToGenerateNextIndex))
+		{
+			return &PackagesToGenerate[PackageToGenerateNextIndex++];
+		}
+		return nullptr;
+	}
+
+	void FGeneratorPackage::Finalize()
+	{
+		if (ensure(!bIsFinalized))
+		{
+			check(PackageToGenerateNextIndex == PackagesToGenerate.Num());
+			GetCookPackageSplitterInstance()->FinalizeGeneratorPackage();
+			bIsFinalized = true;
+		}
+	}
+
 	//////////////////////////////////////////////////////////////////////////
 	// FPendingCookedPlatformData
 
diff --git a/Engine/Source/Editor/UnrealEd/Private/Cooker/CookPackageData.h b/Engine/Source/Editor/UnrealEd/Private/Cooker/CookPackageData.h
--- a/Engine/Source/Editor/UnrealEd/Private/Cooker/CookPackageData.h
+++ b/Engine/Source/Editor/UnrealEd/Private/Cooker/CookPackageData.h
@@ -4,6 +4,7 @@
 
 #include <atomic>
 #include "CookTypes.h"
+#include "CookPackageSplitter.h"
 #include "Containers/Array.h"
 #include "Containers/ArrayView.h"
 #include "Containers/Map.h"
@@ -31,6 +32,7 @@ namespace UE
 {
 namespace Cook
 {
+	struct FGeneratorPackage;
 	struct FPackageDataMonitor;
 	class FPackageDataQueue;
 	struct FPendingCookedPlatformDataCancelManager;
@@ -255,10 +257,6 @@ namespace Cook
 		FName GetCookedFileName() const { return CookedFileName; }
 		void SetCookedFileName(FName InCookedFileName) { CookedFileName = InCookedFileName; }
 
-		/** Get/Set the flag for whether this package is a Generator package. */
-		bool GetIsGeneratorPackage() const { return static_cast<bool>(bIsGeneratorPackage); }
-		void SetIsGeneratorPackage(bool bValue) { bIsGeneratorPackage = bValue != 0; }
-
 		/** Validate that the fields related to the BeginCacheForCookedPlatformData calls are empty, as required when not in the save state. */
 		void CheckCookedPlatformDataEmpty() const;
 		/** Clear the fields related to the BeginCacheForCookedPlatformData calls, when e.g. leaving the save state. Caller is responsible for having already executed any required cancellation steps to avoid dangling pending operations. */
@@ -277,6 +275,16 @@ namespace Cook
 		/** Swap all ITargetPlatform* stored on this instance according to the mapping in @param Remap. */
 		void RemapTargetPlatforms(const TMap<ITargetPlatform*, ITargetPlatform*>& Remap);
 
+		/** Create and return a GeneratorPackage helper object for a given CookPackageSplitter. */
+		FGeneratorPackage* CreateGeneratorPackage(const UObject* InSplitDataObject, ICookPackageSplitter* InCookPackageSplitterInstance, const FString& InGeneratedUncookedRootPath);
+		/** Destroy GeneratorPackage helper object. */
+		void DestroyGeneratorPackage() { GeneratorPackage.Reset(); }
+		/** Get GeneratorPackage helper object. */
+		FGeneratorPackage* GetGeneratorPackage() const { return GeneratorPackage.Get(); }
+
+		/** Return whether a GC is required by a generator package as soon as possible. */
+		bool GeneratorPackageRequiresGC() const;
+
 	private:
 		friend struct UE::Cook::FPackageDatas;
 
@@ -332,6 +340,7 @@ namespace Cook
 		void OnExitHasPackage();
 		void OnEnterHasPackage();
 
+		TUniquePtr<FGeneratorPackage> GeneratorPackage;
 		TArray<const ITargetPlatform*> RequestedPlatforms;
 		TArray<const ITargetPlatform*> CookedPlatforms; // Platform part of the CookedPlatforms set. Always the same length as CookSucceeded.
 		TArray<bool> CookSucceeded; // Success flag part of the CookedPlatforms set. Always the same length as CookedPlatforms.
@@ -352,7 +361,6 @@ namespace Cook
 		uint32 bIsVisited : 1;
 		uint32 bIsPreloadAttempted : 1;
 		uint32 bIsPreloaded : 1;
-		uint32 bIsGeneratorPackage : 1;
 		uint32 bHasSaveCache : 1;
 		uint32 bHasBeginPrepareSaveFailed : 1;
 		uint32 bCookedPlatformDataStarted : 1;
@@ -361,6 +369,42 @@ namespace Cook
 		uint32 bMonitorIsCooked : 1;
 	};
 
+	/**
+	 * Helper that wraps a ICookPackageSplitter, gets/caches packages to generate and provide
+	 * the necessary to iterate over all of them iteratively.
+	 */
+	struct FGeneratorPackage
+	{
+	public:
+		/** Store the provided CookPackageSplitter and prepare the packages to generate. */
+		FGeneratorPackage(UE::Cook::FPackageData& InOwner, const UObject* InSplitDataObject, ICookPackageSplitter* InCookPackageSplitterInstance, const FString& InGeneratedUncookedRootPath);
+		/** Return next package to generate. */
+		const ICookPackageSplitter::FGeneratedPackage* GetNextPackageToGenerate();
+		/** Finalize generator package. */
+		void Finalize();
+
+		/** Return whether generator package is finalized. */
+		bool IsFinalized() const { return bIsFinalized; }
+		/** Return CookPackageSplitter. */
+		ICookPackageSplitter* GetCookPackageSplitterInstance() const { return CookPackageSplitterInstance.Get(); }
+		/** Return owner FPackageData. */
+		const UE::Cook::FPackageData& GetOwner() const { return Owner; }
+		/** Return the SplitDataObject. */
+		const UObject* GetSplitDataObject() const { return SplitDataObject; }
+		/** Returns root path of generated packages. */
+		const FString& GetGeneratedUncookedRootPath() const { return GeneratedUncookedRootPath; }
+
+	private:
+		bool bIsFinalized;
+		const UE::Cook::FPackageData& Owner;
+		const UObject* SplitDataObject;
+		const FString GeneratedUncookedRootPath;
+		TUniquePtr<ICookPackageSplitter> CookPackageSplitterInstance; // Cached CookPackageSplitter
+		TArray<ICookPackageSplitter::FGeneratedPackage> PackagesToGenerate; // Generator's packages to generate
+		int32 PackageToGenerateNextIndex = 0; // Index of next package to generate
+	};
+
+
 	/**
 	 * Stores information about the pending action in response to a single call to BeginCacheForCookedPlatformData that was made on a given object for the given platform, when saving the given PackageData.
 	 * This instance will remain alive until the object returns true from IsCachedCookedPlatformDataLoaded.
diff --git a/Engine/Source/Runtime/Engine/Private/WorldPartition/WorldPartitionRuntimeLevelStreamingCell.cpp b/Engine/Source/Runtime/Engine/Private/WorldPartition/WorldPartitionRuntimeLevelStreamingCell.cpp
--- a/Engine/Source/Runtime/Engine/Private/WorldPartition/WorldPartitionRuntimeLevelStreamingCell.cpp
+++ b/Engine/Source/Runtime/Engine/Private/WorldPartition/WorldPartitionRuntimeLevelStreamingCell.cpp
@@ -116,10 +116,7 @@ bool UWorldPartitionRuntimeLevelStreamingCell::PopulateGeneratedPackageForCook(U
 			if (!Actor)
 			{
 				Actor = ActorDesc->Load();
-			}
-			if (ensure(Actor))
-			{
-				LoadedActorsForCook.Add(Actor);
+				ensure(Actor);
 			}
 		}
 
@@ -139,17 +136,6 @@ bool UWorldPartitionRuntimeLevelStreamingCell::PopulateGeneratedPackageForCook(U
 	return true;
 }
 
-void UWorldPartitionRuntimeLevelStreamingCell::FinalizeGeneratedPackageForCook()
-{
-#if DO_CHECK
-	// Make sure actors are not external anymore
-	for(AActor * Actor : LoadedActorsForCook)
-	{
-		check(!Actor->IsPackageExternal());
-	}
-#endif
-}
-
 int32 UWorldPartitionRuntimeLevelStreamingCell::GetActorCount() const
 {
 	return Packages.Num();
diff --git a/Engine/Source/Runtime/Engine/Public/WorldPartition/WorldPartitionRuntimeLevelStreamingCell.h b/Engine/Source/Runtime/Engine/Public/WorldPartition/WorldPartitionRuntimeLevelStreamingCell.h
--- a/Engine/Source/Runtime/Engine/Public/WorldPartition/WorldPartitionRuntimeLevelStreamingCell.h
+++ b/Engine/Source/Runtime/Engine/Public/WorldPartition/WorldPartitionRuntimeLevelStreamingCell.h
@@ -23,7 +23,6 @@ class UWorldPartitionRuntimeLevelStreamingCell : public UWorldPartitionRuntimeSp
 
 	// Cook methods
 	virtual bool PopulateGeneratedPackageForCook(UPackage* InPackage, const FString& InPackageCookName) override;
-	virtual void FinalizeGeneratedPackageForCook() override;
 	virtual FString GetPackageNameToCreate() const override;
 #endif
 
@@ -48,10 +47,6 @@ private:
 
 	// Actors to be loaded for cook
 	TArray<FGuid> ActorsToLoadForCook;
-
-	// Loaded actors for cook
-	UPROPERTY(Transient)
-	TArray<AActor*> LoadedActorsForCook;
 #endif
 
 	UPROPERTY()
