WorldPartitionConvertCommandlet cannot be run anymore if 'AllowCommandletRendering' is not specified on the commandlet, mainly for landscape conversion and minimap generation.

diff --git a/Engine/Source/Editor/UnrealEd/Classes/Commandlets/WorldPartitionConvertCommandlet.h b/Engine/Source/Editor/UnrealEd/Classes/Commandlets/WorldPartitionConvertCommandlet.h
--- a/Engine/Source/Editor/UnrealEd/Classes/Commandlets/WorldPartitionConvertCommandlet.h
+++ b/Engine/Source/Editor/UnrealEd/Classes/Commandlets/WorldPartitionConvertCommandlet.h
@@ -88,7 +88,6 @@ protected:
 	bool bReportOnly;
 	bool bVerbose;
 	bool bConversionSuffix;
-	bool bAllowCommandletRendering;
 	FString ConversionSuffix;
 
 	UPROPERTY(Config)
diff --git a/Engine/Source/Editor/UnrealEd/Private/Commandlets/WorldPartitionConvertCommandlet.cpp b/Engine/Source/Editor/UnrealEd/Private/Commandlets/WorldPartitionConvertCommandlet.cpp
--- a/Engine/Source/Editor/UnrealEd/Private/Commandlets/WorldPartitionConvertCommandlet.cpp
+++ b/Engine/Source/Editor/UnrealEd/Private/Commandlets/WorldPartitionConvertCommandlet.cpp
@@ -662,7 +662,9 @@ void UWorldPartitionConvertCommandlet::OnWorldLoaded(UWorld* World)
 
 void UWorldPartitionConvertCommandlet::CreateWorldMiniMapTexture(UWorld* World)
 {
-	AWorldPartitionMiniMap* WorldMiniMap = FWorldPartitionMiniMapHelper::GetWorldPartitionMiniMap(World,true);
+	UE_SCOPED_TIMER(TEXT("CreateWorldMiniMapTexture"), LogWorldPartitionConvertCommandlet, Display);
+
+	AWorldPartitionMiniMap* WorldMiniMap = FWorldPartitionMiniMapHelper::GetWorldPartitionMiniMap(World, /*bCreateNewMiniMap*/true);
 	if (!WorldMiniMap)
 	{
 		UE_LOG(LogWorldPartitionConvertCommandlet, Error, TEXT("Failed to create Minimap. WorldPartitionMiniMap actor not found in the persistent level."));
@@ -696,7 +698,12 @@ int32 UWorldPartitionConvertCommandlet::Main(const FString& Params)
 	bGenerateIni = Switches.Contains(TEXT("GenerateIni"));
 	bReportOnly = bGenerateIni || Switches.Contains(TEXT("ReportOnly"));
 	bVerbose = Switches.Contains(TEXT("Verbose"));
-	bAllowCommandletRendering = Switches.Contains(TEXT("AllowCommandletRendering"));
+
+	if (!Switches.Contains(TEXT("AllowCommandletRendering")))
+	{
+		UE_LOG(LogWorldPartitionConvertCommandlet, Error, TEXT("The option \"-AllowCommandletRendering\" is required."));
+		return 1;
+	}
 
 	ReadAdditionalTokensAndSwitches(Tokens, Switches);
 
@@ -1159,15 +1166,9 @@ int32 UWorldPartitionConvertCommandlet::Main(const FString& Params)
 	bool bForceInitializeWorld = false;
 	bool bInitializedPhysicsSceneForSave = GEditor->InitializePhysicsSceneForSaveIfNecessary(MainWorld, bForceInitializeWorld);
 
-	//Create MiniMap of World
-	bool bSkipMiniMapGeneration = Switches.Contains(TEXT("SkipMiniMapGeneration"));
-	if (!bSkipMiniMapGeneration)
+	// Create the world's minimap
+	if (!Switches.Contains(TEXT("SkipMiniMapGeneration")))
 	{
-		if (!bAllowCommandletRendering)
-		{
-			UE_LOG(LogWorldPartitionConvertCommandlet, Error, TEXT("AllowCommandletRendering option is required to generate MiniMap. Use SkipMiniMapGeneration option to skip MiniMap generation."));
-			return 1;
-		}
 		CreateWorldMiniMapTexture(MainWorld);
 	}
 
