CookOnTheFlyServer: Save generated packages using editor wrapper GEditor->Save instead of UPackage::SavePackage (mainly for PreSaveWorld delegates to be executed).

diff --git a/Engine/Source/Editor/UnrealEd/Private/CookOnTheFlyServer.cpp b/Engine/Source/Editor/UnrealEd/Private/CookOnTheFlyServer.cpp
--- a/Engine/Source/Editor/UnrealEd/Private/CookOnTheFlyServer.cpp
+++ b/Engine/Source/Editor/UnrealEd/Private/CookOnTheFlyServer.cpp
@@ -2103,7 +2103,8 @@ void UCookOnTheFlyServer::SplitPackage(UE::Cook::FGeneratorPackage* Generator, b
 		// Save package into the uncooked intermediate directory
 		const FString GeneratedPackageExtension = GeneratedPackage->ContainsMap() ? FPackageName::GetMapPackageExtension() : FPackageName::GetAssetPackageExtension();
 		const FString GeneratedPackageUncookedFileName = FPaths::RemoveDuplicateSlashes(FPackageName::LongPackageNameToFilename(GeneratedPackage->GetName(), GeneratedPackageExtension));
-		if (!UPackage::SavePackage(GeneratedPackage, nullptr, RF_Standalone, *GeneratedPackageUncookedFileName, GError, nullptr, false, true, SAVE_None))
+		FSavePackageResultStruct SaveResult = GEditor->Save(GeneratedPackage, nullptr, RF_Standalone, *GeneratedPackageUncookedFileName, GError, nullptr, /*bForceByteSwapping*/ false, /*bWarnOfLongFilename*/ true, SAVE_None);
+		if (SaveResult.Result != ESavePackageResult::Success)
 		{
 			UE_LOG(LogCook, Error, TEXT("Error saving generated package %s for object %s."), *GeneratedPackageUncookedFileName, *Generator->GetSplitDataObject()->GetFullName());
 			bOutError = true;
