Added guard against nullptr skeleton for an AnimSequence in the anim compression commandlet

diff --git a/Engine/Source/Editor/UnrealEd/Private/Commandlets/PackageUtilities.cpp b/Engine/Source/Editor/UnrealEd/Private/Commandlets/PackageUtilities.cpp
--- a/Engine/Source/Editor/UnrealEd/Private/Commandlets/PackageUtilities.cpp
+++ b/Engine/Source/Editor/UnrealEd/Private/Commandlets/PackageUtilities.cpp
@@ -1847,7 +1847,13 @@ struct CompressAnimationsFunctor
 			}
 
 			USkeleton* Skeleton = AnimSeq->GetSkeleton();
-			check (Skeleton);
+
+			if (Skeleton == nullptr)
+			{
+				UE_LOG(LogPackageUtilities, Warning, TEXT("Animation (%s) is missing its skeleton. Skipping."), *AnimSeq->GetName());
+				continue;
+			}
+
 			if (Skeleton->HasAnyFlags(RF_NeedLoad))
 			{
 				Skeleton->GetLinker()->Preload(Skeleton);
