Toggling a stat command in PIE will use the editor viewport instead of relying on game's viewport.

diff --git a/Engine/Source/Editor/UnrealEd/Private/UnrealEdSrv.cpp b/Engine/Source/Editor/UnrealEd/Private/UnrealEdSrv.cpp
--- a/Engine/Source/Editor/UnrealEd/Private/UnrealEdSrv.cpp
+++ b/Engine/Source/Editor/UnrealEd/Private/UnrealEdSrv.cpp
@@ -16,6 +16,7 @@
 #include "UObject/UnrealType.h"
 #include "InputCoreTypes.h"
 #include "Input/Reply.h"
+#include "Slate/SceneViewport.h"
 #include "Widgets/DeclarativeSyntaxSupport.h"
 #include "Widgets/SWidget.h"
 #include "Widgets/SCompoundWidget.h"
@@ -690,7 +691,8 @@ bool UUnrealEdEngine::Exec( UWorld* InWorld, const TCHAR* Stream, FOutputDevice&
 	}
 
 	// If we don't have a viewport specified to catch the stat commands, use to the active viewport.  If there is a game viewport ignore this as we do not want 
-	if (GStatProcessingViewportClient == NULL && (GameViewport == NULL || GameViewport->IsSimulateInEditorViewport() ) )
+	if (GStatProcessingViewportClient == NULL && (GameViewport == NULL || GameViewport->IsSimulateInEditorViewport()
+		|| (GameViewport->GetGameViewport() && GameViewport->GetGameViewport()->IsPlayInEditorViewport()) ) )
 	{
 		GStatProcessingViewportClient = GLastKeyLevelEditingViewportClient ? GLastKeyLevelEditingViewportClient : GCurrentLevelEditingViewportClient;
 	}
