WorldPartition: Until it's properly supported by WorldPartition, save dirty packages before starting PIE so that PIE data matches what's in-editor.

diff --git a/Engine/Source/Editor/UnrealEd/Private/PlayLevel.cpp b/Engine/Source/Editor/UnrealEd/Private/PlayLevel.cpp
--- a/Engine/Source/Editor/UnrealEd/Private/PlayLevel.cpp
+++ b/Engine/Source/Editor/UnrealEd/Private/PlayLevel.cpp
@@ -115,6 +115,7 @@
 #include "IAssetViewport.h"
 #include "IPIEAuthorizer.h"
 #include "Features/IModularFeatures.h"
+#include "WorldPartition/WorldPartitionSubsystem.h"
 
 DEFINE_LOG_CATEGORY(LogPlayLevel);
 
@@ -1082,7 +1083,10 @@ void UEditorEngine::StartQueuedPlaySessionRequestImpl()
 		bIsSeparateProcess |= NetMode == EPlayNetMode::PIE_Client;
 	}
 
-	if (bIsSeparateProcess && !SaveMapsForPlaySession())
+	// For now, WorldPartition also requires the user to save their content for PIE to match what we have in-editor.
+	bool bSaveMapsForPlaySession = bIsSeparateProcess || GWorld->HasSubsystem<UWorldPartitionSubsystem>();
+
+	if (bSaveMapsForPlaySession && !SaveMapsForPlaySession())
 	{
 		// Maps did not save, print a warning
 		FText ErrorMsg = LOCTEXT("PIEWorldSaveFail", "PIE failed because map save was canceled");
