Fixed slowtask secondary bar displaying the same text as the main bar

diff --git a/Engine/Source/Editor/UnrealEd/Private/FeedbackContextEditor.cpp b/Engine/Source/Editor/UnrealEd/Private/FeedbackContextEditor.cpp
--- a/Engine/Source/Editor/UnrealEd/Private/FeedbackContextEditor.cpp
+++ b/Engine/Source/Editor/UnrealEd/Private/FeedbackContextEditor.cpp
@@ -156,7 +156,7 @@ public:
 
 		for (int SecondaryBarIndex = 0; SecondaryBarIndex < MaxNumSecondaryBars; ++SecondaryBarIndex)
 		{
-			CreateSecondaryBar(SecondaryBarIndex);
+			CreateSecondaryBar(SecondaryBarIndex + 1); // Index 0 is the main bar
 		}
 
 		// Make sure all our bars are set up
