Fix mismatched delegate unregister in AssetEditorSubsystem

diff --git a/Engine/Source/Editor/UnrealEd/Private/Subsystems/AssetEditorSubsystem.cpp b/Engine/Source/Editor/UnrealEd/Private/Subsystems/AssetEditorSubsystem.cpp
--- a/Engine/Source/Editor/UnrealEd/Private/Subsystems/AssetEditorSubsystem.cpp
+++ b/Engine/Source/Editor/UnrealEd/Private/Subsystems/AssetEditorSubsystem.cpp
@@ -60,7 +60,8 @@ void UAssetEditorSubsystem::Deinitialize()
 {
 	FCoreUObjectDelegates::OnPackageReloaded.RemoveAll(this);
 	GEditor->OnEditorClose().RemoveAll(this);
-	FCoreDelegates::OnPreExit.RemoveAll(this);
+	FCoreDelegates::OnEnginePreExit.RemoveAll(this);
+	FCoreDelegates::OnPostEngineInit.RemoveAll(this);
 
 	// Don't attempt to report usage stats if analytics isn't available
 	if (FEngineAnalytics::IsAvailable())
