Ported viewport drag-duplication to use typed elements

This involved splitting and converting several UUnrealEdEngine functions to no longer operate on or manipulate the current selection set:
 - edactCopySelected:
   - Split into CopyComponents and CopyActors, which take an array of components and actors to copy rather than use the selection state.
   - Breaking: If a DestinationData is provided, it is now used *instead of* the clipboard rather than *as well as* the clipboard.
     - In practice nothing seemed to rely on this old method, and in fact some places stashed and restored the clipboard data to workaround it!
 - edactPasteSelected:
   - Split into PasteComponents and PasteActors, which now returns an array of components or actors that were pasted rather than make them the active selection.
   - Note: PasteActors does still update the selection via FactoryCreateText, but we restore it again afterwards.
 - edactDuplicateSelected:
   - Split into DuplicateComponents and DuplicateActors, which take an array of components or actors to duplicate, and also returns an array of components or actors that were created.

The "edact" versions still use and update the selection state when calling into the internal functions, but the typed element layer does not, and handles duplication via UEngineElementsLibrary::DuplicateElements and UTypedElementWorldInterface, with the level editor viewport handling updating the selection state afterwards.

diff --git a/Engine/Source/Editor/EditorFramework/Private/Elements/Framework/TypedElementViewportInteraction.cpp b/Engine/Source/Editor/EditorFramework/Private/Elements/Framework/TypedElementViewportInteraction.cpp
--- a/Engine/Source/Editor/EditorFramework/Private/Elements/Framework/TypedElementViewportInteraction.cpp
+++ b/Engine/Source/Editor/EditorFramework/Private/Elements/Framework/TypedElementViewportInteraction.cpp
@@ -3,6 +3,7 @@
 #include "Elements/Framework/TypedElementViewportInteraction.h"
 #include "Elements/Framework/TypedElementSelectionSet.h"
 #include "Elements/Framework/TypedElementRegistry.h"
+#include "Elements/Framework/TypedElementUtil.h"
 
 void FTypedElementAssetEditorViewportInteractionCustomization::GetElementsToMove(const TTypedElement<UTypedElementWorldInterface>& InElementWorldHandle, const ETypedElementViewportInteractionWorldType InWorldType, const UTypedElementSelectionSet* InSelectionSet, UTypedElementList* OutElementsToMove)
 {
@@ -154,7 +155,7 @@ void UTypedElementViewportInteraction::BeginGizmoManipulation(const UTypedElemen
 {
 	{
 		TMap<FTypedHandleTypeId, TArray<FTypedElementHandle>> ElementsToMoveByType;
-		BatchElementsByType(InElementsToMove, ElementsToMoveByType);
+		TypedElementUtil::BatchElementsByType(InElementsToMove, ElementsToMoveByType);
 
 		for (const auto& ElementsByTypePair : ElementsToMoveByType)
 		{
@@ -204,7 +205,7 @@ void UTypedElementViewportInteraction::EndGizmoManipulation(const UTypedElementL
 	
 	{
 		TMap<FTypedHandleTypeId, TArray<FTypedElementHandle>> ElementsToMoveByType;
-		BatchElementsByType(InElementsToMove, ElementsToMoveByType);
+		TypedElementUtil::BatchElementsByType(InElementsToMove, ElementsToMoveByType);
 
 		for (const auto& ElementsByTypePair : ElementsToMoveByType)
 		{
@@ -243,14 +244,3 @@ FTypedElementViewportInteractionElement UTypedElementViewportInteraction::Resolv
 		? FTypedElementViewportInteractionElement(UTypedElementRegistry::GetInstance()->GetElement<UTypedElementWorldInterface>(InElementHandle), GetAssetEditorCustomizationByTypeId(InElementHandle.GetId().GetTypeId()))
 		: FTypedElementViewportInteractionElement();
 }
-
-void UTypedElementViewportInteraction::BatchElementsByType(const UTypedElementList* InElementsToMove, TMap<FTypedHandleTypeId, TArray<FTypedElementHandle>>& OutElementsByType)
-{
-	OutElementsByType.Reset();
-	InElementsToMove->ForEachElementHandle([&OutElementsByType](const FTypedElementHandle& InElementHandle)
-	{
-		TArray<FTypedElementHandle>& ElementsForType = OutElementsByType.FindOrAdd(InElementHandle.GetId().GetTypeId());
-		ElementsForType.Add(InElementHandle);
-		return true;
-	});
-}
diff --git a/Engine/Source/Editor/EditorFramework/Public/Elements/Framework/TypedElementViewportInteraction.h b/Engine/Source/Editor/EditorFramework/Public/Elements/Framework/TypedElementViewportInteraction.h
--- a/Engine/Source/Editor/EditorFramework/Public/Elements/Framework/TypedElementViewportInteraction.h
+++ b/Engine/Source/Editor/EditorFramework/Public/Elements/Framework/TypedElementViewportInteraction.h
@@ -96,17 +96,17 @@ public:
 	void GetSelectedElementsToMove(const UTypedElementSelectionSet* InSelectionSet, const ETypedElementViewportInteractionWorldType InWorldType, UTypedElementList* OutElementsToMove) const;
 
 	/**
-	 * Notify that the gizmo is potentially about to start manipulating the transform of the given set of objects (calculated from calling GetSelectedElementsToMove).
+	 * Notify that the gizmo is potentially about to start manipulating the transform of the given set of elements (calculated from calling GetSelectedElementsToMove).
 	 */
 	void BeginGizmoManipulation(const UTypedElementList* InElementsToMove, const UE::Widget::EWidgetMode InWidgetMode);
 
 	/**
-	 * Notify that the gizmo has manipulated the transform of the given set of objects (calculated from calling GetSelectedElementsToMove) by the given delta.
+	 * Notify that the gizmo has manipulated the transform of the given set of elements (calculated from calling GetSelectedElementsToMove) by the given delta.
 	 */
 	void UpdateGizmoManipulation(const UTypedElementList* InElementsToMove, const UE::Widget::EWidgetMode InWidgetMode, const EAxisList::Type InDragAxis, const FInputDeviceState& InInputState, const FTransform& InDeltaTransform);
 	
 	/**
-	 * Notify that the gizmo has finished manipulating the transform of the given set of objects (calculated from calling GetSelectedElementsToMove).
+	 * Notify that the gizmo has finished manipulating the transform of the given set of elements (calculated from calling GetSelectedElementsToMove).
 	 */
 	void EndGizmoManipulation(const UTypedElementList* InElementsToMove, const UE::Widget::EWidgetMode InWidgetMode);
 
@@ -126,9 +126,4 @@ private:
 	 * Attempt to resolve the selection interface and asset editor viewport interaction customization for the given element, if any.
 	 */
 	FTypedElementViewportInteractionElement ResolveViewportInteractionElement(const FTypedElementHandle& InElementHandle) const;
-
-	/**
-	 * Batch the given elements by their type.
-	 */
-	static void BatchElementsByType(const UTypedElementList* InElementsToMove, TMap<FTypedHandleTypeId, TArray<FTypedElementHandle>>& OutElementsByType);
 };
diff --git a/Engine/Source/Editor/UnrealEd/Classes/Editor/EditorEngine.h b/Engine/Source/Editor/UnrealEd/Classes/Editor/EditorEngine.h
--- a/Engine/Source/Editor/UnrealEd/Classes/Editor/EditorEngine.h
+++ b/Engine/Source/Editor/UnrealEd/Classes/Editor/EditorEngine.h
@@ -1297,9 +1297,9 @@ public:
 	 * Copy selected actors to the clipboard.
 	 *
 	 * @param	InWorld					World context
-	 * @param	DestinationData			If != NULL, additionally copy data to string
+	 * @param	DestinationData			If != NULL, fill instead of clipboard data
 	 */
-	virtual void edactCopySelected(UWorld* InWorld, FString* DestinationData = NULL) {}
+	virtual void edactCopySelected(UWorld* InWorld, FString* DestinationData = nullptr) {}
 
 	/**
 	 * Paste selected actors from the clipboard.
@@ -1310,13 +1310,13 @@ public:
 	 * @param	bWarnIfHidden		If true displays a warning if the destination level is hidden
 	 * @param	SourceData			If != NULL, use instead of clipboard data
 	 */
-	virtual void edactPasteSelected(UWorld* InWorld, bool bDuplicate, bool bOffsetLocations, bool bWarnIfHidden, FString* SourceData = NULL) {}
+	virtual void edactPasteSelected(UWorld* InWorld, bool bDuplicate, bool bOffsetLocations, bool bWarnIfHidden, const FString* SourceData = nullptr) {}
 
 	/**
 	 * Duplicates selected actors.
 	 *
-	 * @param	InLevel			Level to place duplicate
-	 * @param	bUseOffset		Should the actor locations be offset after they are created?
+	 * @param	InLevel				Level to place duplicate
+	 * @param	bOffsetLocations	Should the actor locations be offset after they are created?
 	 */
 	virtual void edactDuplicateSelected( ULevel* InLevel, bool bOffsetLocations ) {}
 
diff --git a/Engine/Source/Editor/UnrealEd/Classes/Editor/UnrealEdEngine.h b/Engine/Source/Editor/UnrealEd/Classes/Editor/UnrealEdEngine.h
--- a/Engine/Source/Editor/UnrealEd/Classes/Editor/UnrealEdEngine.h
+++ b/Engine/Source/Editor/UnrealEd/Classes/Editor/UnrealEdEngine.h
@@ -422,7 +422,7 @@
 	 * @return								true unless the delete operation was aborted.
 	 */
 	virtual bool edactDeleteSelected( UWorld* InWorld, bool bVerifyDeletionCanHappen=true, bool bWarnAboutReferences = true, bool bWarnAboutSoftReferences = true) override;
-
+
 	/**
 	 * Creates a new group from the current selection removing any existing groups.
 	 */
@@ -464,28 +464,85 @@
 	 * Copy selected actors to the clipboard.  Does not copy PrefabInstance actors or parts of Prefabs.
 	 *
 	 * @param	InWorld					World context
-	 * @param	DestinationData			If != NULL, additionally copy data to string
+	 * @param	DestinationData			If != NULL, fill instead of clipboard data
 	 */
-	virtual void edactCopySelected(UWorld* InWorld, FString* DestinationData = NULL) override;
+	virtual void edactCopySelected(UWorld* InWorld, FString* DestinationData = nullptr) override;
 
 	/**
+	 * Copy the given components to the clipboard.
+	 *
+	 * @param	InComponentsToCopy		Array of components to copy
+	 * @param	DestinationData			If != NULL, fill instead of clipboard data
+	 */
+	void CopyComponents(const TArray<UActorComponent*>& InComponentsToCopy, FString* DestinationData = nullptr) const;
+
+	/**
+	 * Copy the given actors to the clipboard.  Does not copy PrefabInstance actors or parts of Prefabs.
+	 *
+	 * @param	InActorsToCopy			Array of actors to copy
+	 * @param	InWorld					World context
+	 * @param	DestinationData			If != NULL, fill instead of clipboard data
+	 */
+	void CopyActors(const TArray<AActor*>& InActorsToCopy, UWorld* InWorld, FString* DestinationData = nullptr) const;
+
+	/**
 	 * Paste selected actors from the clipboard.
 	 *
-	 * @param	InWorld					World context
+	 * @param	InWorld				World context
 	 * @param	bDuplicate			Is this a duplicate operation (as opposed to a real paste)?
 	 * @param	bOffsetLocations	Should the actor locations be offset after they are created?
 	 * @param	bWarnIfHidden		If true displays a warning if the destination level is hidden
 	 * @param	SourceData			If != NULL, use instead of clipboard data
 	 */
-	virtual void edactPasteSelected(UWorld* InWorld, bool bDuplicate, bool bOffsetLocations, bool bWarnIfHidden, FString* SourceData = NULL) override;
+	virtual void edactPasteSelected(UWorld* InWorld, bool bDuplicate, bool bOffsetLocations, bool bWarnIfHidden, const FString* SourceData = nullptr) override;
 
 	/**
+	 * Paste the components from the clipboard.
+	 *
+	 * @param	OutPastedComponents List of all the components that were pasted
+	 * @param	TargetActor			The actor to attach the pasted components to
+	 * @param	bWarnIfHidden		If true displays a warning if the destination level is hidden
+	 * @param	SourceData			If != NULL, use instead of clipboard data
+	 */
+	void PasteComponents(TArray<UActorComponent*>& OutPastedComponents, AActor* TargetActor, const bool bWarnIfHidden, const FString* SourceData = nullptr);
+
+	/**
+	 * Paste the actors from the clipboard.
+	 *
+	 * @param	OutPastedActors		List of all the actors that were pasted
+	 * @param	InWorld				World context
+	 * @param	bDuplicate			Is this a duplicate operation (as opposed to a real paste)?
+	 * @param	bOffsetLocations	Should the actor locations be offset after they are created?
+	 * @param	bWarnIfHidden		If true displays a warning if the destination level is hidden
+	 * @param	SourceData			If != NULL, use instead of clipboard data
+	 */
+	void PasteActors(TArray<AActor*>& OutPastedActors, UWorld* InWorld, bool bDuplicate, bool bOffsetLocations, bool bWarnIfHidden, const FString* SourceData = nullptr);
+
+	/**
 	 * Duplicates selected actors.  Handles the case where you are trying to duplicate PrefabInstance actors.
 	 *
 	 * @param	InLevel				Level to place duplicate
+	 * @param	bOffsetLocations	Should the actor locations be offset after they are created?
+	 */
+	virtual void edactDuplicateSelected(ULevel* InLevel, bool bOffsetLocations) override;
+
+	/**
+	 * Duplicate the given components.
+	 * 
+	 * @param	InComponentsToDuplicate		Array of components to duplicate
+	 * @param	OutNewComponents			List of all the components that were duplicated
+	 */
+	void DuplicateComponents(const TArray<UActorComponent*>& InComponentsToDuplicate, TArray<UActorComponent*>& OutNewComponents);
+
+	/**
+	 * Duplicates the given actors.  Handles the case where you are trying to duplicate PrefabInstance actors.
+	 *
+	 * @param	InActorsToDuplicate	Array of actors to duplicate
+	 * @param	OutNewActors		List of all the actors that were duplicated
+	 * @param	InLevel				Level to place duplicate
 	 * @param	bUseOffset			Should the actor locations be offset after they are created?
 	 */
-	virtual void edactDuplicateSelected(ULevel* InLevel, bool bUseOffset) override;
+	void DuplicateActors(const TArray<AActor*>& InActorsToDuplicate, TArray<AActor*>& OutNewActors, ULevel* InLevel, bool bOffsetLocations);
 
 	/**
 	 * Replace all selected brushes with the default brush.
@@ -756,7 +813,7 @@
 	 * 
 	 * @param InWorld			World context
 	 */
-	bool WarnIfDestinationLevelIsHidden( UWorld* InWorld );
+	bool WarnIfDestinationLevelIsHidden( UWorld* InWorld ) const;
 
 	/**
 	 * Generate the package thumbails if they are needed. 
diff --git a/Engine/Source/Editor/UnrealEd/Private/EditorActor.cpp b/Engine/Source/Editor/UnrealEd/Private/EditorActor.cpp
--- a/Engine/Source/Editor/UnrealEd/Private/EditorActor.cpp
+++ b/Engine/Source/Editor/UnrealEd/Private/EditorActor.cpp
@@ -20,6 +20,7 @@
 #include "Exporters/Exporter.h"
 #include "Components/SkeletalMeshComponent.h"
 #include "Engine/Brush.h"
+#include "UnrealEdGlobals.h"
 #include "Editor/EditorEngine.h"
 #include "Editor/UnrealEdEngine.h"
 #include "Factories/LevelFactory.h"
@@ -59,6 +60,8 @@
 #include "Settings/EditorProjectSettings.h"
 #include "ActorGroupingUtils.h"
 #include "HAL/PlatformApplicationMisc.h"
+#include "Serialization/ObjectWriter.h"
+#include "Serialization/ObjectReader.h"
 #include "IAssetTools.h"
 #include "AssetToolsModule.h"
 #include "AssetSelection.h"
@@ -90,71 +93,81 @@ void UUnrealEdEngine::edactCopySelected( UWorld* InWorld, FString* DestinationDa
 	{
 		// Copy selected components
 		TArray<UActorComponent*> SelectedComponents;
-		for (FSelectionIterator It(GetSelectedComponentIterator()); It; ++It)
-		{
-			SelectedComponents.Add(CastChecked<UActorComponent>(*It));
-		}
-
-		FComponentEditorUtils::CopyComponents(SelectedComponents);
+		GetSelectedComponents()->GetSelectedObjects<UActorComponent>(SelectedComponents);
+		CopyComponents(SelectedComponents, DestinationData);
 	}
 	else
 	{
-		// Copy Actors
-		// Before copying, deselect:
-		//		- Actors belonging to prefabs unless all actors in the prefab are selected.
-		//		- Builder brushes.
-		//      - World Settings.
-		TArray<AActor*> ActorsToDeselect;
+		// Copy selected actors
+		TArray<AActor*> SelectedActors;
+		GetSelectedActors()->GetSelectedObjects<AActor>(SelectedActors);
+		CopyActors(SelectedActors, InWorld, DestinationData);
+	}
+}
+
+void UUnrealEdEngine::CopyComponents(const TArray<UActorComponent*>& InComponentsToCopy, FString* DestinationData) const
+{
+	// Remove:
+	// - Components that cannot be copied.
+	TArray<UActorComponent*> ComponentsToCopy = InComponentsToCopy;
+	ComponentsToCopy.RemoveAll([](UActorComponent* InComponent)
+	{
+		return !FComponentEditorUtils::CanCopyComponent(InComponent);
+	});
+
+	FComponentEditorUtils::CopyComponents(ComponentsToCopy, DestinationData);
+}
 
+void UUnrealEdEngine::CopyActors(const TArray<AActor*>& InActorsToCopy, UWorld* InWorld, FString* DestinationData) const
+{
+	// Remove:
+	// - Actors belonging to prefabs unless all actors in the prefab are selected.
+	// - Builder brushes.
+	// - World Settings.
+	TArray<AActor*> ActorsToCopy = InActorsToCopy;
+	{
 		bool bSomeSelectedActorsNotInCurrentLevel = false;
-		for (FSelectionIterator It(GetSelectedActorIterator()); It; ++It)
+		ActorsToCopy.RemoveAll([&bSomeSelectedActorsNotInCurrentLevel](AActor* InActor)
 		{
-			AActor* Actor = static_cast<AActor*>( *It );
-			checkSlow(Actor->IsA(AActor::StaticClass()));
-
-			// Deselect any selected builder brushes.
-			ABrush* Brush = Cast< ABrush >(Actor);
-			const bool bActorIsBuilderBrush = ( Brush && FActorEditorUtils::IsABuilderBrush(Brush) );
-			if (bActorIsBuilderBrush)
+			// Remove any selected builder brushes.
+			ABrush* Brush = Cast<ABrush>(InActor);
+			if (Brush && FActorEditorUtils::IsABuilderBrush(Brush))
 			{
-				ActorsToDeselect.Add(Actor);
+				return true;
 			}
 
-			// Deselect world settings
-			if (Actor->IsA( AWorldSettings::StaticClass() ))
+			// Remove world settings
+			if (InActor->IsA<AWorldSettings>())
 			{
-				ActorsToDeselect.Add(Actor);
+				return true;
 			}
 
 			// If any selected actors are not in the current level, warn the user that some actors will not be copied.
-			if (!bSomeSelectedActorsNotInCurrentLevel && !Actor->GetLevel()->IsCurrentLevel())
+			if (!bSomeSelectedActorsNotInCurrentLevel && !InActor->GetLevel()->IsCurrentLevel())
 			{
 				bSomeSelectedActorsNotInCurrentLevel = true;
 				FMessageDialog::Open(EAppMsgType::Ok, NSLOCTEXT("UnrealEd", "CopySelectedActorsInNonCurrentLevel", "Some selected actors are not in the current level and will not be copied."));
 			}
-		}
 
-		const FScopedBusyCursor BusyCursor;
-		for (int32 ActorIndex = 0; ActorIndex < ActorsToDeselect.Num(); ++ActorIndex)
-		{
-			AActor* Actor = ActorsToDeselect[ActorIndex];
-			GetSelectedActors()->Deselect(Actor);
-		}
+			return false;
+		});
+	}
 
-		// Export the actors.
-		FStringOutputDevice Ar;
-		const FSelectedActorExportObjectInnerContext Context;
-		UExporter::ExportToOutputDevice(&Context, InWorld, NULL, Ar, TEXT("copy"), 0, PPF_DeepCompareInstances | PPF_ExportsNotFullyQualified);
+	// Export the actors.
+	FStringOutputDevice Ar;
+	const FSelectedActorExportObjectInnerContext Context(ActorsToCopy);
+	UExporter::ExportToOutputDevice(&Context, InWorld, NULL, Ar, TEXT("copy"), 0, PPF_DeepCompareInstances | PPF_ExportsNotFullyQualified);
+	if (DestinationData)
+	{
+		*DestinationData = MoveTemp(Ar);
+	}
+	else
+	{
 		FPlatformApplicationMisc::ClipboardCopy(*Ar);
-		if (DestinationData)
-		{
-			*DestinationData = MoveTemp(Ar);
-		}
 	}
 }
 
 
-
 /**
  * Creates offsets for locations based on the editor grid size and active viewport.
  */
@@ -181,7 +194,7 @@ static FVector CreateLocationOffset(bool bDuplicate, bool bOffsetLocations)
 }
 
 
-bool UUnrealEdEngine::WarnIfDestinationLevelIsHidden( UWorld* InWorld )
+bool UUnrealEdEngine::WarnIfDestinationLevelIsHidden( UWorld* InWorld ) const
 {
 	bool bShouldLoadHiddenLevels = true;
 	bool bShowPasteHiddenWarning = true;
@@ -245,36 +258,26 @@ bool UUnrealEdEngine::WarnIfDestinationLevelIsHidden( UWorld* InWorld )
 	return !bShouldLoadHiddenLevels;
 }
 
-void UUnrealEdEngine::edactPasteSelected(UWorld* InWorld, bool bDuplicate, bool bOffsetLocations, bool bWarnIfHidden, FString* SourceData)
+void UUnrealEdEngine::edactPasteSelected(UWorld* InWorld, bool bDuplicate, bool bOffsetLocations, bool bWarnIfHidden, const FString* SourceData)
 {
-	//check and warn if the user is trying to paste to a hidden level. This will return if he wishes to abort the process
-	if( bWarnIfHidden && WarnIfDestinationLevelIsHidden( InWorld ) == true )
-	{
-		return;
-	}
-	
 	if (GetSelectedComponentCount() > 0)
 	{
 		AActor* SelectedActor = CastChecked<AActor>(*GetSelectedActorIterator());
 
 		TArray<UActorComponent*> PastedComponents;
-		FComponentEditorUtils::PasteComponents(PastedComponents, SelectedActor, SelectedActor->GetRootComponent());
+		PasteComponents(PastedComponents, SelectedActor, bWarnIfHidden, SourceData);
 
 		if (PastedComponents.Num() > 0)
 		{
-			// Make sure all the SCS trees have a chance to update
-			FLevelEditorModule& LevelEditor = FModuleManager::LoadModuleChecked<FLevelEditorModule>("LevelEditor");
-			LevelEditor.BroadcastComponentsEdited();
-
 			// Select the new clones
 			USelection* ComponentSelection = GetSelectedComponents();
-			ComponentSelection->Modify(false);
+			ComponentSelection->Modify();
 			ComponentSelection->BeginBatchSelectOperation();
 			ComponentSelection->DeselectAll();
 
-			for (UActorComponent* PastedComp : PastedComponents)
+			for (UActorComponent* PastedComponent : PastedComponents)
 			{
-				GEditor->SelectComponent(PastedComp, true, false);
+				GEditor->SelectComponent(PastedComponent, true, false);
 			}
 
 			ComponentSelection->EndBatchSelectOperation(true);
@@ -282,85 +285,139 @@ void UUnrealEdEngine::edactPasteSelected(UWorld* InWorld, bool bDuplicate, bool
 	}
 	else
 	{
-		const FScopedBusyCursor BusyCursor;
+		TArray<AActor*> PastedActors;
+		PasteActors(PastedActors, InWorld, bDuplicate, bOffsetLocations, bWarnIfHidden, SourceData);
 
-		// Create a location offset.
-		const FVector LocationOffset = CreateLocationOffset(bDuplicate, bOffsetLocations);
+		if (PastedActors.Num() > 0)
+		{
+			// Select the new clones
+			USelection* ActorSelection = GetSelectedActors();
+			ActorSelection->Modify();
+			ActorSelection->BeginBatchSelectOperation();
+			ActorSelection->DeselectAll();
 
-		FCachedActorLabels ActorLabels(InWorld);
+			for (AActor* PastedActor : PastedActors)
+			{
+				GEditor->SelectActor(PastedActor, true, false);
+			}
 
-		// Transact the current selection set.
-		USelection* SelectedActors = GetSelectedActors();
-		SelectedActors->Modify();
+			ActorSelection->EndBatchSelectOperation(true);
 
-		// Get pasted text.
-		FString PasteString;
-		if (SourceData)
-		{
-			PasteString = *SourceData;
-		}
-		else
-		{
-			FPlatformApplicationMisc::ClipboardPaste(PasteString);
+			// Note the selection change.  This will also redraw level viewports and update the pivot.
+			NoteSelectionChange();
 		}
-		const TCHAR* Paste = *PasteString;
+	}
+}
 
-		// Turn off automatic BSP update while pasting to save rebuilding geometry potentially multiple times
-		const bool bBSPAutoUpdate = GetDefault<ULevelEditorMiscSettings>()->bBSPAutoUpdate;
-		GetMutableDefault<ULevelEditorMiscSettings>()->bBSPAutoUpdate = false;
+void UUnrealEdEngine::PasteComponents(TArray<UActorComponent*>& OutPastedComponents, AActor* TargetActor, const bool bWarnIfHidden, const FString* SourceData)
+{
+	OutPastedComponents.Reset();
 
-		// Import the actors.
-		ULevelFactory* Factory = NewObject<ULevelFactory>();
-		Factory->FactoryCreateText(ULevel::StaticClass(), InWorld->GetCurrentLevel(), InWorld->GetCurrentLevel()->GetFName(), RF_Transactional, NULL, bDuplicate ? TEXT("move") : TEXT("paste"), Paste, Paste + FCString::Strlen(Paste), GWarn);
+	// Check and warn if the user is trying to paste to a hidden level. This will return if they wish to abort the process...
+	if (bWarnIfHidden && WarnIfDestinationLevelIsHidden(TargetActor->GetWorld()))
+	{
+		return;
+	}
 
-		// Reinstate old BSP update setting, and force a rebuild - any levels whose geometry has changed while pasting will be rebuilt
-		GetMutableDefault<ULevelEditorMiscSettings>()->bBSPAutoUpdate = bBSPAutoUpdate;
-		RebuildAlteredBSP();
+	FComponentEditorUtils::PasteComponents(OutPastedComponents, TargetActor, TargetActor->GetRootComponent(), SourceData);
 
-		// Fire ULevel::LevelDirtiedEvent when falling out of scope.
-		FScopedLevelDirtied			LevelDirtyCallback;
+	if (OutPastedComponents.Num() > 0)
+	{
+		// Make sure all the SCS trees have a chance to update
+		FLevelEditorModule& LevelEditor = FModuleManager::LoadModuleChecked<FLevelEditorModule>("LevelEditor");
+		LevelEditor.BroadcastComponentsEdited();
+	}
+}
 
-		// Update the actors' locations and update the global list of visible layers.
-		ULayersSubsystem* LayersSubsystem = GEditor->GetEditorSubsystem<ULayersSubsystem>();
-		for (FSelectionIterator It(GetSelectedActorIterator()); It; ++It)
-		{
-			AActor* Actor = static_cast<AActor*>( *It );
-			checkSlow(Actor->IsA(AActor::StaticClass()));
+void UUnrealEdEngine::PasteActors(TArray<AActor*>& OutPastedActors, UWorld* InWorld, bool bDuplicate, bool bOffsetLocations, bool bWarnIfHidden, const FString* SourceData)
+{
+	OutPastedActors.Reset();
 
-			// We only want to offset the location if this actor is the root of a selected attachment hierarchy
-			// Offsetting children of an attachment hierarchy would cause them to drift away from the node they're attached to
-			// as the offset would effectively get applied twice
-			const AActor* const ParentActor = Actor->GetAttachParentActor();
-			const FVector& ActorLocationOffset = ( ParentActor && ParentActor->IsSelected() ) ? FVector::ZeroVector : LocationOffset;
+	// Check and warn if the user is trying to paste to a hidden level. This will return if they wish to abort the process...
+	if (bWarnIfHidden && WarnIfDestinationLevelIsHidden(InWorld))
+	{
+		return;
+	}
 
-			// Offset the actor's location.
-			Actor->TeleportTo(Actor->GetActorLocation() + ActorLocationOffset, Actor->GetActorRotation(), false, true);
+	const FScopedBusyCursor BusyCursor;
 
-			if (!GetDefault<ULevelEditorMiscSettings>()->bAvoidRelabelOnPasteSelected)
-			{
-				// Re-label duplicated actors so that labels become unique
-				FActorLabelUtilities::SetActorLabelUnique(Actor, Actor->GetActorLabel(), &ActorLabels);
-				ActorLabels.Add(Actor->GetActorLabel());
-			}
+	// Create a location offset.
+	const FVector LocationOffset = CreateLocationOffset(bDuplicate, bOffsetLocations);
+
+	FCachedActorLabels ActorLabels(InWorld);
+
+	// Get pasted text.
+	FString PasteString;
+	if (SourceData)
+	{
+		PasteString = *SourceData;
+	}
+	else
+	{
+		FPlatformApplicationMisc::ClipboardPaste(PasteString);
+	}
+	const TCHAR* Paste = *PasteString;
 
-			LayersSubsystem->InitializeNewActorLayers(Actor);
+	USelection* ActorSelection = GetSelectedActors();
 
-			// Ensure any layers this actor belongs to are visible
-			LayersSubsystem->SetLayersVisibility(Actor->Layers, true);
+	// FactoryCreateText will mutate the selection, so cache it here so we can restore it later
+	TArray<uint8> OriginalSelectionState;
+	FObjectWriter(ActorSelection, OriginalSelectionState);
 
-			Actor->CheckDefaultSubobjects();
-			Actor->InvalidateLightingCache();
-			// Call PostEditMove to update components, etc.
-			Actor->PostEditMove(true);
-			Actor->PostDuplicate(EDuplicateMode::Normal);
-			Actor->CheckDefaultSubobjects();
+	// Turn off automatic BSP update while pasting to save rebuilding geometry potentially multiple times
+	const bool bBSPAutoUpdate = GetDefault<ULevelEditorMiscSettings>()->bBSPAutoUpdate;
+	GetMutableDefault<ULevelEditorMiscSettings>()->bBSPAutoUpdate = false;
 
-			// Request saves/refreshes.
-			Actor->MarkPackageDirty();
-			LevelDirtyCallback.Request();
+	// Import the actors.
+	ULevelFactory* Factory = NewObject<ULevelFactory>();
+	Factory->FactoryCreateText(ULevel::StaticClass(), InWorld->GetCurrentLevel(), InWorld->GetCurrentLevel()->GetFName(), RF_Transactional, NULL, bDuplicate ? TEXT("move") : TEXT("paste"), Paste, Paste + FCString::Strlen(Paste), GWarn);
+
+	// Reinstate old BSP update setting, and force a rebuild - any levels whose geometry has changed while pasting will be rebuilt
+	GetMutableDefault<ULevelEditorMiscSettings>()->bBSPAutoUpdate = bBSPAutoUpdate;
+	RebuildAlteredBSP();
+
+	// FactoryCreateText set the selection to the new actors, so copy that into OutPastedActors and restore the original selection
+	ActorSelection->GetSelectedObjects<AActor>(OutPastedActors);
+	FObjectReader(ActorSelection, OriginalSelectionState);
+
+	// Fire ULevel::LevelDirtiedEvent when falling out of scope.
+	FScopedLevelDirtied			LevelDirtyCallback;
+
+	// Update the actors' locations and update the global list of visible layers.
+	ULayersSubsystem* LayersSubsystem = GEditor->GetEditorSubsystem<ULayersSubsystem>();
+	for (AActor* Actor : OutPastedActors)
+	{
+		// We only want to offset the location if this actor is the root of a selected attachment hierarchy
+		// Offsetting children of an attachment hierarchy would cause them to drift away from the node they're attached to
+		// as the offset would effectively get applied twice
+		const AActor* const ParentActor = Actor->GetAttachParentActor();
+		const FVector& ActorLocationOffset = (ParentActor && ParentActor->IsSelected()) ? FVector::ZeroVector : LocationOffset;
+
+		// Offset the actor's location.
+		Actor->TeleportTo(Actor->GetActorLocation() + ActorLocationOffset, Actor->GetActorRotation(), false, true);
+
+		if (!GetDefault<ULevelEditorMiscSettings>()->bAvoidRelabelOnPasteSelected)
+		{
+			// Re-label duplicated actors so that labels become unique
+			FActorLabelUtilities::SetActorLabelUnique(Actor, Actor->GetActorLabel(), &ActorLabels);
+			ActorLabels.Add(Actor->GetActorLabel());
 		}
-		// Note the selection change.  This will also redraw level viewports and update the pivot.
-		NoteSelectionChange();
+
+		LayersSubsystem->InitializeNewActorLayers(Actor);
+
+		// Ensure any layers this actor belongs to are visible
+		LayersSubsystem->SetLayersVisibility(Actor->Layers, true);
+
+		Actor->CheckDefaultSubobjects();
+		Actor->InvalidateLightingCache();
+		// Call PostEditMove to update components, etc.
+		Actor->PostEditMove(true);
+		Actor->PostDuplicate(EDuplicateMode::Normal);
+		Actor->CheckDefaultSubobjects();
+
+		// Request saves/refreshes.
+		Actor->MarkPackageDirty();
+		LevelDirtyCallback.Request();
 	}
 }
 
@@ -388,7 +445,7 @@ public:
 	void DuplicateActorsToLevel(TArray<AActor*>& OutNewActors, ULevel* DestLevel, bool bOffsetLocations)
 	{
 		// Check neither level is locked
-		if ( FLevelUtils::IsLevelLocked(SrcLevel) )
+		if (FLevelUtils::IsLevelLocked(SrcLevel))
 		{
 			UE_LOG(LogEditorActor, Warning, TEXT("DuplicateActorsToLevel: The requested operation could not be completed because the level is locked."));
 			return;
@@ -405,39 +462,35 @@ public:
 			return;
 		}
 
-		// Cache the current source level
-		ULevel* OldLevel = SrcLevel->OwningWorld->GetCurrentLevel();
-		// Set the selection set to be precisely the actors belonging to this job.
-		SrcLevel->OwningWorld->SetCurrentLevel( SrcLevel );
-		GEditor->SelectNone( false, true );
-		for ( int32 ActorIndex = 0 ; ActorIndex < Actors.Num() ; ++ActorIndex )
+		FString ScratchData;
+
+		// Copy actors from src level
 		{
-			AActor* Actor = Actors[ ActorIndex ];
-			GEditor->SelectActor( Actor, true, false, true );
+			// Cache the current source level
+			ULevel* CurrentSrcLevel = SrcLevel->OwningWorld->GetCurrentLevel();
+			GUnrealEd->CopyActors(Actors, SrcLevel->OwningWorld, &ScratchData);
+			SrcLevel->OwningWorld->SetCurrentLevel(CurrentSrcLevel);
 		}
 
-		FString ScratchData;
-		// Copy actors from src level.
-		GEditor->edactCopySelected( SrcLevel->OwningWorld, &ScratchData );
-		// Restore source level
-		SrcLevel->OwningWorld->SetCurrentLevel( OldLevel );
-
-		// Cache the current dest level
-		OldLevel = DestLevel->OwningWorld->GetCurrentLevel();
-		// Paste to the dest level.
+		// Paste to the dest level
 		{
+			ULevel* CurrentDestLevel = DestLevel->OwningWorld->GetCurrentLevel();
+
 			FLevelPartitionOperationScope LevelPartitionScope(DestLevel);
 			DestLevel->OwningWorld->SetCurrentLevel(LevelPartitionScope.GetLevel());
-			GEditor->edactPasteSelected(DestLevel->OwningWorld, true, bOffsetLocations, true, &ScratchData);
-			// Restore dest level
-			DestLevel->OwningWorld->SetCurrentLevel(OldLevel);
-		}
-		// The selection set will be the newly created actors; copy them over to the output array.
-		for ( FSelectionIterator It( GEditor->GetSelectedActorIterator() ) ; It ; ++It )
-		{
-			AActor* Actor = static_cast<AActor*>( *It );
-			checkSlow( Actor->IsA(AActor::StaticClass()) );
-			OutNewActors.Add( Actor );
+
+			TArray<AActor*> PastedActors;
+			GUnrealEd->PasteActors(PastedActors, DestLevel->OwningWorld, /*bDuplicate*/true, bOffsetLocations, /*bWarnIfHidden*/true, &ScratchData);
+			
+			DestLevel->OwningWorld->SetCurrentLevel(CurrentDestLevel);
+
+			// We generate new seeds when we duplicate
+			for (AActor* Actor : PastedActors)
+			{
+				Actor->SeedAllRandomStreams();
+			}
+
+			OutNewActors.Append(MoveTemp(PastedActors));
 		}
 	}
 };
@@ -446,41 +499,26 @@ public:
 
 void UUnrealEdEngine::edactDuplicateSelected( ULevel* InLevel, bool bOffsetLocations )
 {
-	int32 NumSelectedComponents = GetSelectedComponentCount();
-	if (NumSelectedComponents > 0)
+	if (GetSelectedComponentCount() > 0)
 	{
-		TArray<UActorComponent*> NewComponentClones;
-		NewComponentClones.Reserve(NumSelectedComponents);
+		USelection* ComponentSelection = GetSelectedComponents();
 
-		// Duplicate selected components
-		for (FSelectionIterator It(GetSelectedComponentIterator()); It; ++It)
-		{
-			UActorComponent* Component = CastChecked<UActorComponent>(*It);
+		TArray<UActorComponent*> SelectedComponents;
+		ComponentSelection->GetSelectedObjects<UActorComponent>(SelectedComponents);
 
-			if (FComponentEditorUtils::CanCopyComponent(Component))
-			{
-				if (UActorComponent* Clone = FComponentEditorUtils::DuplicateComponent(Component))
-				{
-					NewComponentClones.Add(Clone);
-				}
-			}
-		}
+		TArray<UActorComponent*> NewComponents;
+		DuplicateComponents(SelectedComponents, NewComponents);
 
-		if (NewComponentClones.Num() > 0)
+		if (NewComponents.Num() > 0)
 		{
-			// Make sure all the SCS trees have a chance to update
-			FLevelEditorModule& LevelEditor = FModuleManager::LoadModuleChecked<FLevelEditorModule>("LevelEditor");
-			LevelEditor.BroadcastComponentsEdited();
-
 			// Select the new clones
-			USelection* ComponentSelection = GetSelectedComponents();
-			ComponentSelection->Modify(false);
+			ComponentSelection->Modify();
 			ComponentSelection->BeginBatchSelectOperation();
 			ComponentSelection->DeselectAll();
 
-			for (UActorComponent* Clone : NewComponentClones)
+			for (UActorComponent* NewComponent : NewComponents)
 			{
-				GEditor->SelectComponent(Clone, true, false);
+				GEditor->SelectComponent(NewComponent, true, false);
 			}
 
 			ComponentSelection->EndBatchSelectOperation(true);
@@ -488,96 +526,89 @@ void UUnrealEdEngine::edactDuplicateSelected( ULevel* InLevel, bool bOffsetLocat
 	}
 	else
 	{
-		const FScopedBusyCursor BusyCursor;
-		GetSelectedActors()->Modify();
+		USelection* ActorSelection = GetSelectedActors();
 
-		// Create per-level job lists.
-		typedef TMap<ULevel*, FDuplicateJob*>	DuplicateJobMap;
-		DuplicateJobMap							DuplicateJobs;
+		TArray<AActor*> SelectedActors;
+		ActorSelection->GetSelectedObjects<AActor>(SelectedActors);
 
-		// Build set of selected actors before duplication
-		TArray<AActor*> PreDuplicateSelection;
+		TArray<AActor*> NewActors;
+		DuplicateActors(SelectedActors, NewActors, InLevel, bOffsetLocations);
 
-		// Add selected actors to the per-level job lists.
-		bool bHaveActorLocation = false;
-		FVector AnyActorLocation = FVector::ZeroVector;
-		for (FSelectionIterator It(GetSelectedActorIterator()); It; ++It)
+		if (NewActors.Num() > 0)
 		{
-			AActor* Actor = static_cast<AActor*>( *It );
-			checkSlow(Actor->IsA(AActor::StaticClass()));
-
-			if (!bHaveActorLocation)
-			{
-				bHaveActorLocation = true;
-				AnyActorLocation = Actor->GetActorLocation();
-			}
-
-			PreDuplicateSelection.Add(Actor);
+			// Select the new clones
+			ActorSelection->Modify();
+			ActorSelection->BeginBatchSelectOperation();
+			ActorSelection->DeselectAll();
 
-			ULevel* OldLevel = Actor->GetLevel();
-			FDuplicateJob** Job = DuplicateJobs.Find(OldLevel);
-			if (Job)
+			for (AActor* Actor : NewActors)
 			{
-				( *Job )->Actors.Add(Actor);
+				GEditor->SelectActor(Actor, true, false);
 			}
-			else
-			{
-				// Allocate a new job for the level.
-				FDuplicateJob* NewJob = new FDuplicateJob;
-				NewJob->SrcLevel = OldLevel;
-				NewJob->Actors.Add(Actor);
-				DuplicateJobs.Add(OldLevel, NewJob);
-			}
-		}
 
-		UWorld* World = InLevel->OwningWorld;
-		ULevel* DesiredLevel = InLevel;
+			ActorSelection->EndBatchSelectOperation(true);
 
-		USelection* SelectedActors = GetSelectedActors();
-		SelectedActors->BeginBatchSelectOperation();
-		SelectedActors->Modify();
+			// Note the selection change.  This will also redraw level viewports and update the pivot.
+			NoteSelectionChange();
 
-		// For each level, select the actors in that level and copy-paste into the destination level.
-		TArray<AActor*>	NewActors;
-		for (DuplicateJobMap::TIterator It(DuplicateJobs); It; ++It)
-		{
-			FDuplicateJob* Job = It.Value();
-			check(Job);
-			Job->DuplicateActorsToLevel(NewActors, InLevel, bOffsetLocations);
+			GLevelEditorModeTools().ActorsDuplicatedNotify(SelectedActors, NewActors, bOffsetLocations);
 		}
+	}
+}
 
-		// Select any newly created actors and prefabs.
-		SelectNone(false, true);
-		for (int32 ActorIndex = 0; ActorIndex < NewActors.Num(); ++ActorIndex)
-		{
-			AActor* Actor = NewActors[ActorIndex];
-			SelectActor(Actor, true, false);
-		}
-		SelectedActors->EndBatchSelectOperation();
-		NoteSelectionChange();
+void UUnrealEdEngine::DuplicateComponents(const TArray<UActorComponent*>& InComponentsToDuplicate, TArray<UActorComponent*>& OutNewComponents)
+{
+	OutNewComponents.Reset(InComponentsToDuplicate.Num());
 
-		// Finally, cleanup.
-		for (DuplicateJobMap::TIterator It(DuplicateJobs); It; ++It)
+	for (UActorComponent* Component : InComponentsToDuplicate)
+	{
+		if (FComponentEditorUtils::CanCopyComponent(Component))
 		{
-			FDuplicateJob* Job = It.Value();
-			delete Job;
+			if (UActorComponent* NewComponent = FComponentEditorUtils::DuplicateComponent(Component))
+			{
+				OutNewComponents.Add(NewComponent);
+			}
 		}
+	}
 
-		// Build set of selected actors after duplication
-		TArray<AActor*> PostDuplicateSelection;
-		for (FSelectionIterator It(GetSelectedActorIterator()); It; ++It)
-		{
-			AActor* Actor = static_cast<AActor*>( *It );
-			checkSlow(Actor->IsA(AActor::StaticClass()));
+	if (OutNewComponents.Num() > 0)
+	{
+		// Make sure all the SCS trees have a chance to update
+		FLevelEditorModule& LevelEditor = FModuleManager::LoadModuleChecked<FLevelEditorModule>("LevelEditor");
+		LevelEditor.BroadcastComponentsEdited();
+	}
+}
 
-			// We generate new seeds when we duplicate
-			Actor->SeedAllRandomStreams();
+void UUnrealEdEngine::DuplicateActors(const TArray<AActor*>& InActorsToDuplicate, TArray<AActor*>& OutNewActors, ULevel* InLevel, bool bOffsetLocations)
+{
+	OutNewActors.Reset(InActorsToDuplicate.Num());
 
-			PostDuplicateSelection.Add(Actor);
+	const FScopedBusyCursor BusyCursor;
+
+	// Create per-level job lists.
+	TMap<ULevel*, TSharedPtr<FDuplicateJob>> DuplicateJobs;
+	for (AActor* Actor : InActorsToDuplicate)
+	{
+		ULevel* SrcLevel = Actor->GetLevel();
+		TSharedPtr<FDuplicateJob> Job = DuplicateJobs.FindRef(SrcLevel);
+		if (!Job)
+		{
+			// Allocate a new job for the level.
+			Job = DuplicateJobs.Add(SrcLevel, MakeShared<FDuplicateJob>());
+			Job->SrcLevel = SrcLevel;
 		}
+		Job->Actors.Add(Actor);
+	}
 
-		GLevelEditorModeTools().ActorsDuplicatedNotify(PreDuplicateSelection, PostDuplicateSelection, bOffsetLocations);
+	// For each level, select the actors in that level and copy-paste into the destination level.
+	for (const auto& DuplicateJobPair : DuplicateJobs)
+	{
+		FDuplicateJob& Job = *DuplicateJobPair.Value;
+		Job.DuplicateActorsToLevel(OutNewActors, InLevel, bOffsetLocations);
 	}
+
+	// Finally, cleanup.
+	DuplicateJobs.Reset();
 }
 
 bool UUnrealEdEngine::CanDeleteSelectedActors( const UWorld* InWorld, const bool bStopAtFirst, const bool bLogUndeletable, TArray<AActor*>* OutDeletableActors ) const
diff --git a/Engine/Source/Editor/UnrealEd/Private/Elements/Actor/ActorElementEditorWorldInterface.cpp b/Engine/Source/Editor/UnrealEd/Private/Elements/Actor/ActorElementEditorWorldInterface.cpp
--- a/Engine/Source/Editor/UnrealEd/Private/Elements/Actor/ActorElementEditorWorldInterface.cpp
+++ b/Engine/Source/Editor/UnrealEd/Private/Elements/Actor/ActorElementEditorWorldInterface.cpp
@@ -4,7 +4,11 @@
 #include "Elements/Actor/ActorElementData.h"
 #include "GameFramework/Actor.h"
 
+#include "Elements/Framework/EngineElementsLibrary.h"
+
 #include "Editor.h"
+#include "UnrealEdGlobals.h"
+#include "Editor/UnrealEdEngine.h"
 
 void UActorElementEditorWorldInterface::NotifyMovementStarted(const FTypedElementHandle& InElementHandle)
 {
@@ -34,3 +38,35 @@ void UActorElementEditorWorldInterface::NotifyMovementEnded(const FTypedElementH
 		Actor->MarkPackageDirty();
 	}
 }
+
+void UActorElementEditorWorldInterface::DuplicateElements(TArrayView<const FTypedElementHandle> InElementHandles, UWorld* InWorld, bool bOffsetLocations, TArray<FTypedElementHandle>& OutNewElements)
+{
+	const TArray<AActor*> ActorsToDuplicate = ActorElementDataUtil::GetActorsFromHandles(InElementHandles);
+
+	if (ActorsToDuplicate.Num() > 0)
+	{
+		FEditorDelegates::OnDuplicateActorsBegin.Broadcast();
+
+		TArray<AActor*> NewActors;
+		ABrush::SetSuppressBSPRegeneration(true);
+		GUnrealEd->DuplicateActors(ActorsToDuplicate, NewActors, InWorld->GetCurrentLevel(), bOffsetLocations);
+		ABrush::SetSuppressBSPRegeneration(false);
+
+		FEditorDelegates::OnDuplicateActorsEnd.Broadcast();
+
+		bool bRebuildBSP = false;
+		OutNewElements.Reserve(OutNewElements.Num() + NewActors.Num());
+		for (AActor* NewActor : NewActors)
+		{
+			// Only rebuild if the new actors will change the BSP as this is expensive
+			bRebuildBSP |= NewActor->IsA<ABrush>();
+
+			OutNewElements.Add(UEngineElementsLibrary::AcquireEditorActorElementHandle(NewActor));
+		}
+
+		if (bRebuildBSP)
+		{
+			GEditor->RebuildAlteredBSP();
+		}
+	}
+}
diff --git a/Engine/Source/Editor/UnrealEd/Private/Elements/Actor/ActorElementEditorWorldInterface.h b/Engine/Source/Editor/UnrealEd/Private/Elements/Actor/ActorElementEditorWorldInterface.h
--- a/Engine/Source/Editor/UnrealEd/Private/Elements/Actor/ActorElementEditorWorldInterface.h
+++ b/Engine/Source/Editor/UnrealEd/Private/Elements/Actor/ActorElementEditorWorldInterface.h
@@ -14,4 +14,5 @@ public:
 	virtual void NotifyMovementStarted(const FTypedElementHandle& InElementHandle) override;
 	virtual void NotifyMovementOngoing(const FTypedElementHandle& InElementHandle) override;
 	virtual void NotifyMovementEnded(const FTypedElementHandle& InElementHandle) override;
+	virtual void DuplicateElements(TArrayView<const FTypedElementHandle> InElementHandles, UWorld* InWorld, bool bOffsetLocations, TArray<FTypedElementHandle>& OutNewElements) override;
 };
diff --git a/Engine/Source/Editor/UnrealEd/Private/Elements/Component/ComponentElementEditorWorldInterface.cpp b/Engine/Source/Editor/UnrealEd/Private/Elements/Component/ComponentElementEditorWorldInterface.cpp
--- a/Engine/Source/Editor/UnrealEd/Private/Elements/Component/ComponentElementEditorWorldInterface.cpp
+++ b/Engine/Source/Editor/UnrealEd/Private/Elements/Component/ComponentElementEditorWorldInterface.cpp
@@ -4,7 +4,11 @@
 #include "Elements/Component/ComponentElementData.h"
 #include "Components/ActorComponent.h"
 
+#include "Elements/Framework/EngineElementsLibrary.h"
+
 #include "Editor.h"
+#include "UnrealEdGlobals.h"
+#include "Editor/UnrealEdEngine.h"
 
 void UComponentElementEditorWorldInterface::NotifyMovementStarted(const FTypedElementHandle& InElementHandle)
 {
@@ -39,3 +43,20 @@ void UComponentElementEditorWorldInterface::NotifyMovementEnded(const FTypedElem
 		Component->MarkPackageDirty();
 	}
 }
+
+void UComponentElementEditorWorldInterface::DuplicateElements(TArrayView<const FTypedElementHandle> InElementHandles, UWorld* InWorld, bool bOffsetLocations, TArray<FTypedElementHandle>& OutNewElements)
+{
+	const TArray<UActorComponent*> ComponentsToDuplicate = ComponentElementDataUtil::GetComponentsFromHandles(InElementHandles);
+
+	if (ComponentsToDuplicate.Num() > 0)
+	{
+		TArray<UActorComponent*> NewComponents;
+		GUnrealEd->DuplicateComponents(ComponentsToDuplicate, NewComponents);
+
+		OutNewElements.Reserve(OutNewElements.Num() + NewComponents.Num());
+		for (UActorComponent* NewComponent : NewComponents)
+		{
+			OutNewElements.Add(UEngineElementsLibrary::AcquireEditorComponentElementHandle(NewComponent));
+		}
+	}
+}
diff --git a/Engine/Source/Editor/UnrealEd/Private/Elements/Component/ComponentElementEditorWorldInterface.h b/Engine/Source/Editor/UnrealEd/Private/Elements/Component/ComponentElementEditorWorldInterface.h
--- a/Engine/Source/Editor/UnrealEd/Private/Elements/Component/ComponentElementEditorWorldInterface.h
+++ b/Engine/Source/Editor/UnrealEd/Private/Elements/Component/ComponentElementEditorWorldInterface.h
@@ -14,4 +14,5 @@ public:
 	virtual void NotifyMovementStarted(const FTypedElementHandle& InElementHandle) override;
 	virtual void NotifyMovementOngoing(const FTypedElementHandle& InElementHandle) override;
 	virtual void NotifyMovementEnded(const FTypedElementHandle& InElementHandle) override;
+	virtual void DuplicateElements(TArrayView<const FTypedElementHandle> InElementHandles, UWorld* InWorld, bool bOffsetLocations, TArray<FTypedElementHandle>& OutNewElements) override;
 };
diff --git a/Engine/Source/Editor/UnrealEd/Private/Kismet2/ComponentEditorUtils.cpp b/Engine/Source/Editor/UnrealEd/Private/Kismet2/ComponentEditorUtils.cpp
--- a/Engine/Source/Editor/UnrealEd/Private/Kismet2/ComponentEditorUtils.cpp
+++ b/Engine/Source/Editor/UnrealEd/Private/Kismet2/ComponentEditorUtils.cpp
@@ -451,7 +451,7 @@ bool FComponentEditorUtils::CanCopyComponents(const TArray<UActorComponent*>& Co
 	return bCanCopy;
 }
 
-void FComponentEditorUtils::CopyComponents(const TArray<UActorComponent*>& ComponentsToCopy)
+void FComponentEditorUtils::CopyComponents(const TArray<UActorComponent*>& ComponentsToCopy, FString* DestinationData)
 {
 	FStringOutputDevice Archive;
 
@@ -523,27 +523,47 @@ void FComponentEditorUtils::CopyComponents(const TArray<UActorComponent*>& Compo
 	}
 
 	// Copy text to clipboard
-	FString ExportedText = Archive;
-	FPlatformApplicationMisc::ClipboardCopy(*ExportedText);
+	if (DestinationData)
+	{
+		*DestinationData = MoveTemp(Archive);
+	}
+	else
+	{
+		FPlatformApplicationMisc::ClipboardCopy(*Archive);
+	}
 }
 
-bool FComponentEditorUtils::CanPasteComponents(USceneComponent* RootComponent, bool bOverrideCanAttach, bool bPasteAsArchetypes)
+bool FComponentEditorUtils::CanPasteComponents(USceneComponent* RootComponent, bool bOverrideCanAttach, bool bPasteAsArchetypes, const FString* SourceData)
 {
 	FString ClipboardContent;
-	FPlatformApplicationMisc::ClipboardPaste(ClipboardContent);
+	if (SourceData)
+	{
+		ClipboardContent = *SourceData;
+	}
+	else
+	{
+		FPlatformApplicationMisc::ClipboardPaste(ClipboardContent);
+	}
 
 	// Obtain the component object text factory for the clipboard content and return whether or not we can use it
 	TSharedRef<FComponentObjectTextFactory> Factory = FComponentObjectTextFactory::Get(ClipboardContent, bPasteAsArchetypes);
 	return Factory->NewObjectMap.Num() > 0 && ( bOverrideCanAttach || Factory->CanAttachComponentsTo(RootComponent) );
 }
 
-void FComponentEditorUtils::PasteComponents(TArray<UActorComponent*>& OutPastedComponents, AActor* TargetActor, USceneComponent* TargetComponent)
+void FComponentEditorUtils::PasteComponents(TArray<UActorComponent*>& OutPastedComponents, AActor* TargetActor, USceneComponent* TargetComponent, const FString* SourceData)
 {
 	check(TargetActor);
 
 	// Get the text from the clipboard
 	FString TextToImport;
-	FPlatformApplicationMisc::ClipboardPaste(TextToImport);
+	if (SourceData)
+	{
+		TextToImport = *SourceData;
+	}
+	else
+	{
+		FPlatformApplicationMisc::ClipboardPaste(TextToImport);
+	}
 
 	// Get a new component object factory for the clipboard content
 	TSharedRef<FComponentObjectTextFactory> Factory = FComponentObjectTextFactory::Get(TextToImport);
diff --git a/Engine/Source/Editor/UnrealEd/Private/LevelEditorViewport.cpp b/Engine/Source/Editor/UnrealEd/Private/LevelEditorViewport.cpp
--- a/Engine/Source/Editor/UnrealEd/Private/LevelEditorViewport.cpp
+++ b/Engine/Source/Editor/UnrealEd/Private/LevelEditorViewport.cpp
@@ -2571,13 +2571,30 @@ bool FLevelEditorViewportClient::InputWidgetDelta(FViewport* InViewport, EAxisLi
 						// Only duplicate if we're translating or rotating.
 						if ( !Drag.IsNearlyZero() || !Rot.IsZero() )
 						{
-							// Widget hasn't been dragged since ALT+LMB went down.
 							bDuplicateOnNextDrag = false;
-							UUnrealEdEngine* UnrealEdEngine = Cast<UUnrealEdEngine>(GEditor);
-							if (UnrealEdEngine)
+							
+							if (!ModeTools->ProcessEditDuplicate())
 							{
-								UnrealEdEngine->DuplicateSelectedActors(GetWorld());
+								CacheElementsToManipulate();
+								TArray<FTypedElementHandle> DuplicatedElements = UEngineElementsLibrary::DuplicateElements(ElementsToManipulate, GetWorld(), /*bOffsetLocations*/false);
+								ElementsToManipulate->Reset();
+
+								// Exclusively select the new elements, so that future gizmo interaction manipulates those items instead
+								if (DuplicatedElements.Num() > 0)
+								{
+									UTypedElementSelectionSet* SelectionSet = GetMutableSelectionSet();
+
+									const FTypedElementSelectionOptions SelectionOptions = FTypedElementSelectionOptions()
+										.SetAllowLegacyNotifications(false); // Old drag duplicate code didn't used to notify about this selection change
+
+									FTypedElementListLegacySyncScopedBatch LegacySyncBatch(SelectionSet->GetElementList(), SelectionOptions.AllowLegacyNotifications());
+
+									SelectionSet->ClearSelection(SelectionOptions);
+									SelectionSet->SelectElements(DuplicatedElements, SelectionOptions);
+								}
 							}
+
+							RedrawAllViewportsIntoThisScene();
 						}
 					}
 				}
diff --git a/Engine/Source/Editor/UnrealEd/Public/Kismet2/ComponentEditorUtils.h b/Engine/Source/Editor/UnrealEd/Public/Kismet2/ComponentEditorUtils.h
--- a/Engine/Source/Editor/UnrealEd/Public/Kismet2/ComponentEditorUtils.h
+++ b/Engine/Source/Editor/UnrealEd/Public/Kismet2/ComponentEditorUtils.h
@@ -58,24 +58,27 @@ public:
 	/**
 	 * Copies the selected components to the clipboard
 	 * @param ComponentsToCopy The list of components to copy
+	 * @param DestinationData Buffer to fill with the copied data, or null to use the clipboard
 	 */
-	static void CopyComponents(const TArray<UActorComponent*>& ComponentsToCopy);
+	static void CopyComponents(const TArray<UActorComponent*>& ComponentsToCopy, FString* DestinationData = nullptr);
 
 	/**
 	 * Determines whether the current contents of the clipboard contain paste-able component information
 	 * @param RootComponent The root component of the actor being pasted on
 	 * @param bOverrideCanAttach Optional override declaring that components can be attached and a check is not needed
+	 * @param SourceData Component data to paste, or null to use the clipboard
 	 * @return Whether components can be pasted
 	 */
-	static bool CanPasteComponents(USceneComponent* RootComponent, bool bOverrideCanAttach = false, bool bPasteAsArchetypes = false);
+	static bool CanPasteComponents(USceneComponent* RootComponent, bool bOverrideCanAttach = false, bool bPasteAsArchetypes = false, const FString* SourceData = nullptr);
 
 	/**
 	 * Attempts to paste components from the clipboard as siblings of the target component
 	 * @param OutPastedComponents List of all the components that were pasted
 	 * @param TargetActor The actor to attach the pasted components to
 	 * @param TargetComponent The component the paste is targeting (will attempt to paste components as siblings). If null, will attach pasted components to the root.
+	 * @param SourceData Component data to paste, or null to use the clipboard
 	 */
-	static void PasteComponents(TArray<UActorComponent*>& OutPastedComponents, AActor* TargetActor, USceneComponent* TargetComponent = nullptr);
+	static void PasteComponents(TArray<UActorComponent*>& OutPastedComponents, AActor* TargetActor, USceneComponent* TargetComponent = nullptr, const FString* SourceData = nullptr);
 
 	/**
 	 * Gets the copied components from the clipboard without attempting to paste/apply them in any way
diff --git a/Engine/Source/Runtime/Engine/Private/Elements/Framework/EngineElementsLibrary.cpp b/Engine/Source/Runtime/Engine/Private/Elements/Framework/EngineElementsLibrary.cpp
--- a/Engine/Source/Runtime/Engine/Private/Elements/Framework/EngineElementsLibrary.cpp
+++ b/Engine/Source/Runtime/Engine/Private/Elements/Framework/EngineElementsLibrary.cpp
@@ -3,6 +3,8 @@
 #include "Elements/Framework/EngineElementsLibrary.h"
 #include "Elements/Framework/TypedElementRegistry.h"
 #include "Elements/Framework/TypedElementOwnerStore.h"
+#include "Elements/Framework/TypedElementUtil.h"
+#include "Elements/Interfaces/TypedElementWorldInterface.h"
 
 #include "Elements/Object/ObjectElementData.h"
 #include "UObject/Object.h"
@@ -214,3 +216,48 @@ FTypedElementHandle UEngineElementsLibrary::AcquireEditorComponentElementHandle(
 	return EngineElementsLibraryUtil::AcquireEditorTypedElementHandle<UActorComponent, FComponentElementData>(Component, GComponentElementOwnerStore, &UEngineElementsLibrary::CreateComponentElement, bAllowCreate);
 }
 #endif
+
+TArray<FTypedElementHandle> UEngineElementsLibrary::DuplicateElements(const TArray<FTypedElementHandle>& ElementHandles, UWorld* World, bool bOffsetLocations)
+{
+	return DuplicateElements(MakeArrayView(ElementHandles), World, bOffsetLocations);
+}
+
+TArray<FTypedElementHandle> UEngineElementsLibrary::DuplicateElements(TArrayView<const FTypedElementHandle> ElementHandles, UWorld* World, bool bOffsetLocations)
+{
+	TMap<FTypedHandleTypeId, TArray<FTypedElementHandle>> ElementsToDuplicateByType;
+	TypedElementUtil::BatchElementsByType(ElementHandles, ElementsToDuplicateByType);
+
+	TArray<FTypedElementHandle> NewElements;
+	NewElements.Reserve(ElementHandles.Num());
+
+	UTypedElementRegistry* Registry = UTypedElementRegistry::GetInstance();
+	for (const auto& ElementsByTypePair : ElementsToDuplicateByType)
+	{
+		if (UTypedElementWorldInterface* WorldInterface = Registry->GetElementInterface<UTypedElementWorldInterface>(ElementsByTypePair.Key))
+		{
+			WorldInterface->DuplicateElements(ElementsByTypePair.Value, World, bOffsetLocations, NewElements);
+		}
+	}
+
+	return NewElements;
+}
+
+TArray<FTypedElementHandle> UEngineElementsLibrary::DuplicateElements(const UTypedElementList* ElementList, UWorld* World, bool bOffsetLocations)
+{
+	TMap<FTypedHandleTypeId, TArray<FTypedElementHandle>> ElementsToDuplicateByType;
+	TypedElementUtil::BatchElementsByType(ElementList, ElementsToDuplicateByType);
+
+	TArray<FTypedElementHandle> NewElements;
+	NewElements.Reserve(ElementList->Num());
+
+	UTypedElementRegistry* Registry = UTypedElementRegistry::GetInstance();
+	for (const auto& ElementsByTypePair : ElementsToDuplicateByType)
+	{
+		if (UTypedElementWorldInterface* WorldInterface = Registry->GetElementInterface<UTypedElementWorldInterface>(ElementsByTypePair.Key))
+		{
+			WorldInterface->DuplicateElements(ElementsByTypePair.Value, World, bOffsetLocations, NewElements);
+		}
+	}
+
+	return NewElements;
+}
diff --git a/Engine/Source/Runtime/Engine/Private/UnrealExporter.cpp b/Engine/Source/Runtime/Engine/Private/UnrealExporter.cpp
--- a/Engine/Source/Runtime/Engine/Private/UnrealExporter.cpp
+++ b/Engine/Source/Runtime/Engine/Private/UnrealExporter.cpp
@@ -901,31 +901,44 @@ FString DumpObjectToString(UObject* Object)
 
 #if WITH_EDITOR
 FSelectedActorExportObjectInnerContext::FSelectedActorExportObjectInnerContext()
-	//call the empty version of the base class
-	: FExportObjectInnerContext(false)
+	: FExportObjectInnerContext(false) //call the empty version of the base class
 {
 	// For each selected actor...
 	for (FSelectionIterator It(GEditor->GetSelectedActorIterator()); It; ++It)
 	{
 		AActor* Actor = (AActor*)*It;
 		checkSlow(Actor->IsA(AActor::StaticClass()));
+		AddActorInner(Actor);
+	}
+}
 
-		ForEachObjectWithOuter(Actor, [this](UObject* InnerObj)
-		{
-			UObject* OuterObj = InnerObj->GetOuter();
-			InnerList* Inners = ObjectToInnerMap.Find(OuterObj);
-			if (Inners)
-			{
-				// Add object to existing inner list.
-				Inners->Add( InnerObj );
-			}
-			else
-			{
-				// Create a new inner list for the outer object.
-				InnerList& InnersForOuterObject = ObjectToInnerMap.Add(OuterObj, InnerList());
-				InnersForOuterObject.Add(InnerObj);
-			}
-		}, /** bIncludeNestedObjects */ true, RF_NoFlags, EInternalObjectFlags::PendingKill);
+FSelectedActorExportObjectInnerContext::FSelectedActorExportObjectInnerContext(const TArray<AActor*> InSelectedActors)
+	: FExportObjectInnerContext(false) //call the empty version of the base class
+{
+	// For each selected actor...
+	for (AActor* Actor : InSelectedActors)
+	{
+		AddActorInner(Actor);
 	}
 }
-#endif
\ No newline at end of file
+
+void FSelectedActorExportObjectInnerContext::AddActorInner(const AActor* InActor)
+{
+	ForEachObjectWithOuter(InActor, [this](UObject* InnerObj)
+	{
+		UObject* OuterObj = InnerObj->GetOuter();
+		InnerList* Inners = ObjectToInnerMap.Find(OuterObj);
+		if (Inners)
+		{
+			// Add object to existing inner list.
+			Inners->Add(InnerObj);
+		}
+		else
+		{
+			// Create a new inner list for the outer object.
+			InnerList& InnersForOuterObject = ObjectToInnerMap.Add(OuterObj, InnerList());
+			InnersForOuterObject.Add(InnerObj);
+		}
+	}, /** bIncludeNestedObjects */ true, RF_NoFlags, EInternalObjectFlags::PendingKill);
+}
+#endif
diff --git a/Engine/Source/Runtime/Engine/Public/Elements/Framework/EngineElementsLibrary.h b/Engine/Source/Runtime/Engine/Public/Elements/Framework/EngineElementsLibrary.h
--- a/Engine/Source/Runtime/Engine/Public/Elements/Framework/EngineElementsLibrary.h
+++ b/Engine/Source/Runtime/Engine/Public/Elements/Framework/EngineElementsLibrary.h
@@ -15,6 +15,9 @@ struct FActorElementData;
 class UActorComponent;
 struct FComponentElementData;
 
+class UWorld;
+class UTypedElementList;
+
 UCLASS()
 class ENGINE_API UEngineElementsLibrary : public UBlueprintFunctionLibrary
 {
@@ -50,6 +53,11 @@ public:
 	static FTypedElementHandle AcquireEditorComponentElementHandle(const UActorComponent* Component, const bool bAllowCreate = true);
 #endif
 
+	UFUNCTION(BlueprintCallable, Category="TypedElementFramework|Util")
+	static TArray<FTypedElementHandle> DuplicateElements(const TArray<FTypedElementHandle>& ElementHandles, UWorld* World, bool bOffsetLocations);
+	static TArray<FTypedElementHandle> DuplicateElements(TArrayView<const FTypedElementHandle> ElementHandles, UWorld* World, bool bOffsetLocations);
+	static TArray<FTypedElementHandle> DuplicateElements(const UTypedElementList* ElementList, UWorld* World, bool bOffsetLocations);
+
 private:
 #if WITH_EDITOR
 	static void DestroyUnreachableEditorObjectElements();
diff --git a/Engine/Source/Runtime/Engine/Public/Elements/Interfaces/TypedElementWorldInterface.h b/Engine/Source/Runtime/Engine/Public/Elements/Interfaces/TypedElementWorldInterface.h
--- a/Engine/Source/Runtime/Engine/Public/Elements/Interfaces/TypedElementWorldInterface.h
+++ b/Engine/Source/Runtime/Engine/Public/Elements/Interfaces/TypedElementWorldInterface.h
@@ -90,6 +90,26 @@ public:
 	{
 		return false;
 	}
+
+	/**
+	 * Duplicate the given element.
+	 * @note Default version calls DuplicateElements with a single element.
+	 */
+	UFUNCTION(BlueprintCallable, Category="TypedElementInterfaces|World")
+	virtual FTypedElementHandle DuplicateElement(const FTypedElementHandle& InElementHandle, UWorld* InWorld, bool bOffsetLocations)
+	{
+		TArray<FTypedElementHandle> NewElements;
+		DuplicateElements(MakeArrayView(&InElementHandle, 1), InWorld, bOffsetLocations, NewElements);
+		return NewElements.Num() > 0 ? MoveTemp(NewElements[0]) : FTypedElementHandle();
+	}
+
+	/**
+	 * Duplicate the given set of elements.
+	 * @note If you want to duplicate an array of elements that are potentially different types, you probably want to use the higher-level UEngineElementsLibrary::DuplicateElements function instead.
+	 */
+	virtual void DuplicateElements(TArrayView<const FTypedElementHandle> InElementHandles, UWorld* InWorld, bool bOffsetLocations, TArray<FTypedElementHandle>& OutNewElements)
+	{
+	}
 };
 
 template <>
@@ -107,4 +127,5 @@ struct TTypedElement<UTypedElementWorldInterface> : public TTypedElementBase<UTy
 	void NotifyMovementEnded() const { InterfacePtr->NotifyMovementEnded(*this); }
 	bool FindSuitableTransformAtPoint(const FTransform& InPotentialTransform, FTransform& OutSuitableTransform) const { return InterfacePtr->FindSuitableTransformAtPoint(*this, InPotentialTransform, OutSuitableTransform); }
 	bool FindSuitableTransformAlongPath(const FVector& InPathStart, const FVector& InPathEnd, const FCollisionShape& InTestShape, TArrayView<const FTypedElementHandle> InElementsToIgnore, FTransform& OutSuitableTransform) const { return InterfacePtr->FindSuitableTransformAlongPath(*this, InPathStart, InPathEnd, InTestShape, InElementsToIgnore, OutSuitableTransform); }
+	FTypedElementHandle DuplicateElement(UWorld* InWorld, bool bOffsetLocations) const { return InterfacePtr->DuplicateElement(*this, InWorld, bOffsetLocations); }
 };
diff --git a/Engine/Source/Runtime/Engine/Public/UnrealExporter.h b/Engine/Source/Runtime/Engine/Public/UnrealExporter.h
--- a/Engine/Source/Runtime/Engine/Public/UnrealExporter.h
+++ b/Engine/Source/Runtime/Engine/Public/UnrealExporter.h
@@ -61,6 +61,10 @@ class ENGINE_API FSelectedActorExportObjectInnerContext : public FExportObjectIn
 {
 public:
 	FSelectedActorExportObjectInnerContext();
+	explicit FSelectedActorExportObjectInnerContext(const TArray<class AActor*> InSelectedActors);
+
+private:
+	void AddActorInner(const class AActor* InActor);
 };
 #endif
 
diff --git a/Engine/Source/Runtime/TypedElementFramework/Private/Elements/Framework/TypedElementList.cpp b/Engine/Source/Runtime/TypedElementFramework/Private/Elements/Framework/TypedElementList.cpp
--- a/Engine/Source/Runtime/TypedElementFramework/Private/Elements/Framework/TypedElementList.cpp
+++ b/Engine/Source/Runtime/TypedElementFramework/Private/Elements/Framework/TypedElementList.cpp
@@ -73,7 +73,7 @@ void FTypedElementListLegacySync::ForceBatchOperationDirty()
 }
 
 
-FTypedElementListLegacySyncScopedBatch::FTypedElementListLegacySyncScopedBatch(UTypedElementList* InElementList, const bool InNotify)
+FTypedElementListLegacySyncScopedBatch::FTypedElementListLegacySyncScopedBatch(const UTypedElementList* InElementList, const bool InNotify)
 	: ElementListLegacySync(InElementList->Legacy_GetSyncPtr())
 	, bNotify(InNotify)
 {
@@ -291,7 +291,7 @@ FTypedElementListLegacySync& UTypedElementList::Legacy_GetSync()
 	return *LegacySync;
 }
 
-FTypedElementListLegacySync* UTypedElementList::Legacy_GetSyncPtr()
+FTypedElementListLegacySync* UTypedElementList::Legacy_GetSyncPtr() const
 {
 	return LegacySync.Get();
 }
diff --git a/Engine/Source/Runtime/TypedElementFramework/Private/Elements/Framework/TypedElementUtil.cpp b/Engine/Source/Runtime/TypedElementFramework/Private/Elements/Framework/TypedElementUtil.cpp
new file mode 100644
--- /dev/null
+++ b/Engine/Source/Runtime/TypedElementFramework/Private/Elements/Framework/TypedElementUtil.cpp
@@ -0,0 +1,30 @@
+// Copyright Epic Games, Inc. All Rights Reserved.
+
+#include "Elements/Framework/TypedElementUtil.h"
+#include "Elements/Framework/TypedElementList.h"
+
+namespace TypedElementUtil
+{
+
+void BatchElementsByType(const UTypedElementList* InElementsToBatch, TMap<FTypedHandleTypeId, TArray<FTypedElementHandle>>& OutElementsByType)
+{
+	OutElementsByType.Reset();
+	InElementsToBatch->ForEachElementHandle([&OutElementsByType](const FTypedElementHandle& InElementHandle)
+	{
+		TArray<FTypedElementHandle>& ElementsForType = OutElementsByType.FindOrAdd(InElementHandle.GetId().GetTypeId());
+		ElementsForType.Add(InElementHandle);
+		return true;
+	});
+}
+
+void BatchElementsByType(TArrayView<const FTypedElementHandle> InElementsToBatch, TMap<FTypedHandleTypeId, TArray<FTypedElementHandle>>& OutElementsByType)
+{
+	OutElementsByType.Reset();
+	for (const FTypedElementHandle& ElementHandle : InElementsToBatch)
+	{
+		TArray<FTypedElementHandle>& ElementsForType = OutElementsByType.FindOrAdd(ElementHandle.GetId().GetTypeId());
+		ElementsForType.Add(ElementHandle);
+	}
+}
+
+} // namespace TypedElementUtil
diff --git a/Engine/Source/Runtime/TypedElementFramework/Public/Elements/Framework/TypedElementList.h b/Engine/Source/Runtime/TypedElementFramework/Public/Elements/Framework/TypedElementList.h
--- a/Engine/Source/Runtime/TypedElementFramework/Public/Elements/Framework/TypedElementList.h
+++ b/Engine/Source/Runtime/TypedElementFramework/Public/Elements/Framework/TypedElementList.h
@@ -107,7 +107,7 @@ private:
 class TYPEDELEMENTFRAMEWORK_API FTypedElementListLegacySyncScopedBatch
 {
 public:
-	explicit FTypedElementListLegacySyncScopedBatch(UTypedElementList* InElementList, const bool InNotify = true);
+	explicit FTypedElementListLegacySyncScopedBatch(const UTypedElementList* InElementList, const bool InNotify = true);
 	~FTypedElementListLegacySyncScopedBatch();
 
 	FTypedElementListLegacySyncScopedBatch(const FTypedElementListLegacySyncScopedBatch&) = delete;
@@ -568,7 +568,7 @@ public:
 	 * Access the interface to allow external systems (such as USelection) to receive immediate sync notifications as an element list is changed.
 	 * This exists purely as a bridging mechanism and shouldn't be relied on for new code. This will return null if no legacy sync has been created for this instance.
 	 */
-	FTypedElementListLegacySync* Legacy_GetSyncPtr();
+	FTypedElementListLegacySync* Legacy_GetSyncPtr() const;
 
 private:
 	enum class EChangeType : uint8
diff --git a/Engine/Source/Runtime/TypedElementFramework/Public/Elements/Framework/TypedElementUtil.h b/Engine/Source/Runtime/TypedElementFramework/Public/Elements/Framework/TypedElementUtil.h
new file mode 100644
--- /dev/null
+++ b/Engine/Source/Runtime/TypedElementFramework/Public/Elements/Framework/TypedElementUtil.h
@@ -0,0 +1,18 @@
+// Copyright Epic Games, Inc. All Rights Reserved.
+
+#pragma once
+
+#include "Elements/Framework/TypedElementHandle.h"
+
+class UTypedElementList;
+
+namespace TypedElementUtil
+{
+
+/**
+ * Batch the given elements by their type.
+ */
+TYPEDELEMENTFRAMEWORK_API void BatchElementsByType(const UTypedElementList* InElementsToBatch, TMap<FTypedHandleTypeId, TArray<FTypedElementHandle>>& OutElementsByType);
+TYPEDELEMENTFRAMEWORK_API void BatchElementsByType(TArrayView<const FTypedElementHandle> InElementsToBatch, TMap<FTypedHandleTypeId, TArray<FTypedElementHandle>>& OutElementsByType);
+
+} // namespace TypedElementUtil
