Fix comment typo

diff --git a/Engine/Source/Editor/UnrealEd/Private/Commandlets/WorldPartitionConvertCommandlet.cpp b/Engine/Source/Editor/UnrealEd/Private/Commandlets/WorldPartitionConvertCommandlet.cpp
--- a/Engine/Source/Editor/UnrealEd/Private/Commandlets/WorldPartitionConvertCommandlet.cpp
+++ b/Engine/Source/Editor/UnrealEd/Private/Commandlets/WorldPartitionConvertCommandlet.cpp
@@ -794,7 +794,7 @@ int32 UWorldPartitionConvertCommandlet::Main(const FString& Params)
 		return 1;
 	}
 
-	// Make sure the world isn't already partitionned
+	// Make sure the world isn't already partitioned
 	AWorldSettings* MainWorldSettings = MainWorld->GetWorldSettings();
 	if (MainWorldSettings->GetWorldPartition())
 	{
