Adding missing file that had deprecation warnings fixed up

diff --git a/Engine/Source/Editor/UnrealEd/Private/EditorEngine.cpp b/Engine/Source/Editor/UnrealEd/Private/EditorEngine.cpp
--- a/Engine/Source/Editor/UnrealEd/Private/EditorEngine.cpp
+++ b/Engine/Source/Editor/UnrealEd/Private/EditorEngine.cpp
@@ -647,7 +647,6 @@ void UEditorEngine::InitEditor(IEngineLoop* InEngineLoop)
 		const bool bCorrectDeficiency = EditorSettings->bColorVisionDeficiencyCorrection;
 		const bool bShowCorrectionWithDeficiency = EditorSettings->bColorVisionDeficiencyCorrectionPreviewWithDeficiency;
 		FSlateApplication::Get().GetRenderer()->SetColorVisionDeficiencyType(DeficiencyType, Severity, bCorrectDeficiency, bShowCorrectionWithDeficiency);
-		FSlateApplication::Get().EnableMenuAnimations(EditorSettings->bEnableWindowAnimations);
 	}
 
 	UEditorStyleSettings* StyleSettings = GetMutableDefault<UEditorStyleSettings>();
