Initial Placement Factory and Subsystem APIs.

# FActorFactoryAssetProxy uses the placement subsystem first to try to place a new actor.
# Existing UActorFactory implementations added to the subsystem and hooked up to the placement API

diff --git a/Engine/Source/Editor/EditorFramework/EditorFramework.Build.cs b/Engine/Source/Editor/EditorFramework/EditorFramework.Build.cs
--- a/Engine/Source/Editor/EditorFramework/EditorFramework.Build.cs
+++ b/Engine/Source/Editor/EditorFramework/EditorFramework.Build.cs
@@ -31,6 +31,7 @@ public class EditorFramework : ModuleRules
 
 		PublicDependencyModuleNames.AddRange(
 			new string[] {
+				"EditorSubsystem",
 				"InteractiveToolsFramework",
 				"TypedElementFramework",
 				"TypedElementInterfaces",
diff --git a/Engine/Source/Editor/EditorFramework/Private/Factories/AssetFactoryInterface.cpp b/Engine/Source/Editor/EditorFramework/Private/Factories/AssetFactoryInterface.cpp
new file mode 100644
--- /dev/null
+++ b/Engine/Source/Editor/EditorFramework/Private/Factories/AssetFactoryInterface.cpp
@@ -0,0 +1,3 @@
+// Copyright Epic Games, Inc. All Rights Reserved.
+
+#include "Factories/AssetFactoryInterface.h"
diff --git a/Engine/Source/Editor/EditorFramework/Private/Subsystems/PlacementSubsystem.cpp b/Engine/Source/Editor/EditorFramework/Private/Subsystems/PlacementSubsystem.cpp
new file mode 100644
--- /dev/null
+++ b/Engine/Source/Editor/EditorFramework/Private/Subsystems/PlacementSubsystem.cpp
@@ -0,0 +1,96 @@
+// Copyright Epic Games, Inc. All Rights Reserved.
+
+#include "Subsystems/PlacementSubsystem.h"
+
+#include "Factories/AssetFactoryInterface.h"
+
+#include "Elements/Interfaces/TypedElementObjectInterface.h"
+#include "Elements/Framework/TypedElementRegistry.h"
+#include "Elements/Framework/TypedElementSelectionSet.h"
+
+#include "Misc/CoreDelegates.h"
+
+void UPlacementSubsystem::Initialize(FSubsystemCollectionBase& Collection)
+{
+	FCoreDelegates::OnPostEngineInit.AddUObject(this, &UPlacementSubsystem::RegisterPlacementFactories);
+	FCoreDelegates::OnEnginePreExit.AddUObject(this, &UPlacementSubsystem::UnregisterPlacementFactories);
+}
+
+void UPlacementSubsystem::Deinitialize()
+{
+	FCoreDelegates::OnEnginePreExit.RemoveAll(this);
+	FCoreDelegates::OnPostEngineInit.RemoveAll(this);
+}
+
+TArray<FTypedElementHandle> UPlacementSubsystem::PlaceAsset(const FAssetPlacementInfo& InPlacementInfo, const FPlacementOptions& InPlacementOptions)
+{
+	return PlaceAssets(MakeArrayView( {InPlacementInfo} ), InPlacementOptions);
+}
+
+TArray<FTypedElementHandle> UPlacementSubsystem::PlaceAssets(TArrayView<const FAssetPlacementInfo> InPlacementInfos, const FPlacementOptions& InPlacementOptions)
+{
+	TArray<FTypedElementHandle> PlacedElements;
+	for (const FAssetPlacementInfo& PlacementInfo : InPlacementInfos)
+	{
+		const FAssetData& AssetData = PlacementInfo.AssetToPlace;
+		TScriptInterface<IAssetFactoryInterface> FactoryInterface = PlacementInfo.FactoryOverride;
+		if (!FactoryInterface)
+		{
+			FactoryInterface = FindAssetFactoryFromAssetData(AssetData);
+		}
+
+		if (!FactoryInterface || !FactoryInterface->CanPlaceElementsFromAssetData(AssetData))
+		{
+			continue;
+		}
+
+		FAssetPlacementInfo AdjustedPlacementInfo = PlacementInfo;
+		if (!FactoryInterface->PrePlaceAsset(AdjustedPlacementInfo, InPlacementOptions))
+		{
+			continue;
+		}
+
+		TArray<FTypedElementHandle> PlacedHandles = FactoryInterface->PlaceAsset(AdjustedPlacementInfo, InPlacementOptions);
+		if (PlacedHandles.Num())
+		{
+			FactoryInterface->PostPlaceAsset(PlacedHandles, PlacementInfo, InPlacementOptions);
+			PlacedElements.Append(PlacedHandles);
+		}
+	}
+
+	return PlacedElements;
+}
+
+TScriptInterface<IAssetFactoryInterface> UPlacementSubsystem::FindAssetFactoryFromAssetData(const FAssetData& InAssetData)
+{
+	for (const TScriptInterface<IAssetFactoryInterface>& AssetFactory : AssetFactories)
+	{
+		if (AssetFactory && AssetFactory->CanPlaceElementsFromAssetData(InAssetData))
+		{
+			return AssetFactory;
+		}
+	}
+
+	return nullptr;
+}
+
+void UPlacementSubsystem::RegisterPlacementFactories()
+{
+	for (TObjectIterator<UClass> ObjectIt; ObjectIt; ++ObjectIt)
+	{
+		UClass* TestClass = *ObjectIt;
+		if (TestClass->ImplementsInterface(UAssetFactoryInterface::StaticClass()))
+		{
+			if (!TestClass->HasAnyClassFlags(CLASS_Abstract | CLASS_Deprecated))
+			{
+				TScriptInterface<IAssetFactoryInterface> NewFactory = NewObject<UObject>(this, TestClass);
+				AssetFactories.Add(NewFactory);
+			}
+		}
+	}
+}
+
+void UPlacementSubsystem::UnregisterPlacementFactories()
+{
+	AssetFactories.Empty();
+}
diff --git a/Engine/Source/Editor/EditorFramework/Public/Factories/AssetFactoryInterface.h b/Engine/Source/Editor/EditorFramework/Public/Factories/AssetFactoryInterface.h
new file mode 100644
--- /dev/null
+++ b/Engine/Source/Editor/EditorFramework/Public/Factories/AssetFactoryInterface.h
@@ -0,0 +1,61 @@
+// Copyright Epic Games, Inc. All Rights Reserved.
+
+#pragma once
+
+#include "AssetRegistry/AssetData.h"
+#include "Elements/Framework/TypedElementHandle.h"
+#include "UObject/Interface.h"
+
+#include "AssetFactoryInterface.generated.h"
+
+class IAssetFactoryInterface;
+struct FPlacementOptions;
+struct FAssetPlacementInfo;
+
+UINTERFACE(MinimalAPI)
+class UAssetFactoryInterface : public UInterface
+{
+	GENERATED_BODY()
+};
+
+class EDITORFRAMEWORK_API IAssetFactoryInterface
+{
+	GENERATED_BODY()
+
+public:
+	/**
+	 * Given an FAssetData, determine if this asset factory can place any elements.
+	 *
+	 * @returns true if the factory can be used to place elements.
+	 */
+	virtual bool CanPlaceElementsFromAssetData(const FAssetData& InAssetData) = 0;
+
+	/**
+	 * Performs any final tweaking of the PlacementInfo that the asset factory may need to do. This includes final adjustments to things like transforms.
+	 * This should NOT perform any viewport or editor specific adjustments, such as grid snapping, alignment to hit objects, or undo tracking.
+	 *
+	 * @returns true if the asset is still valid to place after final adjustments.
+	 */
+	virtual bool PrePlaceAsset(FAssetPlacementInfo& InPlacementInfo, const FPlacementOptions& InPlacementOptions) = 0;
+
+	/**
+	 * Places the asset.
+	 *
+	 * @returns valid FTypedElementHandles that were placed by the factory from the given asset.
+	 */
+	virtual TArray<FTypedElementHandle> PlaceAsset(const FAssetPlacementInfo& InPlacementInfo, const FPlacementOptions& InPlacementOptions) = 0;
+
+	/**
+	 * Finalizes any placed elements based on adjustments the factory may need to do.
+	 * This should NOT include any adjustments from viewport or asset editor specific functionality, such as finalizing undo tracking.
+	 */
+	virtual void PostPlaceAsset(TArrayView<const FTypedElementHandle> InHandle, const FAssetPlacementInfo& InPlacementInfo, const FPlacementOptions& InPlacementOptions) = 0;
+
+	/**
+	 * Given an FTypedElementHandle, attempt to resolve the handle to the FAssetData which may have placed it.
+	 * The FAssetData may be a wrapped type, like a static mesh component component inside a static mesh actor.
+	 *
+	 * @returns the FAssetData which corresponds to the placed FTypedElementHandle. The returned FAssetData may be invalid.
+	 */
+	virtual FAssetData GetAssetDataFromElementHandle(const FTypedElementHandle& InHandle) = 0;
+};
diff --git a/Engine/Source/Editor/EditorFramework/Public/Subsystems/PlacementSubsystem.h b/Engine/Source/Editor/EditorFramework/Public/Subsystems/PlacementSubsystem.h
new file mode 100644
--- /dev/null
+++ b/Engine/Source/Editor/EditorFramework/Public/Subsystems/PlacementSubsystem.h
@@ -0,0 +1,89 @@
+// Copyright Epic Games, Inc. All Rights Reserved.
+
+#pragma once
+
+#include "AssetRegistry/AssetData.h"
+#include "Elements/Framework/TypedElementHandle.h"
+#include "UObject/Interface.h"
+#include "EditorSubsystem.h"
+
+#include "PlacementSubsystem.generated.h"
+
+class IAssetFactoryInterface;
+
+USTRUCT()
+struct EDITORFRAMEWORK_API FAssetPlacementInfo
+{
+	GENERATED_BODY()
+
+	// The asset data which should be placed.
+	UPROPERTY()
+	FAssetData AssetToPlace;
+
+	// If set, will override the name on placed elements instead of factory defined defaults.
+	UPROPERTY()
+	FName NameOverride;
+
+	// If set, the factory will attempt to place inside the given level. World partitioning may ultimately override this preference.
+	UPROPERTY()
+	TWeakObjectPtr<ULevel> PreferredLevel;
+
+	// The finalized transform where the factory should be place the asset. This should include any location snapping or other considerations from viewports or editor settings.
+	UPROPERTY()
+	FTransform FinalizedTransform;
+
+	// If set, will use the given factory to place the asset, instead of allowing the placement subsystem to determine which factory to use.
+	UPROPERTY()
+	TScriptInterface<IAssetFactoryInterface> FactoryOverride;
+};
+
+USTRUCT()
+struct EDITORFRAMEWORK_API FPlacementOptions
+{
+	GENERATED_BODY()
+
+	// If true, asset factory implementations should defer to placing instanced items (i.e. instanced static mesh instead of individual static mesh actors).
+	UPROPERTY()
+	bool bPreferInstancedPlacement;
+
+	// If true, asset factory implementations should prefer a batch placement algorithm (like duplicating an object) over a single placement algorithm.
+	UPROPERTY()
+	bool bPreferBatchPlacement;
+};
+
+UCLASS(Transient)
+class EDITORFRAMEWORK_API UPlacementSubsystem : public UEditorSubsystem
+{
+	GENERATED_BODY()
+
+public:
+	// USubsystem Interface
+	void Initialize(FSubsystemCollectionBase& Collection) override;
+	void Deinitialize() override;
+	// End USubsystem Interface
+
+	/**
+	 * Places a single asset based on the given FAssetPlacementInfo and FPlacementOptions.
+	 * @returns an array of FTypedElementHandles corresponding to any successfully placed elements.
+	 */
+	TArray<FTypedElementHandle> PlaceAsset(const FAssetPlacementInfo& InPlacementInfo, const FPlacementOptions& InPlacementOptions);
+
+	/**
+	 * Places a multiple asset based on the given FAssetPlacementInfos and FPlacementOptions.
+	 * @returns an array of FTypedElementHandles corresponding to any successfully placed elements.
+	 */
+	TArray<FTypedElementHandle> PlaceAssets(TArrayView<const FAssetPlacementInfo> InPlacementInfos, const FPlacementOptions& InPlacementOptions);
+
+	/**
+	 * Finds a registered AssetFactory for the given FAssetData.
+	 * @returns the first found factory, or nullptr if no valid factory is found.
+	 */
+	TScriptInterface<IAssetFactoryInterface> FindAssetFactoryFromAssetData(const FAssetData& InAssetData);
+
+private:
+	void RegisterPlacementFactories();
+	void UnregisterPlacementFactories();
+
+	UPROPERTY()
+	TArray<TScriptInterface<IAssetFactoryInterface>> AssetFactories;
+};
diff --git a/Engine/Source/Editor/UnrealEd/Classes/ActorFactories/ActorFactory.h b/Engine/Source/Editor/UnrealEd/Classes/ActorFactories/ActorFactory.h
--- a/Engine/Source/Editor/UnrealEd/Classes/ActorFactories/ActorFactory.h
+++ b/Engine/Source/Editor/UnrealEd/Classes/ActorFactories/ActorFactory.h
@@ -7,6 +7,8 @@
 #include "UObject/ObjectMacros.h"
 #include "UObject/Object.h"
 #include "Templates/SubclassOf.h"
+#include "Factories/AssetFactoryInterface.h"
+
 #include "ActorFactory.generated.h"
 
 UNREALED_API DECLARE_LOG_CATEGORY_EXTERN(LogActorFactory, Log, All);
@@ -17,7 +19,7 @@ class UBlueprint;
 class ULevel;
 
 UCLASS(collapsecategories, hidecategories=Object, editinlinenew, config=Editor, abstract, transient)
-class UNREALED_API UActorFactory : public UObject
+class UNREALED_API UActorFactory : public UObject, public IAssetFactoryInterface
 {
 	GENERATED_UCLASS_BODY()
 
@@ -35,7 +37,7 @@ class UNREALED_API UActorFactory : public UObject
 
 	/**  AActor  subclass this ActorFactory creates. */
 	UPROPERTY()
-	TSubclassOf<class AActor>  NewActorClass;
+	TSubclassOf<AActor>  NewActorClass;
 
 	/** Whether to appear in the editor add actor quick menu */
 	UPROPERTY()
@@ -52,6 +54,7 @@ class UNREALED_API UActorFactory : public UObject
 	AActor* CreateActor( UObject* Asset, ULevel* InLevel, FTransform Transform, EObjectFlags InObjectFlags = RF_Transactional, const FName InName = NAME_None );
 
 	/** Called to create a blueprint class that can be used to spawn an actor from this factory */
+	UE_DEPRECATED(5.0, "This function is no longer used. See FKismetEditorUtilities::CreateBlueprint.")
 	UBlueprint* CreateBlueprint( UObject* Instance, UObject* Outer, const FName Name, const FName CallingContext = NAME_None );
 
 	virtual bool CanCreateActorFrom( const FAssetData& AssetData, FText& OutErrorMsg );
@@ -63,14 +66,26 @@ class UNREALED_API UActorFactory : public UObject
 	virtual AActor* GetDefaultActor( const FAssetData& AssetData );
 
 	/** Initialize NewActorClass if necessary, and return that class. */
-	UClass* GetDefaultActorClass( const FAssetData& AssetData );
+	virtual UClass* GetDefaultActorClass( const FAssetData& AssetData );
 
-	/** Given an instance of an actor pertaining to this factory, find the asset that should be used to create a new actor */
+	/** Given an instance of an actor, find the wrapped asset object which can be used to create a valid FAssetData.
+	 *  Returns nullptr if the given ActorInstance is not valid for this factory.
+	 *  Override this function if the factory actor is a different class than the asset data's class which this factory operates on.
+	 *  For example, if this is the static mesh actor factory, the class of the asset data is UStaticMesh, but the actor factory's class is AStaticMeshActor
+	 */
 	virtual UObject* GetAssetFromActorInstance(AActor* ActorInstance);
 
 	/** Return a quaternion which aligns this actor type to the specified surface normal */
 	virtual FQuat AlignObjectToSurfaceNormal(const FVector& InSurfaceNormal, const FQuat& ActorRotation = FQuat::Identity) const;
 
+	// Begin IAssetFactoryInterface Interface
+	virtual bool CanPlaceElementsFromAssetData(const FAssetData& InAssetData) override;
+	virtual bool PrePlaceAsset(FAssetPlacementInfo& InPlacementInfo, const FPlacementOptions& InPlacementOptions) override;
+	virtual TArray<FTypedElementHandle> PlaceAsset(const FAssetPlacementInfo& InPlacementInfo, const FPlacementOptions& InPlacementOptions) override;
+	virtual void PostPlaceAsset(TArrayView<const FTypedElementHandle> InHandle, const FAssetPlacementInfo& InPlacementInfo, const FPlacementOptions& InPlacementOptions) override;
+	virtual FAssetData GetAssetDataFromElementHandle(const FTypedElementHandle& InHandle) override;
+	// End IAssetFactoryInterface Interface
+
 protected:
 
 	virtual bool PreSpawnActor( UObject* Asset, FTransform& InOutLocation);
diff --git a/Engine/Source/Editor/UnrealEd/Classes/ActorFactories/ActorFactoryAnimationAsset.h b/Engine/Source/Editor/UnrealEd/Classes/ActorFactories/ActorFactoryAnimationAsset.h
--- a/Engine/Source/Editor/UnrealEd/Classes/ActorFactories/ActorFactoryAnimationAsset.h
+++ b/Engine/Source/Editor/UnrealEd/Classes/ActorFactories/ActorFactoryAnimationAsset.h
@@ -21,7 +21,8 @@ protected:
 	//~ Begin UActorFactory Interface
 	virtual void PostSpawnActor( UObject* Asset, AActor* NewActor ) override;
 	virtual void PostCreateBlueprint( UObject* Asset, AActor* CDO ) override;
-	virtual bool CanCreateActorFrom( const FAssetData& AssetData, FText& OutErrorMsg ) override;
+	virtual bool CanCreateActorFrom(const FAssetData& AssetData, FText& OutErrorMsg) override;
+	virtual UObject* GetAssetFromActorInstance(AActor* ActorInstance);
 	//~ End UActorFactory Interface
 
 	virtual USkeletalMesh* GetSkeletalMeshFromAsset( UObject* Asset ) override;
diff --git a/Engine/Source/Editor/UnrealEd/Classes/ActorFactories/ActorFactoryBlueprint.h b/Engine/Source/Editor/UnrealEd/Classes/ActorFactories/ActorFactoryBlueprint.h
--- a/Engine/Source/Editor/UnrealEd/Classes/ActorFactories/ActorFactoryBlueprint.h
+++ b/Engine/Source/Editor/UnrealEd/Classes/ActorFactories/ActorFactoryBlueprint.h
@@ -20,6 +20,7 @@ class UNREALED_API UActorFactoryBlueprint : public UActorFactory
 	virtual bool PreSpawnActor( UObject* Asset, FTransform& InOutLocation ) override;
 	virtual bool CanCreateActorFrom( const FAssetData& AssetData, FText& OutErrorMsg ) override;
 	virtual AActor* GetDefaultActor( const FAssetData& AssetData ) override;
+	virtual UClass* GetDefaultActorClass( const FAssetData& AssetData ) override;
 	//~ End UActorFactory Interface	
 };
 
diff --git a/Engine/Source/Editor/UnrealEd/Classes/ActorFactories/ActorFactoryCylinderVolume.h b/Engine/Source/Editor/UnrealEd/Classes/ActorFactories/ActorFactoryCylinderVolume.h
--- a/Engine/Source/Editor/UnrealEd/Classes/ActorFactories/ActorFactoryCylinderVolume.h
+++ b/Engine/Source/Editor/UnrealEd/Classes/ActorFactories/ActorFactoryCylinderVolume.h
@@ -17,7 +17,6 @@ class UActorFactoryCylinderVolume : public UActorFactoryVolume
 	GENERATED_UCLASS_BODY()
 
 	//~ Begin UActorFactory Interface
-	virtual bool CanCreateActorFrom( const FAssetData& AssetData, FText& OutErrorMsg ) override;
 	virtual void PostSpawnActor( UObject* Asset, AActor* NewActor ) override;
 	//~ End UActorFactory Interface
 };
diff --git a/Engine/Source/Editor/UnrealEd/Classes/ActorFactories/ActorFactoryVolume.h b/Engine/Source/Editor/UnrealEd/Classes/ActorFactories/ActorFactoryVolume.h
--- a/Engine/Source/Editor/UnrealEd/Classes/ActorFactories/ActorFactoryVolume.h
+++ b/Engine/Source/Editor/UnrealEd/Classes/ActorFactories/ActorFactoryVolume.h
@@ -6,10 +6,27 @@
 #include "CoreMinimal.h"
 #include "UObject/ObjectMacros.h"
 #include "ActorFactories/ActorFactory.h"
+#include "GameFramework/Volume.h"
 #include "ActorFactoryVolume.generated.h"
 
-UCLASS(Abstract, MinimalAPI, config=Editor, collapsecategories, hidecategories=Object)
+UCLASS(Abstract, MinimalAPI, config=Editor, collapsecategories, hidecategories=Object, Abstract)
 class UActorFactoryVolume : public UActorFactory
 {
 	GENERATED_BODY()
+
+public:
+	virtual bool CanCreateActorFrom( const FAssetData& AssetData, FText& OutErrorMsg ) override
+	{
+		if (UActorFactory::CanCreateActorFrom(AssetData, OutErrorMsg))
+		{
+			return true;
+		}
+
+		if (AssetData.IsValid() && !AssetData.GetClass()->IsChildOf(AVolume::StaticClass()))
+		{
+			return false;
+		}
+
+		return true;
+	}
 };
diff --git a/Engine/Source/Editor/UnrealEd/Private/AssetSelection.cpp b/Engine/Source/Editor/UnrealEd/Private/AssetSelection.cpp
--- a/Engine/Source/Editor/UnrealEd/Private/AssetSelection.cpp
+++ b/Engine/Source/Editor/UnrealEd/Private/AssetSelection.cpp
@@ -51,7 +51,8 @@
 #include "ISourceControlModule.h"
 #include "ISourceControlProvider.h"
 #include "Misc/MessageDialog.h"
-
+#include "Subsystems/PlacementSubsystem.h"
+#include "Elements/Actor/ActorElementData.h"
 
 namespace AssetSelectionUtils
 {
@@ -604,10 +605,32 @@ static AActor* PrivateAddActor( UObject* Asset, UActorFactory* Factory, bool Sel
 	if(bSpawnActor)
 	{
 		FScopedTransaction Transaction( NSLOCTEXT("UnrealEd", "CreateActor", "Create Actor"), (ObjectFlags & RF_Transactional) != 0 );
-
+		
 		// Create the actor.
-		Actor = Factory->CreateActor( Asset, DesiredLevel, ActorTransform, ObjectFlags, Name );
-		if(Actor)
+		UPlacementSubsystem* PlacementSubsystem = GEditor->GetEditorSubsystem<UPlacementSubsystem>();
+		if (PlacementSubsystem)
+		{
+			FAssetPlacementInfo PlacementInfo;
+			PlacementInfo.AssetToPlace = FAssetData(Asset);
+			PlacementInfo.PreferredLevel = DesiredLevel;
+			PlacementInfo.NameOverride = Name;
+			PlacementInfo.FinalizedTransform = ActorTransform;
+			PlacementInfo.FactoryOverride = Factory;
+
+			TArray<FTypedElementHandle> PlacedElements = PlacementSubsystem->PlaceAsset(PlacementInfo, FPlacementOptions());
+			if (PlacedElements.Num())
+			{
+				const FActorElementData* ActorElementData = PlacedElements[0].GetData<FActorElementData>();
+				Actor = ActorElementData ? ActorElementData->Actor : nullptr;
+			}
+		}
+		
+		if (!Actor)
+		{
+			Actor = Factory->CreateActor(Asset, DesiredLevel, ActorTransform, ObjectFlags, Name);
+		}
+
+		if (Actor)
 		{
 			if ( SelectActor )
 			{
diff --git a/Engine/Source/Editor/UnrealEd/Private/EditorEngine.cpp b/Engine/Source/Editor/UnrealEd/Private/EditorEngine.cpp
--- a/Engine/Source/Editor/UnrealEd/Private/EditorEngine.cpp
+++ b/Engine/Source/Editor/UnrealEd/Private/EditorEngine.cpp
@@ -714,8 +714,6 @@ void UEditorEngine::InitEditor(IEngineLoop* InEngineLoop)
 	BookmarkTools.RegisterBookmarkTypeActions(MakeShared<FBookMarkTypeActions>());
 	
 	{
-		FAssetData NoAssetData;
-
 		TArray<UClass*> VolumeClasses;
 		TArray<UClass*> VolumeFactoryClasses;
 
@@ -1020,7 +1018,8 @@ void UEditorEngine::Init(IEngineLoop* InEngineLoop)
 			TEXT("TimeManagementEditor"),
 			TEXT("EditorInteractiveToolsFramework"),
 			TEXT("TraceInsights"),
-			TEXT("StaticMeshEditor")
+			TEXT("StaticMeshEditor"),
+			TEXT("EditorFramework")
 		};
 
 		FScopedSlowTask ModuleSlowTask(UE_ARRAY_COUNT(ModuleNames));
@@ -6380,7 +6379,7 @@ UActorFactory* UEditorEngine::FindActorFactoryByClassForActorClass( const UClass
 		if ( Factory != NULL && Factory->GetClass() == InFactoryClass )
 		{
 			// force NewActorClass update
-			const UObject* const ActorCDO = Factory->GetDefaultActor( FAssetData() );
+			const UObject* const ActorCDO = Factory->GetDefaultActor( FAssetData(InActorClass) );
 			if ( ActorCDO != NULL && ActorCDO->GetClass() == InActorClass )
 			{
 				return Factory;
diff --git a/Engine/Source/Editor/UnrealEd/Private/Factories/ActorFactory.cpp b/Engine/Source/Editor/UnrealEd/Private/Factories/ActorFactory.cpp
--- a/Engine/Source/Editor/UnrealEd/Private/Factories/ActorFactory.cpp
+++ b/Engine/Source/Editor/UnrealEd/Private/Factories/ActorFactory.cpp
@@ -131,12 +131,12 @@ ActorFactory.cpp:
 #include "LevelSequence.h"
 #include "LevelSequenceActor.h"
 #include "Factories/ActorFactoryMovieScene.h"
+#include "Elements/Actor/ActorElementData.h"
+#include "Elements/Framework/EngineElementsLibrary.h"
+#include "Subsystems/PlacementSubsystem.h"
 
-#if WITH_EDITOR
 #include "LevelEditorViewport.h"
-#include "Editor/EditorEngine.h"
 #include "Editor.h"
-#endif
 
 DEFINE_LOG_CATEGORY(LogActorFactory);
 
@@ -249,7 +249,7 @@ UClass* UActorFactory::GetDefaultActorClass( const FAssetData& AssetData )
 
 UObject* UActorFactory::GetAssetFromActorInstance(AActor* ActorInstance)
 {
-	return NULL;
+	return nullptr;
 }
 
 FQuat UActorFactory::AlignObjectToSurfaceNormal(const FVector& InSurfaceNormal, const FQuat& ActorRotation) const
@@ -265,9 +265,94 @@ FQuat UActorFactory::AlignObjectToSurfaceNormal(const FVector& InSurfaceNormal,
 	}
 }
 
+bool UActorFactory::CanPlaceElementsFromAssetData(const FAssetData& InAssetData)
+{
+	UClass* Class = InAssetData.GetClass();
+	if (Class && (Class == GetDefaultActorClass(InAssetData)))
+	{
+		return true;
+	}
+
+	FText TmpMsg;
+	return CanCreateActorFrom(InAssetData, TmpMsg);
+}
+
+bool UActorFactory::PrePlaceAsset(FAssetPlacementInfo& InPlacementInfo, const FPlacementOptions& InPlacementOptions)
+{
+	return PreSpawnActor(InPlacementInfo.AssetToPlace.GetAsset(), InPlacementInfo.FinalizedTransform);
+}
+
+TArray<FTypedElementHandle> UActorFactory::PlaceAsset(const FAssetPlacementInfo& InPlacementInfo, const FPlacementOptions& InPlacementOptions)
+{
+	TArray<FTypedElementHandle> PlacedActorHandles;
+	AActor* NewActor = CreateActor(InPlacementInfo.AssetToPlace.GetAsset(), InPlacementInfo.PreferredLevel.Get(), InPlacementInfo.FinalizedTransform, RF_Transactional, InPlacementInfo.NameOverride);
+	if (NewActor)
+	{
+		PlacedActorHandles.Add(UEngineElementsLibrary::AcquireEditorActorElementHandle(NewActor));
+	}
+	return PlacedActorHandles;
+}
+
+void UActorFactory::PostPlaceAsset(TArrayView<const FTypedElementHandle> InElementHandles, const FAssetPlacementInfo& InPlacementInfo, const FPlacementOptions& InPlacementOptions)
+{
+	for (const FTypedElementHandle& PlacedElement : InElementHandles)
+	{
+		const FActorElementData* ActorElementData = PlacedElement.GetData<FActorElementData>();
+		if (ActorElementData && ActorElementData->Actor)
+		{
+			PostSpawnActor(InPlacementInfo.AssetToPlace.GetAsset(), ActorElementData->Actor);
+
+			// Only do this if the actor wasn't already given a name
+			if (InPlacementInfo.NameOverride.IsNone())
+			{
+				FActorLabelUtilities::SetActorLabelUnique(ActorElementData->Actor, InPlacementInfo.AssetToPlace.AssetName.ToString());
+			}
+
+			ActorElementData->Actor->PostEditChange();
+			ActorElementData->Actor->PostEditMove(true);
+		}
+	}
+}
+
+FAssetData UActorFactory::GetAssetDataFromElementHandle(const FTypedElementHandle& InHandle)
+{
+	const FActorElementData* ActorData = InHandle.GetData<FActorElementData>();
+	if (!ActorData || !ActorData->Actor)
+	{
+		return FAssetData();
+	}
+
+	AActor* ActorInstance = ActorData->Actor;
+	const UClass* DefaultActorClass = GetDefaultActorClass(FAssetData());
+
+	// See if the existing implementation of the actor factory has a valid asset data.
+	// First check that the actor instance derives from the default class
+	if (DefaultActorClass && ActorInstance->IsA(DefaultActorClass))
+	{
+		if (UObject* WrappedAssetObject = GetAssetFromActorInstance(ActorInstance))
+		{
+			return FAssetData(WrappedAssetObject);
+		}
+
+		// Query the actor itself for any contained content objects.
+		TArray<UObject*> ReferencedContentObjects;
+		ActorInstance->GetReferencedContentObjects(ReferencedContentObjects);
+		for (const UObject* ReferencedContentObject : ReferencedContentObjects)
+		{
+			FAssetData ReferencedObjectData = FAssetData(ReferencedContentObject);
+			if (CanPlaceElementsFromAssetData(ReferencedObjectData))
+			{
+				return ReferencedObjectData;
+			}
+		}
+	}
+
+	return FAssetData();
+}
+
 AActor* UActorFactory::CreateActor( UObject* Asset, ULevel* InLevel, FTransform SpawnTransform, EObjectFlags InObjectFlags, const FName Name )
 {
-	AActor* NewActor = NULL;
+	AActor* NewActor = nullptr;
 
 	if ( PreSpawnActor(Asset, SpawnTransform) )
 	{
@@ -282,6 +367,9 @@ AActor* UActorFactory::CreateActor( UObject* Asset, ULevel* InLevel, FTransform
 			{
 				FActorLabelUtilities::SetActorLabelUnique(NewActor, Asset->GetName());
 			}
+
+			NewActor->PostEditChange();
+			NewActor->PostEditMove(true);
 		}
 	}
 
@@ -320,7 +408,7 @@ AActor* UActorFactory::SpawnActor( UObject* Asset, ULevel* InLevel, const FTrans
 		return InLevel->OwningWorld->SpawnActor( DefaultActor->GetClass(), &Transform, SpawnInfo );
 	}
 
-	return NULL;
+	return nullptr;
 }
 
 void UActorFactory::PostSpawnActor( UObject* Asset, AActor* NewActor)
@@ -391,7 +479,7 @@ UObject* UActorFactoryStaticMesh::GetAssetFromActorInstance(AActor* Instance)
 
 void UActorFactoryStaticMesh::PostCreateBlueprint( UObject* Asset, AActor* CDO )
 {
-	if (Asset != NULL && CDO != NULL)
+	if (Asset != nullptr && CDO != nullptr)
 	{
 		UStaticMesh* StaticMesh = CastChecked<UStaticMesh>(Asset);
 		AStaticMeshActor* StaticMeshActor = CastChecked<AStaticMeshActor>(CDO);
@@ -541,7 +629,7 @@ void UActorFactoryDeferredDecal::PostSpawnActor(UObject* Asset, AActor* NewActor
 
 	UMaterialInterface* Material = GetMaterial( Asset );
 
-	if (Material != NULL )
+	if (Material != nullptr )
 	{
 		// Change properties
 		UDecalComponent* DecalComponent = nullptr;
@@ -567,11 +655,11 @@ void UActorFactoryDeferredDecal::PostSpawnActor(UObject* Asset, AActor* NewActor
 
 void UActorFactoryDeferredDecal::PostCreateBlueprint( UObject* Asset, AActor* CDO )
 {
-	if (Asset != NULL && CDO != NULL)
+	if (Asset != nullptr && CDO != nullptr)
 	{
 		UMaterialInterface* Material = GetMaterial(Asset);
 
-		if (Material != NULL)
+		if (Material != nullptr)
 		{
 			UDecalComponent* DecalComponent = nullptr;
 			for (UActorComponent* Component : CDO->GetComponents())
@@ -597,7 +685,7 @@ UMaterialInterface* UActorFactoryDeferredDecal::GetMaterial( UObject* Asset ) co
 		&& TargetMaterial->GetMaterial() 
 		&& TargetMaterial->GetMaterial()->MaterialDomain == MD_DeferredDecal ? 
 TargetMaterial : 
-	NULL;
+	nullptr;
 }
 
 /*-----------------------------------------------------------------------------
@@ -670,13 +758,13 @@ UObject* UActorFactoryEmitter::GetAssetFromActorInstance(AActor* Instance)
 	}
 	else
 	{
-		return NULL;
+		return nullptr;
 	}
 }
 
 void UActorFactoryEmitter::PostCreateBlueprint( UObject* Asset, AActor* CDO )
 {
-	if (Asset != NULL && CDO != NULL)
+	if (Asset != nullptr && CDO != nullptr)
 	{
 		UParticleSystem* ParticleSystem = CastChecked<UParticleSystem>(Asset);
 		AEmitter* Emitter = CastChecked<AEmitter>(CDO);
@@ -785,11 +873,11 @@ void UActorFactoryPhysicsAsset::PostSpawnActor(UObject* Asset, AActor* NewActor)
 
 void UActorFactoryPhysicsAsset::PostCreateBlueprint( UObject* Asset, AActor* CDO )
 {
-	if (CDO != NULL)
+	if (CDO != nullptr)
 	{
 		ASkeletalMeshActor* SkeletalPhysicsActor = CastChecked<ASkeletalMeshActor>(CDO);
 
-		if (Asset != NULL)
+		if (Asset != nullptr)
 		{
 			UPhysicsAsset* PhysicsAsset = CastChecked<UPhysicsAsset>(Asset);
 
@@ -878,7 +966,7 @@ bool UActorFactoryAnimationAsset::CanCreateActorFrom( const FAssetData& AssetDat
 
 USkeletalMesh* UActorFactoryAnimationAsset::GetSkeletalMeshFromAsset( UObject* Asset )
 {
-	USkeletalMesh* SkeletalMesh = NULL;
+	USkeletalMesh* SkeletalMesh = nullptr;
 	
 	if(UAnimSequenceBase* AnimationAsset = Cast<UAnimSequenceBase>(Asset))
 	{
@@ -893,10 +981,10 @@ USkeletalMesh* UActorFactoryAnimationAsset::GetSkeletalMeshFromAsset( UObject* A
 	// Check to see if we are using a custom factory in which case this should probably be ignored. This seems kind of wrong...
 	if( SkeletalMesh && SkeletalMesh->HasCustomActorFactory())
 	{
-		SkeletalMesh = NULL;
+		SkeletalMesh = nullptr;
 	}
 
-	check( SkeletalMesh != NULL );
+	check( SkeletalMesh != nullptr );
 	return SkeletalMesh;
 }
 
@@ -936,7 +1024,7 @@ void UActorFactoryAnimationAsset::PostCreateBlueprint( UObject* Asset,  AActor*
 {
 	Super::PostCreateBlueprint( Asset, CDO );
 	
-	if (Asset != NULL && CDO != NULL)
+	if (Asset != nullptr && CDO != nullptr)
 	{
 		UAnimationAsset* AnimationAsset = Cast<UAnimationAsset>(Asset);
 
@@ -950,6 +1038,21 @@ void UActorFactoryAnimationAsset::PostCreateBlueprint( UObject* Asset,  AActor*
 	}
 }
 
+UObject* UActorFactoryAnimationAsset::GetAssetFromActorInstance(AActor* ActorInstance)
+{
+	if (ASkeletalMeshActor* SkeletalMeshActor = Cast<ASkeletalMeshActor>(ActorInstance))
+	{
+		if (USkeletalMeshComponent* NewSASComponent = (SkeletalMeshActor->GetSkeletalMeshComponent()))
+		{
+			if (NewSASComponent->GetAnimationMode() == EAnimationMode::Type::AnimationSingleNode)
+			{
+				return NewSASComponent->AnimationData.AnimToPlay;
+			}
+		}
+	}
+	return nullptr;
+}
+
 
 /*-----------------------------------------------------------------------------
 UActorFactorySkeletalMesh
@@ -1070,18 +1173,18 @@ USkeletalMesh* UActorFactorySkeletalMesh::GetSkeletalMeshFromAsset( UObject* Ass
 	UAnimBlueprint* AnimBlueprint = Cast<UAnimBlueprint>( Asset );
 	USkeleton* Skeleton = Cast<USkeleton>( Asset );
 
-	if( SkeletalMesh == NULL && AnimBlueprint != NULL && AnimBlueprint->TargetSkeleton )
+	if( SkeletalMesh == nullptr && AnimBlueprint != nullptr && AnimBlueprint->TargetSkeleton )
 	{
 		// base it on preview skeletal mesh, just to have something
 		SkeletalMesh = AnimBlueprint->TargetSkeleton->GetPreviewMesh(true);
 	}
 
-	if( SkeletalMesh == NULL && Skeleton != NULL )
+	if( SkeletalMesh == nullptr && Skeleton != nullptr )
 	{
 		SkeletalMesh = Skeleton->GetPreviewMesh(true);
 	}
 
-	if (SkeletalMesh == NULL)
+	if (SkeletalMesh == nullptr)
 	{
 		for (const auto& Pair : UActorFactorySkeletalMesh::GetSkeletalMeshDelegates)
 		{
@@ -1094,7 +1197,7 @@ USkeletalMesh* UActorFactorySkeletalMesh::GetSkeletalMeshFromAsset( UObject* Ass
 		}
 	}
 
-	check( SkeletalMesh != NULL );
+	check( SkeletalMesh != nullptr );
 	return SkeletalMesh;
 }
 
@@ -1123,7 +1226,7 @@ void UActorFactorySkeletalMesh::PostSpawnActor( UObject* Asset, AActor* NewActor
 		NewSMActor->GetSkeletalMeshComponent()->SetAnimInstanceClass(AnimBlueprint->GeneratedClass);
 	}
 
-	if (ClassUsedForDelegate != NULL)
+	if (ClassUsedForDelegate != nullptr)
 	{
 		UActorFactorySkeletalMesh::PostSkeletalMeshActorSpawnedDelegates.FindChecked(ClassUsedForDelegate).Execute(NewSMActor, Asset);
 	}
@@ -1132,14 +1235,14 @@ void UActorFactorySkeletalMesh::PostSpawnActor( UObject* Asset, AActor* NewActor
 
 void UActorFactorySkeletalMesh::PostCreateBlueprint( UObject* Asset, AActor* CDO )
 {
-	if (Asset != NULL && CDO != NULL)
+	if (Asset != nullptr && CDO != nullptr)
 	{
 		USkeletalMesh* SkeletalMesh = GetSkeletalMeshFromAsset(Asset);
 		UAnimBlueprint* AnimBlueprint = Cast<UAnimBlueprint>(Asset);
 
 		ASkeletalMeshActor* SkeletalMeshActor = CastChecked<ASkeletalMeshActor>(CDO);
 		SkeletalMeshActor->GetSkeletalMeshComponent()->SkeletalMesh = SkeletalMesh;
-		SkeletalMeshActor->GetSkeletalMeshComponent()->AnimClass = AnimBlueprint ? Cast<UAnimBlueprintGeneratedClass>(AnimBlueprint->GeneratedClass) : NULL;
+		SkeletalMeshActor->GetSkeletalMeshComponent()->AnimClass = AnimBlueprint ? Cast<UAnimBlueprintGeneratedClass>(AnimBlueprint->GeneratedClass) : nullptr;
 	}
 }
 
@@ -1278,7 +1381,7 @@ void UActorFactoryAmbientSound::PostSpawnActor( UObject* Asset, AActor* NewActor
 
 	USoundBase* AmbientSound = Cast<USoundBase>(Asset);
 
-	if ( AmbientSound != NULL )
+	if ( AmbientSound != nullptr )
 	{
 		AAmbientSound* NewSound = CastChecked<AAmbientSound>( NewActor );
 		NewSound->GetAudioComponent()->SetSound(AmbientSound);
@@ -1296,11 +1399,11 @@ UObject* UActorFactoryAmbientSound::GetAssetFromActorInstance(AActor* Instance)
 
 void UActorFactoryAmbientSound::PostCreateBlueprint( UObject* Asset, AActor* CDO )
 {
-	if (Asset != NULL && CDO != NULL)
+	if (Asset != nullptr && CDO != nullptr)
 	{
 		USoundBase* AmbientSound = Cast<USoundBase>(Asset);
 
-		if (AmbientSound != NULL)
+		if (AmbientSound != nullptr)
 		{
 			AAmbientSound* NewSound = CastChecked<AAmbientSound>(CDO);
 			NewSound->GetAudioComponent()->SetSound(AmbientSound);
@@ -1322,7 +1425,7 @@ bool UActorFactoryClass::CanCreateActorFrom( const FAssetData& AssetData, FText&
 	if ( AssetData.IsValid() && AssetData.GetClass()->IsChildOf( UClass::StaticClass() ) )
 	{
 		UClass* ActualClass = Cast<UClass>(AssetData.GetAsset());
-		if ( (NULL != ActualClass) && ActualClass->IsChildOf(AActor::StaticClass()) )
+		if ( (nullptr != ActualClass) && ActualClass->IsChildOf(AActor::StaticClass()) )
 		{
 			return true;
 		}
@@ -1336,23 +1439,23 @@ AActor* UActorFactoryClass::GetDefaultActor( const FAssetData& AssetData )
 {
 	if ( AssetData.IsValid() && AssetData.GetClass()->IsChildOf( UClass::StaticClass() ) )
 	{
-		UClass* ActualClass = Cast<UClass>(StaticLoadObject(UClass::StaticClass(), NULL, *AssetData.ObjectPath.ToString(), NULL, LOAD_NoWarn, NULL));
+		UClass* ActualClass = Cast<UClass>(StaticLoadObject(UClass::StaticClass(), nullptr, *AssetData.ObjectPath.ToString(), nullptr, LOAD_NoWarn, nullptr));
 			
 		//Cast<UClass>(AssetData.GetAsset());
-		if ( (NULL != ActualClass) && ActualClass->IsChildOf(AActor::StaticClass()) )
+		if ( (nullptr != ActualClass) && ActualClass->IsChildOf(AActor::StaticClass()) )
 		{
 			return ActualClass->GetDefaultObject<AActor>();
 		}
 	}
 
-	return NULL;
+	return nullptr;
 }
 
 bool UActorFactoryClass::PreSpawnActor( UObject* Asset, FTransform& InOutLocation)
 {
 	UClass* ActualClass = Cast<UClass>(Asset);
 
-	if ( (NULL != ActualClass) && ActualClass->IsChildOf(AActor::StaticClass()) )
+	if ( (nullptr != ActualClass) && ActualClass->IsChildOf(AActor::StaticClass()) )
 	{
 		return true;
 	}
@@ -1364,7 +1467,7 @@ AActor* UActorFactoryClass::SpawnActor( UObject* Asset, ULevel* InLevel, const F
 {
 	UClass* ActualClass = Cast<UClass>(Asset);
 
-	if ( (NULL != ActualClass) && ActualClass->IsChildOf(AActor::StaticClass()) )
+	if ( (nullptr != ActualClass) && ActualClass->IsChildOf(AActor::StaticClass()) )
 	{
 		FActorSpawnParameters SpawnInfo;
 		SpawnInfo.OverrideLevel = InLevel;
@@ -1373,7 +1476,7 @@ AActor* UActorFactoryClass::SpawnActor( UObject* Asset, ULevel* InLevel, const F
 		return InLevel->OwningWorld->SpawnActor( ActualClass, &Transform, SpawnInfo );
 	}
 
-	return NULL;
+	return nullptr;
 }
 
 
@@ -1401,10 +1504,10 @@ bool UActorFactoryBlueprint::CanCreateActorFrom( const FAssetData& AssetData, FT
 		return false;
 	}
 
-	UClass* ParentClass = FindObject<UClass>(NULL, *ParentClassPath);
+	UClass* ParentClass = FindObject<UClass>(nullptr, *ParentClassPath);
 
 	bool bIsActorBased = false;
-	if ( ParentClass != NULL )
+	if ( ParentClass != nullptr )
 	{
 		// The parent class is loaded. Make sure it is derived from AActor
 		bIsActorBased = ParentClass->IsChildOf(AActor::StaticClass());
@@ -1435,25 +1538,31 @@ bool UActorFactoryBlueprint::CanCreateActorFrom( const FAssetData& AssetData, FT
 
 AActor* UActorFactoryBlueprint::GetDefaultActor( const FAssetData& AssetData )
 {
-	if ( !AssetData.IsValid() || !AssetData.GetClass()->IsChildOf( UBlueprint::StaticClass() ) )
+	UClass* GeneratedClass = GetDefaultActorClass(AssetData);
+	return GeneratedClass ? GeneratedClass->GetDefaultObject<AActor>() : nullptr;
+}
+
+UClass* UActorFactoryBlueprint::GetDefaultActorClass(const FAssetData& AssetData)
+{
+	if (!AssetData.IsValid() || !AssetData.GetClass()->IsChildOf(UBlueprint::StaticClass()))
 	{
-		return NULL;
+		return nullptr;
 	}
 
 	const FString GeneratedClassPath = AssetData.GetTagValueRef<FString>(FBlueprintTags::GeneratedClassPath);
-	if ( GeneratedClassPath.IsEmpty() )
+	if (GeneratedClassPath.IsEmpty())
 	{
-		return NULL;
+		return nullptr;
 	}
 
-	UClass* GeneratedClass = Cast<UClass>(StaticLoadObject(UClass::StaticClass(), NULL, *GeneratedClassPath, NULL, LOAD_NoWarn, NULL));
+	UClass* GeneratedClass = Cast<UClass>(StaticLoadObject(UClass::StaticClass(), nullptr, *GeneratedClassPath, nullptr, LOAD_NoWarn, nullptr));
 
-	if ( GeneratedClass == NULL )
+	if (GeneratedClass == nullptr)
 	{
-		return NULL;
+		return nullptr;
 	}
 
-	return GeneratedClass->GetDefaultObject<AActor>();
+	return GeneratedClass;
 }
 
 bool UActorFactoryBlueprint::PreSpawnActor( UObject* Asset, FTransform& InOutLocation)
@@ -1461,7 +1570,7 @@ bool UActorFactoryBlueprint::PreSpawnActor( UObject* Asset, FTransform& InOutLoc
 	UBlueprint* Blueprint = CastChecked<UBlueprint>(Asset);
 
 	// Invalid if there is no generated class, or this is not actor based
-	if (Blueprint == NULL || Blueprint->GeneratedClass == NULL || !FBlueprintEditorUtils::IsActorBased(Blueprint))
+	if (Blueprint == nullptr || Blueprint->GeneratedClass == nullptr || !FBlueprintEditorUtils::IsActorBased(Blueprint))
 	{
 		return false;
 	}
@@ -1519,7 +1628,7 @@ void UActorFactoryMatineeActor::PostSpawnActor( UObject* Asset, AActor* NewActor
 
 void UActorFactoryMatineeActor::PostCreateBlueprint( UObject* Asset, AActor* CDO )
 {
-	if (Asset != NULL && CDO != NULL)
+	if (Asset != nullptr && CDO != nullptr)
 	{
 		UInterpData* MatineeData = Cast<UInterpData>(Asset);
 		AMatineeActor* MatineeActor = CastChecked<AMatineeActor>(CDO);
@@ -1762,10 +1871,10 @@ CreateBrushForVolumeActor
 // Helper function for the volume actor factories, not sure where it should live
 void CreateBrushForVolumeActor( AVolume* NewActor, UBrushBuilder* BrushBuilder )
 {
-	if ( NewActor != NULL )
+	if ( NewActor != nullptr )
 	{
 		// this code builds a brush for the new actor
-		NewActor->PreEditChange(NULL);
+		NewActor->PreEditChange(nullptr);
 
 		// Use the same object flags as the owner volume
 		EObjectFlags ObjectFlags = NewActor->GetFlags() & (RF_Transient | RF_Transactional);
@@ -1784,14 +1893,14 @@ void CreateBrushForVolumeActor( AVolume* NewActor, UBrushBuilder* BrushBuilder )
 
 		FBSPOps::csgPrepMovingBrush( NewActor );
 
-		// Set the texture on all polys to NULL.  This stops invisible textures
+		// Set the texture on all polys to nullptr.  This stops invisible textures
 		// dependencies from being formed on volumes.
 		if ( NewActor->Brush )
 		{
 			for ( int32 poly = 0 ; poly < NewActor->Brush->Polys->Element.Num() ; ++poly )
 			{
 				FPoly* Poly = &(NewActor->Brush->Polys->Element[poly]);
-				Poly->Material = NULL;
+				Poly->Material = nullptr;
 			}
 		}
 
@@ -1829,7 +1938,7 @@ void UActorFactoryBoxVolume::PostSpawnActor( UObject* Asset, AActor* NewActor )
 	Super::PostSpawnActor(Asset, NewActor);
 
 	AVolume* VolumeActor = CastChecked<AVolume>(NewActor);
-	if ( VolumeActor != NULL )
+	if ( VolumeActor != nullptr )
 	{
 		UCubeBuilder* Builder = NewObject<UCubeBuilder>();
 		CreateBrushForVolumeActor( VolumeActor, Builder );
@@ -1866,7 +1975,7 @@ void UActorFactorySphereVolume::PostSpawnActor( UObject* Asset, AActor* NewActor
 	Super::PostSpawnActor(Asset, NewActor);
 
 	AVolume* VolumeActor = CastChecked<AVolume>(NewActor);
-	if ( VolumeActor != NULL )
+	if ( VolumeActor != nullptr )
 	{
 		UTetrahedronBuilder* Builder = NewObject<UTetrahedronBuilder>();
 		Builder->SphereExtrapolation = 2;
@@ -1884,27 +1993,12 @@ UActorFactoryCylinderVolume::UActorFactoryCylinderVolume( const FObjectInitializ
 	DisplayName = LOCTEXT( "CylinderVolumeDisplayName", "Cylinder Volume" );
 	NewActorClass = AVolume::StaticClass();
 }
-bool UActorFactoryCylinderVolume::CanCreateActorFrom( const FAssetData& AssetData, FText& OutErrorMsg )
-{
-	if ( UActorFactory::CanCreateActorFrom( AssetData, OutErrorMsg ) )
-	{
-		return true;
-	}
-
-	if ( AssetData.IsValid() && !AssetData.GetClass()->IsChildOf( AVolume::StaticClass() ) )
-	{
-		return false;
-	}
-
-	return true;
-
-}
 void UActorFactoryCylinderVolume::PostSpawnActor( UObject* Asset, AActor* NewActor )
 {
 	Super::PostSpawnActor(Asset, NewActor);
 
 	AVolume* VolumeActor = CastChecked<AVolume>(NewActor);
-	if ( VolumeActor != NULL )
+	if ( VolumeActor != nullptr )
 	{
 		UCylinderBuilder* Builder = NewObject<UCylinderBuilder>();
 		Builder->OuterRadius = 128.0f;
