Migrated the cooking commandlet code for Studio Analytics from UE4 to UE5. This seemed to have gotten lost in a prevoius merge

diff --git a/Engine/Source/Editor/UnrealEd/Private/Commandlets/CookCommandlet.cpp b/Engine/Source/Editor/UnrealEd/Private/Commandlets/CookCommandlet.cpp
--- a/Engine/Source/Editor/UnrealEd/Private/Commandlets/CookCommandlet.cpp
+++ b/Engine/Source/Editor/UnrealEd/Private/Commandlets/CookCommandlet.cpp
@@ -41,6 +41,7 @@
 #include "HAL/MemoryMisc.h"
 #include "ProfilingDebugging/CookStats.h"
 #include "AssetRegistryModule.h"
+#include "StudioAnalytics.h"
 #include "Cooker/CookProfiling.h"
 
 DEFINE_LOG_CATEGORY_STATIC(LogCookCommandlet, Log, All);
@@ -100,6 +101,31 @@ namespace DetailedCookStats
 
 	static void LogCookStats(const FString& CookCmdLine)
 	{
+		if (FStudioAnalytics::IsAvailable())
+		{
+			// convert filtered stats directly to an analytics event
+			TArray<FAnalyticsEventAttribute> StatAttrs;
+
+			// Sends each cook stat to the studio analytics system.
+			auto SendCookStatsToAnalytics = [&StatAttrs](const FString& StatName, const TArray<FCookStatsManager::StringKeyValue>& StatAttributes)
+			{
+				for (const auto& Attr : StatAttributes)
+				{
+					FString FormattedAttrName = StatName + "." + Attr.Key;
+
+					StatAttrs.Emplace(FormattedAttrName, Attr.Value);
+				}
+			};
+
+			// Now actually grab the stats 
+			FCookStatsManager::LogCookStats(SendCookStatsToAnalytics);
+
+			// Record them all under cooking event
+			FStudioAnalytics::GetProvider().RecordEvent(TEXT("Core.Cooking"), StatAttrs);
+
+			FStudioAnalytics::GetProvider().BlockUntilFlushed(60.0f);
+		}
+
 		bool bSendCookAnalytics = false;
 		GConfig->GetBool(TEXT("CookAnalytics"), TEXT("SendAnalytics"), bSendCookAnalytics, GEngineIni);
 
