Add the WorldBuilder builder object to the root set to avoid being GC'd when collecting garbage in various builders.

diff --git a/Engine/Source/Editor/UnrealEd/Private/Commandlets/WorldPartitionBuilderCommandlet.cpp b/Engine/Source/Editor/UnrealEd/Private/Commandlets/WorldPartitionBuilderCommandlet.cpp
--- a/Engine/Source/Editor/UnrealEd/Private/Commandlets/WorldPartitionBuilderCommandlet.cpp
+++ b/Engine/Source/Editor/UnrealEd/Private/Commandlets/WorldPartitionBuilderCommandlet.cpp
@@ -124,6 +124,8 @@ int32 UWorldPartitionBuilderCommandlet::Main(const FString& Params)
 	// Create builder instance
 	UWorldPartitionBuilder* Builder = NewObject<UWorldPartitionBuilder>(this, BuilderClass);
 	check(Builder);
+
+	Builder->AddToRoot();
 	
 	// Validate builder settings
 	if (Builder->RequiresCommandletRendering() && !IsAllowCommandletRendering())
