Revert CL 14828781

diff --git a/Engine/Source/Editor/UnrealEd/Private/UnrealEdSrv.cpp b/Engine/Source/Editor/UnrealEd/Private/UnrealEdSrv.cpp
--- a/Engine/Source/Editor/UnrealEd/Private/UnrealEdSrv.cpp
+++ b/Engine/Source/Editor/UnrealEd/Private/UnrealEdSrv.cpp
@@ -16,7 +16,6 @@
 #include "UObject/UnrealType.h"
 #include "InputCoreTypes.h"
 #include "Input/Reply.h"
-#include "Slate/SceneViewport.h"
 #include "Widgets/DeclarativeSyntaxSupport.h"
 #include "Widgets/SWidget.h"
 #include "Widgets/SCompoundWidget.h"
@@ -692,8 +691,7 @@ bool UUnrealEdEngine::Exec( UWorld* InWorld, const TCHAR* Stream, FOutputDevice&
 	}
 
 	// If we don't have a viewport specified to catch the stat commands, use to the active viewport.  If there is a game viewport ignore this as we do not want 
-	if (GStatProcessingViewportClient == NULL && (GameViewport == NULL || GameViewport->IsSimulateInEditorViewport()
-		|| (GameViewport->GetGameViewport() && GameViewport->GetGameViewport()->IsPlayInEditorViewport()) ) )
+	if (GStatProcessingViewportClient == NULL && (GameViewport == NULL || GameViewport->IsSimulateInEditorViewport() ) )
 	{
 		GStatProcessingViewportClient = GLastKeyLevelEditingViewportClient ? GLastKeyLevelEditingViewportClient : GCurrentLevelEditingViewportClient;
 	}
