Fix cooker's inisettingsoutofdate feature:
When saving the inisettings, mark inifiles that are in GConfig as platform-agnostic, rather than the platform of the editor.
When saving inifiles, do not mark all of their values as having been read for dependency tracking.
Cache inifiles to improve performance in IniSettingsOutOfDate.
#rnx

diff --git a/Engine/Source/Editor/UnrealEd/Private/CookOnTheFlyServer.cpp b/Engine/Source/Editor/UnrealEd/Private/CookOnTheFlyServer.cpp
--- a/Engine/Source/Editor/UnrealEd/Private/CookOnTheFlyServer.cpp
+++ b/Engine/Source/Editor/UnrealEd/Private/CookOnTheFlyServer.cpp
@@ -82,6 +82,7 @@
 #include "Serialization/BulkDataManifest.h"
 #include "MeshCardRepresentation.h"
 #include "Misc/PathViews.h"
+#include "String/Find.h"
 
 #include "AssetRegistryModule.h"
 #include "AssetRegistryState.h"
@@ -4733,36 +4734,59 @@ void UCookOnTheFlyServer::OnFConfigDeleted(const FConfigFile* Config)
 	OpenConfigFiles.Remove(Config);
 }
 
-
 void UCookOnTheFlyServer::ProcessAccessedIniSettings(const FConfigFile* Config, FIniSettingContainer& OutAccessedIniStrings) const
 {	
 	if (Config->Name == NAME_None)
 	{
 		return;
 	}
-	// try figure out if this config file is for a specific platform 
+
+	// try to figure out if this config file is for a specific platform 
 	FString PlatformName;
 	bool bFoundPlatformName = false;
-	for (auto It : FDataDrivenPlatformInfoRegistry::GetAllPlatformInfos())
-	{
-		FString CurrentPlatformName = It.Key.ToString();
-		for ( const auto& SourceIni : Config->SourceIniHierarchy )
-		{
-			if ( SourceIni.Value.Filename.Contains(CurrentPlatformName) )
+
+	if (!GConfig->ContainsConfigFile(Config)) // If the ConfigFile is in GConfig, then it is the editor's config and is not platform specific
+	{
+		// For the config files not in GConfig, we assume they were loaded from LoadConfigFile, and we match these to a platform
+		// By looking for a platform-specific filepath in their SourceIniHierarchy.
+		// Examples:
+		// (1) ROOT\Engine\Config\Windows\WindowsEngine.ini
+		// (2) ROOT\Engine\Config\Android\DataDrivePlatformInfo.ini
+		// (3) ROOT\Engine\Config\Android\AndroidWindowsCompatability.ini
+		// 
+		// Note that for config files of form #3, we want them to be matched to Android rather than windows;
+		// we assume that an exact match on a directory component is more definitive than a substring match
+		bool bFoundPlatformGuess = false;
+		for (auto It : FDataDrivenPlatformInfoRegistry::GetAllPlatformInfos())
+		{
+			FString CurrentPlatformName = It.Key.ToString();
+			TStringBuilder<128> PlatformDirString;
+			PlatformDirString.Appendf(TEXT("/%s/"), *CurrentPlatformName);
+			for (const auto& SourceIni : Config->SourceIniHierarchy)
+			{
+				// Look for platform in the path, rating a full subdirectory name match (/Android/ or /Windows/) higher than a partial filename match (AndroidEngine.ini or WindowsEngine.ini)
+				bool bFoundPlatformDir = UE::String::FindFirst(SourceIni.Value.Filename, PlatformDirString, ESearchCase::IgnoreCase) != INDEX_NONE;
+				bool bFoundPlatformSubstring = UE::String::FindFirst(SourceIni.Value.Filename, CurrentPlatformName, ESearchCase::IgnoreCase) != INDEX_NONE;
+				if (bFoundPlatformDir)
+				{
+					PlatformName = MoveTemp(CurrentPlatformName);
+					bFoundPlatformName = true;
+					break;
+				}
+				else if (!bFoundPlatformGuess && bFoundPlatformSubstring)
+				{
+					PlatformName = MoveTemp(CurrentPlatformName);
+					bFoundPlatformGuess = true;
+				}
+			}
+			if (bFoundPlatformName)
 			{
-				PlatformName = MoveTemp(CurrentPlatformName);
-				bFoundPlatformName = true;
 				break;
 			}
 		}
-		if ( bFoundPlatformName )
-		{
-			break;
-		}
+		bFoundPlatformName = bFoundPlatformName || bFoundPlatformGuess;
 	}
 
-	
-
 	TStringBuilder<128> ConfigName;
 	if (bFoundPlatformName)
 	{
@@ -4832,9 +4856,6 @@ void UCookOnTheFlyServer::ProcessAccessedIniSettings(const FConfigFile* Config,
 	}
 }
 
-
-
-
 bool UCookOnTheFlyServer::IniSettingsOutOfDate(const ITargetPlatform* TargetPlatform) const
 {
 	FScopeAssign<bool> A = FScopeAssign<bool>(IniSettingRecurse, true);
@@ -4867,20 +4888,23 @@ bool UCookOnTheFlyServer::IniSettingsOutOfDate(const ITargetPlatform* TargetPlat
 		}
 	}
 
-	for ( const auto& OldIniFile : OldIniSettings )
+	for (const auto& OldIniFile : OldIniSettings)
 	{
-		const FName& ConfigNameKey = OldIniFile.Key;
+		FName ConfigNameKey = OldIniFile.Key;
 
 		TArray<FString> ConfigNameArray;
 		ConfigNameKey.ToString().ParseIntoArray(ConfigNameArray, TEXT("."));
 		FString Filename;
 		FString PlatformName;
+		// The input NameKey is of the form 
+		//   Platform.ConfigName:Section:Key:ArrayIndex=Value
+		// The Platform is optional and will not be present if the configfile was an editor config file rather than a platform-specific config file
 		bool bFoundPlatformName = false;
-		if ( ConfigNameArray.Num() <= 1 )
+		if (ConfigNameArray.Num() <= 1)
 		{
 			Filename = ConfigNameKey.ToString();
 		}
-		else if ( ConfigNameArray.Num() == 2 )
+		else if (ConfigNameArray.Num() == 2)
 		{
 			PlatformName = ConfigNameArray[0];
 			Filename = ConfigNameArray[1];
@@ -4888,25 +4912,30 @@ bool UCookOnTheFlyServer::IniSettingsOutOfDate(const ITargetPlatform* TargetPlat
 		}
 		else
 		{
-			UE_LOG( LogCook, Warning, TEXT("Found invalid file name in old ini settings file Filename %s settings file %s"), *ConfigNameKey.ToString(), *TargetPlatform->PlatformName() );
+			UE_LOG(LogCook, Warning, TEXT("Found invalid file name in old ini settings file Filename %s settings file %s"), *ConfigNameKey.ToString(), *TargetPlatform->PlatformName());
 			return true;
 		}
-		
+
 		const FConfigFile* ConfigFile = nullptr;
 		FConfigFile Temp;
-		if ( bFoundPlatformName)
+		if (bFoundPlatformName)
 		{
-			GConfig->LoadLocalIniFile(Temp, *Filename, true, *PlatformName );
+			// For the platform-specific old ini files, load them using LoadLocalIniFiles; this matches the assumption in SaveCurrentIniSettings
+			// that the platform-specific ini files were loaded by LoadLocalIniFiles
+			FConfigCacheIni::LoadLocalIniFile(Temp, *Filename, true, *PlatformName);
 			ConfigFile = &Temp;
 		}
 		else
 		{
-			ConfigFile = GConfig->Find(Filename);
+			// For the platform-agnostic old ini files, read them from GConfig; this matches where we loaded them from in SaveCurrentIniSettings
+			// The ini files may have been saved by fullpath or by shortname; search first for a fullpath match using FindConfigFile and
+			// if that fails search for the shortname match by iterating over all files in GConfig
+			ConfigFile = GConfig->FindConfigFile(Filename);
 		}
-		FName FileFName = FName(*Filename);
-		if ( !ConfigFile )
+		if (!ConfigFile)
 		{
-			for( const FString& ConfigFilename : GConfig->GetFilenames() )
+			FName FileFName = FName(*Filename);
+			for (const FString& ConfigFilename : GConfig->GetFilenames())
 			{
 				FConfigFile* File = GConfig->FindConfigFile(ConfigFilename);
 				if (File->Name == FileFName)
@@ -4915,7 +4944,7 @@ bool UCookOnTheFlyServer::IniSettingsOutOfDate(const ITargetPlatform* TargetPlat
 					break;
 				}
 			}
-			if ( !ConfigFile )
+			if (!ConfigFile)
 			{
 				UE_LOG(LogCook, Display, TEXT("Unable to find config file %s invalidating inisettings"), *FString::Printf(TEXT("%s %s"), *PlatformName, *Filename));
 				return true;
@@ -4923,15 +4952,14 @@ bool UCookOnTheFlyServer::IniSettingsOutOfDate(const ITargetPlatform* TargetPlat
 		}
 		for ( const auto& OldIniSection : OldIniFile.Value )
 		{
-			
 			const FName& SectionName = OldIniSection.Key;
 			const FConfigSection* IniSection = ConfigFile->Find( SectionName.ToString() );
-
-			const FString BlackListSetting = *FString::Printf(TEXT("%s.%s:%s"), *PlatformName, *Filename, *SectionName.ToString());
+			const FString BlackListSetting = FString::Printf(TEXT("%s%s%s:%s"), *PlatformName, bFoundPlatformName ? TEXT(".") : TEXT(""), *Filename, *SectionName.ToString());
 
 			if ( IniSection == nullptr )
 			{
-				UE_LOG(LogCook, Display, TEXT("Inisetting is different for %s, Current section doesn't exist"), *FString::Printf(TEXT("%s %s %s"), *PlatformName, *Filename, *SectionName.ToString()));
+				UE_LOG(LogCook, Display, TEXT("Inisetting is different for %s, Current section doesn't exist"), 
+					*FString::Printf(TEXT("%s %s %s"), *PlatformName, *Filename, *SectionName.ToString()));
 				UE_LOG(LogCook, Display, TEXT("To avoid this add blacklist setting to DefaultEditor.ini [CookSettings] %s"), *BlackListSetting);
 				return true;
 			}
@@ -4945,7 +4973,8 @@ bool UCookOnTheFlyServer::IniSettingsOutOfDate(const ITargetPlatform* TargetPlat
 
 				if ( CurrentValues.Num() != OldIniValue.Value.Num() )
 				{
-					UE_LOG(LogCook, Display, TEXT("Inisetting is different for %s, missmatched num array elements %d != %d "), *FString::Printf(TEXT("%s %s %s %s"), *PlatformName, *Filename, *SectionName.ToString(), *ValueName.ToString()), CurrentValues.Num(), OldIniValue.Value.Num());
+					UE_LOG(LogCook, Display, TEXT("Inisetting is different for %s, missmatched num array elements %d != %d "), *FString::Printf(TEXT("%s %s %s %s"),
+						*PlatformName, *Filename, *SectionName.ToString(), *ValueName.ToString()), CurrentValues.Num(), OldIniValue.Value.Num());
 					UE_LOG(LogCook, Display, TEXT("To avoid this add blacklist setting to DefaultEditor.ini [CookSettings] %s"), *BlackListSetting);
 					return true;
 				}
@@ -4954,12 +4983,13 @@ bool UCookOnTheFlyServer::IniSettingsOutOfDate(const ITargetPlatform* TargetPlat
 					const FString FilteredCurrentValue = CurrentValues[Index].GetSavedValue().Replace(TEXT(":"), TEXT(""));
 					if ( FilteredCurrentValue != OldIniValue.Value[Index] )
 					{
-						UE_LOG(LogCook, Display, TEXT("Inisetting is different for %s, value %s != %s invalidating cook"),  *FString::Printf(TEXT("%s %s %s %s %d"),*PlatformName, *Filename, *SectionName.ToString(), *ValueName.ToString(), Index), *CurrentValues[Index].GetSavedValue(), *OldIniValue.Value[Index] );
+						UE_LOG(LogCook, Display, TEXT("Inisetting is different for %s, value %s != %s invalidating cook"),
+							*FString::Printf(TEXT("%s %s %s %s %d"),*PlatformName, *Filename, *SectionName.ToString(), *ValueName.ToString(), Index),
+							*CurrentValues[Index].GetSavedValue(), *OldIniValue.Value[Index] );
 						UE_LOG(LogCook, Display, TEXT("To avoid this add blacklist setting to DefaultEditor.ini [CookSettings] %s"), *BlackListSetting);
 						return true;
 					}
 				}
-				
 			}
 		}
 	}
diff --git a/Engine/Source/Runtime/Core/Private/Misc/ConfigCacheIni.cpp b/Engine/Source/Runtime/Core/Private/Misc/ConfigCacheIni.cpp
--- a/Engine/Source/Runtime/Core/Private/Misc/ConfigCacheIni.cpp
+++ b/Engine/Source/Runtime/Core/Private/Misc/ConfigCacheIni.cpp
@@ -1574,7 +1574,8 @@ bool FConfigFile::Write(const FString& Filename, bool bDoRemoteWrite, TMap<FStri
 		for( FConfigSection::TConstIterator It2(Section); It2; ++It2 )
 		{
 			const FName PropertyName = It2.Key();
-			const FString& PropertyValue = It2.Value().GetSavedValue();
+			// Use GetSavedValueForWriting rather than GetSavedValue to avoid having this save operation mark the values as having been accessed for dependency tracking
+			const FString& PropertyValue = UE::ConfigCacheIni::Private::FAccessor::GetSavedValueForWriting(It2.Value());
 
 			// Check if the we've already processed a property of this name. If it was part of an array we may have already written it out.
 			if( !PropertiesAddedLookup.Contains( PropertyName ) )
@@ -1629,7 +1630,8 @@ bool FConfigFile::Write(const FString& Filename, bool bDoRemoteWrite, TMap<FStri
 						PropertyName.AppendString(PropertyNameString);
 						for (const FConfigValue* ConfigValue : CompletePropertyToWrite)
 						{
-							AppendExportedPropertyLine(Text, PropertyNameString, ConfigValue->GetSavedValue());
+							// Use GetSavedValueForWriting rather than GetSavedValue to avoid marking these values used during save to disk as having been accessed for dependency tracking
+							AppendExportedPropertyLine(Text, PropertyNameString, UE::ConfigCacheIni::Private::FAccessor::GetSavedValueForWriting(*ConfigValue));
 						}
 					}
 
@@ -1920,7 +1922,8 @@ void FConfigFile::SetString( const TCHAR* Section, const TCHAR* Key, const TCHAR
 		Sec->Add( Key, Value );
 		Dirty = true;
 	}
-	else if( FCString::Strcmp(*ConfigValue->GetSavedValue(),Value)!=0 )
+	// Use GetSavedValueForWriting rather than GetSavedValue to avoid having the is-it-dirty query mark the values as having been accessed for dependency tracking
+	else if( FCString::Strcmp(*UE::ConfigCacheIni::Private::FAccessor::GetSavedValueForWriting(*ConfigValue),Value)!=0 )
 	{
 		Dirty = true;
 		*ConfigValue = FConfigValue(Value);
@@ -1940,7 +1943,8 @@ void FConfigFile::SetText( const TCHAR* Section, const TCHAR* Key, const FText&
 		Sec->Add( Key, StrValue );
 		Dirty = true;
 	}
-	else if( FCString::Strcmp(*ConfigValue->GetSavedValue(), *StrValue)!=0 )
+	// Use GetSavedValueForWriting rather than GetSavedValue to avoid having the is-it-dirty query mark the values as having been accessed for dependency tracking
+	else if( FCString::Strcmp(*UE::ConfigCacheIni::Private::FAccessor::GetSavedValueForWriting(*ConfigValue), *StrValue)!=0 )
 	{
 		Dirty = true;
 		*ConfigValue = FConfigValue(StrValue);
@@ -1988,7 +1992,8 @@ void FConfigFile::SaveSourceToBackupFile()
 		for( FConfigSection::TConstIterator PropertyIterator(Section); PropertyIterator; ++PropertyIterator )
 		{
 			const FName PropertyName = PropertyIterator.Key();
-			const FString& PropertyValue = PropertyIterator.Value().GetSavedValue();
+			// Use GetSavedValueForWriting rather than GetSavedValue to avoid having this save operation mark the values as having been accessed for dependency tracking
+			const FString& PropertyValue = UE::ConfigCacheIni::Private::FAccessor::GetSavedValueForWriting(PropertyIterator.Value());
 			Text += FConfigFile::GenerateExportedPropertyLine(PropertyName.ToString(), PropertyValue);
 		}
 		Text += LINE_TERMINATOR;
@@ -2161,6 +2166,33 @@ FConfigFile* FConfigCacheIni::FindConfigFileWithBaseName(FName BaseName)
 	return nullptr;
 }
 
+bool FConfigCacheIni::ContainsConfigFile(const FConfigFile* ConfigFile) const
+{
+	// Check the normal inis. Note that the FConfigFiles in the map
+	// could have been reallocated if new inis were added to the map
+	// since the point at which the caller received the ConfigFile pointer
+	// they are testing. It is the caller's responsibility to not try to hold
+	// on to the ConfigFile pointer during writes to this ConfigCacheIni
+	for (const TPair<FString, FConfigFile>& CurrentFilePair : *this)
+	{
+		if (ConfigFile == &CurrentFilePair.Value)
+		{
+			return true;
+		}
+	}
+	// Check the known inis
+	for (const FKnownConfigFiles::FKnownConfigFile& KnownFile : KnownFiles.Files)
+	{
+		if (ConfigFile == &KnownFile.IniFile)
+		{
+			return true;
+		}
+	}
+
+	return false;
+}
+
+
 TArray<FString> FConfigCacheIni::GetFilenames()
 {
 	TArray<FString> Result;
@@ -2566,7 +2598,8 @@ void FConfigCacheIni::SetString( const TCHAR* Section, const TCHAR* Key, const T
 		Sec->Add( Key, Value );
 		File->Dirty = true;
 	}
-	else if( FCString::Strcmp(*ConfigValue->GetSavedValue(),Value)!=0 )
+	// Use GetSavedValueForWriting rather than GetSavedValue to avoid having the is-it-dirty query mark the values as having been accessed for dependency tracking
+	else if( FCString::Strcmp(*UE::ConfigCacheIni::Private::FAccessor::GetSavedValueForWriting(*ConfigValue),Value)!=0 )
 	{
 		File->Dirty = true;
 		*ConfigValue = FConfigValue(Value);
@@ -2593,7 +2626,8 @@ void FConfigCacheIni::SetText( const TCHAR* Section, const TCHAR* Key, const FTe
 		Sec->Add( Key, StrValue );
 		File->Dirty = true;
 	}
-	else if( FCString::Strcmp(*ConfigValue->GetSavedValue(), *StrValue)!=0 )
+	// Use GetSavedValueForWriting rather than GetSavedValue to avoid having the is-it-dirty query mark the values as having been accessed for dependency tracking
+	else if( FCString::Strcmp(*UE::ConfigCacheIni::Private::FAccessor::GetSavedValueForWriting(*ConfigValue), *StrValue)!=0 )
 	{
 		File->Dirty = true;
 		*ConfigValue = FConfigValue(StrValue);
@@ -3931,7 +3965,7 @@ bool FConfigCacheIni::InitializeKnownConfigFiles(const TCHAR* PlatformName, bool
 		const TCHAR* PlatformNameForThisIni = KnownIndex == (uint8)EKnownIniFile::Scalability ? ScalabilityPlatformOverride : PlatformName;
 		// load the hierarchy!
 		bool bConfigCreated = FConfigCacheIni::LoadGlobalIniFile(KnownFiles.Files[KnownIndex].IniPath, *KnownFiles.Files[KnownIndex].IniName.ToString(), PlatformName, 
-			bDefaultEngineIniRequired, false, bAllowGeneratedIniWhenCooked, bAllowRemoteConfig, *FPaths::GeneratedConfigDir(), this);
+			false, bDefaultEngineIniRequired, bAllowGeneratedIniWhenCooked, bAllowRemoteConfig, *FPaths::GeneratedConfigDir(), this);
 
 		// we want to return if the Engine config was successfully created (to not remove any functionality from old code)
 		if (KnownIndex == (uint8)EKnownIniFile::Engine)
@@ -4717,7 +4751,8 @@ bool FConfigFile::UpdateSinglePropertyInSection(const TCHAR* DiskFilename, const
 	{
 		if (const FConfigValue* ConfigValue = LocalSection->Find(PropertyName))
 		{
-			PropertyValue = ConfigValue->GetSavedValue();
+			// Use GetSavedValueForWriting rather than GetSavedValue to avoid having this save operation mark the value as having been accessed for dependency tracking
+			PropertyValue = UE::ConfigCacheIni::Private::FAccessor::GetSavedValueForWriting(*ConfigValue);
 			FSinglePropertyConfigHelper SinglePropertyConfigHelper(DiskFilename, SectionName, PropertyName, PropertyValue);
 			bSuccessfullyUpdatedFile = SinglePropertyConfigHelper.UpdateConfigFile();
 		}
diff --git a/Engine/Source/Runtime/Core/Public/Misc/ConfigCacheIni.h b/Engine/Source/Runtime/Core/Public/Misc/ConfigCacheIni.h
--- a/Engine/Source/Runtime/Core/Public/Misc/ConfigCacheIni.h
+++ b/Engine/Source/Runtime/Core/Public/Misc/ConfigCacheIni.h
@@ -60,6 +60,17 @@ enum class EKnownIniFile : uint8
 	NumKnownFiles,
 };
 
+namespace UE
+{
+namespace ConfigCacheIni
+{
+namespace Private
+{
+struct FAccessor;
+}
+}
+}
+
 
 struct FConfigValue
 {
@@ -232,6 +243,13 @@ private:
 	/** Internal version of ExpandValue that expands SavedValue into ExpandedValue, or produces an empty ExpandedValue if no expansion occurred. */
 	CORE_API void ExpandValueInternal();
 
+	/** Gets the SavedValue without marking it as having been accessed for e.g. writing out to a ConfigFile to disk */
+	friend struct UE::ConfigCacheIni::Private::FAccessor;
+	const FString& GetSavedValueForWriting() const
+	{
+		return SavedValue;
+	};
+
 	FString SavedValue;
 	FString ExpandedValue;
 #if CONFIG_REMEMBER_ACCESS_PATTERN 
@@ -239,6 +257,28 @@ private:
 #endif
 };
 
+namespace UE
+{
+namespace ConfigCacheIni
+{
+namespace Private
+{
+/** An accessor class to access functions that should be restricted only to FConfigFileCache Internal use */
+struct FAccessor
+{
+private:
+	friend class ::FConfigCacheIni;
+	friend class ::FConfigFile;
+
+	static const FString& GetSavedValueForWriting(const FConfigValue& ConfigValue)
+	{
+		return ConfigValue.GetSavedValueForWriting();
+	}
+};
+}
+}
+}
+
 typedef TMultiMap<FName,FConfigValue> FConfigSectionMap;
 
 // One section in a config file.
@@ -616,6 +656,13 @@ public:
 	FConfigFile* FindConfigFile( const FString& Filename );
 	FConfigFile* Find(const FString& InFilename);
 
+	/**
+	 * Reports whether an FConfigFile* is pointing to a config file inside of this
+	 * Used for downstream functions to check whether a config file they were passed came from this ConfigCacheIni or from 
+	 * a different source such as LoadLocalIniFile
+	 */
+	bool ContainsConfigFile(const FConfigFile* ConfigFile) const;
+
 	UE_DEPRECATED(5.0, "CreateIfNotFound is deprecated, please use the overload without this parameter")
 	FConfigFile* Find(const FString& Filename, bool CreateIfNotFound);
 
