UCookOnTheFlyServer::ProcessAccessedIniSettings: Fix calling Find on empty string due to MoveTemp on string still in use.
#rnx

diff --git a/Engine/Source/Editor/UnrealEd/Private/CookOnTheFlyServer.cpp b/Engine/Source/Editor/UnrealEd/Private/CookOnTheFlyServer.cpp
--- a/Engine/Source/Editor/UnrealEd/Private/CookOnTheFlyServer.cpp
+++ b/Engine/Source/Editor/UnrealEd/Private/CookOnTheFlyServer.cpp
@@ -4759,7 +4759,7 @@ void UCookOnTheFlyServer::ProcessAccessedIniSettings(const FConfigFile* Config,
 		bool bFoundPlatformGuess = false;
 		for (auto It : FDataDrivenPlatformInfoRegistry::GetAllPlatformInfos())
 		{
-			FString CurrentPlatformName = It.Key.ToString();
+			const FString CurrentPlatformName = It.Key.ToString();
 			TStringBuilder<128> PlatformDirString;
 			PlatformDirString.Appendf(TEXT("/%s/"), *CurrentPlatformName);
 			for (const auto& SourceIni : Config->SourceIniHierarchy)
@@ -4769,13 +4769,13 @@ void UCookOnTheFlyServer::ProcessAccessedIniSettings(const FConfigFile* Config,
 				bool bFoundPlatformSubstring = UE::String::FindFirst(SourceIni.Value.Filename, CurrentPlatformName, ESearchCase::IgnoreCase) != INDEX_NONE;
 				if (bFoundPlatformDir)
 				{
-					PlatformName = MoveTemp(CurrentPlatformName);
+					PlatformName = CurrentPlatformName;
 					bFoundPlatformName = true;
 					break;
 				}
 				else if (!bFoundPlatformGuess && bFoundPlatformSubstring)
 				{
-					PlatformName = MoveTemp(CurrentPlatformName);
+					PlatformName = CurrentPlatformName;
 					bFoundPlatformGuess = true;
 				}
 			}
