Fix #jira MHA-73 Maya navigation bug for none-zero values of 'roll'. Made the test for 'upside down' more specific to give a nicer behaviour when you navigate in 'Maya' mode starting from a view orientation of none-zero 'roll'. The STR in the Jira ticket shows a very extreme case where 'roll' is set to -45 degrees

diff --git a/Engine/Source/Editor/UnrealEd/Private/EditorViewportClient.cpp b/Engine/Source/Editor/UnrealEd/Private/EditorViewportClient.cpp
--- a/Engine/Source/Editor/UnrealEd/Private/EditorViewportClient.cpp
+++ b/Engine/Source/Editor/UnrealEd/Private/EditorViewportClient.cpp
@@ -750,7 +750,8 @@ void FEditorViewportClient::ToggleOrbitCamera( bool bEnableOrbitCamera )
 		{
 			FRotator ViewRotation = ViewTransform.GetRotation();
 
-			bool bUpsideDown = (ViewRotation.Pitch < -90.0f || ViewRotation.Pitch > 90.0f || !FMath::IsNearlyZero(ViewRotation.Roll, KINDA_SMALL_NUMBER));
+			bool bUpsideDown = (ViewRotation.Pitch < -90.0f || ViewRotation.Pitch > 90.0f ||
+				FMath::IsNearlyZero(FMath::Abs(ViewRotation.Roll) - 180.f, KINDA_SMALL_NUMBER));
 
 			// if the camera is upside down compute the rotation differently to preserve pitch
 			// otherwise the view will pop to right side up when transferring to orbit controls
