Revised sprite category visibility to be localized to the scene/level view logic where it is being used.

These changes allow a Blueprint class that is based on an abstract native parent class to dynamically add new sprite categories inherited from the native parent when a level is loaded with one or more instances of the Blueprint class. Previously, this could crash the editor due to a global mapping table that was not being updated to include new sprite category visibility commands created in that case.

diff --git a/Engine/Source/Editor/LevelEditor/Private/LevelViewportActions.cpp b/Engine/Source/Editor/LevelEditor/Private/LevelViewportActions.cpp
--- a/Engine/Source/Editor/LevelEditor/Private/LevelViewportActions.cpp
+++ b/Engine/Source/Editor/LevelEditor/Private/LevelViewportActions.cpp
@@ -278,8 +278,7 @@ void FLevelViewportCommands::RegisterShowVolumeCommands()
 
 void FLevelViewportCommands::RegisterShowSpriteCommands()
 {
-	// get all the known layers
-	// Get a fresh list as GUnrealEd->SortedSpriteInfo may not yet be built.
+	// Get all the known layers.
 	TArray<FSpriteCategoryInfo> SortedSpriteInfo;
 	UUnrealEdEngine::MakeSortedSpriteInfo(SortedSpriteInfo);
 
@@ -301,9 +300,25 @@ void FLevelViewportCommands::RegisterShowSpriteCommands()
 				= FUICommandInfoDecl(this->AsShared(), CommandName, LocalizedName, SpriteInfo.Description)
 				.UserInterfaceType(EUserInterfaceActionType::ToggleButton);
 
-			ShowSpriteCommands.Add(FLevelViewportCommands::FShowMenuCommand(ShowSpriteCommand, SpriteInfo.DisplayName));
+			const int32 ShowSpriteCommandIndex = ShowSpriteCommands.Add(FLevelViewportCommands::FShowMenuCommand(ShowSpriteCommand, SpriteInfo.DisplayName));
+
+			// Map the unlocalized category name to the index of the command in the list.
+			SpriteCategoryToCommandIndexMap.Add(SpriteInfo.Category, ShowSpriteCommandIndex);
 		}
 	}
+
+	// These should both end up being the same size when we're done.
+	check(ShowSpriteCommands.Num() == SpriteCategoryToCommandIndexMap.Num());
+}
+
+int32 FLevelViewportCommands::GetShowCommandIndexForSpriteCategory(const FName& InSpriteCategoryName) const
+{
+	if (const int32* ValuePtr = SpriteCategoryToCommandIndexMap.Find(InSpriteCategoryName))
+	{
+		return *ValuePtr;
+	}
+
+	return INDEX_NONE;
 }
 
 PRAGMA_ENABLE_OPTIMIZATION
diff --git a/Engine/Source/Editor/LevelEditor/Private/SLevelViewport.cpp b/Engine/Source/Editor/LevelEditor/Private/SLevelViewport.cpp
--- a/Engine/Source/Editor/LevelEditor/Private/SLevelViewport.cpp
+++ b/Engine/Source/Editor/LevelEditor/Private/SLevelViewport.cpp
@@ -2196,16 +2196,16 @@ void SLevelViewport::OnToggleAllSpriteCategories( bool bVisible )
 }
 
 /** Called when the user toggles a category from the sprite sub-menu. **/
-void SLevelViewport::ToggleSpriteCategory( int32 CategoryID )
+void SLevelViewport::ToggleSpriteCategory( int32 ShowSpriteCommandIndex )
 {
-	LevelViewportClient->SetSpriteCategoryVisibility( CategoryID, !LevelViewportClient->GetSpriteCategoryVisibility( CategoryID ) );
+	LevelViewportClient->SetSpriteCategoryVisibility( ShowSpriteCommandIndex, !LevelViewportClient->GetSpriteCategoryVisibility( ShowSpriteCommandIndex ) );
 	LevelViewportClient->Invalidate();
 }
 
 /** Called to determine if a category from the sprite sub-menu is visible. **/
-bool SLevelViewport::IsSpriteCategoryVisible( int32 CategoryID ) const
+bool SLevelViewport::IsSpriteCategoryVisible( int32 ShowSpriteCommandIndex ) const
 {
-	return LevelViewportClient->GetSpriteCategoryVisibility( CategoryID );
+	return LevelViewportClient->GetSpriteCategoryVisibility( ShowSpriteCommandIndex );
 }
 
 void SLevelViewport::OnToggleAllStatCommands( bool bVisible )
diff --git a/Engine/Source/Editor/LevelEditor/Public/LevelViewportActions.h b/Engine/Source/Editor/LevelEditor/Public/LevelViewportActions.h
--- a/Engine/Source/Editor/LevelEditor/Public/LevelViewportActions.h
+++ b/Engine/Source/Editor/LevelEditor/Public/LevelViewportActions.h
@@ -207,7 +207,13 @@ public:
 	/** Registers our commands with the binding system for showing sprites. */
 	void RegisterShowSpriteCommands();
 
+	/** Returns the index of the "show sprite" command for the given category name. */
+	int32 GetShowCommandIndexForSpriteCategory(const FName& InSpriteCategoryName) const;
+
 private:
+	/** Maps sprite category names to a matching index in the sorted ShowSpriteCommands list */
+	TMap<FName, int32> SpriteCategoryToCommandIndexMap;
+
 	/** Registers additional commands as they are loaded */
 	void HandleNewStatGroup(const TArray<FStatNameAndInfo>& NameAndInfos);
 	void HandleNewStat(const FName& InStatName, const FName& InStatCategory, const FText& InStatDescription);
diff --git a/Engine/Source/Editor/LevelEditor/Public/SLevelViewport.h b/Engine/Source/Editor/LevelEditor/Public/SLevelViewport.h
--- a/Engine/Source/Editor/LevelEditor/Public/SLevelViewport.h
+++ b/Engine/Source/Editor/LevelEditor/Public/SLevelViewport.h
@@ -533,16 +533,16 @@ private:
 	/**
 	 * Toggles sprite category visibility in this viewport
 	 *
-	 * @param CategoryID				Index of the category
+	 * @param ShowSpriteCommandIndex	Index of the visibility command that corresponds to a sprite category
 	 */
-	void ToggleSpriteCategory( int32 CategoryID );
+	void ToggleSpriteCategory( int32 ShowSpriteCommandIndex );
 
 	/**
-	 * Checks if sprite category is visible in this viewport
+	 * Checks if a sprite category is visible in this viewport
 	 *
-	 * @param CategoryID				Index of the category
+	 * @param ShowSpriteCommandIndex	Index of the visibility command that corresponds to a sprite category
 	 */
-	bool IsSpriteCategoryVisible( int32 CategoryID ) const;
+	bool IsSpriteCategoryVisible( int32 ShowSpriteCommandIndex ) const;
 
 	/**
 	 * Toggles all Stat commands visibility
diff --git a/Engine/Source/Editor/UnrealEd/Classes/Editor/UnrealEdEngine.h b/Engine/Source/Editor/UnrealEd/Classes/Editor/UnrealEdEngine.h
--- a/Engine/Source/Editor/UnrealEd/Classes/Editor/UnrealEdEngine.h
+++ b/Engine/Source/Editor/UnrealEd/Classes/Editor/UnrealEdEngine.h
@@ -169,7 +169,8 @@ public:
 	TMap<TWeakObjectPtr<UPackage>, uint8> PackageToNotifyState;
 
 	/** Mapping of sprite category ids to their matching indices in the sorted sprite categories array */
-	TMap<FName, int32>			SpriteIDToIndexMap;
+	UE_DEPRECATED(5.0, "This mapping has been moved to FLevelViewportCommands::SpriteCategoryToCommandIndexMap.")
+	TMap<FName, int32> SpriteIDToIndexMap;
 
 	/** Map from component class to visualizer object to use */
 	TMap< FName, TSharedPtr<class FComponentVisualizer> > ComponentVisualizerMap;
@@ -776,15 +777,6 @@ public:
 	 * Identify any brushes whose sense is inverted and repair them
 	 */
 	void FixAnyInvertedBrushes(UWorld* World);
-
-	/**
-	 * Get the index of the provided sprite category
-	 *
-	 * @param	InSpriteCategory	Sprite category to get the index of
-	 *
-	 * @return	Index of the provided sprite category, if possible; INDEX_NONE otherwise
-	 */
-	virtual int32 GetSpriteCategoryIndex( const FName& InSpriteCategory ) override;
 	
 	/**
 	 * Shows the LightingStaticMeshInfoWindow, creating it first if it hasn't been initialized.
diff --git a/Engine/Source/Editor/UnrealEd/Private/LevelEditorViewport.cpp b/Engine/Source/Editor/UnrealEd/Private/LevelEditorViewport.cpp
--- a/Engine/Source/Editor/UnrealEd/Private/LevelEditorViewport.cpp
+++ b/Engine/Source/Editor/UnrealEd/Private/LevelEditorViewport.cpp
@@ -1686,6 +1686,13 @@ void FTrackingTransaction::OnEditorSelectionChanged(UObject* NewSelection)
 
 }
 
+static const FLevelViewportCommands& GetLevelViewportCommands()
+{
+	static FName LevelEditorName("LevelEditor");
+	FLevelEditorModule& LevelEditor = FModuleManager::LoadModuleChecked<FLevelEditorModule>(LevelEditorName);
+	return LevelEditor.GetLevelViewportCommands();
+}
+
 
 FLevelEditorViewportClient::FLevelEditorViewportClient(const TSharedPtr<SLevelViewport>& InLevelViewport)
 	: FEditorViewportClient(&GLevelEditorModeTools(), nullptr, StaticCastSharedPtr<SEditorViewport>(InLevelViewport))
@@ -1806,8 +1813,8 @@ void FLevelEditorViewportClient::InitializeVisibilityFlags()
 	UUnrealEdEngine::GetSortedVolumeClasses(&VolumeClasses);
 	VolumeActorVisibility.Init(true, VolumeClasses.Num());
 
-	// Initialize all sprite categories to visible
-	SpriteCategoryVisibility.Init(true, GUnrealEd->SpriteIDToIndexMap.Num());
+	// Initialize all sprite categories to visible (one per registered "show sprite" command)
+	SpriteCategoryVisibility.Init(true, GetLevelViewportCommands().ShowSpriteCommands.Num());
 }
 
 void FLevelEditorViewportClient::InitializeViewportInteraction()
@@ -1875,9 +1882,12 @@ FSceneView* FLevelEditorViewportClient::CalcSceneView(FSceneViewFamily* ViewFami
 	FSceneView* View = FEditorViewportClient::CalcSceneView(ViewFamily, StereoPass);
 
 	View->ViewActor = ActorLocks.GetLock().GetLockedActor();
-	View->SpriteCategoryVisibility = SpriteCategoryVisibility;
 	View->bCameraCut = bEditorCameraCut;
 	View->bHasSelectedComponents = GEditor->GetSelectedComponentCount() > 0;
+	View->GetSpriteCategoryVisibilityOverride().BindLambda([this](const FName& InCategoryName)
+	{
+		return this->GetSpriteCategoryVisibility(InCategoryName);
+	});
 
 #if RHI_RAYTRACING
 	View->SetupRayTracedRendering();
@@ -2708,13 +2718,6 @@ static bool CommandAcceptsInput( FLevelEditorViewportClient& ViewportClient, FKe
 	return bAccepted;
 }
 
-static const FLevelViewportCommands& GetLevelViewportCommands()
-{
-	static FName LevelEditorName("LevelEditor");
-	FLevelEditorModule& LevelEditor = FModuleManager::LoadModuleChecked<FLevelEditorModule>( LevelEditorName );
-	return LevelEditor.GetLevelViewportCommands();
-}
-
 void FLevelEditorViewportClient::SetCurrentViewport()
 {
 	// Set the current level editing viewport client to the dropped-in viewport client
@@ -4680,30 +4683,26 @@ void FLevelEditorViewportClient::OnPreBeginPIE(const bool bIsSimulating)
 
 bool FLevelEditorViewportClient::GetSpriteCategoryVisibility( const FName& InSpriteCategory ) const
 {
-	const int32 CategoryIndex = GEngine->GetSpriteCategoryIndex( InSpriteCategory );
-	check( CategoryIndex != INDEX_NONE && CategoryIndex < SpriteCategoryVisibility.Num() );
-
-	return SpriteCategoryVisibility[ CategoryIndex ];
+	const int32 ShowSpriteCommandIndex = GetLevelViewportCommands().GetShowCommandIndexForSpriteCategory( InSpriteCategory );
+	return GetSpriteCategoryVisibility( ShowSpriteCommandIndex );
 }
 
-bool FLevelEditorViewportClient::GetSpriteCategoryVisibility( int32 Index ) const
+bool FLevelEditorViewportClient::GetSpriteCategoryVisibility( int32 InShowSpriteCommandIndex ) const
 {
-	check( Index >= 0 && Index < SpriteCategoryVisibility.Num() );
-	return SpriteCategoryVisibility[ Index ];
+	check( InShowSpriteCommandIndex >= 0 && InShowSpriteCommandIndex < SpriteCategoryVisibility.Num() );
+	return SpriteCategoryVisibility[ InShowSpriteCommandIndex ];
 }
 
 void FLevelEditorViewportClient::SetSpriteCategoryVisibility( const FName& InSpriteCategory, bool bVisible )
 {
-	const int32 CategoryIndex = GEngine->GetSpriteCategoryIndex( InSpriteCategory );
-	check( CategoryIndex != INDEX_NONE && CategoryIndex < SpriteCategoryVisibility.Num() );
-
-	SpriteCategoryVisibility[ CategoryIndex ] = bVisible;
+	const int32 ShowSpriteCommandIndex = GetLevelViewportCommands().GetShowCommandIndexForSpriteCategory( InSpriteCategory );
+	SetSpriteCategoryVisibility( ShowSpriteCommandIndex, bVisible );
 }
 
-void FLevelEditorViewportClient::SetSpriteCategoryVisibility( int32 Index, bool bVisible )
+void FLevelEditorViewportClient::SetSpriteCategoryVisibility( int32 InShowSpriteCommandIndex, bool bVisible )
 {
-	check( Index >= 0 && Index < SpriteCategoryVisibility.Num() );
-	SpriteCategoryVisibility[ Index ] = bVisible;
+	check( InShowSpriteCommandIndex >= 0 && InShowSpriteCommandIndex < SpriteCategoryVisibility.Num() );
+	SpriteCategoryVisibility[ InShowSpriteCommandIndex ] = bVisible;
 }
 
 void FLevelEditorViewportClient::SetAllSpriteCategoryVisibility( bool bVisible )
diff --git a/Engine/Source/Editor/UnrealEd/Private/UnrealEdEngine.cpp b/Engine/Source/Editor/UnrealEd/Private/UnrealEdEngine.cpp
--- a/Engine/Source/Editor/UnrealEd/Private/UnrealEdEngine.cpp
+++ b/Engine/Source/Editor/UnrealEd/Private/UnrealEdEngine.cpp
@@ -147,18 +147,6 @@ void UUnrealEdEngine::Init(IEngineLoop* InEngineLoop)
 		}
 	}
 
-	// Populate the data structures related to the sprite category visibility feature for use elsewhere in the editor later
-	TArray<FSpriteCategoryInfo> SortedSpriteInfo;
-	UUnrealEdEngine::MakeSortedSpriteInfo(SortedSpriteInfo);
-
-	// Iterate over the sorted list, constructing a mapping of unlocalized categories to the index the localized category
-	// resides in. This is an optimization to prevent having to localize values repeatedly.
-	for( int32 InfoIndex = 0; InfoIndex < SortedSpriteInfo.Num(); ++InfoIndex )
-	{
-		const FSpriteCategoryInfo& SpriteInfo = SortedSpriteInfo[InfoIndex];
-		SpriteIDToIndexMap.Add( SpriteInfo.Category, InfoIndex );
-	}
-
 	if (FPaths::IsProjectFilePathSet() && GIsEditor && !FApp::IsUnattended())
 	{
 		AutoReimportManager = NewObject<UAutoReimportManager>();
@@ -383,6 +371,7 @@ void UUnrealEdEngine::PreExit()
 }
 
 
+PRAGMA_DISABLE_DEPRECATION_WARNINGS
 UUnrealEdEngine::~UUnrealEdEngine()
 {
 	if (this == GUnrealEd)
@@ -390,6 +379,7 @@ UUnrealEdEngine::~UUnrealEdEngine()
 		GUnrealEd = NULL; 
 	}
 }
+PRAGMA_ENABLE_DEPRECATION_WARNINGS
 
 
 void UUnrealEdEngine::FinishDestroy()
@@ -1054,17 +1044,6 @@ void UUnrealEdEngine::MakeSelectedActorsLevelCurrent()
 }
 
 
-int32 UUnrealEdEngine::GetSpriteCategoryIndex( const FName& InSpriteCategory )
-{
-	// Find the sprite category in the unlocalized to index map, if possible
-	const int32* CategoryIndexPtr = SpriteIDToIndexMap.Find( InSpriteCategory );
-	
-	const int32 CategoryIndex = CategoryIndexPtr ? *CategoryIndexPtr : INDEX_NONE;
-
-	return CategoryIndex;
-}
-
-
 void UUnrealEdEngine::ShowLightingStaticMeshInfoWindow()
 {
 	// first invoke the stats viewer tab
diff --git a/Engine/Source/Editor/UnrealEd/Public/LevelEditorViewport.h b/Engine/Source/Editor/UnrealEd/Public/LevelEditorViewport.h
--- a/Engine/Source/Editor/UnrealEd/Public/LevelEditorViewport.h
+++ b/Engine/Source/Editor/UnrealEd/Public/LevelEditorViewport.h
@@ -354,11 +354,11 @@ public:
 	/**
 	 * Returns whether the sprite category specified by the provided index is visible in the viewport or not
 	 *
-	 * @param	Index	Index of the sprite category to check
+	 * @param	InShowSpriteCommandIndex	Index of the visibility command that corresponds to a sprite category
 	 *
 	 * @return	true if the category specified by the index is visible in the viewport; false if it is not
 	 */
-	bool GetSpriteCategoryVisibility( int32 Index ) const;
+	bool GetSpriteCategoryVisibility( int32 InShowSpriteCommandIndex ) const;
 
 	/**
 	 * Sets the visibility of the provided unlocalized category to the provided value
@@ -371,10 +371,10 @@ public:
 	/**
 	 * Sets the visibility of the category specified by the provided index to the provided value
 	 *
-	 * @param	Index		Index of the sprite category to set the visibility of
-	 * @param	bVisible	true if the category should be made visible, false if it should be hidden
+	 * @param	InShowSpriteCommandIndex	Index of the visibility command that corresponds to a sprite category
+	 * @param	bVisible					true if the category should be made visible, false if it should be hidden
 	 */
-	void SetSpriteCategoryVisibility( int32 Index, bool bVisible );
+	void SetSpriteCategoryVisibility( int32 InShowSpriteCommandIndex, bool bVisible );
 
 	/**
 	 * Sets the visibility of all sprite categories to the provided value
diff --git a/Engine/Source/Runtime/Engine/Classes/Engine/Engine.h b/Engine/Source/Runtime/Engine/Classes/Engine/Engine.h
--- a/Engine/Source/Runtime/Engine/Classes/Engine/Engine.h
+++ b/Engine/Source/Runtime/Engine/Classes/Engine/Engine.h
@@ -2279,6 +2279,7 @@ public:
 	 *
 	 * @return	Index of the provided sprite category, if possible; INDEX_NONE otherwise
 	 */
+	UE_DEPRECATED(5.0, "Please use FSceneView::GetSpriteCategoryIndex() instead, which can now be accessed from within an FPrimitiveSceneProxy::GetViewRelevance() override.")
 	virtual int32 GetSpriteCategoryIndex( const FName& InSpriteCategory )
 	{
 		// The editor may override this to handle sprite categories as necessary
diff --git a/Engine/Source/Runtime/Engine/Private/Components/ArrowComponent.cpp b/Engine/Source/Runtime/Engine/Private/Components/ArrowComponent.cpp
--- a/Engine/Source/Runtime/Engine/Private/Components/ArrowComponent.cpp
+++ b/Engine/Source/Runtime/Engine/Private/Components/ArrowComponent.cpp
@@ -57,9 +57,9 @@ public:
 		bWillEverBeLit = false;
 #if WITH_EDITOR
 		// If in the editor, extract the sprite category from the component
-		if ( GIsEditor )
+		if (GIsEditor)
 		{
-			SpriteCategoryIndex = GEngine->GetSpriteCategoryIndex( Component->SpriteInfo.Category );
+			SpriteCategoryName = Component->SpriteInfo.Category;
 		}
 #endif	//WITH_EDITOR
 
@@ -175,7 +175,7 @@ public:
 #if WITH_EDITOR
 		if (bTreatAsASprite)
 		{
-			if ( GIsEditor && SpriteCategoryIndex != INDEX_NONE && SpriteCategoryIndex < View->SpriteCategoryVisibility.Num() && !View->SpriteCategoryVisibility[ SpriteCategoryIndex ] )
+			if ( GIsEditor && !View->GetSpriteCategoryVisibility(SpriteCategoryName) )
 			{
 				Result.bDrawRelevance = false;
 			}
@@ -209,7 +209,7 @@ private:
 #if WITH_EDITORONLY_DATA
 	bool bLightAttachment;
 	bool bTreatAsASprite;
-	int32 SpriteCategoryIndex;
+	FName SpriteCategoryName;
 	bool bUseInEditorScaling;
 	float EditorScale;
 #endif // #if WITH_EDITORONLY_DATA
diff --git a/Engine/Source/Runtime/Engine/Private/Components/BillboardComponent.cpp b/Engine/Source/Runtime/Engine/Private/Components/BillboardComponent.cpp
--- a/Engine/Source/Runtime/Engine/Private/Components/BillboardComponent.cpp
+++ b/Engine/Source/Runtime/Engine/Private/Components/BillboardComponent.cpp
@@ -53,7 +53,7 @@ public:
 		// Extract the sprite category from the component if in the editor
 		if ( GIsEditor )
 		{
-			SpriteCategoryIndex = GEngine->GetSpriteCategoryIndex( InComponent->SpriteInfo.Category );
+			SpriteCategoryName = InComponent->SpriteInfo.Category;
 		}
 #endif //WITH_EDITOR
 
@@ -230,9 +230,9 @@ public:
 	{
 		bool bVisible = View->Family->EngineShowFlags.BillboardSprites;
 #if WITH_EDITOR
-		if ( GIsEditor && bVisible && SpriteCategoryIndex != INDEX_NONE && SpriteCategoryIndex < View->SpriteCategoryVisibility.Num() )
+		if ( GIsEditor && bVisible )
 		{
-			bVisible = View->SpriteCategoryVisibility[ SpriteCategoryIndex ];
+			bVisible = View->GetSpriteCategoryVisibility(SpriteCategoryName);
 		}
 #endif
 		FPrimitiveViewRelevance Result;
@@ -265,7 +265,7 @@ private:
 	const uint32 bIsScreenSizeScaled : 1;
 	uint32 bIsActorLocked : 1;
 #if WITH_EDITORONLY_DATA
-	int32 SpriteCategoryIndex;
+	FName SpriteCategoryName;
 	bool bUseInEditorScaling;
 	float EditorScale;
 #endif // #if WITH_EDITORONLY_DATA
diff --git a/Engine/Source/Runtime/Engine/Private/SceneView.cpp b/Engine/Source/Runtime/Engine/Private/SceneView.cpp
--- a/Engine/Source/Runtime/Engine/Private/SceneView.cpp
+++ b/Engine/Source/Runtime/Engine/Private/SceneView.cpp
@@ -2475,6 +2475,19 @@ const FExposureBufferData* FSceneView::GetEyeAdaptationBuffer() const
 	return nullptr;
 }
 
+#if WITH_EDITOR
+bool FSceneView::GetSpriteCategoryVisibility(const FName& InCategoryName) const
+{
+	if (GetSpriteCategoryVisibilityDelegate.IsBound())
+	{
+		return GetSpriteCategoryVisibilityDelegate.Execute(InCategoryName);
+	}
+
+	// Visible by default if no delegate is bound to potentially override it.
+	return true;
+}
+#endif
+
 FSceneViewFamily::FSceneViewFamily(const ConstructionValues& CVS)
 	:
 	ViewMode(VMI_Lit),
diff --git a/Engine/Source/Runtime/Engine/Public/SceneView.h b/Engine/Source/Runtime/Engine/Public/SceneView.h
--- a/Engine/Source/Runtime/Engine/Public/SceneView.h
+++ b/Engine/Source/Runtime/Engine/Public/SceneView.h
@@ -206,7 +206,6 @@ struct FSceneViewInitOptions : public FSceneViewProjectionData
 		, OverrideLODViewOrigin(ForceInitToZero)
 		, bUseFauxOrthoViewPos(false)
 		, bDisableGameScreenPercentage(false)
-		//@TODO: , const TBitArray<>& InSpriteCategoryVisibility=TBitArray<>()
 #endif
 	{
 	}
@@ -1180,6 +1179,7 @@ public:
 	bool bAllowTranslucentPrimitivesInHitProxy;
 
 	/** BitArray representing the visibility state of the various sprite categories in the editor for this view */
+	UE_DEPRECATED(5.0, "Please use GetSpriteCategoryVisibility() instead.")
 	TBitArray<> SpriteCategoryVisibility;
 	/** Selection color for the editor (used by post processing) */
 	FLinearColor SelectionOutlineColor;
@@ -1414,9 +1414,20 @@ public:
 	/** Returns the eye adaptation buffer (mobile) or null if it doesn't exist. */
 	const FExposureBufferData* GetEyeAdaptationBuffer() const;
 
+#if WITH_EDITOR
+	/** Delegate that can provide access to the visibility state of the various sprite categories in the editor for this view */
+	DECLARE_DELEGATE_RetVal_OneParam(bool, FGetSpriteCategoryVisibility, const FName&);
+	/** Allows this view to bind to an optional override to determine sprite category visibility. */
+	FGetSpriteCategoryVisibility& GetSpriteCategoryVisibilityOverride() { return GetSpriteCategoryVisibilityDelegate; }
+	/** Returns whether the sprite category specified by the given name should be visible or not. */
+	bool GetSpriteCategoryVisibility(const FName& InCategoryName) const;
+#endif
 
 protected:
 	FSceneViewStateInterface* EyeAdaptationViewState = nullptr;
+#if WITH_EDITOR
+	FGetSpriteCategoryVisibility GetSpriteCategoryVisibilityDelegate;
+#endif
 };
 
 //////////////////////////////////////////////////////////////////////////
