Fixed noPch compile error.

#rnx

diff --git a/Engine/Source/Editor/UnrealEd/Public/Tools/InteractiveToolsCommands.h b/Engine/Source/Editor/UnrealEd/Public/Tools/InteractiveToolsCommands.h
--- a/Engine/Source/Editor/UnrealEd/Public/Tools/InteractiveToolsCommands.h
+++ b/Engine/Source/Editor/UnrealEd/Public/Tools/InteractiveToolsCommands.h
@@ -7,6 +7,7 @@
 #include "Framework/Commands/UICommandList.h"
 #include "InteractiveTool.h"
 #include "StandardToolModeCommands.h"
+#include "Styling/ISlateStyle.h"
 
 
 
