Fix crash on attempt to get view relevance for a sprite category that doesn't map to a registered visibility toggle command in the level viewport command list.

This was an unintentional regression introduced by me at 14886456. The previous behavior was to default to visible in this case, which has now been restored.

#rnx

diff --git a/Engine/Source/Editor/UnrealEd/Private/LevelEditorViewport.cpp b/Engine/Source/Editor/UnrealEd/Private/LevelEditorViewport.cpp
--- a/Engine/Source/Editor/UnrealEd/Private/LevelEditorViewport.cpp
+++ b/Engine/Source/Editor/UnrealEd/Private/LevelEditorViewport.cpp
@@ -4681,7 +4681,13 @@ void FLevelEditorViewportClient::OnPreBeginPIE(const bool bIsSimulating)
 bool FLevelEditorViewportClient::GetSpriteCategoryVisibility( const FName& InSpriteCategory ) const
 {
 	const int32 ShowSpriteCommandIndex = GetLevelViewportCommands().GetShowCommandIndexForSpriteCategory( InSpriteCategory );
-	return GetSpriteCategoryVisibility( ShowSpriteCommandIndex );
+	if (ShowSpriteCommandIndex != INDEX_NONE)
+	{
+		return GetSpriteCategoryVisibility(ShowSpriteCommandIndex);
+	}
+
+	// If the given category is not registered as a visibility command, default to being visible.
+	return true;
 }
 
 bool FLevelEditorViewportClient::GetSpriteCategoryVisibility( int32 InShowSpriteCommandIndex ) const
@@ -4693,7 +4699,10 @@ bool FLevelEditorViewportClient::GetSpriteCategoryVisibility( int32 InShowSprite
 void FLevelEditorViewportClient::SetSpriteCategoryVisibility( const FName& InSpriteCategory, bool bVisible )
 {
 	const int32 ShowSpriteCommandIndex = GetLevelViewportCommands().GetShowCommandIndexForSpriteCategory( InSpriteCategory );
-	SetSpriteCategoryVisibility( ShowSpriteCommandIndex, bVisible );
+	if (ShowSpriteCommandIndex != INDEX_NONE)
+	{
+		SetSpriteCategoryVisibility(ShowSpriteCommandIndex, bVisible);
+	}
 }
 
 void FLevelEditorViewportClient::SetSpriteCategoryVisibility( int32 InShowSpriteCommandIndex, bool bVisible )
