Make alternate skinning working with interchange.

diff --git a/Engine/Source/Editor/UnrealEd/Private/SkinWeightsUtilities.cpp b/Engine/Source/Editor/UnrealEd/Private/SkinWeightsUtilities.cpp
--- a/Engine/Source/Editor/UnrealEd/Private/SkinWeightsUtilities.cpp
+++ b/Engine/Source/Editor/UnrealEd/Private/SkinWeightsUtilities.cpp
@@ -32,6 +32,8 @@
 #include "DesktopPlatformModule.h"
 #include "EditorDirectories.h"
 #include "Framework/Application/SlateApplication.h"
+#include "InterchangeManager.h"
+#include "Settings/EditorExperimentalSettings.h"
 
 DEFINE_LOG_CATEGORY_STATIC(LogSkinWeightsUtilities, Log, All);
 
@@ -57,56 +59,19 @@ bool FSkinWeightsUtilities::ImportAlternateSkinWeight(USkeletalMesh* SkeletalMes
 	FScopedSuspendAlternateSkinWeightPreview ScopedSuspendAlternateSkinnWeightPreview(SkeletalMesh);
 	FScopedSkeletalMeshPostEditChange ScopePostEditChange(SkeletalMesh);
 
-	UnFbx::FBXImportOptions ImportOptions;
-	//Import the alternate fbx into a temporary skeletal mesh using the same import options
-	UFbxFactory* FbxFactory = NewObject<UFbxFactory>(UFbxFactory::StaticClass());
-	FbxFactory->AddToRoot();
+	//If Interchange is enable use it if not use the old path
 
-	FbxFactory->ImportUI = NewObject<UFbxImportUI>(FbxFactory);
-	UFbxSkeletalMeshImportData* OriginalSkeletalMeshImportData = UFbxSkeletalMeshImportData::GetImportDataForSkeletalMesh(SkeletalMesh, nullptr);
-	if (OriginalSkeletalMeshImportData != nullptr)
-	{
-		//Copy the skeletal mesh import data options
-		FbxFactory->ImportUI->SkeletalMeshImportData = DuplicateObject<UFbxSkeletalMeshImportData>(OriginalSkeletalMeshImportData, FbxFactory);
-	}
-	//Skip the auto detect type on import, the test set a specific value
-	FbxFactory->SetDetectImportTypeOnImport(false);
-	FbxFactory->ImportUI->bImportAsSkeletal = true;
-	FbxFactory->ImportUI->MeshTypeToImport = FBXIT_SkeletalMesh;
-	FbxFactory->ImportUI->bIsReimport = false;
-	FbxFactory->ImportUI->ReimportMesh = nullptr;
-	FbxFactory->ImportUI->bAllowContentTypeImport = true;
-	FbxFactory->ImportUI->bImportAnimations = false;
-	FbxFactory->ImportUI->bAutomatedImportShouldDetectType = false;
-	FbxFactory->ImportUI->bCreatePhysicsAsset = false;
-	FbxFactory->ImportUI->bImportMaterials = false;
-	FbxFactory->ImportUI->bImportTextures = false;
-	FbxFactory->ImportUI->bImportMesh = true;
-	FbxFactory->ImportUI->bImportRigidMesh = false;
-	FbxFactory->ImportUI->bIsObjImport = false;
-	FbxFactory->ImportUI->bOverrideFullName = true;
-	FbxFactory->ImportUI->Skeleton = nullptr;
-	
-	//Force some skeletal mesh import options
-	if (FbxFactory->ImportUI->SkeletalMeshImportData)
-	{
-		FbxFactory->ImportUI->SkeletalMeshImportData->bImportMeshLODs = false;
-		FbxFactory->ImportUI->SkeletalMeshImportData->bImportMorphTargets = false;
-		FbxFactory->ImportUI->SkeletalMeshImportData->bUpdateSkeletonReferencePose = false;
-		FbxFactory->ImportUI->SkeletalMeshImportData->ImportContentType = EFBXImportContentType::FBXICT_All; //We need geo and skinning, so we can match the weights
-	}
-	//Force some material options
-	if (FbxFactory->ImportUI->TextureImportData)
-	{
-		FbxFactory->ImportUI->TextureImportData->MaterialSearchLocation = EMaterialSearchLocation::DoNotSearch;
-		FbxFactory->ImportUI->TextureImportData->BaseMaterialName.Reset();
-	}
+	bool bUseInterchangeFramework = false;
+	UInterchangeManager& InterchangeManager = UInterchangeManager::GetInterchangeManager();
+#if WITH_EDITOR
+	bUseInterchangeFramework = GetDefault<UEditorExperimentalSettings>()->bEnableInterchangeFramework;
+#endif
 
 	FAssetToolsModule& AssetToolsModule = FModuleManager::LoadModuleChecked<FAssetToolsModule>("AssetTools");
 	FString ImportAssetPath = TEXT("/Engine/TempEditor/SkeletalMeshTool");
 	//Empty the temporary path
 	FAssetRegistryModule& AssetRegistryModule = FModuleManager::LoadModuleChecked<FAssetRegistryModule>("AssetRegistry");
-	
+
 	auto DeletePathAssets = [&AssetRegistryModule, &ImportAssetPath]()
 	{
 		TArray<FAssetData> AssetsToDelete;
@@ -129,41 +94,104 @@ bool FSkinWeightsUtilities::ImportAlternateSkinWeight(USkeletalMesh* SkeletalMes
 
 	DeletePathAssets();
 
-	ApplyImportUIToImportOptions(FbxFactory->ImportUI, ImportOptions);
-
-	TArray<FString> ImportFilePaths;
-	ImportFilePaths.Add(AbsoluteFilePath);
-
-	UAssetImportTask* Task = NewObject<UAssetImportTask>();
-	Task->AddToRoot();
-	Task->bAutomated = true;
-	Task->bReplaceExisting = true;
-	Task->DestinationPath = ImportAssetPath;
-	Task->bSave = false;
-	Task->DestinationName = FGuid::NewGuid().ToString(EGuidFormats::Digits);
-	Task->Options = FbxFactory->ImportUI;
-	Task->Filename = AbsoluteFilePath;
-	Task->Factory = FbxFactory;
-	FbxFactory->SetAssetImportTask(Task);
-	TArray<UAssetImportTask*> Tasks;
-	Tasks.Add(Task);
-	AssetToolsModule.Get().ImportAssetTasks(Tasks);
-
 	UObject* ImportedObject = nullptr;
-	
-	for (FString AssetPath : Task->ImportedObjectPaths)
+	UnFbx::FBXImportOptions ImportOptions;
+
+	if (bUseInterchangeFramework)
 	{
-		FAssetData AssetData = AssetRegistryModule.Get().GetAssetByObjectPath(FName(*AssetPath));
-		ImportedObject = AssetData.GetAsset();
-		if (ImportedObject != nullptr)
+		UE::Interchange::FScopedSourceData ScopedSourceData(AbsoluteFilePath);
+
+		FImportAssetParameters ImportAssetParameters;
+		ImportAssetParameters.bIsAutomated = true; // From the InterchangeManager point of view, this is considered an automated import
+
+		//TODO create a pipeline that set all the proper skeletalmesh options (look at the legacy system setup)
+		UE::Interchange::FAssetImportResultRef AssetImportResult = InterchangeManager.ImportAssetAsync(ImportAssetPath, ScopedSourceData.GetSourceData(), ImportAssetParameters);
+		AssetImportResult->WaitUntilDone();
+		if (USkeletalMesh* ImportedSkeletalMesh = Cast< USkeletalMesh >(AssetImportResult->GetFirstAssetOfClass(USkeletalMesh::StaticClass())))
 		{
-			break;
+			ImportedObject = ImportedSkeletalMesh;
 		}
 	}
-	
-	//Factory and task can now be garbage collected
-	Task->RemoveFromRoot();
-	FbxFactory->RemoveFromRoot();
+	else
+	{
+		//Import the alternate fbx into a temporary skeletal mesh using the same import options
+		UFbxFactory* FbxFactory = NewObject<UFbxFactory>(UFbxFactory::StaticClass());
+		FbxFactory->AddToRoot();
+
+		FbxFactory->ImportUI = NewObject<UFbxImportUI>(FbxFactory);
+		UFbxSkeletalMeshImportData* OriginalSkeletalMeshImportData = UFbxSkeletalMeshImportData::GetImportDataForSkeletalMesh(SkeletalMesh, nullptr);
+		if (OriginalSkeletalMeshImportData != nullptr)
+		{
+			//Copy the skeletal mesh import data options
+			FbxFactory->ImportUI->SkeletalMeshImportData = DuplicateObject<UFbxSkeletalMeshImportData>(OriginalSkeletalMeshImportData, FbxFactory);
+		}
+		//Skip the auto detect type on import, the test set a specific value
+		FbxFactory->SetDetectImportTypeOnImport(false);
+		FbxFactory->ImportUI->bImportAsSkeletal = true;
+		FbxFactory->ImportUI->MeshTypeToImport = FBXIT_SkeletalMesh;
+		FbxFactory->ImportUI->bIsReimport = false;
+		FbxFactory->ImportUI->ReimportMesh = nullptr;
+		FbxFactory->ImportUI->bAllowContentTypeImport = true;
+		FbxFactory->ImportUI->bImportAnimations = false;
+		FbxFactory->ImportUI->bAutomatedImportShouldDetectType = false;
+		FbxFactory->ImportUI->bCreatePhysicsAsset = false;
+		FbxFactory->ImportUI->bImportMaterials = false;
+		FbxFactory->ImportUI->bImportTextures = false;
+		FbxFactory->ImportUI->bImportMesh = true;
+		FbxFactory->ImportUI->bImportRigidMesh = false;
+		FbxFactory->ImportUI->bIsObjImport = false;
+		FbxFactory->ImportUI->bOverrideFullName = true;
+		FbxFactory->ImportUI->Skeleton = nullptr;
+
+		//Force some skeletal mesh import options
+		if (FbxFactory->ImportUI->SkeletalMeshImportData)
+		{
+			FbxFactory->ImportUI->SkeletalMeshImportData->bImportMeshLODs = false;
+			FbxFactory->ImportUI->SkeletalMeshImportData->bImportMorphTargets = false;
+			FbxFactory->ImportUI->SkeletalMeshImportData->bUpdateSkeletonReferencePose = false;
+			FbxFactory->ImportUI->SkeletalMeshImportData->ImportContentType = EFBXImportContentType::FBXICT_All; //We need geo and skinning, so we can match the weights
+		}
+		//Force some material options
+		if (FbxFactory->ImportUI->TextureImportData)
+		{
+			FbxFactory->ImportUI->TextureImportData->MaterialSearchLocation = EMaterialSearchLocation::DoNotSearch;
+			FbxFactory->ImportUI->TextureImportData->BaseMaterialName.Reset();
+		}
+
+		ApplyImportUIToImportOptions(FbxFactory->ImportUI, ImportOptions);
+
+		TArray<FString> ImportFilePaths;
+		ImportFilePaths.Add(AbsoluteFilePath);
+
+		UAssetImportTask* Task = NewObject<UAssetImportTask>();
+		Task->AddToRoot();
+		Task->bAutomated = true;
+		Task->bReplaceExisting = true;
+		Task->DestinationPath = ImportAssetPath;
+		Task->bSave = false;
+		Task->DestinationName = FGuid::NewGuid().ToString(EGuidFormats::Digits);
+		Task->Options = FbxFactory->ImportUI;
+		Task->Filename = AbsoluteFilePath;
+		Task->Factory = FbxFactory;
+		FbxFactory->SetAssetImportTask(Task);
+		TArray<UAssetImportTask*> Tasks;
+		Tasks.Add(Task);
+		AssetToolsModule.Get().ImportAssetTasks(Tasks);
+
+		for (FString AssetPath : Task->ImportedObjectPaths)
+		{
+			FAssetData AssetData = AssetRegistryModule.Get().GetAssetByObjectPath(FName(*AssetPath));
+			ImportedObject = AssetData.GetAsset();
+			if (ImportedObject != nullptr)
+			{
+				break;
+			}
+		}
+
+		//Factory and task can now be garbage collected
+		Task->RemoveFromRoot();
+		FbxFactory->RemoveFromRoot();
+	}
 
 	USkeletalMesh* TmpSkeletalMesh = Cast<USkeletalMesh>(ImportedObject);
 	if (TmpSkeletalMesh == nullptr || TmpSkeletalMesh->GetSkeleton() == nullptr)
