Preserve captured thumbnails through rename operations

UnrealEdSrv.cpp
Finalize load of package before the rename, to load latent thumbnail data
Borrow a pointer to the thumbnail data before the asset rename
Put thumbnail into new package keyed on new asset name

ObjectTools.cpp
Fixup bad indention (minor)

Testing done
Capture thumbnail of non-art asset and rename
Load project with stored captured thumbnail and rename

diff --git a/Engine/Source/Editor/UnrealEd/Private/ObjectTools.cpp b/Engine/Source/Editor/UnrealEd/Private/ObjectTools.cpp
--- a/Engine/Source/Editor/UnrealEd/Private/ObjectTools.cpp
+++ b/Engine/Source/Editor/UnrealEd/Private/ObjectTools.cpp
@@ -3749,6 +3749,12 @@ namespace ObjectTools
 				Object->SetFlags(RF_Standalone);
 			}
 
+			// The object must be fully loaded to realize latent thumbnail data
+			EnsureLoadingComplete(Object->GetOutermost());
+
+			// Look for a thumbnail for this asset before we rename it
+			FObjectThumbnail* Thumbnail = ThumbnailTools::GetThumbnailForObject(Object);
+
 			FString OldObjectFullName = Object->GetFullName();
 			FString OldObjectPathName = Object->GetPathName();
 			GEditor->RenameObject( Object, NewPackage, *ObjName, bLeaveRedirector ? REN_None : REN_DontCreateRedirectors );
@@ -3763,6 +3769,12 @@ namespace ObjectTools
 				OldPackage->MetaData->RemoveMetaDataOutsidePackage();
 			}
 
+			// Migrate any thumbnail from the old package to the new one
+			if (Thumbnail)
+			{
+				ThumbnailTools::CacheThumbnail(Object->GetFullName(), Thumbnail, NewPackage);
+			}
+
 			// Notify the asset registry of the rename
 			FAssetRegistryModule::AssetRenamed(Object, OldObjectPathName);
 
