Rolling back CL 14886456 for M4 as it introduced new regressions with sprite visibility.

I'm going to narrow the scope of the fix for UE-54964 instead in an additional CL to follow.

#rnx

diff --git a/Engine/Source/Editor/LevelEditor/Private/LevelViewportActions.cpp b/Engine/Source/Editor/LevelEditor/Private/LevelViewportActions.cpp
--- a/Engine/Source/Editor/LevelEditor/Private/LevelViewportActions.cpp
+++ b/Engine/Source/Editor/LevelEditor/Private/LevelViewportActions.cpp
@@ -278,7 +278,8 @@ void FLevelViewportCommands::RegisterShowVolumeCommands()
 
 void FLevelViewportCommands::RegisterShowSpriteCommands()
 {
-	// Get all the known layers.
+	// get all the known layers
+	// Get a fresh list as GUnrealEd->SortedSpriteInfo may not yet be built.
 	TArray<FSpriteCategoryInfo> SortedSpriteInfo;
 	UUnrealEdEngine::MakeSortedSpriteInfo(SortedSpriteInfo);
 
@@ -300,25 +301,9 @@ void FLevelViewportCommands::RegisterShowSpriteCommands()
 				= FUICommandInfoDecl(this->AsShared(), CommandName, LocalizedName, SpriteInfo.Description)
 				.UserInterfaceType(EUserInterfaceActionType::ToggleButton);
 
-			const int32 ShowSpriteCommandIndex = ShowSpriteCommands.Add(FLevelViewportCommands::FShowMenuCommand(ShowSpriteCommand, SpriteInfo.DisplayName));
-
-			// Map the unlocalized category name to the index of the command in the list.
-			SpriteCategoryToCommandIndexMap.Add(SpriteInfo.Category, ShowSpriteCommandIndex);
+			ShowSpriteCommands.Add(FLevelViewportCommands::FShowMenuCommand(ShowSpriteCommand, SpriteInfo.DisplayName));
 		}
 	}
-
-	// These should both end up being the same size when we're done.
-	check(ShowSpriteCommands.Num() == SpriteCategoryToCommandIndexMap.Num());
-}
-
-int32 FLevelViewportCommands::GetShowCommandIndexForSpriteCategory(const FName& InSpriteCategoryName) const
-{
-	if (const int32* ValuePtr = SpriteCategoryToCommandIndexMap.Find(InSpriteCategoryName))
-	{
-		return *ValuePtr;
-	}
-
-	return INDEX_NONE;
 }
 
 PRAGMA_ENABLE_OPTIMIZATION
diff --git a/Engine/Source/Editor/LevelEditor/Private/SLevelViewport.cpp b/Engine/Source/Editor/LevelEditor/Private/SLevelViewport.cpp
--- a/Engine/Source/Editor/LevelEditor/Private/SLevelViewport.cpp
+++ b/Engine/Source/Editor/LevelEditor/Private/SLevelViewport.cpp
@@ -191,12 +191,6 @@ void SLevelViewport::Construct(const FArguments& InArgs, const FAssetEditorViewp
 	StartingSimulateBorder = FEditorStyle::GetBrush( "LevelViewport.StartingSimulateBorder" );
 	ReturningToEditorBorder = FEditorStyle::GetBrush( "LevelViewport.ReturningToEditorBorder" );
 
-	// Register/update viewport commands that are somewhat dynamic (can grow in size) based on loaded assets.
-	// Note: This needs to happen prior to constructing the viewport client and SEditorViewport instances below.
-	FLevelViewportCommands& LevelViewportCommands = FLevelViewportCommands::Get();
-	LevelViewportCommands.RegisterShowSpriteCommands();
-	LevelViewportCommands.RegisterShowVolumeCommands();
-
 	// Default level viewport client values for settings that could appear in layout config ini
 	FLevelEditorViewportInstanceSettings ViewportInstanceSettings;
 	ViewportInstanceSettings.ViewportType = InConstructionArguments.ViewportType;
@@ -1570,7 +1564,8 @@ void SLevelViewport::BindShowCommands( FUICommandList& OutCommandList )
 			LevelViewportCommands.HideAllVolumes,
 			FExecuteAction::CreateSP( this, &SLevelViewport::OnToggleAllVolumeActors, false ) );
 		
-		// Note: This list is somewhat dynamic and is registered/updated at viewport construction time, @see Construct().
+
+		LevelViewportCommands.RegisterShowVolumeCommands();
 		const TArray<FLevelViewportCommands::FShowMenuCommand>& ShowVolumeCommands = LevelViewportCommands.ShowVolumeCommands;
 		for (int32 VolumeCommandIndex = 0; VolumeCommandIndex < ShowVolumeCommands.Num(); ++VolumeCommandIndex)
 		{
@@ -1623,7 +1618,7 @@ void SLevelViewport::BindShowCommands( FUICommandList& OutCommandList )
 			FExecuteAction::CreateSP( this, &SLevelViewport::OnToggleAllSpriteCategories, false ) );
 
 		// Bind each show flag to the same delegate.  We use the delegate payload system to figure out what show flag we are dealing with
-		// Note: This list is somewhat dynamic and is registered/updated at viewport construction time, @see Construct().
+		LevelViewportCommands.RegisterShowSpriteCommands();
 		const TArray<FLevelViewportCommands::FShowMenuCommand>& ShowSpriteCommands = LevelViewportCommands.ShowSpriteCommands;
 		for (int32 SpriteCommandIndex = 0; SpriteCommandIndex < ShowSpriteCommands.Num(); ++SpriteCommandIndex)
 		{
@@ -2201,16 +2196,16 @@ void SLevelViewport::OnToggleAllSpriteCategories( bool bVisible )
 }
 
 /** Called when the user toggles a category from the sprite sub-menu. **/
-void SLevelViewport::ToggleSpriteCategory( int32 ShowSpriteCommandIndex )
+void SLevelViewport::ToggleSpriteCategory( int32 CategoryID )
 {
-	LevelViewportClient->SetSpriteCategoryVisibility( ShowSpriteCommandIndex, !LevelViewportClient->GetSpriteCategoryVisibility( ShowSpriteCommandIndex ) );
+	LevelViewportClient->SetSpriteCategoryVisibility( CategoryID, !LevelViewportClient->GetSpriteCategoryVisibility( CategoryID ) );
 	LevelViewportClient->Invalidate();
 }
 
 /** Called to determine if a category from the sprite sub-menu is visible. **/
-bool SLevelViewport::IsSpriteCategoryVisible( int32 ShowSpriteCommandIndex ) const
+bool SLevelViewport::IsSpriteCategoryVisible( int32 CategoryID ) const
 {
-	return LevelViewportClient->GetSpriteCategoryVisibility( ShowSpriteCommandIndex );
+	return LevelViewportClient->GetSpriteCategoryVisibility( CategoryID );
 }
 
 void SLevelViewport::OnToggleAllStatCommands( bool bVisible )
diff --git a/Engine/Source/Editor/LevelEditor/Public/LevelViewportActions.h b/Engine/Source/Editor/LevelEditor/Public/LevelViewportActions.h
--- a/Engine/Source/Editor/LevelEditor/Public/LevelViewportActions.h
+++ b/Engine/Source/Editor/LevelEditor/Public/LevelViewportActions.h
@@ -207,13 +207,7 @@ public:
 	/** Registers our commands with the binding system for showing sprites. */
 	void RegisterShowSpriteCommands();
 
-	/** Returns the index of the "show sprite" command for the given category name. */
-	int32 GetShowCommandIndexForSpriteCategory(const FName& InSpriteCategoryName) const;
-
 private:
-	/** Maps sprite category names to a matching index in the sorted ShowSpriteCommands list */
-	TMap<FName, int32> SpriteCategoryToCommandIndexMap;
-
 	/** Registers additional commands as they are loaded */
 	void HandleNewStatGroup(const TArray<FStatNameAndInfo>& NameAndInfos);
 	void HandleNewStat(const FName& InStatName, const FName& InStatCategory, const FText& InStatDescription);
diff --git a/Engine/Source/Editor/LevelEditor/Public/SLevelViewport.h b/Engine/Source/Editor/LevelEditor/Public/SLevelViewport.h
--- a/Engine/Source/Editor/LevelEditor/Public/SLevelViewport.h
+++ b/Engine/Source/Editor/LevelEditor/Public/SLevelViewport.h
@@ -533,16 +533,16 @@ private:
 	/**
 	 * Toggles sprite category visibility in this viewport
 	 *
-	 * @param ShowSpriteCommandIndex	Index of the visibility command that corresponds to a sprite category
+	 * @param CategoryID				Index of the category
 	 */
-	void ToggleSpriteCategory( int32 ShowSpriteCommandIndex );
+	void ToggleSpriteCategory( int32 CategoryID );
 
 	/**
-	 * Checks if a sprite category is visible in this viewport
+	 * Checks if sprite category is visible in this viewport
 	 *
-	 * @param ShowSpriteCommandIndex	Index of the visibility command that corresponds to a sprite category
+	 * @param CategoryID				Index of the category
 	 */
-	bool IsSpriteCategoryVisible( int32 ShowSpriteCommandIndex ) const;
+	bool IsSpriteCategoryVisible( int32 CategoryID ) const;
 
 	/**
 	 * Toggles all Stat commands visibility
diff --git a/Engine/Source/Editor/UnrealEd/Classes/Editor/UnrealEdEngine.h b/Engine/Source/Editor/UnrealEd/Classes/Editor/UnrealEdEngine.h
--- a/Engine/Source/Editor/UnrealEd/Classes/Editor/UnrealEdEngine.h
+++ b/Engine/Source/Editor/UnrealEd/Classes/Editor/UnrealEdEngine.h
@@ -169,8 +169,7 @@ public:
 	TMap<TWeakObjectPtr<UPackage>, uint8> PackageToNotifyState;
 
 	/** Mapping of sprite category ids to their matching indices in the sorted sprite categories array */
-	UE_DEPRECATED(5.0, "This mapping has been moved to FLevelViewportCommands::SpriteCategoryToCommandIndexMap.")
-	TMap<FName, int32> SpriteIDToIndexMap;
+	TMap<FName, int32>			SpriteIDToIndexMap;
 
 	/** Map from component class to visualizer object to use */
 	TMap< FName, TSharedPtr<class FComponentVisualizer> > ComponentVisualizerMap;
@@ -777,6 +776,15 @@ public:
 	 * Identify any brushes whose sense is inverted and repair them
 	 */
 	void FixAnyInvertedBrushes(UWorld* World);
+
+	/**
+	 * Get the index of the provided sprite category
+	 *
+	 * @param	InSpriteCategory	Sprite category to get the index of
+	 *
+	 * @return	Index of the provided sprite category, if possible; INDEX_NONE otherwise
+	 */
+	virtual int32 GetSpriteCategoryIndex( const FName& InSpriteCategory ) override;
 	
 	/**
 	 * Shows the LightingStaticMeshInfoWindow, creating it first if it hasn't been initialized.
diff --git a/Engine/Source/Editor/UnrealEd/Private/LevelEditorViewport.cpp b/Engine/Source/Editor/UnrealEd/Private/LevelEditorViewport.cpp
--- a/Engine/Source/Editor/UnrealEd/Private/LevelEditorViewport.cpp
+++ b/Engine/Source/Editor/UnrealEd/Private/LevelEditorViewport.cpp
@@ -1686,13 +1686,6 @@ void FTrackingTransaction::OnEditorSelectionChanged(UObject* NewSelection)
 
 }
 
-static const FLevelViewportCommands& GetLevelViewportCommands()
-{
-	static FName LevelEditorName("LevelEditor");
-	FLevelEditorModule& LevelEditor = FModuleManager::LoadModuleChecked<FLevelEditorModule>(LevelEditorName);
-	return LevelEditor.GetLevelViewportCommands();
-}
-
 
 FLevelEditorViewportClient::FLevelEditorViewportClient(const TSharedPtr<SLevelViewport>& InLevelViewport)
 	: FEditorViewportClient(&GLevelEditorModeTools(), nullptr, StaticCastSharedPtr<SEditorViewport>(InLevelViewport))
@@ -1813,8 +1806,8 @@ void FLevelEditorViewportClient::InitializeVisibilityFlags()
 	UUnrealEdEngine::GetSortedVolumeClasses(&VolumeClasses);
 	VolumeActorVisibility.Init(true, VolumeClasses.Num());
 
-	// Initialize all sprite categories to visible (one per registered "show sprite" command)
-	SpriteCategoryVisibility.Init(true, GetLevelViewportCommands().ShowSpriteCommands.Num());
+	// Initialize all sprite categories to visible
+	SpriteCategoryVisibility.Init(true, GUnrealEd->SpriteIDToIndexMap.Num());
 }
 
 void FLevelEditorViewportClient::InitializeViewportInteraction()
@@ -1882,12 +1875,9 @@ FSceneView* FLevelEditorViewportClient::CalcSceneView(FSceneViewFamily* ViewFami
 	FSceneView* View = FEditorViewportClient::CalcSceneView(ViewFamily, StereoPass);
 
 	View->ViewActor = ActorLocks.GetLock().GetLockedActor();
+	View->SpriteCategoryVisibility = SpriteCategoryVisibility;
 	View->bCameraCut = bEditorCameraCut;
 	View->bHasSelectedComponents = GEditor->GetSelectedComponentCount() > 0;
-	View->GetSpriteCategoryVisibilityOverride().BindLambda([this](const FName& InCategoryName)
-	{
-		return this->GetSpriteCategoryVisibility(InCategoryName);
-	});
 
 #if RHI_RAYTRACING
 	View->SetupRayTracedRendering();
@@ -2715,6 +2705,13 @@ static bool CommandAcceptsInput( FLevelEditorViewportClient& ViewportClient, FKe
 	return bAccepted;
 }
 
+static const FLevelViewportCommands& GetLevelViewportCommands()
+{
+	static FName LevelEditorName("LevelEditor");
+	FLevelEditorModule& LevelEditor = FModuleManager::LoadModuleChecked<FLevelEditorModule>( LevelEditorName );
+	return LevelEditor.GetLevelViewportCommands();
+}
+
 void FLevelEditorViewportClient::SetCurrentViewport()
 {
 	// Set the current level editing viewport client to the dropped-in viewport client
@@ -4680,35 +4677,30 @@ void FLevelEditorViewportClient::OnPreBeginPIE(const bool bIsSimulating)
 
 bool FLevelEditorViewportClient::GetSpriteCategoryVisibility( const FName& InSpriteCategory ) const
 {
-	const int32 ShowSpriteCommandIndex = GetLevelViewportCommands().GetShowCommandIndexForSpriteCategory( InSpriteCategory );
-	if (ShowSpriteCommandIndex != INDEX_NONE)
-	{
-		return GetSpriteCategoryVisibility(ShowSpriteCommandIndex);
-	}
+	const int32 CategoryIndex = GEngine->GetSpriteCategoryIndex( InSpriteCategory );
+	check( CategoryIndex != INDEX_NONE && CategoryIndex < SpriteCategoryVisibility.Num() );
 
-	// If the given category is not registered as a visibility command, default to being visible.
-	return true;
+	return SpriteCategoryVisibility[ CategoryIndex ];
 }
 
-bool FLevelEditorViewportClient::GetSpriteCategoryVisibility( int32 InShowSpriteCommandIndex ) const
+bool FLevelEditorViewportClient::GetSpriteCategoryVisibility( int32 Index ) const
 {
-	check( InShowSpriteCommandIndex >= 0 && InShowSpriteCommandIndex < SpriteCategoryVisibility.Num() );
-	return SpriteCategoryVisibility[ InShowSpriteCommandIndex ];
+	check( Index >= 0 && Index < SpriteCategoryVisibility.Num() );
+	return SpriteCategoryVisibility[ Index ];
 }
 
 void FLevelEditorViewportClient::SetSpriteCategoryVisibility( const FName& InSpriteCategory, bool bVisible )
 {
-	const int32 ShowSpriteCommandIndex = GetLevelViewportCommands().GetShowCommandIndexForSpriteCategory( InSpriteCategory );
-	if (ShowSpriteCommandIndex != INDEX_NONE)
-	{
-		SetSpriteCategoryVisibility(ShowSpriteCommandIndex, bVisible);
-	}
+	const int32 CategoryIndex = GEngine->GetSpriteCategoryIndex( InSpriteCategory );
+	check( CategoryIndex != INDEX_NONE && CategoryIndex < SpriteCategoryVisibility.Num() );
+
+	SpriteCategoryVisibility[ CategoryIndex ] = bVisible;
 }
 
-void FLevelEditorViewportClient::SetSpriteCategoryVisibility( int32 InShowSpriteCommandIndex, bool bVisible )
+void FLevelEditorViewportClient::SetSpriteCategoryVisibility( int32 Index, bool bVisible )
 {
-	check( InShowSpriteCommandIndex >= 0 && InShowSpriteCommandIndex < SpriteCategoryVisibility.Num() );
-	SpriteCategoryVisibility[ InShowSpriteCommandIndex ] = bVisible;
+	check( Index >= 0 && Index < SpriteCategoryVisibility.Num() );
+	SpriteCategoryVisibility[ Index ] = bVisible;
 }
 
 void FLevelEditorViewportClient::SetAllSpriteCategoryVisibility( bool bVisible )
diff --git a/Engine/Source/Editor/UnrealEd/Private/UnrealEdEngine.cpp b/Engine/Source/Editor/UnrealEd/Private/UnrealEdEngine.cpp
--- a/Engine/Source/Editor/UnrealEd/Private/UnrealEdEngine.cpp
+++ b/Engine/Source/Editor/UnrealEd/Private/UnrealEdEngine.cpp
@@ -147,6 +147,18 @@ void UUnrealEdEngine::Init(IEngineLoop* InEngineLoop)
 		}
 	}
 
+	// Populate the data structures related to the sprite category visibility feature for use elsewhere in the editor later
+	TArray<FSpriteCategoryInfo> SortedSpriteInfo;
+	UUnrealEdEngine::MakeSortedSpriteInfo(SortedSpriteInfo);
+
+	// Iterate over the sorted list, constructing a mapping of unlocalized categories to the index the localized category
+	// resides in. This is an optimization to prevent having to localize values repeatedly.
+	for( int32 InfoIndex = 0; InfoIndex < SortedSpriteInfo.Num(); ++InfoIndex )
+	{
+		const FSpriteCategoryInfo& SpriteInfo = SortedSpriteInfo[InfoIndex];
+		SpriteIDToIndexMap.Add( SpriteInfo.Category, InfoIndex );
+	}
+
 	if (FPaths::IsProjectFilePathSet() && GIsEditor && !FApp::IsUnattended())
 	{
 		AutoReimportManager = NewObject<UAutoReimportManager>();
@@ -371,7 +383,6 @@ void UUnrealEdEngine::PreExit()
 }
 
 
-PRAGMA_DISABLE_DEPRECATION_WARNINGS
 UUnrealEdEngine::~UUnrealEdEngine()
 {
 	if (this == GUnrealEd)
@@ -379,7 +390,6 @@ UUnrealEdEngine::~UUnrealEdEngine()
 		GUnrealEd = NULL; 
 	}
 }
-PRAGMA_ENABLE_DEPRECATION_WARNINGS
 
 
 void UUnrealEdEngine::FinishDestroy()
@@ -1044,6 +1054,17 @@ void UUnrealEdEngine::MakeSelectedActorsLevelCurrent()
 }
 
 
+int32 UUnrealEdEngine::GetSpriteCategoryIndex( const FName& InSpriteCategory )
+{
+	// Find the sprite category in the unlocalized to index map, if possible
+	const int32* CategoryIndexPtr = SpriteIDToIndexMap.Find( InSpriteCategory );
+	
+	const int32 CategoryIndex = CategoryIndexPtr ? *CategoryIndexPtr : INDEX_NONE;
+
+	return CategoryIndex;
+}
+
+
 void UUnrealEdEngine::ShowLightingStaticMeshInfoWindow()
 {
 	// first invoke the stats viewer tab
diff --git a/Engine/Source/Editor/UnrealEd/Public/LevelEditorViewport.h b/Engine/Source/Editor/UnrealEd/Public/LevelEditorViewport.h
--- a/Engine/Source/Editor/UnrealEd/Public/LevelEditorViewport.h
+++ b/Engine/Source/Editor/UnrealEd/Public/LevelEditorViewport.h
@@ -354,11 +354,11 @@ public:
 	/**
 	 * Returns whether the sprite category specified by the provided index is visible in the viewport or not
 	 *
-	 * @param	InShowSpriteCommandIndex	Index of the visibility command that corresponds to a sprite category
+	 * @param	Index	Index of the sprite category to check
 	 *
 	 * @return	true if the category specified by the index is visible in the viewport; false if it is not
 	 */
-	bool GetSpriteCategoryVisibility( int32 InShowSpriteCommandIndex ) const;
+	bool GetSpriteCategoryVisibility( int32 Index ) const;
 
 	/**
 	 * Sets the visibility of the provided unlocalized category to the provided value
@@ -371,10 +371,10 @@ public:
 	/**
 	 * Sets the visibility of the category specified by the provided index to the provided value
 	 *
-	 * @param	InShowSpriteCommandIndex	Index of the visibility command that corresponds to a sprite category
-	 * @param	bVisible					true if the category should be made visible, false if it should be hidden
+	 * @param	Index		Index of the sprite category to set the visibility of
+	 * @param	bVisible	true if the category should be made visible, false if it should be hidden
 	 */
-	void SetSpriteCategoryVisibility( int32 InShowSpriteCommandIndex, bool bVisible );
+	void SetSpriteCategoryVisibility( int32 Index, bool bVisible );
 
 	/**
 	 * Sets the visibility of all sprite categories to the provided value
diff --git a/Engine/Source/Runtime/Engine/Classes/Engine/Engine.h b/Engine/Source/Runtime/Engine/Classes/Engine/Engine.h
--- a/Engine/Source/Runtime/Engine/Classes/Engine/Engine.h
+++ b/Engine/Source/Runtime/Engine/Classes/Engine/Engine.h
@@ -2279,7 +2279,6 @@ public:
 	 *
 	 * @return	Index of the provided sprite category, if possible; INDEX_NONE otherwise
 	 */
-	UE_DEPRECATED(5.0, "Please use FSceneView::GetSpriteCategoryIndex() instead, which can now be accessed from within an FPrimitiveSceneProxy::GetViewRelevance() override.")
 	virtual int32 GetSpriteCategoryIndex( const FName& InSpriteCategory )
 	{
 		// The editor may override this to handle sprite categories as necessary
diff --git a/Engine/Source/Runtime/Engine/Private/Components/ArrowComponent.cpp b/Engine/Source/Runtime/Engine/Private/Components/ArrowComponent.cpp
--- a/Engine/Source/Runtime/Engine/Private/Components/ArrowComponent.cpp
+++ b/Engine/Source/Runtime/Engine/Private/Components/ArrowComponent.cpp
@@ -57,9 +57,9 @@ public:
 		bWillEverBeLit = false;
 #if WITH_EDITOR
 		// If in the editor, extract the sprite category from the component
-		if (GIsEditor)
+		if ( GIsEditor )
 		{
-			SpriteCategoryName = Component->SpriteInfo.Category;
+			SpriteCategoryIndex = GEngine->GetSpriteCategoryIndex( Component->SpriteInfo.Category );
 		}
 #endif	//WITH_EDITOR
 
@@ -175,7 +175,7 @@ public:
 #if WITH_EDITOR
 		if (bTreatAsASprite)
 		{
-			if ( GIsEditor && !View->GetSpriteCategoryVisibility(SpriteCategoryName) )
+			if ( GIsEditor && SpriteCategoryIndex != INDEX_NONE && SpriteCategoryIndex < View->SpriteCategoryVisibility.Num() && !View->SpriteCategoryVisibility[ SpriteCategoryIndex ] )
 			{
 				Result.bDrawRelevance = false;
 			}
@@ -209,7 +209,7 @@ private:
 #if WITH_EDITORONLY_DATA
 	bool bLightAttachment;
 	bool bTreatAsASprite;
-	FName SpriteCategoryName;
+	int32 SpriteCategoryIndex;
 	bool bUseInEditorScaling;
 	float EditorScale;
 #endif // #if WITH_EDITORONLY_DATA
diff --git a/Engine/Source/Runtime/Engine/Private/Components/BillboardComponent.cpp b/Engine/Source/Runtime/Engine/Private/Components/BillboardComponent.cpp
--- a/Engine/Source/Runtime/Engine/Private/Components/BillboardComponent.cpp
+++ b/Engine/Source/Runtime/Engine/Private/Components/BillboardComponent.cpp
@@ -53,7 +53,7 @@ public:
 		// Extract the sprite category from the component if in the editor
 		if ( GIsEditor )
 		{
-			SpriteCategoryName = InComponent->SpriteInfo.Category;
+			SpriteCategoryIndex = GEngine->GetSpriteCategoryIndex( InComponent->SpriteInfo.Category );
 		}
 #endif //WITH_EDITOR
 
@@ -230,9 +230,9 @@ public:
 	{
 		bool bVisible = View->Family->EngineShowFlags.BillboardSprites;
 #if WITH_EDITOR
-		if ( GIsEditor && bVisible )
+		if ( GIsEditor && bVisible && SpriteCategoryIndex != INDEX_NONE && SpriteCategoryIndex < View->SpriteCategoryVisibility.Num() )
 		{
-			bVisible = View->GetSpriteCategoryVisibility(SpriteCategoryName);
+			bVisible = View->SpriteCategoryVisibility[ SpriteCategoryIndex ];
 		}
 #endif
 		FPrimitiveViewRelevance Result;
@@ -265,7 +265,7 @@ private:
 	const uint32 bIsScreenSizeScaled : 1;
 	uint32 bIsActorLocked : 1;
 #if WITH_EDITORONLY_DATA
-	FName SpriteCategoryName;
+	int32 SpriteCategoryIndex;
 	bool bUseInEditorScaling;
 	float EditorScale;
 #endif // #if WITH_EDITORONLY_DATA
diff --git a/Engine/Source/Runtime/Engine/Private/SceneView.cpp b/Engine/Source/Runtime/Engine/Private/SceneView.cpp
--- a/Engine/Source/Runtime/Engine/Private/SceneView.cpp
+++ b/Engine/Source/Runtime/Engine/Private/SceneView.cpp
@@ -2475,19 +2475,6 @@ const FExposureBufferData* FSceneView::GetEyeAdaptationBuffer() const
 	return nullptr;
 }
 
-#if WITH_EDITOR
-bool FSceneView::GetSpriteCategoryVisibility(const FName& InCategoryName) const
-{
-	if (GetSpriteCategoryVisibilityDelegate.IsBound())
-	{
-		return GetSpriteCategoryVisibilityDelegate.Execute(InCategoryName);
-	}
-
-	// Visible by default if no delegate is bound to potentially override it.
-	return true;
-}
-#endif
-
 FSceneViewFamily::FSceneViewFamily(const ConstructionValues& CVS)
 	:
 	ViewMode(VMI_Lit),
diff --git a/Engine/Source/Runtime/Engine/Public/SceneView.h b/Engine/Source/Runtime/Engine/Public/SceneView.h
--- a/Engine/Source/Runtime/Engine/Public/SceneView.h
+++ b/Engine/Source/Runtime/Engine/Public/SceneView.h
@@ -206,6 +206,7 @@ struct FSceneViewInitOptions : public FSceneViewProjectionData
 		, OverrideLODViewOrigin(ForceInitToZero)
 		, bUseFauxOrthoViewPos(false)
 		, bDisableGameScreenPercentage(false)
+		//@TODO: , const TBitArray<>& InSpriteCategoryVisibility=TBitArray<>()
 #endif
 	{
 	}
@@ -1180,6 +1181,8 @@ public:
 	/** True if we should draw translucent objects when rendering hit proxies */
 	bool bAllowTranslucentPrimitivesInHitProxy;
 
+	/** BitArray representing the visibility state of the various sprite categories in the editor for this view */
+	TBitArray<> SpriteCategoryVisibility;
 	/** Selection color for the editor (used by post processing) */
 	FLinearColor SelectionOutlineColor;
 	/** Selection color for use in the editor with inactive primitives */
@@ -1413,27 +1416,9 @@ public:
 	/** Returns the eye adaptation buffer (mobile) or null if it doesn't exist. */
 	const FExposureBufferData* GetEyeAdaptationBuffer() const;
 
-#if WITH_EDITOR
-	/** Delegate that can provide access to the visibility state of the various sprite categories in the editor for this view */
-	DECLARE_DELEGATE_RetVal_OneParam(bool, FGetSpriteCategoryVisibility, const FName&);
-	/** Allows this view to bind to an optional override to determine sprite category visibility. */
-	FGetSpriteCategoryVisibility& GetSpriteCategoryVisibilityOverride() { return GetSpriteCategoryVisibilityDelegate; }
-	/** Returns whether the sprite category specified by the given name should be visible or not. */
-	bool GetSpriteCategoryVisibility(const FName& InCategoryName) const;
-#endif
 
 protected:
 	FSceneViewStateInterface* EyeAdaptationViewState = nullptr;
-#if WITH_EDITOR
-	FGetSpriteCategoryVisibility GetSpriteCategoryVisibilityDelegate;
-#endif
-
-private:
-#if WITH_EDITOR
-	/** BitArray representing the visibility state of the various sprite categories in the editor for this view */
-	// Deprecated - Use GetSpriteCategoryVisibility() instead.
-	TBitArray<> SpriteCategoryVisibility;
-#endif
 };
 
 //////////////////////////////////////////////////////////////////////////
