Added CookCultures and TargetPlatforms to cook stats

diff --git a/Engine/Source/Editor/UnrealEd/Private/Commandlets/CookCommandlet.cpp b/Engine/Source/Editor/UnrealEd/Private/Commandlets/CookCommandlet.cpp
--- a/Engine/Source/Editor/UnrealEd/Private/Commandlets/CookCommandlet.cpp
+++ b/Engine/Source/Editor/UnrealEd/Private/Commandlets/CookCommandlet.cpp
@@ -57,6 +57,7 @@ ASSETREGISTRY_API void DumpBuildDependencyTrackerData();
 namespace DetailedCookStats
 {
 	FString CookProject;
+	FString CookCultures;
 	FString TargetPlatforms;
 	double CookWallTimeSec = 0.0;
 	double StartupWallTimeSec = 0.0;
@@ -96,6 +97,9 @@ namespace DetailedCookStats
 		ADD_COOK_STAT_FLT(" 0. 1. 4", TickLoopShaderProcessAsyncResultsTimeSec);
 		ADD_COOK_STAT_FLT(" 0. 1. 5", TickLoopProcessDeferredCommandsTimeSec);
 		ADD_COOK_STAT_FLT(" 0. 1. 6", TickLoopTickCommandletStatsTimeSec);
+		ADD_COOK_STAT_FLT(" 0. 1. 7", TargetPlatforms);
+		ADD_COOK_STAT_FLT(" 0. 1. 8", CookProject);
+		ADD_COOK_STAT_FLT(" 0. 1. 9", CookCultures);
 		#undef ADD_COOK_STAT_FLT
 	});
 
@@ -621,7 +625,8 @@ int32 UCookCommandlet::Main(const FString& CmdLineParams)
 
 	COOK_STAT(DetailedCookStats::CookProject = FApp::GetProjectName());
 
-
+	COOK_STAT(FParse::Value(*Params, TEXT("CookCultures="), DetailedCookStats::CookCultures));
+	
 	ITargetPlatformManagerModule& TPM = GetTargetPlatformManagerRef();
 	if ( bCookOnTheFly )
 	{
