FScopedDurationTimeCustomLogger now uses UE_LOG_REF.

diff --git a/Engine/Source/Editor/UnrealEd/Private/Commandlets/WorldPartitionBuilderCommandlet.cpp b/Engine/Source/Editor/UnrealEd/Private/Commandlets/WorldPartitionBuilderCommandlet.cpp
--- a/Engine/Source/Editor/UnrealEd/Private/Commandlets/WorldPartitionBuilderCommandlet.cpp
+++ b/Engine/Source/Editor/UnrealEd/Private/Commandlets/WorldPartitionBuilderCommandlet.cpp
@@ -24,7 +24,7 @@ int32 UWorldPartitionBuilderCommandlet::Main(const FString& Params)
 {
 	TRACE_CPUPROFILER_EVENT_SCOPE(UWorldPartitionBuilderCommandlet::Main);
 
-	UE_SCOPED_TIMER(TEXT("Execution"), LogWorldPartitionBuilderCommandlet, Display);
+	UE_SCOPED_TIMER(TEXT("Execution"), LogWorldPartitionBuilderCommandlet);
 
 	TArray<FString> Tokens, Switches;
 	ParseCommandLine(*Params, Tokens, Switches);
diff --git a/Engine/Source/Editor/UnrealEd/Private/Commandlets/WorldPartitionConvertCommandlet.cpp b/Engine/Source/Editor/UnrealEd/Private/Commandlets/WorldPartitionConvertCommandlet.cpp
--- a/Engine/Source/Editor/UnrealEd/Private/Commandlets/WorldPartitionConvertCommandlet.cpp
+++ b/Engine/Source/Editor/UnrealEd/Private/Commandlets/WorldPartitionConvertCommandlet.cpp
@@ -405,7 +405,7 @@ void UWorldPartitionConvertCommandlet::ChangeObjectOuter(UObject* Object, UObjec
 
 void UWorldPartitionConvertCommandlet::FixupSoftObjectPaths(UPackage* OuterPackage)
 {
-	UE_SCOPED_TIMER(TEXT("FixupSoftObjectPaths"), LogWorldPartitionConvertCommandlet, Display);
+	UE_SCOPED_TIMER(TEXT("FixupSoftObjectPaths"), LogWorldPartitionConvertCommandlet);
 
 	struct FSoftPathFixupSerializer : public FArchiveUObject
 	{
@@ -662,7 +662,7 @@ void UWorldPartitionConvertCommandlet::OnWorldLoaded(UWorld* World)
 
 void UWorldPartitionConvertCommandlet::CreateWorldMiniMapTexture(UWorld* World)
 {
-	UE_SCOPED_TIMER(TEXT("CreateWorldMiniMapTexture"), LogWorldPartitionConvertCommandlet, Display);
+	UE_SCOPED_TIMER(TEXT("CreateWorldMiniMapTexture"), LogWorldPartitionConvertCommandlet);
 
 	AWorldPartitionMiniMap* WorldMiniMap = FWorldPartitionMiniMapHelper::GetWorldPartitionMiniMap(World, /*bCreateNewMiniMap*/true);
 	if (!WorldMiniMap)
@@ -676,7 +676,7 @@ void UWorldPartitionConvertCommandlet::CreateWorldMiniMapTexture(UWorld* World)
 
 int32 UWorldPartitionConvertCommandlet::Main(const FString& Params)
 {
-	UE_SCOPED_TIMER(TEXT("Conversion"), LogWorldPartitionConvertCommandlet, Display);
+	UE_SCOPED_TIMER(TEXT("Conversion"), LogWorldPartitionConvertCommandlet);
 
 	TArray<FString> Tokens, Switches;
 	ParseCommandLine(*Params, Tokens, Switches);
@@ -748,7 +748,7 @@ int32 UWorldPartitionConvertCommandlet::Main(const FString& Params)
 	// Delete existing result from running the commandlet, even if not using the suffix mode to cleanup previous conversion
 	if (!bReportOnly)
 	{
-		UE_SCOPED_TIMER(TEXT("Deleting existing conversion results"), LogWorldPartitionConvertCommandlet, Display);
+		UE_SCOPED_TIMER(TEXT("Deleting existing conversion results"), LogWorldPartitionConvertCommandlet);
 
 		FString OldLevelName = Tokens[0] + ConversionSuffix;
 		FString ExternalActorsPath = ULevel::GetExternalActorsPath(OldLevelName);
@@ -986,7 +986,7 @@ int32 UWorldPartitionConvertCommandlet::Main(const FString& Params)
 		// do loop after as it may modify Level->Actors
 		if (IFAs.Num())
 		{
-			UE_SCOPED_TIMER(TEXT("PartitionFoliage"), LogWorldPartitionConvertCommandlet, Display);
+			UE_SCOPED_TIMER(TEXT("PartitionFoliage"), LogWorldPartitionConvertCommandlet);
 
 			for (AInstancedFoliageActor* IFA : IFAs)
 			{
@@ -996,7 +996,7 @@ int32 UWorldPartitionConvertCommandlet::Main(const FString& Params)
 
 		if (LandscapeInfos.Num())
 		{
-			UE_SCOPED_TIMER(TEXT("PartitionLandscape"), LogWorldPartitionConvertCommandlet, Display);
+			UE_SCOPED_TIMER(TEXT("PartitionLandscape"), LogWorldPartitionConvertCommandlet);
 
 			for (ULandscapeInfo* LandscapeInfo : LandscapeInfos)
 			{
diff --git a/Engine/Source/Runtime/Core/Public/ProfilingDebugging/ScopedTimers.h b/Engine/Source/Runtime/Core/Public/ProfilingDebugging/ScopedTimers.h
--- a/Engine/Source/Runtime/Core/Public/ProfilingDebugging/ScopedTimers.h
+++ b/Engine/Source/Runtime/Core/Public/ProfilingDebugging/ScopedTimers.h
@@ -206,20 +206,20 @@ public:
 class FScopedDurationTimeCustomLogger
 {
 public:
-	explicit FScopedDurationTimeCustomLogger(FString InMsg, TFunction<void(const FString&)> InOutputFunction)
-		: Msg(MoveTemp(InMsg))
-		, OutputFunction(InOutputFunction)
+	explicit FScopedDurationTimeCustomLogger(const TCHAR* InMsg, const FLogCategoryBase& InLogCategory)
+		: Msg(InMsg)
+		, LogCategory(InLogCategory)
 		, Accumulator(0.0)
 		, Timer(Accumulator)
 	{
-		OutputFunction(FString::Printf(TEXT("%s started..."), *Msg));
+		UE_LOG_REF(LogCategory, Display, TEXT("%s"), *Msg);
 		Timer.Start();
 	}
 
 	~FScopedDurationTimeCustomLogger()
 	{
 		Timer.Stop();
-		OutputFunction(FString::Printf(TEXT("%s took %s"), *Msg, *SecondsToString(Accumulator)));
+		UE_LOG_REF(LogCategory, Display, TEXT("%s took %s"), *Msg, *SecondsToString(Accumulator));
 	}
 
 private:
@@ -258,10 +258,10 @@ private:
 	}
 
 	FString Msg;
-	TFunction<void(const FString&)> OutputFunction;
+	const FLogCategoryBase& LogCategory;
 	double Accumulator;
 	FDurationTimer Timer;
 };
 
-#define UE_SCOPED_TIMER(Title, Category, Verbosity) \
-	FScopedDurationTimeCustomLogger BODY_MACRO_COMBINE(Scoped,Timer,_,__LINE__)(Title, [](const FString& Msg) { UE_LOG(Category, Verbosity, TEXT("%s"), *Msg); });
+#define UE_SCOPED_TIMER(Title, Category) \
+	FScopedDurationTimeCustomLogger BODY_MACRO_COMBINE(Scoped,Timer,_,__LINE__)(Title, Category)
diff --git a/Engine/Source/Runtime/Engine/Private/WorldPartition/WorldPartition.cpp b/Engine/Source/Runtime/Engine/Private/WorldPartition/WorldPartition.cpp
--- a/Engine/Source/Runtime/Engine/Private/WorldPartition/WorldPartition.cpp
+++ b/Engine/Source/Runtime/Engine/Private/WorldPartition/WorldPartition.cpp
@@ -142,7 +142,7 @@ FName UWorldPartition::GetWorldPartitionEditorName()
 
 void UWorldPartition::Initialize(UWorld* InWorld, const FTransform& InTransform)
 {
-	UE_SCOPED_TIMER(TEXT("WorldPartition initialize"), LogWorldPartition, Log);
+	UE_SCOPED_TIMER(TEXT("WorldPartition initialize"), LogWorldPartition);
 	
 	check(!World || (World == InWorld));
 	if (!ensure(!IsInitialized()))
diff --git a/Engine/Source/Runtime/Engine/Private/WorldPartition/WorldPartitionMiniMapHelper.cpp b/Engine/Source/Runtime/Engine/Private/WorldPartition/WorldPartitionMiniMapHelper.cpp
--- a/Engine/Source/Runtime/Engine/Private/WorldPartition/WorldPartitionMiniMapHelper.cpp
+++ b/Engine/Source/Runtime/Engine/Private/WorldPartition/WorldPartitionMiniMapHelper.cpp
@@ -55,7 +55,7 @@ AWorldPartitionMiniMap* FWorldPartitionMiniMapHelper::GetWorldPartitionMiniMap(U
 void FWorldPartitionMiniMapHelper::CaptureWorldMiniMapToTexture(UWorld* InWorld, UObject* InOuterForTexture, uint32 InMiniMapSize, UTexture2D*& InOutMiniMapTexture, FBox& OutWorldBounds)
 {
 	auto WaitForShaderCompilation = [](){
-		UE_SCOPED_TIMER(TEXT("Shader Compilation"), LogWorldPartitionMiniMapHelper, Display);
+		UE_SCOPED_TIMER(TEXT("Shader Compilation"), LogWorldPartitionMiniMapHelper);
 		auto RemainingShaders = GShaderCompilingManager->GetNumRemainingJobs();
 		if (RemainingShaders > 0)
 		{
@@ -64,7 +64,7 @@ void FWorldPartitionMiniMapHelper::CaptureWorldMiniMapToTexture(UWorld* InWorld,
 	};
 	auto WaitForTextureBuilding = []()
 	{
-		UE_SCOPED_TIMER(TEXT("Building Textures"), LogWorldPartitionMiniMapHelper, Display);
+		UE_SCOPED_TIMER(TEXT("Building Textures"), LogWorldPartitionMiniMapHelper);
 		auto RemainingTextures = FTextureCompilingManager::Get().GetNumRemainingTextures();
 		if (RemainingTextures > 0)
 		{
@@ -72,7 +72,7 @@ void FWorldPartitionMiniMapHelper::CaptureWorldMiniMapToTexture(UWorld* InWorld,
 		}
 	};
 	auto WaitForDistanceMeshFieldBuilding = []() {
-		UE_SCOPED_TIMER(TEXT("Building Mesh Distance Fields"), LogWorldPartitionMiniMapHelper, Display);
+		UE_SCOPED_TIMER(TEXT("Building Mesh Distance Fields"), LogWorldPartitionMiniMapHelper);
 		if (GDistanceFieldAsyncQueue && GDistanceFieldAsyncQueue->GetNumOutstandingTasks() > 0)
 		{
 			GDistanceFieldAsyncQueue->BlockUntilAllBuildsComplete();
diff --git a/Engine/Source/Runtime/Engine/Private/WorldPartition/WorldPartitionRuntimeSpatialHash.cpp b/Engine/Source/Runtime/Engine/Private/WorldPartition/WorldPartitionRuntimeSpatialHash.cpp
--- a/Engine/Source/Runtime/Engine/Private/WorldPartition/WorldPartitionRuntimeSpatialHash.cpp
+++ b/Engine/Source/Runtime/Engine/Private/WorldPartition/WorldPartitionRuntimeSpatialHash.cpp
@@ -730,7 +730,7 @@ struct FSquare2DGridHelper
 	// Validates that actor is not referenced by multiple cells
 	void ValidateSingleActorReferer()
 	{
-		UE_SCOPED_TIMER(TEXT("ValidateSingleActorReferer"), LogWorldPartitionRuntimeSpatialHash, Log);
+		UE_SCOPED_TIMER(TEXT("ValidateSingleActorReferer"), LogWorldPartitionRuntimeSpatialHash);
 
 		TSet<FGuid> ActorUsage;
 		for (int32 Level = 0; Level < Levels.Num() - 1; Level++)
@@ -1059,7 +1059,7 @@ static FSquare2DGridHelper GetGridHelper(const FBox& WorldBounds, int32 GridCell
 
 static FSquare2DGridHelper GetPartitionedActors(const UWorldPartition* WorldPartition, const FBox& WorldBounds, const FSpatialHashRuntimeGrid& Grid, const TArray<FActorCluster>& GridActors)
 {
-	UE_SCOPED_TIMER(TEXT("GetPartitionedActors"), LogWorldPartitionRuntimeSpatialHash, Log);
+	UE_SCOPED_TIMER(TEXT("GetPartitionedActors"), LogWorldPartitionRuntimeSpatialHash);
 
 	//
 	// Create the hierarchical grids for the game
@@ -1249,7 +1249,7 @@ bool UWorldPartitionRuntimeSpatialHash::GenerateStreaming(EWorldPartitionStreami
 	TRACE_CPUPROFILER_EVENT_SCOPE(UWorldPartitionRuntimeSpatialHash::GenerateStreaming);
 	UWorldPartition* WorldPartition = GetOuterUWorldPartition();
 
-	UE_SCOPED_TIMER(TEXT("GenerateStreaming"), LogWorldPartitionRuntimeSpatialHash, Log);
+	UE_SCOPED_TIMER(TEXT("GenerateStreaming"), LogWorldPartitionRuntimeSpatialHash);
 	
 	if (!Grids.Num())
 	{
