IniSettingsOutOfDate: Fix ConfigSettingsBlacklist *.Engine entries failing to apply to the settings read from the platform-independent inis used in the cooker.
#rnx

diff --git a/Engine/Source/Editor/UnrealEd/Private/CookOnTheFlyServer.cpp b/Engine/Source/Editor/UnrealEd/Private/CookOnTheFlyServer.cpp
--- a/Engine/Source/Editor/UnrealEd/Private/CookOnTheFlyServer.cpp
+++ b/Engine/Source/Editor/UnrealEd/Private/CookOnTheFlyServer.cpp
@@ -4567,6 +4567,7 @@ bool UCookOnTheFlyServer::GetCurrentIniVersionStrings( const ITargetPlatform* Ta
 
 
 	// remove any which are filtered out
+	FString EditorPrefix(TEXT("Editor."));
 	for ( const FString& Filter : ConfigSettingBlacklist )
 	{
 		TArray<FString> FilterArray;
@@ -4592,7 +4593,11 @@ bool UCookOnTheFlyServer::GetCurrentIniVersionStrings( const ITargetPlatform* Ta
 		{
 			for ( auto ConfigFile = IniVersionStrings.CreateIterator(); ConfigFile; ++ConfigFile )
 			{
-				if ( ConfigFile.Key().ToString().MatchesWildcard(*ConfigFileName) )
+				// Some ConfigBlacklistSettings are written as *.Engine, and are intended to affect the platform-less Editor Engine.ini, which is just "Engine"
+				// To make *.Engine match the editor-only config files as well, we check whether the wildcard matches either Engine or Editor.Engine for the editor files
+				FString IniVersionStringFilename = ConfigFile.Key().ToString();
+				if (IniVersionStringFilename.MatchesWildcard(*ConfigFileName) ||
+					(!IniVersionStringFilename.Contains(TEXT(".")) && (EditorPrefix + IniVersionStringFilename).MatchesWildcard(*ConfigFileName)))
 				{
 					if ( SectionName )
 					{
