Fix incorrect filename extension used for .umap packages that contain a redirector to the ULevel/UWorld.
#rnx

diff --git a/Engine/Source/Editor/UnrealEd/Private/Cooker/PackageNameCache.h b/Engine/Source/Editor/UnrealEd/Private/Cooker/PackageNameCache.h
--- a/Engine/Source/Editor/UnrealEd/Private/Cooker/PackageNameCache.h
+++ b/Engine/Source/Editor/UnrealEd/Private/Cooker/PackageNameCache.h
@@ -155,8 +155,24 @@ inline bool FPackageNameCache::DoesPackageExist(const FName& PackageName, FStrin
 
 	if (OutFilename)
 	{
-		const bool ContainsMap = Algo::FindByPredicate(Assets, [](const FAssetData& Asset) { return Asset.PackageFlags & PKG_ContainsMap; }) != nullptr;
-		const FString& PackageExtension = ContainsMap ? FPackageName::GetMapPackageExtension() : FPackageName::GetAssetPackageExtension();
+		FName ClassRedirector = UObjectRedirector::StaticClass()->GetFName();
+		bool bContainsMap = false;
+		bool bContainsRedirector = false;
+		for (const FAssetData& Asset : Assets)
+		{
+			bContainsMap = bContainsMap | ((Asset.PackageFlags & PKG_ContainsMap) != 0);
+			bContainsRedirector = bContainsRedirector | (Asset.AssetClass == ClassRedirector);
+		}
+		if (!bContainsMap && bContainsRedirector)
+		{
+			// presence of map -> .umap
+			// But we can only assume lack of map -> .uasset if we know the type of every object in the package.
+			// If we don't, because there was a redirector, we have to check the package on disk
+			// TODO: Have the AssetRegistry store the extension of the package so that we don't have to look it up
+			// Guessing the extension based on map vs non-map also does not support text assets and maps which have a different extension
+			return FPackageName::DoesPackageExist(PackageNameStr, NULL, OutFilename, false);
+		}
+		const FString& PackageExtension = bContainsMap ? FPackageName::GetMapPackageExtension() : FPackageName::GetAssetPackageExtension();
 		*OutFilename = FPackageName::LongPackageNameToFilename(PackageNameStr, PackageExtension);
 	}
 
