Update NULL -> nullptr
#rnx

diff --git a/Engine/Source/Editor/UnrealEd/Private/Cooker/PackageNameCache.h b/Engine/Source/Editor/UnrealEd/Private/Cooker/PackageNameCache.h
--- a/Engine/Source/Editor/UnrealEd/Private/Cooker/PackageNameCache.h
+++ b/Engine/Source/Editor/UnrealEd/Private/Cooker/PackageNameCache.h
@@ -142,7 +142,7 @@ inline bool FPackageNameCache::DoesPackageExist(const FName& PackageName, FStrin
 
 	if (!AssetRegistry)
 	{
-		return FPackageName::DoesPackageExist(PackageNameStr, NULL, OutFilename, false);
+		return FPackageName::DoesPackageExist(PackageNameStr, nullptr, OutFilename, false);
 	}
 
 	TArray<FAssetData> Assets;
@@ -170,7 +170,7 @@ inline bool FPackageNameCache::DoesPackageExist(const FName& PackageName, FStrin
 			// If we don't, because there was a redirector, we have to check the package on disk
 			// TODO: Have the AssetRegistry store the extension of the package so that we don't have to look it up
 			// Guessing the extension based on map vs non-map also does not support text assets and maps which have a different extension
-			return FPackageName::DoesPackageExist(PackageNameStr, NULL, OutFilename, false);
+			return FPackageName::DoesPackageExist(PackageNameStr, nullptr, OutFilename, false);
 		}
 		const FString& PackageExtension = bContainsMap ? FPackageName::GetMapPackageExtension() : FPackageName::GetAssetPackageExtension();
 		*OutFilename = FPackageName::LongPackageNameToFilename(PackageNameStr, PackageExtension);
@@ -208,7 +208,7 @@ inline const FPackageNameCache::FCachedPackageFilename& FPackageNameCache::Cache
 
 	FCachedPackageFilename *Cached = PackageFilenameCache.Find(PackageName);
 
-	if (Cached != NULL)
+	if (Cached != nullptr)
 	{
 		return *Cached;
 	}
