When using the placement subsystem, set the object flags on the created actor.

diff --git a/Engine/Source/Editor/UnrealEd/Private/AssetSelection.cpp b/Engine/Source/Editor/UnrealEd/Private/AssetSelection.cpp
--- a/Engine/Source/Editor/UnrealEd/Private/AssetSelection.cpp
+++ b/Engine/Source/Editor/UnrealEd/Private/AssetSelection.cpp
@@ -622,6 +622,7 @@ static AActor* PrivateAddActor( UObject* Asset, UActorFactory* Factory, bool Sel
 			{
 				const FActorElementData* ActorElementData = PlacedElements[0].GetData<FActorElementData>();
 				Actor = ActorElementData ? ActorElementData->Actor : nullptr;
+				Actor->SetFlags(ObjectFlags);
 			}
 		}
 		
