Fixed an import problem with FBXs containing LODs with meshes which can't be triangulated (e.g. due to being non-manifold).
Previously, an error in an individual FBX mesh would cause the import for the current LOD to stop but not fail, resulting in an incomplete mesh. Now, problematic meshes will be omitted and a warning generated. If a LOD cannot be imported at all, it will be auto-generated instead.

diff --git a/Engine/Source/Editor/UnrealEd/Private/Fbx/FbxFactory.cpp b/Engine/Source/Editor/UnrealEd/Private/Fbx/FbxFactory.cpp
--- a/Engine/Source/Editor/UnrealEd/Private/Fbx/FbxFactory.cpp
+++ b/Engine/Source/Editor/UnrealEd/Private/Fbx/FbxFactory.cpp
@@ -562,10 +562,18 @@ UObject* UFbxFactory::FactoryCreateFile
 							}
 							else
 							{
-								FbxImporter->ImportStaticMeshAsSingle(InParent, LODMeshesArray, Name, Flags, ImportUI->StaticMeshImportData, NewStaticMesh, LODIndex);
-								if (NewStaticMesh && NewStaticMesh->IsSourceModelValid(LODIndex))
+								UStaticMesh* ImportedStaticMesh = FbxImporter->ImportStaticMeshAsSingle(InParent, LODMeshesArray, Name, Flags, ImportUI->StaticMeshImportData, NewStaticMesh, LODIndex);
+								if (ImportedStaticMesh)
 								{
-									NewStaticMesh->GetSourceModel(LODIndex).bImportWithBaseMesh = true;
+									check(ImportedStaticMesh == NewStaticMesh);
+									if (NewStaticMesh && NewStaticMesh->IsSourceModelValid(LODIndex))
+									{
+										NewStaticMesh->GetSourceModel(LODIndex).bImportWithBaseMesh = true;
+									}
+								}
+								else
+								{
+									FbxImporter->AddStaticMeshSourceModelGeneratedLOD(NewStaticMesh, LODIndex);
 								}
 							}
 							bOperationCanceled |= FbxImporter->GetImportOperationCancelled() || SlowTask.ShouldCancel();
diff --git a/Engine/Source/Editor/UnrealEd/Private/Fbx/FbxStaticMeshImport.cpp b/Engine/Source/Editor/UnrealEd/Private/Fbx/FbxStaticMeshImport.cpp
--- a/Engine/Source/Editor/UnrealEd/Private/Fbx/FbxStaticMeshImport.cpp
+++ b/Engine/Source/Editor/UnrealEd/Private/Fbx/FbxStaticMeshImport.cpp
@@ -476,8 +476,8 @@ bool UnFbx::FFbxImporter::BuildStaticMeshFromGeometry(FbxNode* Node, UStaticMesh
 		}
 		else
 		{
-			AddTokenizedErrorMessage(FTokenizedMessage::Create(EMessageSeverity::Warning, FText::Format(LOCTEXT("Error_FailedToTriangulate", "Unable to triangulate mesh '{0}'"), FText::FromString(Mesh->GetName()))), FFbxErrors::Generic_Mesh_TriangulationFailed);
-			return false; // not clean, missing some dealloc
+			AddTokenizedErrorMessage(FTokenizedMessage::Create(EMessageSeverity::Warning, FText::Format(LOCTEXT("Error_FailedToTriangulate", "Unable to triangulate mesh '{0}': it will be omitted."), FText::FromString(FbxNodeName))), FFbxErrors::Generic_Mesh_TriangulationFailed);
+			return false;
 		}
 		
 		BuildStaticMeshSlowTask.EnterProgressFrame(1);
@@ -1780,6 +1780,7 @@ UStaticMesh* UnFbx::FFbxImporter::ImportStaticMeshAsSingle(UObject* InParent, TA
 
 	bool bAllDegenerated = true;
 	TArray<FFbxMaterial> MeshMaterials;
+	int32 NodeFailCount = 0;
 	for (int32 MeshIndex = 0; MeshIndex < MeshNodeArray.Num(); MeshIndex++)
 	{
 		FScopedSlowTask BuildMeshSlowTask(MeshNodeArray.Num(), FText::Format(LOCTEXT("FbxStaticMeshBuildGeometryTask", "Importing Static Mesh Geometry {0} of {1}."), FText::AsNumber(1), FText::AsNumber(MeshNodeArray.Num())));
@@ -1811,17 +1812,21 @@ UStaticMesh* UnFbx::FFbxImporter::ImportStaticMeshAsSingle(UObject* InParent, TA
 				if (!BuildStaticMeshFromGeometry(Node, StaticMesh, MeshMaterials, LODIndex,
 					VertexColorImportOption, ExistingVertexColorData, ImportOptions->VertexOverrideColor))
 				{
-					bBuildStatus = false;
-					break;
+					// If this FBX node failed to build, make a note and continue.
+					// We should only fail the mesh/LOD import if every node failed.
+					NodeFailCount++;
 				}
 			}
 			else
 			{
 				CreateTokenizedErrorForDegeneratedPart(this, StaticMesh->GetName(), UTF8_TO_TCHAR(Node->GetName()));
+				NodeFailCount++;
 			}
 		}
 	}
-	if (bAllDegenerated)
+
+	// If every node in the list failed to build, this counts as a failure to import the entire mesh/LOD.
+	if (NodeFailCount == MeshNodeArray.Num() || bAllDegenerated)
 	{
 		bBuildStatus = false;
 	}
