When reverting Open for Adds. Also delete the files.

diff --git a/Engine/Source/Editor/UnrealEd/Private/PackageSourceControlHelper.cpp b/Engine/Source/Editor/UnrealEd/Private/PackageSourceControlHelper.cpp
--- a/Engine/Source/Editor/UnrealEd/Private/PackageSourceControlHelper.cpp
+++ b/Engine/Source/Editor/UnrealEd/Private/PackageSourceControlHelper.cpp
@@ -67,6 +67,12 @@ bool FPackageSourceControlHelper::Delete(const FString& PackageName) const
 					UE_LOG(LogCommandletPackageHelper, Error, TEXT("Error reverting package %s from source control"), *Filename);
 					return false;
 				}
+
+				if (!IFileManager::Get().Delete(*Filename, false, true))
+				{
+					UE_LOG(LogCommandletPackageHelper, Error, TEXT("Error deleting package %s locally"), *Filename);
+					return false;
+				}
 			}
 			else
 			{
