CookOnTheFlyServer: fix load-blocked diagnostic not to fire in periods where we aren't trying to load. Fix load-blocked and save-blocked diagnostics to not fire when we are making some progress, and to clamp the amount displayed.
#rnx

diff --git a/Engine/Source/Editor/UnrealEd/Private/CookOnTheFlyServer.cpp b/Engine/Source/Editor/UnrealEd/Private/CookOnTheFlyServer.cpp
--- a/Engine/Source/Editor/UnrealEd/Private/CookOnTheFlyServer.cpp
+++ b/Engine/Source/Editor/UnrealEd/Private/CookOnTheFlyServer.cpp
@@ -1166,7 +1166,7 @@ uint32 UCookOnTheFlyServer::TickCookOnTheSide(const float TimeSlice, uint32 &Coo
 				break;
 			case ECookAction::Load:
 				PumpLoads(StackData, 0, NumPushed, bBusy);
-				SetLoadBusy(bBusy);
+				SetLoadBusy(bBusy && NumPushed == 0); // Mark as busy if pump was blocked and we did not make any progress
 				if (NumPushed > 0)
 				{
 					SetSaveBusy(false);
@@ -1174,7 +1174,7 @@ uint32 UCookOnTheFlyServer::TickCookOnTheSide(const float TimeSlice, uint32 &Coo
 				break;
 			case ECookAction::LoadLimited:
 				PumpLoads(StackData, DesiredLoadQueueLength, NumPushed, bBusy);
-				SetLoadBusy(bBusy);
+				SetLoadBusy(bBusy && NumPushed == 0); // Mark as busy if pump was blocked and we did not make any progress
 				if (NumPushed > 0)
 				{
 					SetSaveBusy(false);
@@ -1182,11 +1182,11 @@ uint32 UCookOnTheFlyServer::TickCookOnTheSide(const float TimeSlice, uint32 &Coo
 				break;
 			case ECookAction::Save:
 				PumpSaves(StackData, 0, NumPushed, bBusy);
-				SetSaveBusy(bBusy);
+				SetSaveBusy(bBusy && NumPushed == 0); // Mark as busy if pump was blocked and we did not make any progress
 				break;
 			case ECookAction::SaveLimited:
 				PumpSaves(StackData, DesiredSaveQueueLength, NumPushed, bBusy);
-				SetSaveBusy(bBusy);
+				SetSaveBusy(bBusy && NumPushed == 0); // Mark as busy if pump was blocked and we did not make any progress
 				break;
 			case ECookAction::Done:
 				bContinueTick = false;
@@ -1256,6 +1256,8 @@ void UCookOnTheFlyServer::TickCookStatus(UE::Cook::FTickStackData& StackData)
 
 void UCookOnTheFlyServer::SetSaveBusy(bool bInBusy)
 {
+	using namespace UE::Cook;
+
 	if (bSaveBusy != bInBusy)
 	{
 		bSaveBusy = bInBusy;
@@ -1270,10 +1272,60 @@ void UCookOnTheFlyServer::SetSaveBusy(bool bInBusy)
 			SaveBusyTimeLastRetry = 0;
 		}
 	}
+	else if (bSaveBusy)
+	{
+		const float CurrentTime = FPlatformTime::Seconds();
+		if (CurrentTime - SaveBusyTimeStarted > GCookProgressWarnBusyTime)
+		{
+			UE_LOG(LogCook, Warning, TEXT("Cooker has been blocked from saving the current packages for %f seconds."), GCookProgressWarnBusyTime);
+			FPackageDataQueue& SaveQueue = PackageDatas->GetSaveQueue();
+			UE_LOG(LogCook, Display, TEXT("%d packages in the savequeue: "), SaveQueue.Num());
+			int DisplayCount = 0;
+			const int DisplayMax = 10;
+			for (UE::Cook::FPackageData* PackageData: SaveQueue)
+			{
+				if (DisplayCount == DisplayMax)
+				{
+					UE_LOG(LogCook, Display, TEXT("    ..."));
+					break;
+				}
+				UE_LOG(LogCook, Display, TEXT("    %s"), *PackageData->GetFileName().ToString());
+				++DisplayCount;
+			}
+			if (DisplayCount == 0)
+			{
+				UE_LOG(LogCook, Display, TEXT("    <None>"));
+			}
+
+			UE_LOG(LogCook, Display, TEXT("%d objects that have not yet returned true from IsCachedCookedPlatformDataLoaded:"), PackageDatas->GetPendingCookedPlatformDatas().Num());
+			DisplayCount = 0;
+			for (const FPendingCookedPlatformData& Data : PackageDatas->GetPendingCookedPlatformDatas())
+			{
+				if (Data.Object.IsValid())
+				{
+					if (DisplayCount == DisplayMax)
+					{
+						UE_LOG(LogCook, Display, TEXT("    ..."));
+						break;
+					}
+					UE_LOG(LogCook, Display, TEXT("    %s"), *Data.Object.Get()->GetFullName());
+					++DisplayCount;
+				}
+			}
+			if (DisplayCount == 0)
+			{
+				UE_LOG(LogCook, Display, TEXT("    <None>"));
+			}
+
+			SaveBusyTimeStarted = CurrentTime;
+		}
+	}
 }
 
 void UCookOnTheFlyServer::SetLoadBusy(bool bInLoadBusy)
 {
+	using namespace UE::Cook;
+
 	if (bLoadBusy != bInLoadBusy)
 	{
 		bLoadBusy = bInLoadBusy;
@@ -1288,6 +1340,42 @@ void UCookOnTheFlyServer::SetLoadBusy(bool bInLoadBusy)
 			LoadBusyTimeLastRetry = 0;
 		}
 	}
+	else if (bLoadBusy)
+	{
+		const float CurrentTime = FPlatformTime::Seconds();
+		if (CurrentTime - LoadBusyTimeStarted > GCookProgressWarnBusyTime)
+		{
+			int DisplayCount = 0;
+			const int DisplayMax = 10;
+			FLoadPrepareQueue& LoadPrepareQueue = PackageDatas->GetLoadPrepareQueue();
+			UE_LOG(LogCook, Warning, TEXT("Cooker has been blocked from loading the current packages for %f seconds. %d packages in the loadqueue:"), GCookProgressWarnBusyTime, LoadPrepareQueue.PreloadingQueue.Num() + LoadPrepareQueue.EntryQueue.Num());
+			for (FPackageData* PackageData : LoadPrepareQueue.PreloadingQueue)
+			{
+				if (DisplayCount == DisplayMax)
+				{
+					UE_LOG(LogCook, Display, TEXT("    ..."));
+					break;
+				}
+				UE_LOG(LogCook, Display, TEXT("    %s"), *PackageData->GetFileName().ToString());
+				++DisplayCount;
+			}
+			for (FPackageData* PackageData : LoadPrepareQueue.EntryQueue)
+			{
+				if (DisplayCount == DisplayMax)
+				{
+					UE_LOG(LogCook, Display, TEXT("    ..."));
+					break;
+				}
+				UE_LOG(LogCook, Display, TEXT("    %s"), *PackageData->GetFileName().ToString());
+				++DisplayCount;
+			}
+			if (DisplayCount == 0)
+			{
+				UE_LOG(LogCook, Display, TEXT("    <None>"));
+			}
+			LoadBusyTimeStarted = CurrentTime;
+		}
+	}
 }
 
 void UCookOnTheFlyServer::UpdateDisplay(ECookTickFlags TickFlags, bool bForceDisplay)
@@ -1326,57 +1414,9 @@ void UCookOnTheFlyServer::UpdateDisplay(ECookTickFlags TickFlags, bool bForceDis
 			TEXT("Cook Diagnostics: OpenFileHandles=%d, VirtualMemory=%dMiB"),
 			OpenFileHandles, FPlatformMemory::GetStats().UsedVirtual / 1024 / 1024);
 		LastDiagnosticsDisplayTime = CurrentTime;
-		if (bSaveBusy && CurrentTime - SaveBusyTimeStarted > GCookProgressWarnBusyTime)
-		{
-			UE_LOG(LogCook, Warning, TEXT("Cooker has been blocked from saving the current packages for %f seconds."), GCookProgressWarnBusyTime);
-			UE_LOG(LogCook, Display, TEXT("Current packages in the savequeue : "), GCookProgressWarnBusyTime);
-			FPackageDataQueue& SaveQueue = PackageDatas->GetSaveQueue();
-			bool bFound = false;
-			for (UE::Cook::FPackageData* PackageData : SaveQueue)
-			{
-				UE_LOG(LogCook, Display, TEXT("    %s"), *PackageData->GetFileName().ToString());
-				bFound = true;
-			}
-			if (!bFound)
-			{
-				UE_LOG(LogCook, Display, TEXT("    <None>"));
-			}
-
-			UE_LOG(LogCook, Display, TEXT("Current objects that have not yet returned true from IsCachedCookedPlatformDataLoaded:"));
-			bFound = false;
-			for (const FPendingCookedPlatformData& Data : PackageDatas->GetPendingCookedPlatformDatas())
-			{
-				if (Data.Object.IsValid())
-				{
-					bFound = true;
-					UE_LOG(LogCook, Display, TEXT("    %s"), *Data.Object.Get()->GetFullName());
-				}
-			}
-			if (!bFound)
-			{
-				UE_LOG(LogCook, Display, TEXT("    <None>"));
-			}
-
-			SaveBusyTimeStarted = CurrentTime;
-			SaveBusyTimeLastRetry = CurrentTime;
-		}
-		if (bLoadBusy && CurrentTime - LoadBusyTimeStarted > GCookProgressWarnBusyTime)
-		{
-			UE_LOG(LogCook, Warning, TEXT("Cooker has been blocked from loading the current packages for %f seconds. Current packages in the loadqueue:"), GCookProgressWarnBusyTime);
-			FLoadPrepareQueue& LoadPrepareQueue = PackageDatas->GetLoadPrepareQueue();
-			for (FPackageData* PackageData : LoadPrepareQueue.PreloadingQueue)
-			{
-				UE_LOG(LogCook, Display, TEXT("%s"), *PackageData->GetFileName().ToString());
-			}
-			for (FPackageData* PackageData : LoadPrepareQueue.EntryQueue)
-			{
-				UE_LOG(LogCook, Display, TEXT("%s"), *PackageData->GetFileName().ToString());
-			}
-			LoadBusyTimeStarted = CurrentTime;
-			LoadBusyTimeLastRetry = CurrentTime;
-		}
 	}
 }
+
 UCookOnTheFlyServer::ECookAction UCookOnTheFlyServer::DecideNextCookAction(UE::Cook::FTickStackData& StackData)
 {
 	if (IsCookByTheBookMode() && CookByTheBookOptions->bCancel)
