Update AsyncCompilationManager from CookOnTheFlyServer

Previously, the cooker only made progress on async builds when it blocked on them. Now it will make progress from the main cooking loop.

diff --git a/Engine/Source/Editor/UnrealEd/Private/CookOnTheFlyServer.cpp b/Engine/Source/Editor/UnrealEd/Private/CookOnTheFlyServer.cpp
--- a/Engine/Source/Editor/UnrealEd/Private/CookOnTheFlyServer.cpp
+++ b/Engine/Source/Editor/UnrealEd/Private/CookOnTheFlyServer.cpp
@@ -7,6 +7,7 @@
 #include "CookOnTheSide/CookOnTheFlyServer.h"
 #include "Cooker/PackageNameCache.h"
 
+#include "AssetCompilingManager.h"
 #include "Cooker/CookPackageData.h"
 #include "Cooker/CookPlatformManager.h"
 #include "Cooker/CookProfiling.h"
@@ -8044,6 +8045,7 @@ uint32 UCookOnTheFlyServer::FullLoadAndSave(uint32& CookedPackageCount)
 				{
 					GShaderCompilingManager->ProcessAsyncResults(true, false);
 				}
+				FAssetCompilingManager::Get().ProcessAsyncTasks(true);
 				GIsCookerLoadingPackage = false;
 			}
 		}
@@ -8262,6 +8264,7 @@ uint32 UCookOnTheFlyServer::FullLoadAndSave(uint32& CookedPackageCount)
 					{
 						GShaderCompilingManager->ProcessAsyncResults(true, false);
 					}
+					FAssetCompilingManager::Get().ProcessAsyncTasks(true);
 				}
 			}
 			GIsCookerLoadingPackage = false;
@@ -8314,6 +8317,9 @@ uint32 UCookOnTheFlyServer::FullLoadAndSave(uint32& CookedPackageCount)
 		GCardRepresentationAsyncQueue->BlockUntilAllBuildsComplete();
 	}
 
+	// Wait for all assets to finish compiling
+	FAssetCompilingManager::Get().ProcessAsyncTasks(false);
+
 	// Wait for all platform data to be loaded
 	{
 		UE_LOG(LogCook, Display, TEXT("Waiting for cooked platform data..."));
@@ -8339,6 +8345,7 @@ uint32 UCookOnTheFlyServer::FullLoadAndSave(uint32& CookedPackageCount)
 				}
 			}
 
+			FAssetCompilingManager::Get().ProcessAsyncTasks(true);
 			FPlatformProcess::Sleep(0.001f);
 		}
 
diff --git a/Engine/Source/Editor/UnrealEd/Private/Cooker/CookPackageData.cpp b/Engine/Source/Editor/UnrealEd/Private/Cooker/CookPackageData.cpp
--- a/Engine/Source/Editor/UnrealEd/Private/Cooker/CookPackageData.cpp
+++ b/Engine/Source/Editor/UnrealEd/Private/Cooker/CookPackageData.cpp
@@ -3,6 +3,7 @@
 #include "CookPackageData.h"
 
 #include "Algo/AnyOf.h"
+#include "AssetCompilingManager.h"
 #include "CookOnTheSide/CookOnTheFlyServer.h"
 #include "CookPlatformManager.h"
 #include "Containers/StringView.h"
@@ -1427,6 +1428,7 @@ namespace Cook
 		}
 
 		GShaderCompilingManager->ProcessAsyncResults(true /* bLimitExecutionTime */, false /* bBlockOnGlobalShaderCompletion */);
+		FAssetCompilingManager::Get().ProcessAsyncTasks(true);
 
 		FPendingCookedPlatformData* Datas = PendingCookedPlatformDatas.GetData();
 		for (int Index = 0; Index < PendingCookedPlatformDatas.Num();)
diff --git a/Engine/Source/Runtime/Engine/Private/AsyncCompilationHelpers.cpp b/Engine/Source/Runtime/Engine/Private/AsyncCompilationHelpers.cpp
--- a/Engine/Source/Runtime/Engine/Private/AsyncCompilationHelpers.cpp
+++ b/Engine/Source/Runtime/Engine/Private/AsyncCompilationHelpers.cpp
@@ -72,9 +72,12 @@ void FAsyncCompilationNotification::Update(int32 NumJobs)
 			NotificationItemPtr = NotificationItem;
 		}
 
-		NotificationItem->SetCompletionState(SNotificationItem::CS_Pending);
-		NotificationItem->SetVisibility(EVisibility::HitTestInvisible);
-		NotificationItem->SetText(ProgressMessage);
+		if (NotificationItem.IsValid())
+		{
+			NotificationItem->SetCompletionState(SNotificationItem::CS_Pending);
+			NotificationItem->SetVisibility(EVisibility::HitTestInvisible);
+			NotificationItem->SetText(ProgressMessage);
+		}
 	}
 };
 
