Fix non-unity

diff --git a/Engine/Source/Editor/UnrealEd/Classes/Editor/PropertyEditorTestObject.h b/Engine/Source/Editor/UnrealEd/Classes/Editor/PropertyEditorTestObject.h
--- a/Engine/Source/Editor/UnrealEd/Classes/Editor/PropertyEditorTestObject.h
+++ b/Engine/Source/Editor/UnrealEd/Classes/Editor/PropertyEditorTestObject.h
@@ -12,6 +12,7 @@
 #include "Engine/EngineTypes.h"
 
 #include "Curves/RichCurve.h"
+#include "PerPlatformProperties.h"
 
 #include "PropertyEditorTestObject.generated.h"
 
