Update single hull generation in editor to use the same path as the multi convex hull generation
Benefits are the following:
- Generation will properly use the Max Hull Verts parameter
- Results are of better quality and storage is lower ( as the FConvexElem does not store all the verts associated with the bone, but only up to Max Hull Verts )

diff --git a/Engine/Source/Editor/UnrealEd/Private/PhysicsAssetUtils.cpp b/Engine/Source/Editor/UnrealEd/Private/PhysicsAssetUtils.cpp
--- a/Engine/Source/Editor/UnrealEd/Private/PhysicsAssetUtils.cpp
+++ b/Engine/Source/Editor/UnrealEd/Private/PhysicsAssetUtils.cpp
@@ -349,7 +349,7 @@ bool CreateFromSkeletalMesh(UPhysicsAsset* PhysicsAsset, USkeletalMesh* SkelMesh
 
 	FSkinnedBoneTriangleCache TriangleCache(*SkelMesh, Params);
 
-	if ( Params.GeomType == EFG_MultiConvexHull )
+	if (Params.GeomType == EFG_SingleConvexHull || Params.GeomType == EFG_MultiConvexHull)
 	{
 		TriangleCache.BuildCache();
 	}
@@ -447,7 +447,8 @@ FVector ComputeEigenVector(const FMatrix& A)
 bool CreateCollisionFromBoneInternal(UBodySetup* bs, USkeletalMesh* skelMesh, int32 BoneIndex, const FPhysAssetCreateParams& Params, const FBoneVertInfo& Info, const FSkinnedBoneTriangleCache& TriangleCache)
 {
 #if WITH_EDITOR
-	if (Params.GeomType != EFG_MultiConvexHull)	//multi convex hull can fail so wait to clear it
+	// multi convex hull can fail so wait to clear it ( will be called in DecomposeMeshToHulls() )
+	if (Params.GeomType != EFG_SingleConvexHull && Params.GeomType != EFG_MultiConvexHull)
 	{
 		// Empty any existing collision.
 		bs->RemoveSimpleCollision();
@@ -526,30 +527,7 @@ bool CreateCollisionFromBoneInternal(UBodySetup* bs, USkeletalMesh* skelMesh, in
 
 		bs->AggGeom.SphereElems.Add(SphereElem);
 	}
-	// Deal with creating a single convex hull
-	else if (Params.GeomType == EFG_SingleConvexHull)
-	{
-		if (Info.Positions.Num())
-		{
-			FKConvexElem ConvexElem;
-
-			// Add all of the vertices for this bone to the convex element
-			for (int32 index = 0; index < Info.Positions.Num(); index++)
-			{
-				ConvexElem.VertexData.Add(Info.Positions[index]);
-			}
-			ConvexElem.UpdateElemBox();
-			bs->AggGeom.ConvexElems.Add(ConvexElem);
-		}
-		else
-		{
-			FMessageLog EditorErrors("EditorErrors");
-			EditorErrors.Warning(NSLOCTEXT("PhysicsAssetUtils", "ConvexNoPositions", "Unable to create a convex hull for the given bone as there are no vertices associated with the bone."));
-			EditorErrors.Open();
-			return false;
-		}
-	}
-	else if (Params.GeomType == EFG_MultiConvexHull)
+	else if (Params.GeomType == EFG_SingleConvexHull || Params.GeomType == EFG_MultiConvexHull)
 	{
 		TArray<FVector> Verts;
 		TArray<uint32> Indices;
@@ -557,7 +535,8 @@ bool CreateCollisionFromBoneInternal(UBodySetup* bs, USkeletalMesh* skelMesh, in
 
 		if (Verts.Num())
 		{
-			DecomposeMeshToHulls(bs, Verts, Indices, Params.HullCount, Params.MaxHullVerts);
+			const int32 HullCount = Params.GeomType == EFG_SingleConvexHull ? 1 : Params.HullCount;
+			DecomposeMeshToHulls(bs, Verts, Indices, HullCount, Params.MaxHullVerts);
 		}
 		else
 		{
@@ -596,8 +575,6 @@ bool CreateCollisionFromBoneInternal(UBodySetup* bs, USkeletalMesh* skelMesh, in
 			SphylElem.Length = BoxExtent.Z * 1.01f;
 		}
 
-
-
 		bs->AggGeom.SphylElems.Add(SphylElem);
 	}
 	else if (Params.GeomType == EFG_TaperedCapsule)
@@ -641,7 +618,7 @@ bool CreateCollisionFromBone(UBodySetup* bs, USkeletalMesh* skelMesh, int32 Bone
 
 	FSkinnedBoneTriangleCache TriangleCache(*skelMesh, Params);
 
-	if ( Params.GeomType == EFG_MultiConvexHull )
+	if (Params.GeomType == EFG_SingleConvexHull || Params.GeomType == EFG_MultiConvexHull)
 	{
 		TriangleCache.BuildCache();
 	}
@@ -655,7 +632,7 @@ bool CreateCollisionFromBones(UBodySetup* bs, USkeletalMesh* skelMesh, const TAr
 
 	FSkinnedBoneTriangleCache TriangleCache(*skelMesh, Params);
 
-	if (Params.GeomType == EFG_MultiConvexHull)
+	if (Params.GeomType == EFG_SingleConvexHull || Params.GeomType == EFG_MultiConvexHull)
 	{
 		TriangleCache.BuildCache();
 	}
