Use Typed Element handle to deselect actor before deletion, instead of clearing all selection in the editor.

diff --git a/Engine/Source/Editor/UnrealEd/Private/Subsystems/EditorActorSubsystem.cpp b/Engine/Source/Editor/UnrealEd/Private/Subsystems/EditorActorSubsystem.cpp
--- a/Engine/Source/Editor/UnrealEd/Private/Subsystems/EditorActorSubsystem.cpp
+++ b/Engine/Source/Editor/UnrealEd/Private/Subsystems/EditorActorSubsystem.cpp
@@ -21,6 +21,9 @@
 #include "Engine/Selection.h"
 #include "LevelEditorViewport.h"
 
+#include "Elements/Framework/EngineElementsLibrary.h"
+#include "Elements/Framework/TypedElementHandle.h"
+
 #define LOCTEXT_NAMESPACE "EditorActorUtilities"
 
 namespace InternalActorUtilitiesSubsystemLibrary
@@ -511,10 +514,10 @@ bool UEditorActorSubsystem::DestroyActor(class AActor* ToDestroyActor)
 	}
 
 	//To avoid dangling gizmo after actor has been destroyed
-	if (ToDestroyActor->IsSelected())
-	{
-		GEditor->SelectNone(true, true, false);
-	}
+	FTypedElementHandle ActorHandle = UEngineElementsLibrary::AcquireEditorActorElementHandle(ToDestroyActor);
+	USelection* ActorSelection = GEditor->GetSelectedActors();
+	UTypedElementSelectionSet* SelectionSet = ActorSelection->GetElementSelectionSet();
+	SelectionSet->DeselectElement(ActorHandle, FTypedElementSelectionOptions());
 
 	ULayersSubsystem* Layers = GEditor->GetEditorSubsystem<ULayersSubsystem>();
 	Layers->DisassociateActorFromLayers(ToDestroyActor);
