Created Static Mesh placement factory which handles both actor static mesh placement and instanced static mesh placement.

diff --git a/Engine/Source/Editor/AssetPlacementEdMode/Private/Tools/PlacementPlaceSingleTool.cpp b/Engine/Source/Editor/AssetPlacementEdMode/Private/Tools/PlacementPlaceSingleTool.cpp
--- a/Engine/Source/Editor/AssetPlacementEdMode/Private/Tools/PlacementPlaceSingleTool.cpp
+++ b/Engine/Source/Editor/AssetPlacementEdMode/Private/Tools/PlacementPlaceSingleTool.cpp
@@ -31,11 +31,15 @@ void UPlacementModePlaceSingleTool::OnEndDrag(const FRay& Ray)
 			PlacementInfo.FinalizedTransform = FTransform(LastBrushStamp.HitResult.ImpactPoint);
 			PlacementInfo.PreferredLevel = GEditor->GetEditorWorldContext().World()->GetCurrentLevel();
 
+			FPlacementOptions PlacementOptions;
+			PlacementOptions.bPreferBatchPlacement = true;
+			PlacementOptions.bPreferInstancedPlacement = true;
+
 			UPlacementSubsystem* PlacementSubsystem = GEditor->GetEditorSubsystem<UPlacementSubsystem>();
 			if (PlacementSubsystem)
 			{
 				FScopedTransaction Transaction(NSLOCTEXT("PlacementMode", "SinglePlaceAsset", "Place Single Asset"));
-				PlacementSubsystem->PlaceAsset(PlacementInfo, FPlacementOptions());
+				PlacementSubsystem->PlaceAsset(PlacementInfo, PlacementOptions);
 			}
 		}
 	}
diff --git a/Engine/Source/Editor/UnrealEd/Classes/ActorFactories/ActorFactoryStaticMesh.h b/Engine/Source/Editor/UnrealEd/Classes/ActorFactories/ActorFactoryStaticMesh.h
--- a/Engine/Source/Editor/UnrealEd/Classes/ActorFactories/ActorFactoryStaticMesh.h
+++ b/Engine/Source/Editor/UnrealEd/Classes/ActorFactories/ActorFactoryStaticMesh.h
@@ -11,7 +11,7 @@
 class AActor;
 struct FAssetData;
 
-UCLASS(MinimalAPI, config=Editor)
+UCLASS(MinimalAPI, Abstract, Transient)
 class UActorFactoryStaticMesh : public UActorFactory
 {
 	GENERATED_UCLASS_BODY()
@@ -24,6 +24,3 @@ class UActorFactoryStaticMesh : public UActorFactory
 	virtual FQuat AlignObjectToSurfaceNormal(const FVector& InSurfaceNormal, const FQuat& ActorRotation) const override;
 	//~ End UActorFactory Interface
 };
-
-
-
diff --git a/Engine/Source/Editor/UnrealEd/Private/Factories/EditorStaticMeshFactory.cpp b/Engine/Source/Editor/UnrealEd/Private/Factories/EditorStaticMeshFactory.cpp
new file mode 100644
--- /dev/null
+++ b/Engine/Source/Editor/UnrealEd/Private/Factories/EditorStaticMeshFactory.cpp
@@ -0,0 +1,104 @@
+// Copyright Epic Games, Inc. All Rights Reserved.
+
+#include "Factories/EditorStaticMeshFactory.h"
+
+#include "Elements/Actor/ActorElementData.h"
+#include "Elements/Component/ComponentElementData.h"
+#include "Elements/Framework/EngineElementsLibrary.h"
+
+#include "ActorPartition/ActorPartitionSubsystem.h"
+#include "InstancedFoliageActor.h"
+#include "InstancedFoliage.h"
+#include "FoliageType_InstancedStaticMesh.h"
+
+#include "Components/InstancedStaticMeshComponent.h"
+#include "Engine/StaticMesh.h"
+
+#include "Subsystems/PlacementSubsystem.h"
+
+TArray<FTypedElementHandle> UEditorStaticMeshFactory::PlaceAsset(const FAssetPlacementInfo& InPlacementInfo, const FPlacementOptions& InPlacementOptions)
+{
+	if (!InPlacementOptions.bPreferInstancedPlacement)
+	{
+		return Super::PlaceAsset(InPlacementInfo, InPlacementOptions);
+	}
+
+	if (InPlacementInfo.PreferredLevel.IsValid())
+	{
+		if (UActorPartitionSubsystem* PartitionSubsystem = UWorld::GetSubsystem<UActorPartitionSubsystem>(InPlacementInfo.PreferredLevel->GetWorld()))
+		{
+			// Create or find the foliage partition actor
+			constexpr bool bCreatePartitionActorIfMissing = true;
+			FActorPartitionGetParams PartitionActorFindParams(AInstancedFoliageActor::StaticClass(), bCreatePartitionActorIfMissing, InPlacementInfo.PreferredLevel.Get(), InPlacementInfo.FinalizedTransform.GetLocation());
+			AInstancedFoliageActor* FoliageActor = Cast<AInstancedFoliageActor>(PartitionSubsystem->GetActor(PartitionActorFindParams));
+			check(FoliageActor);
+
+			// Create or find the instanced static mesh foliage type for the given mesh
+			UObject* AssetDataObject = InPlacementInfo.AssetToPlace.GetAsset();
+			UFoliageType_InstancedStaticMesh* FoundOrCreatedType = nullptr;
+			for (auto& FoliageTypePair : FoliageActor->FoliageInfos)
+			{
+				FAssetData FoliageTypeSource(FoliageTypePair.Key->GetSource());
+				if (FoliageTypeSource == InPlacementInfo.AssetToPlace)
+				{
+					FoundOrCreatedType = Cast<UFoliageType_InstancedStaticMesh>(FoliageTypePair.Key);
+					break;
+				}
+			}
+			if (!FoundOrCreatedType)
+			{
+				FoundOrCreatedType = NewObject<UFoliageType_InstancedStaticMesh>(FoliageActor);
+				FoundOrCreatedType->SetSource(InPlacementInfo.AssetToPlace.GetAsset());
+			}
+			FFoliageInfo* FoliageInfo = FoliageActor->FindOrAddMesh(FoundOrCreatedType);
+			check(FoliageInfo);
+
+			FDesiredFoliageInstance InstancePlaceInfo;
+			FFoliageInstance FoliagePlacementInfo;
+			FoliagePlacementInfo.Location = InPlacementInfo.FinalizedTransform.GetLocation();
+			FoliagePlacementInfo.Rotation = InPlacementInfo.FinalizedTransform.Rotator();
+			FoliagePlacementInfo.DrawScale3D = InPlacementInfo.FinalizedTransform.GetScale3D();
+			FoliageInfo->AddInstance(FoliageActor, FoundOrCreatedType, FoliagePlacementInfo);
+
+			// Todo: deal with returning instanced handles
+			// In the meantime, we return the handle of the the HISM component added.
+			return TArray<FTypedElementHandle>({ UEngineElementsLibrary::AcquireEditorComponentElementHandle(FoliageInfo->GetComponent()) });
+		}
+	}
+
+	return TArray<FTypedElementHandle>();
+}
+
+FAssetData UEditorStaticMeshFactory::GetAssetDataFromElementHandle(const FTypedElementHandle& InHandle)
+{
+	UInstancedStaticMeshComponent* ISMComponent = nullptr;
+
+	// Try to pull from component handle
+	const FComponentElementData* ComponentData = InHandle.GetData<FComponentElementData>();
+	if (ComponentData && ComponentData->Component)
+	{
+		ISMComponent = Cast<UInstancedStaticMeshComponent>(ComponentData->Component);
+	}
+	else
+	{
+		// Try to pull from actor handle
+		const FActorElementData* ActorData = InHandle.GetData<FActorElementData>();
+		if (ActorData && ActorData->Actor)
+		{
+			ISMComponent = ActorData->Actor->FindComponentByClass<UInstancedStaticMeshComponent>();
+		}
+	}
+
+	if (ISMComponent)
+	{
+		FAssetData ISMStaticMeshAssetData = FAssetData(ISMComponent->GetStaticMesh());
+		if (CanPlaceElementsFromAssetData(ISMStaticMeshAssetData))
+		{
+			return ISMStaticMeshAssetData;
+		}
+	}
+
+	// Todo: deal with instanced handles
+
+	return Super::GetAssetDataFromElementHandle(InHandle);
+}
diff --git a/Engine/Source/Editor/UnrealEd/Private/Factories/EditorStaticMeshFactory.h b/Engine/Source/Editor/UnrealEd/Private/Factories/EditorStaticMeshFactory.h
new file mode 100644
--- /dev/null
+++ b/Engine/Source/Editor/UnrealEd/Private/Factories/EditorStaticMeshFactory.h
@@ -0,0 +1,18 @@
+// Copyright Epic Games, Inc. All Rights Reserved.
+
+#include "Factories/AssetFactoryInterface.h"
+#include "ActorFactories/ActorFactoryStaticMesh.h"
+
+#include "EditorStaticMeshFactory.generated.h"
+
+UCLASS(Transient)
+class UEditorStaticMeshFactory : public UActorFactoryStaticMesh
+{
+	GENERATED_BODY()
+
+public:
+	// Begin IAssetFactoryInterface
+	virtual TArray<FTypedElementHandle> PlaceAsset(const FAssetPlacementInfo& InPlacementInfo, const FPlacementOptions& InPlacementOptions) override;
+	virtual FAssetData GetAssetDataFromElementHandle(const FTypedElementHandle& InHandle) override;
+	// End IAssetFactoryInterface
+};
\ No newline at end of file
