Animation data MVC refactor

diff --git a/Engine/Source/Editor/Persona/Private/AnimTimeline/AnimTimelineTrack_FloatCurve.cpp b/Engine/Source/Editor/Persona/Private/AnimTimeline/AnimTimelineTrack_FloatCurve.cpp
--- a/Engine/Source/Editor/Persona/Private/AnimTimeline/AnimTimelineTrack_FloatCurve.cpp
+++ b/Engine/Source/Editor/Persona/Private/AnimTimeline/AnimTimelineTrack_FloatCurve.cpp
@@ -20,15 +20,17 @@
 #include "Widgets/Layout/SBox.h"
 #include "AnimModel_AnimSequenceBase.h"
 #include "SAnimOutlinerItem.h"
+#include "Animation/AnimData/AnimDataController.h"
 
 #define LOCTEXT_NAMESPACE "FAnimTimelineTrack_FloatCurve"
 
 ANIMTIMELINE_IMPLEMENT_TRACK(FAnimTimelineTrack_FloatCurve);
 
-FAnimTimelineTrack_FloatCurve::FAnimTimelineTrack_FloatCurve(FFloatCurve& InCurve, const TSharedRef<FAnimModel>& InModel)
-	: FAnimTimelineTrack_Curve(InCurve.FloatCurve, InCurve.Name, 0, ERawCurveTrackTypes::RCT_Float, FText::FromName(InCurve.Name.DisplayName), FText::FromName(InCurve.Name.DisplayName), InCurve.GetColor(), InCurve.GetColor(), InModel)
+FAnimTimelineTrack_FloatCurve::FAnimTimelineTrack_FloatCurve(const FFloatCurve* InCurve, const TSharedRef<FAnimModel>& InModel)
+	: FAnimTimelineTrack_Curve(&InCurve->FloatCurve, InCurve->Name, 0, ERawCurveTrackTypes::RCT_Float, FText::FromName(InCurve->Name.DisplayName), FText::FromName(InCurve->Name.DisplayName), InCurve->GetColor(), InCurve->GetColor(), InModel)
 	, FloatCurve(InCurve)
-	, CurveName(InCurve.Name)
+	, CurveName(InCurve->Name)
+	, CurveId(FAnimationCurveIdentifier(InCurve->Name, ERawCurveTrackTypes::RCT_Float))
 {
 	SetHeight(32.0f);
 }
@@ -48,14 +50,14 @@ TSharedRef<SWidget> FAnimTimelineTrack_FloatCurve::MakeTimelineWidgetContainer()
 		}
 		else
 		{
-			return FloatCurve.GetCurveTypeFlag(AACF_Metadata) ? FloatCurve.GetColor().Desaturate(0.25f) : FloatCurve.GetColor().Desaturate(0.75f); 
+			return FloatCurve->GetCurveTypeFlag(AACF_Metadata) ? FloatCurve->GetColor().Desaturate(0.25f) : FloatCurve->GetColor().Desaturate(0.75f);
 		}
 	};
 
 	return
 		SAssignNew(TimelineWidgetContainer, SBorder)
 		.Padding(0.0f)
-		.BorderImage_Lambda([this](){ return FloatCurve.GetCurveTypeFlag(AACF_Metadata) ? FEditorStyle::GetBrush("Sequencer.Section.SelectedSectionOverlay") : FEditorStyle::GetBrush("AnimTimeline.Outliner.DefaultBorder"); })
+		.BorderImage_Lambda([this](){ return FloatCurve->GetCurveTypeFlag(AACF_Metadata) ? FEditorStyle::GetBrush("Sequencer.Section.SelectedSectionOverlay") : FEditorStyle::GetBrush("AnimTimeline.Outliner.DefaultBorder"); })
 		.BorderBackgroundColor_Lambda(ColorLambda)
 		[
 			CurveWidget
@@ -99,7 +101,7 @@ TSharedRef<SWidget> FAnimTimelineTrack_FloatCurve::BuildCurveTrackMenu()
 {
 	FMenuBuilder MenuBuilder(true, GetModel()->GetCommandList());
 
-	bool bIsMetadata = FloatCurve.GetCurveTypeFlag(AACF_Metadata);
+	bool bIsMetadata = FloatCurve->GetCurveTypeFlag(AACF_Metadata);
 
 	MenuBuilder.BeginSection("Curve", bIsMetadata ? LOCTEXT("CurveMetadataMenuSection", "Curve Metadata") : LOCTEXT("CurveMenuSection", "Curve"));
 	{
@@ -154,13 +156,10 @@ void FAnimTimelineTrack_FloatCurve::ConvertCurveToMetaData()
 	IAnimationEditor::FCurveEditInfo EditInfo(CurveName, ERawCurveTrackTypes::RCT_Float, 0);
 	StaticCastSharedRef<FAnimModel_AnimSequenceBase>(GetModel())->OnStopEditingCurves.ExecuteIfBound(TArray<IAnimationEditor::FCurveEditInfo>({ EditInfo }));
 
-	FScopedTransaction Transaction(LOCTEXT("CurvePanel_ConvertCurveToMetaData", "Convert curve to metadata"));
-	AnimSequenceBase->Modify(true);
-	FloatCurve.SetCurveTypeFlag(AACF_Metadata, true);
-
-	// We're moving to a metadata curve, we need to clear out the keys.
-	FloatCurve.FloatCurve.Reset();
-	FloatCurve.FloatCurve.AddKey(0.0f, 1.0f);
+	UAnimDataController* Controller = AnimSequenceBase->GetController();
+	UAnimDataController::FScopedBracket ScopedBracket(Controller, LOCTEXT("ConvertCurveToMetaData_Bracket", "Converting curve to metadata"));
+	Controller->SetCurveFlag(CurveId, AACF_Metadata, true);
+	Controller->SetCurveKeys(CurveId, { FRichCurveKey(0.f, 1.f) });	
 
 	ZoomToFit();
 }
@@ -169,32 +168,31 @@ void FAnimTimelineTrack_FloatCurve::ConvertMetaDataToCurve()
 {
 	UAnimSequenceBase* AnimSequenceBase = GetModel()->GetAnimSequenceBase();
 
-	FScopedTransaction Transaction(LOCTEXT("CurvePanel_ConvertMetaDataToCurve", "Convert metadata to curve"));
-	AnimSequenceBase->Modify(true);
-	FloatCurve.SetCurveTypeFlag(AACF_Metadata, false);
+	UAnimDataController* Controller = AnimSequenceBase->GetController();
+	UAnimDataController::FScopedBracket ScopedBracket(Controller, LOCTEXT("CurvePanel_ConvertMetaDataToCurve", "Convert metadata to curve"));
+	Controller->SetCurveFlag(CurveId, AACF_Metadata, false);	
 }
 
 void FAnimTimelineTrack_FloatCurve::RemoveCurve()
 {
 	UAnimSequenceBase* AnimSequenceBase = GetModel()->GetAnimSequenceBase();
 
-	FScopedTransaction Transaction(LOCTEXT("CurvePanel_RemoveCurve", "Remove Curve"));
-	
-	if(AnimSequenceBase->RawCurveData.GetCurveData(FloatCurve.Name.UID))
+	if(AnimSequenceBase->GetDataModel()->FindFloatCurve(FAnimationCurveIdentifier(FloatCurve->Name, ERawCurveTrackTypes::RCT_Float)))
 	{
 		FSmartName TrackName;
-		if (AnimSequenceBase->GetSkeleton()->GetSmartNameByUID(USkeleton::AnimCurveMappingName, FloatCurve.Name.UID, TrackName))
+		if (AnimSequenceBase->GetSkeleton()->GetSmartNameByUID(USkeleton::AnimCurveMappingName, FloatCurve->Name.UID, TrackName))
 		{
 			// Stop editing this curve in the external editor window
 			IAnimationEditor::FCurveEditInfo EditInfo(CurveName, ERawCurveTrackTypes::RCT_Float, 0);
 			StaticCastSharedRef<FAnimModel_AnimSequenceBase>(GetModel())->OnStopEditingCurves.ExecuteIfBound(TArray<IAnimationEditor::FCurveEditInfo>({ EditInfo }));
 
 			AnimSequenceBase->Modify(true);
-			AnimSequenceBase->RawCurveData.DeleteCurveData(TrackName);
-			AnimSequenceBase->MarkRawDataAsModified();
+
+			UAnimDataController* Controller = AnimSequenceBase->GetController();
+			Controller->RemoveCurve(CurveId);
+
+			REFACTOR_V2_WARNING("V2 what is the reasoning behind this call?");
 			AnimSequenceBase->PostEditChange();
-			
-			GetModel()->RefreshTracks();
 		}
 	}
 }
@@ -218,18 +216,16 @@ void FAnimTimelineTrack_FloatCurve::OnCommitCurveName(const FText& InText, EText
 
 			const FSmartNameMapping* NameMapping = Skeleton->GetSmartNameContainer(USkeleton::AnimCurveMappingName);
 
-			FScopedTransaction Transaction(LOCTEXT("CurveEditor_RenameCurve", "Rename Curve"));
-
-			AnimSequenceBase->Modify();
-
 			FSmartName NewSmartName;
 			if (NameMapping->FindSmartName(RequestedName, NewSmartName))
 			{
 				// Already in use in this sequence, and if it's not my UID
-				if (NewSmartName.UID != FloatCurve.Name.UID && AnimSequenceBase->RawCurveData.GetCurveData(NewSmartName.UID) != nullptr)
+				const TArray<FFloatCurve>& FloatCurves = AnimSequenceBase->GetDataModel()->GetFloatCurves();
+				if (NewSmartName.UID != FloatCurve->Name.UID && !FloatCurves.ContainsByPredicate([NewSmartName](FFloatCurve& Curve)
+				{
+					return Curve.Name == NewSmartName;
+				}))
 				{
-					Transaction.Cancel();
-
 					FFormatNamedArguments Args;
 					Args.Add(TEXT("InvalidName"), FText::FromName(RequestedName));
 					FNotificationInfo Info(FText::Format(LOCTEXT("AnimCurveRenamedInUse", "The name \"{InvalidName}\" is already used."), Args));
@@ -249,7 +245,6 @@ void FAnimTimelineTrack_FloatCurve::OnCommitCurveName(const FText& InText, EText
 			{
 				if(!Skeleton->AddSmartNameAndModify(USkeleton::AnimCurveMappingName, RequestedName, NewSmartName))
 				{
-					Transaction.Cancel();
 					FNotificationInfo Info(LOCTEXT("AnimCurveRenamedError", "Failed to rename curve smart name, check the log for errors."));
 
 					Info.bUseLargeFont = false;
@@ -264,18 +259,19 @@ void FAnimTimelineTrack_FloatCurve::OnCommitCurveName(const FText& InText, EText
 				}
 			}
 
-			FloatCurve.Name.UID = NewSmartName.UID;
-			FloatCurve.Name.DisplayName = NewSmartName.DisplayName;
-
-			CurveName = FloatCurve.Name;
-			FullCurveName = FText::FromName(FloatCurve.Name.DisplayName);
+			UAnimDataController* Controller = AnimSequenceBase->GetController();
+			UAnimDataController::FScopedBracket ScopedBracket(Controller, LOCTEXT("CurveEditor_RenameCurve", "Renaming Curve"));
+          
+			FAnimationCurveIdentifier NewCurveId(NewSmartName, ERawCurveTrackTypes::RCT_Float);
+			Controller->RenameCurve(CurveId, NewCurveId);
+			Controller->RemoveBoneTracksMissingFromSkeleton(AnimSequenceBase->GetSkeleton());           
 		}
 	}
 }
 
 FText FAnimTimelineTrack_FloatCurve::GetLabel() const
 {
-	return FAnimTimelineTrack_FloatCurve::GetFloatCurveName(GetModel(), FloatCurve.Name);
+	return FAnimTimelineTrack_FloatCurve::GetFloatCurveName(GetModel(), FloatCurve->Name);
 }
 
 FText FAnimTimelineTrack_FloatCurve::GetFloatCurveName(const TSharedRef<FAnimModel>& InModel, const FSmartName& InSmartName)
@@ -295,7 +291,7 @@ FText FAnimTimelineTrack_FloatCurve::GetFloatCurveName(const TSharedRef<FAnimMod
 
 bool FAnimTimelineTrack_FloatCurve::CanEditCurve(int32 InCurveIndex) const
 {
-	return !FloatCurve.GetCurveTypeFlag(AACF_Metadata);
+	return !FloatCurve->GetCurveTypeFlag(AACF_Metadata);
 }
 
 void FAnimTimelineTrack_FloatCurve::RequestRename()
@@ -319,17 +315,13 @@ void FAnimTimelineTrack_FloatCurve::AddCurveTrackButton(TSharedPtr<SHorizontalBo
 
 	auto GetValue = [this]()
 	{
-		return FloatCurve.Color;
+		return FloatCurve->Color;
 	};
 
 	auto SetValue = [this](FLinearColor InNewColor)
 	{
 		UAnimSequenceBase* AnimSequenceBase = GetModel()->GetAnimSequenceBase();
-		FScopedTransaction Transaction(LOCTEXT("SetCurveColor", "Set Curve Color"));
-		AnimSequenceBase->Modify(true);
-		FloatCurve.Color = InNewColor;
-
-		Color = InNewColor;
+		AnimSequenceBase->GetController()->SetCurveColor(CurveId, InNewColor);
 
 		// Set display curves too
 		for(const TPair<FCurveModelID, TUniquePtr<FCurveModel>>& CurvePair : CurveEditor->GetCurves())
@@ -338,14 +330,24 @@ void FAnimTimelineTrack_FloatCurve::AddCurveTrackButton(TSharedPtr<SHorizontalBo
 		}
 	};
 
-	auto OnGetMenuContent = [GetValue, SetValue]()
+	auto OnGetMenuContent = [this, GetValue, SetValue]()
 	{
 		// Open a color picker
 		return SNew(SColorPicker)
 			.TargetColorAttribute_Lambda(GetValue)
 			.UseAlpha(false)
 			.DisplayInlineVersion(true)
-			.OnColorCommitted_Lambda(SetValue);
+			.OnColorCommitted_Lambda(SetValue)
+			.OnInteractivePickBegin_Lambda([this]()
+			{
+				UAnimSequenceBase* AnimSequenceBase = GetModel()->GetAnimSequenceBase();
+				AnimSequenceBase->GetController()->OpenBracket(LOCTEXT("EditCurveColor_Bracket", "Editing Curve Color"));				
+			})
+			.OnInteractivePickEnd_Lambda([this]()
+			{
+				UAnimSequenceBase* AnimSequenceBase = GetModel()->GetAnimSequenceBase();
+				AnimSequenceBase->GetController()->CloseBracket();
+			});
 	};
 
 	InnerHorizontalBox->AddSlot()
@@ -375,7 +377,7 @@ void FAnimTimelineTrack_FloatCurve::AddCurveTrackButton(TSharedPtr<SHorizontalBo
 
 FLinearColor FAnimTimelineTrack_FloatCurve::GetCurveColor(int32 InCurveIndex) const
 { 
-	return FloatCurve.Color; 
+	return FloatCurve->Color; 
 }
 
 void FAnimTimelineTrack_FloatCurve::GetCurveEditInfo(int32 InCurveIndex, FSmartName& OutName, ERawCurveTrackTypes& OutType, int32& OutCurveIndex) const
diff --git a/Engine/Source/Editor/Persona/Private/AnimTimeline/AnimTimelineTrack_FloatCurve.h b/Engine/Source/Editor/Persona/Private/AnimTimeline/AnimTimelineTrack_FloatCurve.h
--- a/Engine/Source/Editor/Persona/Private/AnimTimeline/AnimTimelineTrack_FloatCurve.h
+++ b/Engine/Source/Editor/Persona/Private/AnimTimeline/AnimTimelineTrack_FloatCurve.h
@@ -13,7 +13,7 @@ class FAnimTimelineTrack_FloatCurve : public FAnimTimelineTrack_Curve
 	ANIMTIMELINE_DECLARE_TRACK(FAnimTimelineTrack_FloatCurve, FAnimTimelineTrack_Curve);
 
 public:
-	FAnimTimelineTrack_FloatCurve(FFloatCurve& InCurve, const TSharedRef<FAnimModel>& InModel);
+	FAnimTimelineTrack_FloatCurve(const FFloatCurve* InCurve, const TSharedRef<FAnimModel>& InModel);
 
 	/** FAnimTimelineTrack_Curve interface */
 	virtual TSharedRef<SWidget> MakeTimelineWidgetContainer() override;
@@ -28,7 +28,7 @@ public:
 	virtual void GetCurveEditInfo(int32 InCurveIndex, FSmartName& OutName, ERawCurveTrackTypes& OutType, int32& OutCurveIndex) const override;
 
 	/** Access the curve we are editing */
-	FFloatCurve& GetFloatCurve() { return FloatCurve; }
+	const FFloatCurve* GetFloatCurve() { return FloatCurve; }
 
 	/** Helper function used to get a smart name for a curve */
 	static FText GetFloatCurveName(const TSharedRef<FAnimModel>& InModel, const FSmartName& InSmartName);
@@ -44,10 +44,11 @@ private:
 
 private:
 	/** The curve we are editing */
-	FFloatCurve& FloatCurve;
+	const FFloatCurve* FloatCurve;
 
-	/** The curve name */
+	/** The curve name and identifier */
 	FSmartName CurveName;
+	FAnimationCurveIdentifier CurveId;
 
 	/** Label we can edit */
 	TSharedPtr<SInlineEditableTextBlock> EditableTextLabel;
diff --git a/Engine/Source/Editor/Persona/Private/AnimTimeline/AnimTimelineTrack_TransformCurve.cpp b/Engine/Source/Editor/Persona/Private/AnimTimeline/AnimTimelineTrack_TransformCurve.cpp
--- a/Engine/Source/Editor/Persona/Private/AnimTimeline/AnimTimelineTrack_TransformCurve.cpp
+++ b/Engine/Source/Editor/Persona/Private/AnimTimeline/AnimTimelineTrack_TransformCurve.cpp
@@ -12,25 +12,27 @@
 #include "Animation/DebugSkelMeshComponent.h"
 #include "AnimPreviewInstance.h"
 #include "AnimModel_AnimSequenceBase.h"
+#include "Animation/AnimData/AnimDataController.h"
 
 #define LOCTEXT_NAMESPACE "FAnimTimelineTrack_TransformCurve"
 
 ANIMTIMELINE_IMPLEMENT_TRACK(FAnimTimelineTrack_TransformCurve);
 
-FAnimTimelineTrack_TransformCurve::FAnimTimelineTrack_TransformCurve(FTransformCurve& InCurve, const TSharedRef<FAnimModel>& InModel)
-	: FAnimTimelineTrack_Curve(FAnimTimelineTrack_TransformCurve::GetTransformCurveName(InModel, InCurve.Name), FAnimTimelineTrack_TransformCurve::GetTransformCurveName(InModel, InCurve.Name), InCurve.GetColor(), InCurve.GetColor(), InModel)
+FAnimTimelineTrack_TransformCurve::FAnimTimelineTrack_TransformCurve(const FTransformCurve* InCurve, const TSharedRef<FAnimModel>& InModel)
+	: FAnimTimelineTrack_Curve(FAnimTimelineTrack_TransformCurve::GetTransformCurveName(InModel, InCurve->Name), FAnimTimelineTrack_TransformCurve::GetTransformCurveName(InModel, InCurve->Name), InCurve->GetColor(), InCurve->GetColor(), InModel)
 	, TransformCurve(InCurve)
-	, CurveName(InCurve.Name)
+	, CurveName(InCurve->Name)
+	, CurveId(InCurve->Name, ERawCurveTrackTypes::RCT_Transform)
 {
-	Curves.Add(&InCurve.TranslationCurve.FloatCurves[0]);
-	Curves.Add(&InCurve.TranslationCurve.FloatCurves[1]);
-	Curves.Add(&InCurve.TranslationCurve.FloatCurves[2]);
-	Curves.Add(&InCurve.RotationCurve.FloatCurves[0]);
-	Curves.Add(&InCurve.RotationCurve.FloatCurves[1]);
-	Curves.Add(&InCurve.RotationCurve.FloatCurves[2]);
-	Curves.Add(&InCurve.ScaleCurve.FloatCurves[0]);
-	Curves.Add(&InCurve.ScaleCurve.FloatCurves[1]);
-	Curves.Add(&InCurve.ScaleCurve.FloatCurves[2]);
+	Curves.Add(&InCurve->TranslationCurve.FloatCurves[0]);
+	Curves.Add(&InCurve->TranslationCurve.FloatCurves[1]);
+	Curves.Add(&InCurve->TranslationCurve.FloatCurves[2]);
+	Curves.Add(&InCurve->RotationCurve.FloatCurves[0]);
+	Curves.Add(&InCurve->RotationCurve.FloatCurves[1]);
+	Curves.Add(&InCurve->RotationCurve.FloatCurves[2]);
+	Curves.Add(&InCurve->ScaleCurve.FloatCurves[0]);
+	Curves.Add(&InCurve->ScaleCurve.FloatCurves[1]);
+	Curves.Add(&InCurve->ScaleCurve.FloatCurves[2]);
 }
 
 FLinearColor FAnimTimelineTrack_TransformCurve::GetCurveColor(int32 InCurveIndex) const
@@ -120,11 +122,11 @@ void FAnimTimelineTrack_TransformCurve::DeleteTrack()
 	UAnimSequenceBase* AnimSequenceBase = GetModel()->GetAnimSequenceBase();
 	TSharedRef<FAnimModel_AnimSequenceBase> BaseModel = StaticCastSharedRef<FAnimModel_AnimSequenceBase>(GetModel());
 
-	if(AnimSequenceBase->RawCurveData.GetCurveData(TransformCurve.Name.UID, ERawCurveTrackTypes::RCT_Transform))
+	if(AnimSequenceBase->GetDataModel()->FindTransformCurve(CurveId))
 	{
 		const FScopedTransaction Transaction(LOCTEXT("AnimCurve_DeleteTrack", "Delete Curve"));
 		FSmartName CurveToDelete;
-		if (AnimSequenceBase->GetSkeleton()->GetSmartNameByUID(USkeleton::AnimTrackCurveMappingName, TransformCurve.Name.UID, CurveToDelete))
+		if (AnimSequenceBase->GetSkeleton()->GetSmartNameByUID(USkeleton::AnimTrackCurveMappingName, TransformCurve->Name.UID, CurveToDelete))
 		{
 			// Stop editing these curves in the external editor window
 			TArray<IAnimationEditor::FCurveEditInfo> CurveEditInfo;
@@ -140,12 +142,9 @@ void FAnimTimelineTrack_TransformCurve::DeleteTrack()
 
 			BaseModel->OnStopEditingCurves.ExecuteIfBound(CurveEditInfo);
 
-			AnimSequenceBase->Modify();
-			AnimSequenceBase->RawCurveData.DeleteCurveData(CurveToDelete, ERawCurveTrackTypes::RCT_Transform);
 
-			AnimSequenceBase->MarkRawDataAsModified();
-
-			GetModel()->RefreshTracks();
+			UAnimDataController* Controller = AnimSequenceBase->GetController();
+			Controller->RemoveCurve(CurveId);
 
 			if (GetModel()->GetPreviewScene()->GetPreviewMeshComponent()->PreviewInstance != nullptr)
 			{
@@ -158,7 +157,7 @@ void FAnimTimelineTrack_TransformCurve::DeleteTrack()
 bool FAnimTimelineTrack_TransformCurve::IsEnabled() const
 {
 	UAnimSequenceBase* AnimSequenceBase = GetModel()->GetAnimSequenceBase();
-	FAnimCurveBase* Curve = AnimSequenceBase->RawCurveData.GetCurveData(TransformCurve.Name.UID, ERawCurveTrackTypes::RCT_Transform);
+	const FAnimCurveBase* Curve = AnimSequenceBase->GetDataModel()->FindTransformCurve(CurveId);
 	return Curve && !Curve->GetCurveTypeFlag(AACF_Disabled);
 }
 
@@ -166,29 +165,19 @@ void FAnimTimelineTrack_TransformCurve::ToggleEnabled()
 {
 	UAnimSequenceBase* AnimSequenceBase = GetModel()->GetAnimSequenceBase();
 
-	FAnimCurveBase* Curve = AnimSequenceBase->RawCurveData.GetCurveData(TransformCurve.Name.UID, ERawCurveTrackTypes::RCT_Transform);
-	if (Curve)
-	{
-		bool bEnabled = !Curve->GetCurveTypeFlag(AACF_Disabled);
-
-		const FScopedTransaction Transaction(bEnabled ? LOCTEXT("AnimCurve_DisableTrack", "Disable track") : LOCTEXT("AnimCurve_EnableTrack", "Enable track"));
-		AnimSequenceBase->Modify();
-
-		Curve->SetCurveTypeFlag(AACF_Disabled, bEnabled);
+	UAnimDataController* Controller = AnimSequenceBase->GetController();
+	Controller->SetCurveFlag(CurveId, AACF_Disabled, !IsEnabled());
 
-		AnimSequenceBase->MarkRawDataAsModified();
-
-		// need to update curves, otherwise they're not disabled
-		if (GetModel()->GetPreviewScene()->GetPreviewMeshComponent()->PreviewInstance != nullptr)
-		{
-			GetModel()->GetPreviewScene()->GetPreviewMeshComponent()->PreviewInstance->RefreshCurveBoneControllers();
-		}
+	// need to update curves, otherwise they're not disabled
+	if (GetModel()->GetPreviewScene()->GetPreviewMeshComponent()->PreviewInstance != nullptr)
+	{
+		GetModel()->GetPreviewScene()->GetPreviewMeshComponent()->PreviewInstance->RefreshCurveBoneControllers();
 	}
 }
 
 void FAnimTimelineTrack_TransformCurve::GetCurveEditInfo(int32 InCurveIndex, FSmartName& OutName, ERawCurveTrackTypes& OutType, int32& OutCurveIndex) const
 {
-	OutName = TransformCurve.Name;
+	OutName = TransformCurve->Name;
 	OutType = ERawCurveTrackTypes::RCT_Transform;
 	OutCurveIndex = InCurveIndex;
 }
diff --git a/Engine/Source/Editor/Persona/Private/AnimTimeline/AnimTimelineTrack_TransformCurve.h b/Engine/Source/Editor/Persona/Private/AnimTimeline/AnimTimelineTrack_TransformCurve.h
--- a/Engine/Source/Editor/Persona/Private/AnimTimeline/AnimTimelineTrack_TransformCurve.h
+++ b/Engine/Source/Editor/Persona/Private/AnimTimeline/AnimTimelineTrack_TransformCurve.h
@@ -13,7 +13,7 @@ class FAnimTimelineTrack_TransformCurve : public FAnimTimelineTrack_Curve
 	ANIMTIMELINE_DECLARE_TRACK(FAnimTimelineTrack_TransformCurve, FAnimTimelineTrack_Curve);
 
 public:
-	FAnimTimelineTrack_TransformCurve(FTransformCurve& InCurve, const TSharedRef<FAnimModel>& InModel);
+	FAnimTimelineTrack_TransformCurve(const FTransformCurve* InCurve, const TSharedRef<FAnimModel>& InModel);
 
 	/** FAnimTimelineTrack_Curve interface */
 	virtual FLinearColor GetCurveColor(int32 InCurveIndex) const override;
@@ -23,7 +23,7 @@ public:
 	virtual void GetCurveEditInfo(int32 InCurveIndex, FSmartName& OutName, ERawCurveTrackTypes& OutType, int32& OutCurveIndex) const override;
 
 	/** Access the curve we are editing */
-	FTransformCurve& GetTransformCurve() { return TransformCurve; }
+	const FTransformCurve& GetTransformCurve() { return *TransformCurve; }
 
 	/** Helper function used to get a smart name for a curve */
 	static FText GetTransformCurveName(const TSharedRef<FAnimModel>& InModel, const FSmartName& InSmartName);
@@ -43,8 +43,9 @@ private:
 
 private:
 	/** The curve we are editing */
-	FTransformCurve& TransformCurve;
+	const FTransformCurve* TransformCurve;
 
-	/** The curve name */
+	/** The curve name and identifier */
 	FSmartName CurveName;
+	FAnimationCurveIdentifier CurveId;
 };
\ No newline at end of file
diff --git a/Engine/Source/Editor/Persona/Private/AnimTimeline/AnimTimelineTrack_VectorCurve.cpp b/Engine/Source/Editor/Persona/Private/AnimTimeline/AnimTimelineTrack_VectorCurve.cpp
--- a/Engine/Source/Editor/Persona/Private/AnimTimeline/AnimTimelineTrack_VectorCurve.cpp
+++ b/Engine/Source/Editor/Persona/Private/AnimTimeline/AnimTimelineTrack_VectorCurve.cpp
@@ -9,7 +9,7 @@
 
 ANIMTIMELINE_IMPLEMENT_TRACK(FAnimTimelineTrack_VectorCurve);
 
-FAnimTimelineTrack_VectorCurve::FAnimTimelineTrack_VectorCurve(FVectorCurve& InCurve, const FSmartName& InName, int32 InCurveIndex, ERawCurveTrackTypes InType, const FText& InCurveName, const FText& InFullCurveName, const FLinearColor& InColor, const TSharedRef<FAnimModel>& InModel)
+FAnimTimelineTrack_VectorCurve::FAnimTimelineTrack_VectorCurve(const FVectorCurve* InCurve, const FSmartName& InName, int32 InCurveIndex, ERawCurveTrackTypes InType, const FText& InCurveName, const FText& InFullCurveName, const FLinearColor& InColor, const TSharedRef<FAnimModel>& InModel)
 	: FAnimTimelineTrack_Curve(InCurveName, InFullCurveName, InColor, InColor, InModel)
 	, VectorCurve(InCurve)
 {
@@ -17,9 +17,9 @@ FAnimTimelineTrack_VectorCurve::FAnimTimelineTrack_VectorCurve(FVectorCurve& InC
 	OuterCurveIndex = InCurveIndex;
 	OuterType = InType;
 
-	Curves.Add(&InCurve.FloatCurves[0]);
-	Curves.Add(&InCurve.FloatCurves[1]);
-	Curves.Add(&InCurve.FloatCurves[2]);
+	Curves.Add(&InCurve->FloatCurves[0]);
+	Curves.Add(&InCurve->FloatCurves[1]);
+	Curves.Add(&InCurve->FloatCurves[2]);
 }
 
 FLinearColor FAnimTimelineTrack_VectorCurve::GetCurveColor(int32 InCurveIndex) const
diff --git a/Engine/Source/Editor/Persona/Private/AnimTimeline/AnimTimelineTrack_VectorCurve.h b/Engine/Source/Editor/Persona/Private/AnimTimeline/AnimTimelineTrack_VectorCurve.h
--- a/Engine/Source/Editor/Persona/Private/AnimTimeline/AnimTimelineTrack_VectorCurve.h
+++ b/Engine/Source/Editor/Persona/Private/AnimTimeline/AnimTimelineTrack_VectorCurve.h
@@ -13,7 +13,7 @@ class FAnimTimelineTrack_VectorCurve : public FAnimTimelineTrack_Curve
 	ANIMTIMELINE_DECLARE_TRACK(FAnimTimelineTrack_VectorCurve, FAnimTimelineTrack_Curve);
 
 public:
-	FAnimTimelineTrack_VectorCurve(FVectorCurve& InCurve, const FSmartName& InName, int32 InCurveIndex, ERawCurveTrackTypes InType, const FText& InCurveName, const FText& InFullCurveName, const FLinearColor& InColor, const TSharedRef<FAnimModel>& InModel);
+	FAnimTimelineTrack_VectorCurve(const FVectorCurve* InCurve, const FSmartName& InName, int32 InCurveIndex, ERawCurveTrackTypes InType, const FText& InCurveName, const FText& InFullCurveName, const FLinearColor& InColor, const TSharedRef<FAnimModel>& InModel);
 
 	/** FAnimTimelineTrack_Curve interface */
 	virtual FLinearColor GetCurveColor(int32 InCurveIndex) const override;
@@ -22,9 +22,9 @@ public:
 	virtual void GetCurveEditInfo(int32 InCurveIndex, FSmartName& OutName, ERawCurveTrackTypes& OutType, int32& OutCurveIndex) const override;
 
 	/** Access the curve we are editing */
-	FVectorCurve& GetVectorCurve() { return VectorCurve; }
+	const FVectorCurve& GetVectorCurve() { return *VectorCurve; }
 
 private:
 	/** The curve we are editing */
-	FVectorCurve& VectorCurve;
+	const FVectorCurve* VectorCurve;
 };
\ No newline at end of file
diff --git a/Engine/Source/Editor/Persona/Private/AnimTimeline/SAnimTimeline.cpp b/Engine/Source/Editor/Persona/Private/AnimTimeline/SAnimTimeline.cpp
--- a/Engine/Source/Editor/Persona/Private/AnimTimeline/SAnimTimeline.cpp
+++ b/Engine/Source/Editor/Persona/Private/AnimTimeline/SAnimTimeline.cpp
@@ -33,6 +33,9 @@
 #include "Widgets/Input/SSpinBox.h"
 #include "SAnimTimelineTransportControls.h"
 #include "Animation/AnimSequence.h"
+#include "Animation/AnimData/AnimDataController.h"
+#include "Animation/AnimData/AnimDataModel.h"
+#include "Runtime/Engine/Classes/Animation/AnimSequenceHelpers.h"
 
 #define LOCTEXT_NAMESPACE "SAnimTimeline"
 
@@ -491,14 +494,16 @@ void SAnimTimeline::OnCropAnimSequence( bool bFromStart, float CurrentTime )
 
 				//Call modify to restore anim sequence current state
 				AnimSequence->Modify();
+				
+				const float TrimStart = bFromStart ? 0.f : CurrentTime;
+				const float TrimEnd = bFromStart ? CurrentTime : AnimSequence->GetPlayLength();
+
+				// Trim off the user-selected part of the raw anim data.
+				UE::Anim::AnimationData::Trim(AnimSequence, TrimStart, TrimEnd);
 
-				// Crop the raw anim data.
-				AnimSequence->CropRawAnimData( CurrentTime, bFromStart );
 
 				//Resetting slider position to the first frame
 				PreviewInstance->SetPosition( 0.0f, false );
-
-				Model.Pin()->RefreshTracks();
 			}
 		}
 	}
@@ -521,12 +526,8 @@ void SAnimTimeline::OnAppendAnimSequence( bool bFromStart, int32 NumOfFrames )
 			AnimSequence->Modify();
 
 			// Crop the raw anim data.
-			int32 StartFrame = (bFromStart)? 0 : AnimSequence->GetNumberOfSampledKeys() - 1;
-			int32 EndFrame = StartFrame + NumOfFrames;
-			int32 CopyFrame = StartFrame;
-			AnimSequence->InsertFramesToRawAnimData(StartFrame, EndFrame, CopyFrame);
-
-			Model.Pin()->RefreshTracks();
+			const int32 StartFrame = (bFromStart)? 0 : AnimSequence->GetDataModel()->GetNumberOfFrames() - 1;
+			UE::Anim::AnimationData::DuplicateKeys(AnimSequence, StartFrame, 1, StartFrame);
 		}
 	}
 }
@@ -548,11 +549,8 @@ void SAnimTimeline::OnInsertAnimSequence( bool bBefore, int32 CurrentFrame )
 			AnimSequence->Modify();
 
 			// Crop the raw anim data.
-			int32 StartFrame = (bBefore)? CurrentFrame : CurrentFrame + 1;
-			int32 EndFrame = StartFrame + 1;
-			AnimSequence->InsertFramesToRawAnimData(StartFrame, EndFrame, CurrentFrame);
-
-			Model.Pin()->RefreshTracks();
+			const int32 StartFrame = (bBefore)? CurrentFrame : CurrentFrame + 1;
+			UE::Anim::AnimationData::DuplicateKeys(AnimSequence, StartFrame, 1, CurrentFrame);
 		}
 	}
 }
@@ -575,7 +573,8 @@ void SAnimTimeline::OnReZeroAnimSequence(int32 FrameIndex)
 				AnimSequence->Modify();
 
 				// As above, animations don't have any idea of hierarchy, so we don't know for sure if track 0 is the root bone's track.
-				FRawAnimSequenceTrack& RawTrack = AnimSequence->GetRawAnimationTrack(0);
+				const FBoneAnimationTrack& AnimationTrack = AnimSequence->GetDataModel()->GetBoneTrackByIndex(0);
+				FRawAnimSequenceTrack RawTrack = AnimationTrack.InternalTrackData;
 
 				// Find vector that would translate current root bone location onto origin.
 				FVector FrameTransform = FVector::ZeroVector;
@@ -601,13 +600,10 @@ void SAnimTimeline::OnReZeroAnimSequence(int32 FrameIndex)
 					RawTrack.PosKeys[i] += ApplyTranslation;
 				}
 
-				// Handle Raw Data changing
-				AnimSequence->MarkRawDataAsModified();
-				AnimSequence->OnRawDataChanged();
+				UAnimDataController* Controller = AnimSequence->GetController();
+				Controller->SetBoneTrackKeys(AnimationTrack.Name, RawTrack.PosKeys, RawTrack.RotKeys, RawTrack.ScaleKeys);
 
 				AnimSequence->MarkPackageDirty();
-
-				Model.Pin()->RefreshTracks();
 			}
 		}
 	}
@@ -741,7 +737,7 @@ bool SAnimTimeline::GetGridMetrics(float PhysicalWidth, double& OutMajorInterval
 	FSlateFontInfo SmallLayoutFont = FCoreStyle::GetDefaultFontStyle("Regular", 8);
 	TSharedRef<FSlateFontMeasure> FontMeasureService = FSlateApplication::Get().GetRenderer()->GetFontMeasureService();
 
-	FFrameRate DisplayRate(FMath::RoundToInt(Model.Pin()->GetFrameRate()), 1);
+	FFrameRate DisplayRate(FMath::Max(FMath::RoundToInt(Model.Pin()->GetFrameRate()), 1), 1);
 	double BiggestTime = ViewRange.Get().GetUpperBoundValue();
 	FString TickString = NumericTypeInterface->ToString((BiggestTime * DisplayRate).FrameNumber.Value);
 	FVector2D MaxTextSize = FontMeasureService->Measure(TickString, SmallLayoutFont);
@@ -751,7 +747,7 @@ bool SAnimTimeline::GetGridMetrics(float PhysicalWidth, double& OutMajorInterval
 	float MinTickPx = MaxTextSize.X + 5.f;
 	float DesiredMajorTickPx = MaxTextSize.X * MajorTickMultiplier;
 
-	if (PhysicalWidth > 0)
+	if (PhysicalWidth > 0 && DisplayRate.AsDecimal() > 0)
 	{
 		return ComputeGridSpacing(
 			DisplayRate,
@@ -784,6 +780,7 @@ void SAnimTimeline::OnColumnFillCoefficientChanged(float FillCoefficient, int32
 
 void SAnimTimeline::HandleKeyComplete()
 {
+	REFACTOR_V2_WARNING("V2 remove this function and verify that behaviour remained the same");
 	Model.Pin()->RefreshTracks();
 }
 
diff --git a/Engine/Source/Editor/Persona/Private/PersonaModule.cpp b/Engine/Source/Editor/Persona/Private/PersonaModule.cpp
--- a/Engine/Source/Editor/Persona/Private/PersonaModule.cpp
+++ b/Engine/Source/Editor/Persona/Private/PersonaModule.cpp
@@ -92,6 +92,7 @@
 #include "SAnimSequenceCurveEditor.h"
 #include "AnimSequenceTimelineCommands.h"
 #include "SAnimMontageSectionsPanel.h"
+#include "Animation/AnimSequenceHelpers.h"
 
 IMPLEMENT_MODULE( FPersonaModule, Persona );
 
@@ -614,7 +615,7 @@ void FPersonaModule::TestSkeletonCurveNamesForUse(const TSharedRef<IEditableSkel
 					UAnimSequence* Seq = Cast<UAnimSequence>(Anim.GetAsset());
 
 					TSharedPtr<FTokenizedMessage> Message;
-					for (FFloatCurve& Curve : Seq->RawCurveData.FloatCurves)
+					for (const FFloatCurve& Curve : Seq->GetDataModel()->GetCurveData().FloatCurves)
 					{
 						if (UnusedNames.Contains(Curve.Name.DisplayName))
 						{
@@ -829,15 +830,18 @@ bool FPersonaModule::ExportToFBX(TArray<TWeakObjectPtr<UAnimSequence>>& AnimSequ
 
 void FPersonaModule::AddLoopingInterpolation(TArray<TWeakObjectPtr<UAnimSequence>>& AnimSequences)
 {
-	FText WarningMessage = LOCTEXT("AddLoopiingInterpolation", "This will add an extra first frame at the end of the animation to create a better looping interpolation. This action cannot be undone. Would you like to proceed?");
+	FText WarningMessage = LOCTEXT("AddLoopingInterpolation", "This will add an extra first frame at the end of the animation to create a better looping interpolation. This action cannot be undone. Would you like to proceed?");
 
 	if (FMessageDialog::Open(EAppMsgType::YesNo, WarningMessage) == EAppReturnType::Yes)
 	{
-		for (auto Animation : AnimSequences)
+		for (TWeakObjectPtr<UAnimSequence>& Animation : AnimSequences)
 		{
 			// get first frame and add to the last frame and go through track
 			// now calculating old animated space bases
-			Animation->AddLoopingInterpolation();
+			if (UAnimSequence* AnimSequence = Animation.Get())
+			{
+				UE::Anim::AnimationData::AddLoopingInterpolation(AnimSequence);
+			}
 		}
 	}
 }
diff --git a/Engine/Source/Editor/Persona/Private/SAnimCompositePanel.cpp b/Engine/Source/Editor/Persona/Private/SAnimCompositePanel.cpp
--- a/Engine/Source/Editor/Persona/Private/SAnimCompositePanel.cpp
+++ b/Engine/Source/Editor/Persona/Private/SAnimCompositePanel.cpp
@@ -209,10 +209,11 @@ void SAnimCompositePanel::PostUndoRedo()
 		RegisterActiveTimer(0.f, FWidgetActiveTimerDelegate::CreateSP(this, &SAnimCompositePanel::TriggerRebuildPanel));
 	}
 
+	REFACTOR_V2_WARNING("V2 why is this call necessary? Is this handled through notify callbacks now?");
 	// when undo or redo happens, we still have to recalculate length, so we can't rely on sequence length changes or not
 	if (Composite->GetPlayLength())
 	{
-		Composite->SetSequenceLength(0.f);
+		Composite->SetCompositeLength(0.f);
 	}
 }
 
diff --git a/Engine/Source/Editor/Persona/Private/SAnimMontagePanel.cpp b/Engine/Source/Editor/Persona/Private/SAnimMontagePanel.cpp
--- a/Engine/Source/Editor/Persona/Private/SAnimMontagePanel.cpp
+++ b/Engine/Source/Editor/Persona/Private/SAnimMontagePanel.cpp
@@ -582,9 +582,10 @@ void SAnimMontagePanel::PostRedo( bool bSuccess )
 void SAnimMontagePanel::PostRedoUndo()
 {
 	// when undo or redo happens, we still have to recalculate length, so we can't rely on sequence length changes or not
+	REFACTOR_V2_WARNING("V2 is this still necessary or is this handled through notify callbacks by now?");
 	if (Montage->GetPlayLength())
 	{
-		Montage->SetSequenceLength(0.f);
+		Montage->SetCompositeLength(MINIMUM_ANIMATION_LENGTH);
 	}
 
 	RebuildMontagePanel(); //Rebuild here, undoing adds can cause slate to crash later on if we don't (using dummy args since they aren't used by the method
diff --git a/Engine/Source/Editor/Persona/Private/SAnimNotifyPanel.cpp b/Engine/Source/Editor/Persona/Private/SAnimNotifyPanel.cpp
--- a/Engine/Source/Editor/Persona/Private/SAnimNotifyPanel.cpp
+++ b/Engine/Source/Editor/Persona/Private/SAnimNotifyPanel.cpp
@@ -4311,15 +4311,11 @@ void SAnimNotifyPanel::DeleteSelectedNodeObjects()
 	Update();
 }
 
-void SAnimNotifyPanel::SetSequence(class UAnimSequenceBase *	InSequence)
+void SAnimNotifyPanel::SetSequence(class UAnimSequenceBase*	InSequence)
 {
 	if (InSequence != Sequence)
 	{
 		Sequence = InSequence;
-		// @todo anim : this is kinda hack to make sure it has 1 track is alive
-		// we can do this whenever import or asset is created, but it's more places to handle than here
-		// the function name in that case will need to change
-		Sequence->InitializeNotifyTrack();
 		Update();
 	}
 }
diff --git a/Engine/Source/Editor/Persona/Private/SAnimSequenceCurveEditor.cpp b/Engine/Source/Editor/Persona/Private/SAnimSequenceCurveEditor.cpp
--- a/Engine/Source/Editor/Persona/Private/SAnimSequenceCurveEditor.cpp
+++ b/Engine/Source/Editor/Persona/Private/SAnimSequenceCurveEditor.cpp
@@ -16,67 +16,189 @@
 #include "Editor.h"
 #include "Framework/MultiBox/MultiBoxBuilder.h"
 #include "EditorStyleSet.h"
+#include "Animation/AnimData/AnimDataController.h"
 
 #define LOCTEXT_NAMESPACE "SAnimSequenceCurveEditor"
 
+FRichCurveEditorModelNamed::FRichCurveEditorModelNamed(const FSmartName& InName, ERawCurveTrackTypes InType, int32 InCurveIndex, UAnimSequenceBase* InAnimSequence, FCurveEditorTreeItemID InTreeId /*= FCurveEditorTreeItemID()*/) : FRichCurveEditorModel(InAnimSequence)
+, Name(InName)
+, AnimSequence(InAnimSequence)
+, CurveIndex(InCurveIndex)
+, Type(InType)
+, TreeId(InTreeId)
+, CurveId(FAnimationCurveIdentifier(Name, Type))
+{
+	CurveModifiedDelegate.AddRaw(this, &FRichCurveEditorModelNamed::CurveHasChanged);
+
+	InAnimSequence->GetDataModel()->GetModifiedEvent().AddRaw(this, &FRichCurveEditorModelNamed::OnModelHasChanged);
+
+	if (Type == ERawCurveTrackTypes::RCT_Transform)
+	{
+		UAnimationCurveIdentifierExtensions::GetTransformChildCurveIdentifier(CurveId, (ETransformCurveChannel)(CurveIndex / 3), (EVectorCurveChannel)(CurveIndex % 3));
+	}
+
+	UpdateCachedCurve();
+}
+
+FRichCurveEditorModelNamed::~FRichCurveEditorModelNamed()
+{
+	AnimSequence->GetDataModel()->GetModifiedEvent().RemoveAll(this);
+}
+
 bool FRichCurveEditorModelNamed::IsValid() const
 {
-	return AnimSequence->GetCurveData().GetCurveData(Name.UID, Type) != nullptr;
+	return AnimSequence->GetDataModel()->FindCurve(FAnimationCurveIdentifier(Name, Type)) != nullptr;
 }
 
 FRichCurve& FRichCurveEditorModelNamed::GetRichCurve()
 {
 	check(AnimSequence.Get() != nullptr);
+	return CachedCurve;
+}
+
+const FRichCurve& FRichCurveEditorModelNamed::GetReadOnlyRichCurve() const
+{
+	return const_cast<FRichCurveEditorModelNamed*>(this)->GetRichCurve();
+}
+
+void FRichCurveEditorModelNamed::CurveHasChanged()
+{
+	UAnimDataController* Controller = AnimSequence->GetController();
 
-	FAnimCurveBase* CurveBase = AnimSequence->RawCurveData.GetCurveData(Name.UID, Type);
-	check(CurveBase);	// If this fails lifetime contracts have been violated - this curve should always be present if this model exists
-		
 	switch (Type)
 	{
-	case ERawCurveTrackTypes::RCT_Vector:
-	{
-		FVectorCurve& VectorCurve = *(static_cast<FVectorCurve*>(CurveBase));
-		check(CurveIndex < 3);
-		return VectorCurve.FloatCurves[CurveIndex];
-	}
-	case ERawCurveTrackTypes::RCT_Transform:
-	{
-		FTransformCurve& TransformCurve = *(static_cast<FTransformCurve*>(CurveBase));
-		check(CurveIndex < 9);
-		const int32 SubCurveIndex = CurveIndex % 3;
-		switch(CurveIndex)
+		case ERawCurveTrackTypes::RCT_Vector:
+		{
+			ensure(false);
+			break;
+		}
+		case ERawCurveTrackTypes::RCT_Transform:
+		case ERawCurveTrackTypes::RCT_Float:
 		{
-		default:
-			check(false);
-			// fall through
-		case 0:
-		case 1:
-		case 2:
-			return TransformCurve.TranslationCurve.FloatCurves[SubCurveIndex];
-		case 3:
-		case 4:
-		case 5:
-			return TransformCurve.RotationCurve.FloatCurves[SubCurveIndex];
-		case 6:
-		case 7:
-		case 8:
-			return TransformCurve.ScaleCurve.FloatCurves[SubCurveIndex];
+			Controller->SetCurveKeys(CurveId, CachedCurve.GetConstRefOfKeys());
+			break;
 		}
-			
 	}
-	case ERawCurveTrackTypes::RCT_Float:
-	default:
+}
+
+void FRichCurveEditorModelNamed::OnModelHasChanged(const EAnimDataModelNotifyType& NotifyType, UAnimDataModel* Model, const FAnimDataModelNotifPayload& Payload)
+{
+	NotifyCollector.Handle(NotifyType);
+
+	switch (NotifyType)
 	{
-		FFloatCurve& FloatCurve = *(static_cast<FFloatCurve*>(CurveBase));
-		check(CurveIndex == 0);
-		return FloatCurve.FloatCurve;
-	}
+		case EAnimDataModelNotifyType::CurveAdded:
+		case EAnimDataModelNotifyType::CurveChanged:
+		case EAnimDataModelNotifyType::CurveFlagsChanged:
+		case EAnimDataModelNotifyType::CurveScaled:
+		{
+			if (NotifyCollector.IsNotWithinBracket())
+			{
+				const FCurvePayload& TypedPayload = Payload.GetPayload<FCurvePayload>();
+				if (TypedPayload.Identifier.InternalName == Name)
+				{
+					UpdateCachedCurve();
+				}
+			}
+
+			break;
+		}
+
+		case EAnimDataModelNotifyType::CurveRemoved:
+		{
+			// Curve was removed
+			break;
+		}
+
+		case EAnimDataModelNotifyType::CurveRenamed:
+		{
+			const FCurveRenamedPayload& TypedPayload = Payload.GetPayload<FCurveRenamedPayload>();
+			if (TypedPayload.Identifier == CurveId)
+			{
+				Name = TypedPayload.NewIdentifier.InternalName;
+				CurveId = TypedPayload.NewIdentifier;
+
+				if (NotifyCollector.IsNotWithinBracket())
+				{
+					UpdateCachedCurve();
+				}
+			}
+
+			break;
+		}
+
+		case EAnimDataModelNotifyType::BracketClosed:
+		{
+			if (NotifyCollector.IsNotWithinBracket())
+			{
+				if (NotifyCollector.Contains({ EAnimDataModelNotifyType::CurveAdded, EAnimDataModelNotifyType::CurveChanged, EAnimDataModelNotifyType::CurveFlagsChanged, EAnimDataModelNotifyType::CurveScaled, EAnimDataModelNotifyType::CurveRenamed }))
+				{
+					UpdateCachedCurve();
+				}
+			}
+			break;
+		}
 	}
 }
 
-const FRichCurve& FRichCurveEditorModelNamed::GetReadOnlyRichCurve() const
+void FRichCurveEditorModelNamed::UpdateCachedCurve()
 {
-	return const_cast<FRichCurveEditorModelNamed*>(this)->GetRichCurve();
+	const FAnimCurveBase* CurveBase = AnimSequence->GetDataModel()->FindCurve(CurveId);
+	
+	check(CurveBase);	// If this fails lifetime contracts have been violated - this curve should always be present if this model exists
+
+	const FRichCurve* CurveToCopyFrom = [this, CurveBase]() -> const FRichCurve*
+	{
+		switch (Type)
+		{
+			case ERawCurveTrackTypes::RCT_Vector:
+			{
+				ensure(false);
+				const FVectorCurve& VectorCurve = *(static_cast<const FVectorCurve*>(CurveBase));
+				check(CurveIndex < 3);
+				return &VectorCurve.FloatCurves[CurveIndex];
+			}
+			case ERawCurveTrackTypes::RCT_Transform:
+			{
+				const FTransformCurve& TransformCurve = *(static_cast<const FTransformCurve*>(CurveBase));
+				check(CurveIndex < 9);
+				const int32 SubCurveIndex = CurveIndex % 3;
+				switch (CurveIndex)
+				{
+				default:
+					check(false);
+					// fall through
+				case 0:
+				case 1:
+				case 2:
+					return &TransformCurve.TranslationCurve.FloatCurves[SubCurveIndex];
+				case 3:
+				case 4:
+				case 5:
+					return &TransformCurve.RotationCurve.FloatCurves[SubCurveIndex];
+				case 6:
+				case 7:
+				case 8:
+					return &TransformCurve.ScaleCurve.FloatCurves[SubCurveIndex];
+				}
+
+			}
+			case ERawCurveTrackTypes::RCT_Float:
+			default:
+			{
+				const FFloatCurve& FloatCurve = *(static_cast<const FFloatCurve*>(CurveBase));
+				check(CurveIndex == 0);
+				return &FloatCurve.FloatCurve;
+			}
+		}
+
+		return nullptr;
+	}();
+	
+	if (ensure(CurveToCopyFrom))
+	{
+		CachedCurve = *CurveToCopyFrom;
+	}
 }
 
 class FAnimSequenceCurveEditorItem : public ICurveEditorTreeItem
diff --git a/Engine/Source/Editor/Persona/Private/SAnimSequenceCurveEditor.h b/Engine/Source/Editor/Persona/Private/SAnimSequenceCurveEditor.h
--- a/Engine/Source/Editor/Persona/Private/SAnimSequenceCurveEditor.h
+++ b/Engine/Source/Editor/Persona/Private/SAnimSequenceCurveEditor.h
@@ -23,25 +23,27 @@ class FTabManager;
 class FRichCurveEditorModelNamed : public FRichCurveEditorModel
 {
 public:
-	FRichCurveEditorModelNamed(const FSmartName& InName, ERawCurveTrackTypes InType, int32 InCurveIndex, UAnimSequenceBase* InAnimSequence, FCurveEditorTreeItemID InTreeId = FCurveEditorTreeItemID())
-		: FRichCurveEditorModel(InAnimSequence)
-		, Name(InName)
-		, AnimSequence(InAnimSequence)
-		, CurveIndex(InCurveIndex)
-		, Type(InType)
-		, TreeId(InTreeId)
-	{
-	}
+	FRichCurveEditorModelNamed(const FSmartName& InName, ERawCurveTrackTypes InType, int32 InCurveIndex, UAnimSequenceBase* InAnimSequence, FCurveEditorTreeItemID InTreeId = FCurveEditorTreeItemID());
+
+	virtual ~FRichCurveEditorModelNamed();
 
 	virtual bool IsValid() const override;
 	virtual FRichCurve& GetRichCurve() override;
 	virtual const FRichCurve& GetReadOnlyRichCurve() const override;
 
+	void CurveHasChanged();
+	void OnModelHasChanged(const EAnimDataModelNotifyType& NotifyType, UAnimDataModel* Model, const FAnimDataModelNotifPayload& Payload);
+	void UpdateCachedCurve();
+
 	FSmartName Name;
 	TWeakObjectPtr<UAnimSequenceBase> AnimSequence;
 	int32 CurveIndex;
 	ERawCurveTrackTypes Type;
 	FCurveEditorTreeItemID TreeId;
+	
+	FAnimationCurveIdentifier CurveId;
+	UE::Anim::FAnimDataModelNotifyCollector NotifyCollector;
+	FRichCurve CachedCurve;
 };
 
 class SAnimSequenceCurveEditor : public IAnimSequenceCurveEditor, public FEditorUndoClient
diff --git a/Engine/Source/Editor/Persona/Private/SAnimationScrubPanel.cpp b/Engine/Source/Editor/Persona/Private/SAnimationScrubPanel.cpp
--- a/Engine/Source/Editor/Persona/Private/SAnimationScrubPanel.cpp
+++ b/Engine/Source/Editor/Persona/Private/SAnimationScrubPanel.cpp
@@ -10,6 +10,9 @@
 #include "ScopedTransaction.h"
 #include "Animation/BlendSpaceBase.h"
 #include "AnimationEditorPreviewScene.h"
+#include "Animation/AnimData/AnimDataModel.h"
+#include "Animation/AnimData/AnimDataController.h"
+#include "Animation/AnimSequenceHelpers.h"
 
 #define LOCTEXT_NAMESPACE "AnimationScrubPanel"
 
@@ -474,8 +477,12 @@ void SAnimationScrubPanel::OnCropAnimSequence( bool bFromStart, float CurrentTim
 				//Call modify to restore anim sequence current state
 				AnimSequence->Modify();
 
-				// Crop the raw anim data.
-				AnimSequence->CropRawAnimData( CurrentTime, bFromStart );
+
+				const float TrimStart = bFromStart ? 0.f : CurrentTime;
+				const float TrimEnd = bFromStart ? CurrentTime : AnimSequence->GetPlayLength();
+
+				// Trim off the user-selected part of the raw anim data.
+				UE::Anim::AnimationData::Trim(AnimSequence, TrimStart, TrimEnd);
 
 				//Resetting slider position to the first frame
 				PreviewInstance->SetPosition( 0.0f, false );
@@ -503,10 +510,10 @@ void SAnimationScrubPanel::OnAppendAnimSequence( bool bFromStart, int32 NumOfFra
 			AnimSequence->Modify();
 
 			// Crop the raw anim data.
-			int32 StartFrame = (bFromStart)? 0 : AnimSequence->GetNumberOfSampledKeys() - 1;
+			int32 StartFrame = (bFromStart)? 0 : AnimSequence->GetDataModel()->GetNumberOfFrames() - 1;
 			int32 EndFrame = StartFrame + NumOfFrames;
 			int32 CopyFrame = StartFrame;
-			AnimSequence->InsertFramesToRawAnimData(StartFrame, EndFrame, CopyFrame);
+			UE::Anim::AnimationData::DuplicateKeys(AnimSequence, StartFrame, NumOfFrames, CopyFrame);
 
 			OnSetInputViewRange.ExecuteIfBound(0, AnimSequence->GetPlayLength());
 		}
@@ -529,10 +536,9 @@ void SAnimationScrubPanel::OnInsertAnimSequence( bool bBefore, int32 CurrentFram
 			//Call modify to restore anim sequence current state
 			AnimSequence->Modify();
 
-			// Crop the raw anim data.
-			int32 StartFrame = (bBefore)? CurrentFrame : CurrentFrame + 1;
-			int32 EndFrame = StartFrame + 1;
-			AnimSequence->InsertFramesToRawAnimData(StartFrame, EndFrame, CurrentFrame);
+			// Duplicate specified key
+			const int32 StartFrame = (bBefore)? CurrentFrame : CurrentFrame + 1;
+			UE::Anim::AnimationData::DuplicateKeys(AnimSequence, StartFrame, 1, CurrentFrame);
 
 			OnSetInputViewRange.ExecuteIfBound(0, AnimSequence->GetPlayLength());
 		}
@@ -557,7 +563,8 @@ void SAnimationScrubPanel::OnReZeroAnimSequence(int32 FrameIndex)
 				AnimSequence->Modify();
 
 				// As above, animations don't have any idea of hierarchy, so we don't know for sure if track 0 is the root bone's track.
-				FRawAnimSequenceTrack& RawTrack = AnimSequence->GetRawAnimationTrack(0);
+				const FBoneAnimationTrack& AnimationTrack = AnimSequence->GetDataModel()->GetBoneTrackByIndex(0);
+				FRawAnimSequenceTrack RawTrack = AnimationTrack.InternalTrackData;
 
 				// Find vector that would translate current root bone location onto origin.
 				FVector FrameTransform = FVector::ZeroVector;
@@ -583,9 +590,8 @@ void SAnimationScrubPanel::OnReZeroAnimSequence(int32 FrameIndex)
 					RawTrack.PosKeys[i] += ApplyTranslation;
 				}
 
-				// Handle Raw Data changing
-				AnimSequence->MarkRawDataAsModified();
-				AnimSequence->OnRawDataChanged();
+				UAnimDataController* Controller = AnimSequence->GetController();
+				Controller->SetBoneTrackKeys(AnimationTrack.Name, RawTrack.PosKeys, RawTrack.RotKeys, RawTrack.ScaleKeys);
 
 				AnimSequence->MarkPackageDirty();
 			}
diff --git a/Engine/Source/Editor/SequenceRecorder/Private/AnimationRecorder.cpp b/Engine/Source/Editor/SequenceRecorder/Private/AnimationRecorder.cpp
--- a/Engine/Source/Editor/SequenceRecorder/Private/AnimationRecorder.cpp
+++ b/Engine/Source/Editor/SequenceRecorder/Private/AnimationRecorder.cpp
@@ -19,6 +19,7 @@
 #include "Animation/AnimNotifies/AnimNotify.h"
 #include "Animation/AnimNotifies/AnimNotifyState.h"
 #include "Subsystems/AssetEditorSubsystem.h"
+#include "Animation/AnimData/AnimDataController.h"
 
 #define LOCTEXT_NAMESPACE "FAnimationRecorder"
 
@@ -216,7 +217,6 @@ void FAnimationRecorder::StartRecord(USkeletalMeshComponent* Component, UAnimSeq
 	TimePassed = 0.f;
 	AnimationObject = InAnimationObject;
 
-	AnimationObject->RecycleAnimSequence();
 	AnimationObject->BoneCompressionSettings = FAnimationUtils::GetDefaultAnimationRecorderBoneCompressionSettings();
 
 	FAnimationRecorder::GetBoneTransforms(Component, PreviousSpacesBases);
@@ -224,8 +224,14 @@ void FAnimationRecorder::StartRecord(USkeletalMeshComponent* Component, UAnimSeq
 	PreviousComponentToWorld = Component->GetComponentTransform();
 
 	LastFrame = 0;
-	AnimationObject->SetSequenceLength(0.f);
-	AnimationObject->SetNumberOfSampledKeys(0);
+
+	UAnimDataController* Controller = AnimationObject->GetController();
+	Controller->SetModel(AnimationObject->GetDataModel());
+
+	Controller->OpenBracket(LOCTEXT("StartRecord_Bracket", "Starting Animation Recording"));
+
+	AnimationObject->DeleteNotifyTrackData();
+	Controller->ResetModel();
 
 	RecordedCurves.Reset();
 	RecordedTimes.Empty();
@@ -240,15 +246,13 @@ void FAnimationRecorder::StartRecord(USkeletalMeshComponent* Component, UAnimSeq
 		if (BoneTreeIndex != INDEX_NONE)
 		{
 			// add tracks for the bone existing
-			FName BoneTreeName = AnimSkeleton->GetReferenceSkeleton().GetBoneName(BoneTreeIndex);
-			AnimationObject->AddNewRawTrack(BoneTreeName);
+			const FName BoneTreeName = AnimSkeleton->GetReferenceSkeleton().GetBoneName(BoneTreeIndex);
+			Controller->AddBoneTrack(BoneTreeName);
 		}
 	}
 
 	AnimationObject->RetargetSource = Component->SkeletalMesh ? AnimSkeleton->GetRetargetSourceForMesh(Component->SkeletalMesh) : NAME_None;
 
-	// init notifies
-	AnimationObject->InitializeNotifyTrack();
 	// record the first frame
 	Record(Component, PreviousComponentToWorld, PreviousSpacesBases, PreviousAnimCurves,  0);
 }
@@ -289,11 +293,15 @@ UAnimSequence* FAnimationRecorder::StopRecord(bool bShowMessage)
 
 	if (AnimationObject)
 	{
+		UAnimDataController* Controller = AnimationObject->GetController();
 		int32 NumKeys = LastFrame  + 1;
-		AnimationObject->SetNumberOfSampledKeys(NumKeys);
 
 		// can't use TimePassed. That is just total time that has been passed, not necessarily match with frame count
-		AnimationObject->SetSequenceLength((NumKeys>1) ? (NumKeys-1) * IntervalTime : MINIMUM_ANIMATION_LENGTH);
+		Controller->SetPlayLength( (NumKeys>1) ? (NumKeys-1) * IntervalTime : MINIMUM_ANIMATION_LENGTH );
+
+		const float FloatDenominator = 1000.0f;
+		const float Numerator = FloatDenominator / IntervalTime;
+		Controller->SetFrameRate(FFrameRate(Numerator, FloatDenominator));
 
 		FixupNotifies();
 
@@ -304,7 +312,6 @@ UAnimSequence* FAnimationRecorder::StopRecord(bool bShowMessage)
 		if (RecordedCurves.Num() == NumKeys && UIDToArrayIndexLUT)
 		{
 			StartTime = FPlatformTime::Seconds();
-
 			USkeleton* SkeletonObj = AnimationObject->GetSkeleton();
 			for (int32 CurveUID = 0; CurveUID < UIDToArrayIndexLUT->Num(); ++CurveUID)
 			{
@@ -312,7 +319,7 @@ UAnimSequence* FAnimationRecorder::StopRecord(bool bShowMessage)
 
 				if (CurveIndex != MAX_uint16)
 				{
-					FFloatCurve* FloatCurveData = nullptr;
+					const FFloatCurve* FloatCurveData = nullptr;
 
 					TArray<float> TimesToRecord;
 					TArray<float> ValuesToRecord;
@@ -335,8 +342,10 @@ UAnimSequence* FAnimationRecorder::StopRecord(bool bShowMessage)
 								if (SkeletonObj->GetSmartNameByUID(USkeleton::AnimCurveMappingName, CurveUID, CurveName))
 								{
 									// give default curve flag for recording 
-									AnimationObject->RawCurveData.AddFloatCurveKey(CurveName, AACF_DefaultCurve, TimeToRecord, CurCurveValue);
-									FloatCurveData = static_cast<FFloatCurve*>(AnimationObject->RawCurveData.GetCurveData(CurveUID, ERawCurveTrackTypes::RCT_Float));
+									const FAnimationCurveIdentifier CurveId(CurveName, ERawCurveTrackTypes::RCT_Float);
+									Controller->AddCurve(CurveId, AACF_DefaultCurve);
+									FloatCurveData = AnimationObject->GetDataModel()->FindFloatCurve(CurveId);
+
 								}
 							}
 
@@ -360,7 +369,8 @@ UAnimSequence* FAnimationRecorder::StopRecord(bool bShowMessage)
 							Keys.Add(Key);
 						}
 
-						FloatCurveData->FloatCurve.SetKeys(Keys);
+						const FAnimationCurveIdentifier CurveId(FloatCurveData->Name, ERawCurveTrackTypes::RCT_Float);
+						Controller->SetCurveKeys(CurveId, Keys);
 					}
 				}
 			}	
@@ -369,8 +379,8 @@ UAnimSequence* FAnimationRecorder::StopRecord(bool bShowMessage)
 			UE_LOG(LogAnimation, Log, TEXT("Animation Recorder set keys in %0.02f seconds"), ElapsedTime);
 		}
 
-		//AnimationObject->RawCurveData.RemoveRedundantKeys();
-		AnimationObject->PostProcessSequence();
+		Controller->NotifyPopulated();
+		Controller->CloseBracket();
 
 		AnimationObject->MarkPackageDirty();
 		
@@ -396,7 +406,7 @@ UAnimSequence* FAnimationRecorder::StopRecord(bool bShowMessage)
 		{
 			const FText NotificationText = FText::Format(LOCTEXT("RecordAnimation", "'{0}' has been successfully recorded [{1} keys : {2} sec(s) @ {3} Hz]"),
 				FText::FromString(AnimationObject->GetName()),
-				FText::AsNumber(AnimationObject->GetNumberOfSampledKeys()),
+				FText::AsNumber(AnimationObject->GetDataModel()->GetNumberOfKeys()),
 				FText::AsNumber(AnimationObject->GetPlayLength()),
 				FText::AsNumber(1.f / IntervalTime)
 				);
@@ -433,7 +443,7 @@ UAnimSequence* FAnimationRecorder::StopRecord(bool bShowMessage)
 	UniqueNotifies.Empty();
 	UniqueNotifyStates.Empty();
 
-	return NULL;
+	return nullptr;
 }
 
 void FAnimationRecorder::ProcessRecordedTimes(UAnimSequence* AnimSequence, USkeletalMeshComponent* SkeletalMeshComponent, const FString& HoursName, const FString& MinutesName, const FString& SecondsName, const FString& FramesName, const FString& SubFramesName, const FString& SlateName, const FString& Slate)
@@ -635,25 +645,29 @@ bool FAnimationRecorder::Record(USkeletalMeshComponent* Component, FTransform co
 {
 	if (ensure(AnimationObject))
 	{
+		UAnimDataController* Controller = AnimationObject->GetController();
 		USkeletalMesh* SkeletalMesh = Component->MasterPoseComponent != nullptr ? Component->MasterPoseComponent->SkeletalMesh : Component->SkeletalMesh;
 
+		const TArray<FBoneAnimationTrack>& BoneAnimationTracks = AnimationObject->GetDataModel()->GetBoneAnimationTracks();
+
 		if (FrameToAdd == 0)
 		{
 			// Find the root bone & store its transform
 			SkeletonRootIndex = INDEX_NONE;
 			USkeleton* AnimSkeleton = AnimationObject->GetSkeleton();
-			for (int32 TrackIndex = 0; TrackIndex < AnimationObject->GetRawAnimationData().Num(); ++TrackIndex)
+
+			for (const FBoneAnimationTrack& AnimationTrack : BoneAnimationTracks)
 			{
 				// verify if this bone exists in skeleton
-				int32 BoneTreeIndex = AnimationObject->GetSkeletonIndexFromRawDataTrackIndex(TrackIndex);
+				const int32 BoneTreeIndex = AnimationTrack.BoneTreeIndex;
 				if (BoneTreeIndex != INDEX_NONE)
 				{
-					int32 BoneIndex = AnimSkeleton->GetMeshBoneIndexFromSkeletonBoneIndex(SkeletalMesh, BoneTreeIndex);
-					int32 ParentIndex = SkeletalMesh->GetRefSkeleton().GetParentIndex(BoneIndex);
-					FTransform LocalTransform = SpacesBases[BoneIndex];
+					const int32 BoneIndex = AnimSkeleton->GetMeshBoneIndexFromSkeletonBoneIndex(SkeletalMesh, BoneTreeIndex);
+					const int32 ParentIndex = SkeletalMesh->GetRefSkeleton().GetParentIndex(BoneIndex);
+					const FTransform LocalTransform = SpacesBases[BoneIndex];
 					if (ParentIndex == INDEX_NONE)
 					{
-						if (bRemoveRootTransform && AnimationObject->GetRawAnimationData().Num() > 1)
+						if (bRemoveRootTransform && BoneAnimationTracks.Num() > 1)
 						{
 							// Store initial root transform.
 							// We remove the initial transform of the root bone and transform root's children
@@ -679,14 +693,16 @@ bool FAnimationRecorder::Record(USkeletalMeshComponent* Component, FTransform co
 
 		FSerializedAnimation  SerializedAnimation;
 		USkeleton* AnimSkeleton = AnimationObject->GetSkeleton();
-		for (int32 TrackIndex = 0; TrackIndex < AnimationObject->GetRawAnimationData().Num(); ++TrackIndex)
+
+		for (int32 TrackIndex = 0; TrackIndex < BoneAnimationTracks.Num(); ++TrackIndex)
 		{
+			const FBoneAnimationTrack& AnimationTrack = BoneAnimationTracks[TrackIndex];
 			// verify if this bone exists in skeleton
-			int32 BoneTreeIndex = AnimationObject->GetSkeletonIndexFromRawDataTrackIndex(TrackIndex);
+			const int32 BoneTreeIndex = AnimationTrack.BoneTreeIndex;
 			if (BoneTreeIndex != INDEX_NONE)
 			{
-				int32 BoneIndex = AnimSkeleton->GetMeshBoneIndexFromSkeletonBoneIndex(SkeletalMesh, BoneTreeIndex);
-				int32 ParentIndex = SkeletalMesh->GetRefSkeleton().GetParentIndex(BoneIndex);
+				const int32 BoneIndex = AnimSkeleton->GetMeshBoneIndexFromSkeletonBoneIndex(SkeletalMesh, BoneTreeIndex);
+				const int32 ParentIndex = SkeletalMesh->GetRefSkeleton().GetParentIndex(BoneIndex);
 				FTransform LocalTransform = SpacesBases[BoneIndex];
 				if ( ParentIndex != INDEX_NONE )
 				{
@@ -701,17 +717,23 @@ bool FAnimationRecorder::Record(USkeletalMeshComponent* Component, FTransform co
 					}
 				}
 
-				FRawAnimSequenceTrack& RawTrack = AnimationObject->GetRawAnimationTrack(TrackIndex);
+				// Expensive copy, might be worth accumulating track data internally and only setting the keys once when recording is stopped
+				FRawAnimSequenceTrack RawTrack = AnimationTrack.InternalTrackData;
+
 				if (bRecordTransforms)
 				{
 					RawTrack.PosKeys.Add(LocalTransform.GetTranslation());
 					RawTrack.RotKeys.Add(LocalTransform.GetRotation());
 					RawTrack.ScaleKeys.Add(LocalTransform.GetScale3D());
+  
+					Controller->SetBoneTrackKeys(AnimationTrack.Name, RawTrack.PosKeys, RawTrack.RotKeys, RawTrack.ScaleKeys);			
+  
 					if (AnimationSerializer)
 					{
 						SerializedAnimation.AddTransform(TrackIndex, LocalTransform);
 					}
 				}
+
 				// verification
 				if (FrameToAdd != RawTrack.PosKeys.Num()-1)
 				{
diff --git a/Engine/Source/Editor/SequenceRecorder/Private/Sections/MovieScene3DTransformSectionRecorder.cpp b/Engine/Source/Editor/SequenceRecorder/Private/Sections/MovieScene3DTransformSectionRecorder.cpp
--- a/Engine/Source/Editor/SequenceRecorder/Private/Sections/MovieScene3DTransformSectionRecorder.cpp
+++ b/Engine/Source/Editor/SequenceRecorder/Private/Sections/MovieScene3DTransformSectionRecorder.cpp
@@ -11,6 +11,8 @@
 #include "SequenceRecorderUtils.h"
 #include "Algo/Transform.h"
 #include "Channels/MovieSceneChannelProxy.h"
+#include "Animation/AnimData/AnimDataModel.h"
+#include "Animation/AnimSequence.h"
 
 TSharedPtr<IMovieSceneSectionRecorder> FMovieScene3DTransformSectionRecorderFactory::CreateSectionRecorder(const FActorRecordingSettings& InActorRecordingSettings) const
 {
@@ -136,14 +138,16 @@ void FMovieScene3DTransformSectionRecorder::FinalizeSection(float CurrentTime)
 				// find the root bone
 				int32 RootIndex = INDEX_NONE;
 				USkeleton* AnimSkeleton = AnimSequence->GetSkeleton();
-				for (int32 TrackIndex = 0; TrackIndex < AnimSequence->GetRawAnimationData().Num(); ++TrackIndex)
+
+				const TArray<FBoneAnimationTrack>& BoneAnimationTracks = AnimSequence->GetDataModel()->GetBoneAnimationTracks();
+				for (const FBoneAnimationTrack& AnimationTrack : BoneAnimationTracks)
 				{
 					// verify if this bone exists in skeleton
-					int32 BoneTreeIndex = AnimSequence->GetSkeletonIndexFromRawDataTrackIndex(TrackIndex);
+					const int32 BoneTreeIndex = AnimationTrack.BoneTreeIndex;
 					if (BoneTreeIndex != INDEX_NONE)
 					{
-						int32 BoneIndex = AnimSkeleton->GetMeshBoneIndexFromSkeletonBoneIndex(SkeletalMesh, BoneTreeIndex);
-						int32 ParentIndex = SkeletalMesh->GetRefSkeleton().GetParentIndex(BoneIndex);
+						const int32 BoneIndex = AnimSkeleton->GetMeshBoneIndexFromSkeletonBoneIndex(SkeletalMesh, BoneTreeIndex);
+						const int32 ParentIndex = SkeletalMesh->GetRefSkeleton().GetParentIndex(BoneIndex);
 						if (ParentIndex == INDEX_NONE)
 						{
 							// found root
@@ -153,6 +157,7 @@ void FMovieScene3DTransformSectionRecorder::FinalizeSection(float CurrentTime)
 					}
 				}
 
+				ensure(BoneAnimationTracks.IsValidIndex(RootIndex));
 				check(RootIndex != INDEX_NONE);
 
 				FFrameRate TickResolution = MovieSceneSection->GetTypedOuter<UMovieScene>()->GetTickResolution();
@@ -161,7 +166,7 @@ void FMovieScene3DTransformSectionRecorder::FinalizeSection(float CurrentTime)
 				// we may need to offset the transform here if the animation was not recorded on the root component
 				FTransform InvComponentTransform = AnimRecorder->GetComponentTransform().Inverse();
 
-				const FRawAnimSequenceTrack& RawTrack = AnimSequence->GetRawAnimationData()[RootIndex];
+				const FRawAnimSequenceTrack& RawTrack = BoneAnimationTracks[RootIndex].InternalTrackData;
 				const int32 KeyCount = FMath::Max(FMath::Max(RawTrack.PosKeys.Num(), RawTrack.RotKeys.Num()), RawTrack.ScaleKeys.Num());
 				for (int32 KeyIndex = 0; KeyIndex < KeyCount; KeyIndex++)
 				{
diff --git a/Engine/Source/Editor/SkeletonEditor/Private/EditableSkeleton.cpp b/Engine/Source/Editor/SkeletonEditor/Private/EditableSkeleton.cpp
--- a/Engine/Source/Editor/SkeletonEditor/Private/EditableSkeleton.cpp
+++ b/Engine/Source/Editor/SkeletonEditor/Private/EditableSkeleton.cpp
@@ -30,6 +30,7 @@
 #include "Engine/DataAsset.h"
 #include "Animation/PreviewCollectionInterface.h"
 #include "HAL/PlatformApplicationMisc.h"
+#include "Animation/AnimData/AnimDataController.h"
 
 const FString FEditableSkeleton::SocketCopyPasteHeader = TEXT("SocketCopyPasteBuffer");
 
@@ -305,6 +306,8 @@ void FEditableSkeleton::RenameSmartname(const FName InContainerName, SmartName::
 			GWarn->BeginSlowTask(FText::Format(LOCTEXT("RenameCurvesTaskDesc", "Renaming curve for skeleton {0}"), FText::FromString(Skeleton->GetName())), true);
 			FScopedTransaction Transaction(LOCTEXT("RenameCurvesTransactionName", "Rename skeleton curve"));
 
+			FAnimationCurveIdentifier CurveId(CurveToRename, ERawCurveTrackTypes::RCT_Float);
+
 			// Remove curves from animation assets
 			for (FAssetData& Data : AnimationAssets)
 			{
@@ -314,16 +317,15 @@ void FEditableSkeleton::RenameSmartname(const FName InContainerName, SmartName::
 				{
 					SequenceBase->Modify();
 
-					if (FAnimCurveBase* CurrentCurveData = SequenceBase->RawCurveData.GetCurveData(CurveToRename.UID))
+					if (SequenceBase->GetDataModel()->FindFloatCurve(CurveId))
 					{
-						CurrentCurveData->Name.DisplayName = InNewName;
-						SequenceBase->MarkRawDataAsModified();
+						UAnimDataController* Controller = SequenceBase->GetController();
+
+						FAnimationCurveIdentifier NewCurveId(FSmartName(InNewName, InNameUid), ERawCurveTrackTypes::RCT_Float);
+						Controller->RenameCurve(CurveId, NewCurveId);
 						if (UAnimSequence* Seq = Cast<UAnimSequence>(SequenceBase))
 						{
 							SequencesToRecompress.Add(Seq);
-							PRAGMA_DISABLE_DEPRECATION_WARNINGS
-							Seq->ClearCompressedCurveData();
-							PRAGMA_ENABLE_DEPRECATION_WARNINGS
 						}
 					}
 				}
@@ -339,7 +341,6 @@ void FEditableSkeleton::RenameSmartname(const FName InContainerName, SmartName::
 			GWarn->BeginSlowTask(LOCTEXT("RebuildingAnimations", "Rebaking/compressing modified animations"), true);
 
 			//Make sure skeleton is correct before compression 
-
 			Skeleton->RenameSmartnameAndModify(InContainerName, InNameUid, InNewName);
 
 			// Rebake/compress the animations
@@ -477,15 +478,18 @@ void FEditableSkeleton::RemoveSmartnamesAndFixupAnimations(const FName& InContai
 				{
 					USkeleton* MySkeleton = Sequence->GetSkeleton();
 					Sequence->Modify(true);
+
+					UAnimDataController* Controller = Sequence->GetController();
+					UAnimDataController::FScopedBracket ScopedBracket(Controller, LOCTEXT("DeleteCurvesNotOnSkeletonBracket", "Deleting curves that no longer exist on Skeleton"));
 					for (FName Name : InNames)
 					{
 						FSmartName CurveToDelete;
 						if (MySkeleton->GetSmartNameByName(USkeleton::AnimCurveMappingName, Name, CurveToDelete))
 						{
-							Sequence->RawCurveData.DeleteCurveData(CurveToDelete);
+							const FAnimationCurveIdentifier CurveId(CurveToDelete, ERawCurveTrackTypes::RCT_Float);
+							Controller->RemoveCurve(CurveId);
 						}
 					}
-					Sequence->MarkRawDataAsModified();
 				}
 				else if (UPoseAsset* PoseAsset = Cast<UPoseAsset>(Asset))
 				{
@@ -495,18 +499,6 @@ void FEditableSkeleton::RemoveSmartnamesAndFixupAnimations(const FName& InContai
 			}
 			GWarn->EndSlowTask();
 
-			GWarn->BeginSlowTask(LOCTEXT("RebuildingAnimations", "Rebaking/compressing modified animations"), true);
-
-			// Rebake/compress the animations
-			for (TObjectIterator<UAnimSequence> It; It; ++It)
-			{
-				UAnimSequence* Seq = *It;
-
-				GWarn->StatusUpdate(1, 2, FText::Format(LOCTEXT("RebuildingAnimationsStatus", "Rebuilding {0}"), FText::FromString(Seq->GetName())));
-				Seq->RequestSyncAnimRecompression();
-			}
-			GWarn->EndSlowTask();
-
 			// Remove names from skeleton
 			Skeleton->RemoveSmartnamesAndModify(InContainerName, InNames);
 		}
diff --git a/Engine/Source/Editor/UnrealEd/Classes/Editor/EditorEngine.h b/Engine/Source/Editor/UnrealEd/Classes/Editor/EditorEngine.h
--- a/Engine/Source/Editor/UnrealEd/Classes/Editor/EditorEngine.h
+++ b/Engine/Source/Editor/UnrealEd/Classes/Editor/EditorEngine.h
@@ -894,6 +894,7 @@ public:
 	//~ Begin Transaction Interfaces.
 	virtual int32 BeginTransaction(const TCHAR* TransactionContext, const FText& Description, UObject* PrimaryObject) override;
 	int32 BeginTransaction(const FText& Description);
+	virtual bool CanTransact() override;
 	virtual int32 EndTransaction() override;
 	virtual void CancelTransaction(int32 Index) override;
 	void ResetTransaction(const FText& Reason);
diff --git a/Engine/Source/Editor/UnrealEd/Classes/Editor/Transactor.h b/Engine/Source/Editor/UnrealEd/Classes/Editor/Transactor.h
--- a/Engine/Source/Editor/UnrealEd/Classes/Editor/Transactor.h
+++ b/Engine/Source/Editor/UnrealEd/Classes/Editor/Transactor.h
@@ -14,6 +14,7 @@
 #include "Serialization/ArchiveUObject.h"
 #include "Misc/ITransaction.h"
 #include "Serialization/ArchiveSerializedPropertyChain.h"
+#include "Algo/Find.h"
 #include "Transactor.generated.h"
 
 /*-----------------------------------------------------------------------------
@@ -653,6 +654,34 @@ public:
 		return Title;
 	}
 
+	/** Returns the description of each contained Object Record */
+	FText GetDescription() const
+	{			
+		if (Algo::FindByPredicate(Records, [](const FObjectRecord& Record)
+		{
+			return Record.CustomChange.IsValid();
+		}))
+        {
+			FString ConcatenatedRecords;
+			for (const FObjectRecord& Record : Records)
+			{
+				if (Record.CustomChange.IsValid())
+				{
+					if (ConcatenatedRecords.Len())
+					{
+						ConcatenatedRecords += TEXT("\n");
+					}
+					ConcatenatedRecords += Record.CustomChange->ToString();
+				}
+			}
+
+			return FText::FromString(ConcatenatedRecords);
+        }
+
+		return GetTitle();
+	}
+
+
 	/** Serializes a reference to a transaction in a given archive. */
 	friend FArchive& operator<<( FArchive& Ar, FTransaction& T )
 	{
diff --git a/Engine/Source/Editor/UnrealEd/Private/EditorAnimUtils.cpp b/Engine/Source/Editor/UnrealEd/Private/EditorAnimUtils.cpp
--- a/Engine/Source/Editor/UnrealEd/Private/EditorAnimUtils.cpp
+++ b/Engine/Source/Editor/UnrealEd/Private/EditorAnimUtils.cpp
@@ -19,6 +19,8 @@
 #include "ContentBrowserModule.h"
 #include "Subsystems/AssetEditorSubsystem.h"
 #include "Editor.h"
+#include "Animation/AnimData/AnimDataController.h"
+#include "AnimationBlueprintLibrary.h"
 
 #define LOCTEXT_NAMESPACE "EditorAnimUtils"
 
@@ -251,11 +253,13 @@ namespace EditorAnimUtils
 				// Copy curve data from source asset, preserving data in the target if present.
 				if (OldSkeleton)
 				{
-					EditorAnimUtils::CopyAnimCurves(OldSkeleton, NewSkeleton, AnimSequenceToRetarget, USkeleton::AnimCurveMappingName, ERawCurveTrackTypes::RCT_Float);
+					UAnimationBlueprintLibrary::CopyAnimationCurveNamesToSkeleton(OldSkeleton, NewSkeleton, AnimSequenceToRetarget, ERawCurveTrackTypes::RCT_Float);
 
 					// clear transform curves since those curves won't work in new skeleton
 					// since we're deleting curves, mark this rebake flag off
-					AnimSequenceToRetarget->RawCurveData.TransformCurves.Empty();
+					UAnimDataController* Controller = AnimSequenceToRetarget->GetController();
+					Controller->RemoveAllCurvesOfType(ERawCurveTrackTypes::RCT_Transform);
+
 					// I can't copy transform curves yet because transform curves need retargeting. 
 					//EditorAnimUtils::CopyAnimCurves(OldSkeleton, NewSkeleton, AssetToRetarget, USkeleton::AnimTrackCurveMappingName, FRawCurveTracks::TransformType);
 				}
@@ -496,39 +500,9 @@ namespace EditorAnimUtils
 		}
 	}
 
-	void CopyAnimCurves(USkeleton* OldSkeleton, USkeleton* NewSkeleton, UAnimSequenceBase *SequenceBase, const FName ContainerName, ERawCurveTrackTypes CurveType )
+	void CopyAnimCurves(USkeleton* OldSkeleton, USkeleton* NewSkeleton, UAnimSequenceBase* SequenceBase, const FName ContainerName, ERawCurveTrackTypes CurveType)
 	{
-		// Copy curve data from source asset, preserving data in the target if present.
-		const FSmartNameMapping* OldNameMapping = OldSkeleton->GetSmartNameContainer(ContainerName);
-		SequenceBase->RawCurveData.RefreshName(OldNameMapping, CurveType);
-
-		switch (CurveType)
-		{
-		case ERawCurveTrackTypes::RCT_Float:
-			{
-				for(FFloatCurve& Curve : SequenceBase->RawCurveData.FloatCurves)
-				{
-					NewSkeleton->AddSmartNameAndModify(ContainerName, Curve.Name.DisplayName, Curve.Name);
-				}
-				break;
-			}
-		case ERawCurveTrackTypes::RCT_Vector:
-			{
-				for(FVectorCurve& Curve : SequenceBase->RawCurveData.VectorCurves)
-				{
-					NewSkeleton->AddSmartNameAndModify(ContainerName, Curve.Name.DisplayName, Curve.Name);
-				}
-				break;
-			}
-		case ERawCurveTrackTypes::RCT_Transform:
-			{
-				for(FTransformCurve& Curve : SequenceBase->RawCurveData.TransformCurves)
-				{
-					NewSkeleton->AddSmartNameAndModify(ContainerName, Curve.Name.DisplayName, Curve.Name);
-				}
-				break;
-			}
-		}
+		UAnimationBlueprintLibrary::CopyAnimationCurveNamesToSkeleton(OldSkeleton, NewSkeleton, SequenceBase, CurveType);
 	}
 
 	FString FNameDuplicationRule::Rename(const UObject* Asset) const
diff --git a/Engine/Source/Editor/UnrealEd/Private/EditorServer.cpp b/Engine/Source/Editor/UnrealEd/Private/EditorServer.cpp
--- a/Engine/Source/Editor/UnrealEd/Private/EditorServer.cpp
+++ b/Engine/Source/Editor/UnrealEd/Private/EditorServer.cpp
@@ -1096,6 +1096,11 @@ int32 UEditorEngine::BeginTransaction(const FText& Description)
 	return BeginTransaction(NULL, Description, NULL);
 }
 
+bool UEditorEngine::CanTransact()
+{
+	return Trans != nullptr;
+}
+
 int32 UEditorEngine::EndTransaction()
 {
 	int32 Index = INDEX_NONE;
diff --git a/Engine/Source/Editor/UnrealEd/Private/Factories/AnimCompositeFactory.cpp b/Engine/Source/Editor/UnrealEd/Private/Factories/AnimCompositeFactory.cpp
--- a/Engine/Source/Editor/UnrealEd/Private/Factories/AnimCompositeFactory.cpp
+++ b/Engine/Source/Editor/UnrealEd/Private/Factories/AnimCompositeFactory.cpp
@@ -89,7 +89,7 @@ UObject* UAnimCompositeFactory::FactoryCreateNew(UClass* Class, UObject* InParen
 			NewSegment.StartPos = 0.f;
 
 			AnimComposite->AnimationTrack.AnimSegments.Add(NewSegment);
-			AnimComposite->SetSequenceLength(AnimComposite->AnimationTrack.GetLength());
+			AnimComposite->SetCompositeLength(AnimComposite->AnimationTrack.GetLength());
 		}
 		
 		AnimComposite->SetSkeleton( TargetSkeleton );
diff --git a/Engine/Source/Editor/UnrealEd/Private/Factories/AnimMontageFactory.cpp b/Engine/Source/Editor/UnrealEd/Private/Factories/AnimMontageFactory.cpp
--- a/Engine/Source/Editor/UnrealEd/Private/Factories/AnimMontageFactory.cpp
+++ b/Engine/Source/Editor/UnrealEd/Private/Factories/AnimMontageFactory.cpp
@@ -91,7 +91,7 @@ UObject* UAnimMontageFactory::FactoryCreateNew(UClass* Class, UObject* InParent,
 			FSlotAnimationTrack& NewTrack = AnimMontage->SlotAnimTracks[0];
 			NewTrack.AnimTrack.AnimSegments.Add(NewSegment);
 
-			AnimMontage->SetSequenceLength(SourceAnimation->GetPlayLength());
+			AnimMontage->SetCompositeLength(SourceAnimation->GetPlayLength());
 		}
 
 		AnimMontage->SetSkeleton(TargetSkeleton);
diff --git a/Engine/Source/Editor/UnrealEd/Private/Factories/AnimSequenceFactory.cpp b/Engine/Source/Editor/UnrealEd/Private/Factories/AnimSequenceFactory.cpp
--- a/Engine/Source/Editor/UnrealEd/Private/Factories/AnimSequenceFactory.cpp
+++ b/Engine/Source/Editor/UnrealEd/Private/Factories/AnimSequenceFactory.cpp
@@ -66,12 +66,7 @@ bool UAnimSequenceFactory::ConfigureProperties()
 
 UObject* UAnimSequenceFactory::FactoryCreateNew(UClass* Class, UObject* InParent, FName Name, EObjectFlags Flags, UObject* Context, FFeedbackContext* Warn)
 {
-	//
 	UAnimSequence* AnimSequence = NewObject<UAnimSequence>(InParent, Class, Name, Flags);
-
-	// @todo I think this will crash, we should support differentoptions
-	AnimSequence->SetSequenceLength(0.f);
-	AnimSequence->SetNumberOfSampledKeys(0);
 	
 	if (TargetSkeleton)
 	{
diff --git a/Engine/Source/Editor/UnrealEd/Private/Fbx/FbxAnimationExport.cpp b/Engine/Source/Editor/UnrealEd/Private/Fbx/FbxAnimationExport.cpp
--- a/Engine/Source/Editor/UnrealEd/Private/Fbx/FbxAnimationExport.cpp
+++ b/Engine/Source/Editor/UnrealEd/Private/Fbx/FbxAnimationExport.cpp
@@ -32,7 +32,7 @@ namespace UnFbx
 			return false;
 		}
 
-		const float FrameRate = AnimSeq->GetSamplingFrameRate().AsDecimal();
+		const float FrameRate = AnimSeq->GetDataModel()->GetFrameRate().AsDecimal();
 		//Configure the scene time line
 		{
 			FbxGlobalSettings& SceneGlobalSettings = Scene->GetGlobalSettings();
@@ -332,7 +332,7 @@ void FFbxExporter::IterateInsideAnimSequence(const UAnimSequence* AnimSeq, float
 {
 	float AnimTime = AnimStartOffset;
 	float AnimEndTime = (AnimSeq->GetPlayLength() - AnimEndOffset);
-	double TimePerKey = AnimSeq->GetSamplingFrameRate().AsInterval();
+	double TimePerKey = AnimSeq->GetDataModel()->GetFrameRate().AsInterval();
 	const float AnimTimeIncrement = TimePerKey * AnimPlayRate;
 	uint32 AnimFrameIndex = 0;
 
diff --git a/Engine/Source/Editor/UnrealEd/Private/Fbx/ReimportFbxSceneFactory.cpp b/Engine/Source/Editor/UnrealEd/Private/Fbx/ReimportFbxSceneFactory.cpp
--- a/Engine/Source/Editor/UnrealEd/Private/Fbx/ReimportFbxSceneFactory.cpp
+++ b/Engine/Source/Editor/UnrealEd/Private/Fbx/ReimportFbxSceneFactory.cpp
@@ -1435,10 +1435,6 @@ EReimportResult::Type UReimportFbxSceneFactory::ReimportSkeletalMesh(void* VoidF
 								// Notify the asset registry
 								FAssetRegistryModule::AssetCreated(DestSeq);
 							}
-							else
-							{
-								DestSeq->CleanAnimSequenceForImport();
-							}
 							DestSeq->SetSkeleton(Mesh->GetSkeleton());
 							// since to know full path, reimport will need to do same
 							UFbxAnimSequenceImportData* ImportData = UFbxAnimSequenceImportData::GetImportDataForAnimSequence(DestSeq, AnimSequenceImportData);
diff --git a/Engine/Source/Editor/UnrealEd/Private/SkeletalMeshEdit.cpp b/Engine/Source/Editor/UnrealEd/Private/SkeletalMeshEdit.cpp
--- a/Engine/Source/Editor/UnrealEd/Private/SkeletalMeshEdit.cpp
+++ b/Engine/Source/Editor/UnrealEd/Private/SkeletalMeshEdit.cpp
@@ -33,6 +33,11 @@
 #include "SkelImport.h"
 #include "UObject/UObjectIterator.h"
 #include "UObject/Object.h"
+#include "ComponentReregisterContext.h"
+#include "Components/SkeletalMeshComponent.h"
+#include "FbxAnimUtils.h"
+#include "Animation/AnimData/AnimDataController.h"
+#include "Animation/AnimSequenceHelpers.h"
 
 #define LOCTEXT_NAMESPACE "SkeletalMeshEdit"
 
@@ -122,7 +127,7 @@ bool UEditorEngine::ReimportFbxAnimation( USkeleton* Skeleton, UAnimSequence* An
 
 	UnFbx::FFbxImporter* FbxImporter = UnFbx::FFbxImporter::GetInstance();
 	
-	const bool bPrevImportMorph = (AnimSequence->RawCurveData.FloatCurves.Num() > 0) ;
+	const bool bPrevImportMorph = (AnimSequence->GetDataModel()->GetNumberOfFloatCurves() > 0);
 
 	if (!ReimportUI)
 	{
@@ -637,10 +642,6 @@ UAnimSequence * UnFbx::FFbxImporter::ImportAnimations(USkeleton* Skeleton, UObje
 			// Notify the asset registry
 			FAssetRegistryModule::AssetCreated(DestSeq);
 		}
-		else
-		{
-			DestSeq->CleanAnimSequenceForImport();
-		}
 
 		DestSeq->SetSkeleton(Skeleton);
 
@@ -1229,38 +1230,41 @@ bool UnFbx::FFbxImporter::ImportCurveToAnimSequence(class UAnimSequence * Target
 		FSmartName NewName;
 		Skeleton->AddSmartNameAndModify(USkeleton::AnimCurveMappingName, Name, NewName);
 
-		FFloatCurve * CurveToImport = static_cast<FFloatCurve *>(TargetSequence->RawCurveData.GetCurveData(NewName.UID, ERawCurveTrackTypes::RCT_Float));
-		if(CurveToImport==NULL)
+		FAnimationCurveIdentifier FloatCurveId(NewName, ERawCurveTrackTypes::RCT_Float);
+
+
+		UAnimDataModel* DataModel = TargetSequence->GetDataModel();
+		UAnimDataController* Controller = TargetSequence->GetController();
+
+		const FFloatCurve* TargetCurve = DataModel->FindFloatCurve(FloatCurveId);
+		if (TargetCurve == nullptr)
 		{
-			if (TargetSequence->RawCurveData.AddCurveData(NewName, AACF_DefaultCurve | CurveFlags))
-			{
-				CurveToImport = static_cast<FFloatCurve *> (TargetSequence->RawCurveData.GetCurveData(NewName.UID, ERawCurveTrackTypes::RCT_Float));
-				CurveToImport->Name = NewName;
-			}
-			else
-			{
-				// this should not happen, we already checked before adding
-				ensureMsgf(0, TEXT("FBX Import: Critical error: no memory?"));
-			}
+			// Need to add the curve first
+			Controller->AddCurve(FloatCurveId, AACF_DefaultCurve | CurveFlags);
+			TargetCurve = DataModel->FindFloatCurve(FloatCurveId);
 		}
 		else
 		{
-			CurveToImport->FloatCurve.Reset();
-			// if existing add these curve flags. 
-			CurveToImport->SetCurveTypeFlags(CurveFlags | CurveToImport->GetCurveTypeFlags());
+			// Need to update any of the flags
+			Controller->SetCurveFlags(FloatCurveId, CurveFlags | TargetCurve->GetCurveTypeFlags());
 		}
+		
+		// Should be valid at this point
+		ensure(TargetCurve);
 
-		// update last observed name. If not, sometimes it adds new UID while fixing up that will confuse Compressed Raw Data
-		const FSmartNameMapping* Mapping = Skeleton->GetSmartNameContainer(USkeleton::AnimCurveMappingName);
-		TargetSequence->RawCurveData.RefreshName(Mapping);
+		Controller->UpdateCurveNamesFromSkeleton(Skeleton, ERawCurveTrackTypes::RCT_Float);
 
-		TargetSequence->MarkRawDataAsModified();
-		if (CurveToImport && ImportCurve(FbxCurve, CurveToImport->FloatCurve, AnimTimeSpan, ValueScale))
+		FRichCurve RichCurve;
+		if (ImportCurve(FbxCurve, RichCurve, AnimTimeSpan, ValueScale))
 		{
 			if (ImportOptions->bRemoveRedundantKeys)
 			{
-				CurveToImport->FloatCurve.RemoveRedundantKeys(SMALL_NUMBER);
+				RichCurve.RemoveRedundantKeys(SMALL_NUMBER);
 			}
+
+			// Set actual keys on curve within the model
+			Controller->SetCurveKeys(FloatCurveId, RichCurve.GetConstRefOfKeys());
+
 			return true;
 		}
 	}
@@ -1404,26 +1408,29 @@ bool ShouldImportCurve(FbxAnimCurve* Curve, bool bDoNotImportWithZeroValues)
 
 bool UnFbx::FFbxImporter::ImportAnimation(USkeleton* Skeleton, UAnimSequence * DestSeq, const FString& FileName, TArray<FbxNode*>& SortedLinks, TArray<FbxNode*>& NodeArray, FbxAnimStack* CurAnimStack, const int32 ResampleRate, const FbxTimeSpan AnimTimeSpan)
 {
-	//This destroy all previously imported animation raw data
-	DestSeq->CleanAnimSequenceForImport();
-
 	// @todo : the length might need to change w.r.t. sampling keys
 	FbxTime SequenceLength = AnimTimeSpan.GetDuration();
 	float PreviousSequenceLength = DestSeq->GetPlayLength();
 
-	// if you have one pose(thus 0.f duration), it still contains animation, so we'll need to consider that as MINIMUM_ANIMATION_LENGTH time length
-	DestSeq->SetSequenceLength(FGenericPlatformMath::Max<float>(SequenceLength.GetSecondDouble(), MINIMUM_ANIMATION_LENGTH));
+	UAnimDataController* Controller = DestSeq->GetController();
+	Controller->OpenBracket(LOCTEXT("ImportAnimation_Bracket", "Importing Animation"));
+
+	//This destroy all previously imported animation raw data
+	Controller->ResetModel();
 
-	if(PreviousSequenceLength > MINIMUM_ANIMATION_LENGTH && DestSeq->RawCurveData.FloatCurves.Num() > 0)
+	// if you have one pose(thus 0.f duration), it still contains animation, so we'll need to consider that as MINIMUM_ANIMATION_LENGTH time length
+	Controller->SetPlayLength(FGenericPlatformMath::Max<float>(SequenceLength.GetSecondDouble(), MINIMUM_ANIMATION_LENGTH));
+	if(PreviousSequenceLength > MINIMUM_ANIMATION_LENGTH && DestSeq->GetDataModel()->GetNumberOfFloatCurves() > 0)
 	{
 		// The sequence already existed when we began the import. We need to scale the key times for all curves to match the new 
 		// duration before importing over them. This is to catch any user-added curves
 		float ScaleFactor = DestSeq->GetPlayLength() / PreviousSequenceLength;
 		if (!FMath::IsNearlyEqual(ScaleFactor, 1.f))
 		{
-			for (FFloatCurve& Curve : DestSeq->RawCurveData.FloatCurves)
+			for (const FFloatCurve& Curve : DestSeq->GetDataModel()->GetFloatCurves())
 			{
-				Curve.FloatCurve.ScaleCurve(0.0f, ScaleFactor);
+				const FAnimationCurveIdentifier CurveId(Curve.Name, ERawCurveTrackTypes::RCT_Float);
+				Controller->ScaleCurve(CurveId, 0.f, ScaleFactor);
 			}
 		}
 	}
@@ -1433,21 +1440,25 @@ bool UnFbx::FFbxImporter::ImportAnimation(USkeleton* Skeleton, UAnimSequence * D
 
 	if (ImportOptions->bDeleteExistingMorphTargetCurves || ImportOptions->bDeleteExistingCustomAttributeCurves)
 	{
-		for (int32 CurveIdx=0; CurveIdx<DestSeq->RawCurveData.FloatCurves.Num(); ++CurveIdx)
+		TArray<FSmartName> CurveSmartNamesToRemove;
+		for (const FFloatCurve& Curve : DestSeq->GetDataModel()->GetFloatCurves())
 		{
-			auto& Curve = DestSeq->RawCurveData.FloatCurves[CurveIdx];
 			const FCurveMetaData* MetaData = MySkeleton->GetCurveMetaData(Curve.Name);
 			if (MetaData)
 			{
 				bool bDeleteCurve = MetaData->Type.bMorphtarget ? ImportOptions->bDeleteExistingMorphTargetCurves : ImportOptions->bDeleteExistingCustomAttributeCurves;
 				if (bDeleteCurve)
 				{
-					DestSeq->RawCurveData.FloatCurves.RemoveAt(CurveIdx, 1, false);
-					--CurveIdx;
+					CurveSmartNamesToRemove.Add(Curve.Name);
 				}
 			}
 		}
-		DestSeq->RawCurveData.FloatCurves.Shrink();
+
+		for (auto CurveName : CurveSmartNamesToRemove)
+		{
+			const FAnimationCurveIdentifier CurveId(CurveName, ERawCurveTrackTypes::RCT_Float);
+			Controller->RemoveCurve(CurveId);
+		}
 	}
 
 	if (ImportOptions->bDeleteExistingNonCurveCustomAttributes)
@@ -1484,16 +1495,15 @@ bool UnFbx::FFbxImporter::ImportAnimation(USkeleton* Skeleton, UAnimSequence * D
 	}
 	else
 	{
-		for (int32 CurveIdx = 0; CurveIdx < DestSeq->RawCurveData.FloatCurves.Num(); ++CurveIdx)
-		{
-			auto& Curve = DestSeq->RawCurveData.FloatCurves[CurveIdx];
-			const FCurveMetaData* MetaData = MySkeleton->GetCurveMetaData(Curve.Name);
-
-			if (MetaData && !MetaData->Type.bMorphtarget)
-			{
-				CurvesNotFound.Add(Curve.Name.DisplayName.ToString());
-			}
-		}
+	  // Store float curve tracks which use to exist on the animation
+	  for (const FFloatCurve& Curve : DestSeq->GetDataModel()->GetFloatCurves())
+	  {
+		  const FCurveMetaData* MetaData = MySkeleton->GetCurveMetaData(Curve.Name);
+		  if (MetaData && !MetaData->Type.bMorphtarget)
+		  {
+			  CurvesNotFound.Add(Curve.Name.DisplayName.ToString());
+		  }
+	  }
 	}
 
 	if (bReimportWarnings && CurvesNotFound.Num())
@@ -1514,28 +1524,12 @@ bool UnFbx::FFbxImporter::ImportAnimation(USkeleton* Skeleton, UAnimSequence * D
 	{
 		FbxNode* SkeletalMeshRootNode = NodeArray.Num() > 0 ? NodeArray[0] : nullptr;
 		ImportBoneTracks(Skeleton, AnimImportSettings, SkeletalMeshRootNode, ResampleRate, TransformDebugData, TotalNumKeys);
-	}
-	else if (CurveAttributeKeyCount > 0)
-	{
-		DestSeq->SetNumberOfSampledKeys(CurveAttributeKeyCount);
-		DestSeq->MarkRawDataAsModified();
-	}
-	// compress animation
-	{
-		GWarn->BeginSlowTask( LOCTEXT("BeginCompressAnimation", "Compress Animation"), true);
-		GWarn->StatusForceUpdate(1, 1, LOCTEXT("CompressAnimation", "Compressing Animation"));
 
-		DestSeq->PostProcessSequence();
-
-		// run debug mode
 		AnimationTransformDebug::OutputAnimationTransformDebugData(TransformDebugData, TotalNumKeys, RefSkeleton);
-		GWarn->EndSlowTask();
 	}
-
-	// Reregister skeletal mesh components so they reflect the updated animation
-	for (TObjectIterator<USkeletalMeshComponent> Iter; Iter; ++Iter)
+	else if (CurveAttributeKeyCount > 0)
 	{
-		FComponentReregisterContext ReregisterContext(*Iter);
+		Controller->SetFrameRate(FFrameRate(ResampleRate, 1));
 	}
 
 	// Import bone metadata to AnimSequence
@@ -1544,6 +1538,14 @@ bool UnFbx::FFbxImporter::ImportAnimation(USkeleton* Skeleton, UAnimSequence * D
 		ImportNodeCustomProperties(DestSeq, SkeletonNode, true);
 	}
 
+	Controller->NotifyPopulated();
+	Controller->CloseBracket();
+
+	// Reregister skeletal mesh components so they reflect the updated animation
+	for (TObjectIterator<USkeletalMeshComponent> Iter; Iter; ++Iter)
+	{
+		FComponentReregisterContext ReregisterContext(*Iter);
+	}
 	return true;
 }
 
@@ -1646,15 +1648,20 @@ void UnFbx::FFbxImporter::ImportAnimationCustomAttribute(FAnimCurveImportSetting
 	// Store float curve tracks which use to exist on the animation
 	UAnimSequence* DestSeq = AnimImportSettings.DestSeq;
 	USkeleton* MySkeleton = DestSeq->GetSkeleton();
-	OutCurvesNotFound.Reset(DestSeq->RawCurveData.FloatCurves.Num());
-	for (int32 CurveIdx = 0; CurveIdx < DestSeq->RawCurveData.FloatCurves.Num(); ++CurveIdx)
+
+	const UAnimDataModel* DataModel = DestSeq->GetDataModel();
+	const int32 NumFloatCurves = DataModel->GetNumberOfFloatCurves();
+	const FAnimationCurveData& CurveData = DataModel->GetCurveData();
+
+	OutCurvesNotFound.Reset(NumFloatCurves);
+
+	for (const FFloatCurve& FloatCurve : CurveData.FloatCurves)
 	{
-		auto& Curve = DestSeq->RawCurveData.FloatCurves[CurveIdx];
-		const FCurveMetaData* MetaData = MySkeleton->GetCurveMetaData(Curve.Name);
+		const FCurveMetaData* MetaData = MySkeleton->GetCurveMetaData(FloatCurve.Name);
 
 		if (MetaData && !MetaData->Type.bMorphtarget)
 		{
-			OutCurvesNotFound.Add(Curve.Name.DisplayName.ToString());
+			OutCurvesNotFound.Add(FloatCurve.Name.DisplayName.ToString());
 		}
 	}
 
@@ -1758,6 +1765,10 @@ void UnFbx::FFbxImporter::ImportBoneTracks(USkeleton* Skeleton, FAnimCurveImport
 	const FbxTime TimeComparisonThreshold = (KINDA_SMALL_NUMBER * static_cast<float>(FBXSDK_TC_SECOND));
 	const FReferenceSkeleton& RefSkeleton = Skeleton->GetReferenceSkeleton();
 
+	UAnimDataModel* DataModel = DestSeq->GetDataModel();
+	UAnimDataController* Controller = DestSeq->GetController();	
+	Controller->SetFrameRate(FFrameRate(ResampleRate, 1));
+
 	for (int32 SourceTrackIdx = 0; SourceTrackIdx < AnimImportSettings.FbxRawBoneNames.Num(); ++SourceTrackIdx)
 	{
 		int32 NumKeysForTrack = 0;
@@ -1875,9 +1886,11 @@ void UnFbx::FFbxImporter::ImportBoneTracks(USkeleton* Skeleton, FAnimCurveImport
 			if (bSuccess)
 			{
 				//add new track
-				int32 NewTrackIdx = DestSeq->AddNewRawTrack(BoneName, &RawTrack);
+				Controller->AddBoneTrack(BoneName);
+				Controller->SetBoneTrackKeys(BoneName, RawTrack.PosKeys, RawTrack.RotKeys, RawTrack.ScaleKeys);
+				const int32 TrackIndex = DestSeq->GetDataModel()->GetBoneTrackIndexByName(BoneName);
 
-				NewDebugData.SetTrackData(NewTrackIdx, BoneTreeIndex, BoneName);
+				NewDebugData.SetTrackData(TrackIndex, BoneTreeIndex, BoneName);
 
 				// add mapping to skeleton bone track
 				TransformDebugData.Add(NewDebugData);
@@ -1886,10 +1899,6 @@ void UnFbx::FFbxImporter::ImportBoneTracks(USkeleton* Skeleton, FAnimCurveImport
 
 		OutTotalNumKeys = FMath::Max(OutTotalNumKeys, NumKeysForTrack);
 	}
-
-	DestSeq->SetNumberOfSampledKeys(OutTotalNumKeys);
-
-	DestSeq->MarkRawDataAsModified();
 }
 
 #undef LOCTEXT_NAMESPACE
diff --git a/Engine/Source/Editor/UnrealEd/Public/EditorAnimUtils.h b/Engine/Source/Editor/UnrealEd/Public/EditorAnimUtils.h
--- a/Engine/Source/Editor/UnrealEd/Public/EditorAnimUtils.h
+++ b/Engine/Source/Editor/UnrealEd/Public/EditorAnimUtils.h
@@ -214,5 +214,7 @@ namespace EditorAnimUtils
 	}
 
 	// utility functions
-	UNREALED_API void CopyAnimCurves(USkeleton* OldSkeleton, USkeleton* NewSkeleton, UAnimSequenceBase *SequenceBase, const FName ContainerName, ERawCurveTrackTypes CurveType );
+	UE_DEPRECATED(5.0, "CopyAnimCurves has been deprecated, see UAnimationBlueprintLibrary::CopyAnimationCurveNamesToSkeleton instead")
+	UNREALED_API void CopyAnimCurves(USkeleton* OldSkeleton, USkeleton* NewSkeleton, UAnimSequenceBase* SequenceBase, const FName ContainerName, ERawCurveTrackTypes CurveType);
+
 } // namespace EditorAnimUtils
diff --git a/Engine/Source/Editor/UnrealEd/UnrealEd.Build.cs b/Engine/Source/Editor/UnrealEd/UnrealEd.Build.cs
--- a/Engine/Source/Editor/UnrealEd/UnrealEd.Build.cs
+++ b/Engine/Source/Editor/UnrealEd/UnrealEd.Build.cs
@@ -195,7 +195,8 @@ public class UnrealEd : ModuleRules
 				"IoStoreUtilities",
 				"EditorInteractiveToolsFramework",
 				"TraceLog",
-				"DeveloperSettings"
+				"DeveloperSettings",
+				"AnimationModifiers"
 			}
 		);
 
@@ -251,7 +252,6 @@ public class UnrealEd : ModuleRules
 				"PortalServices",
 				"BlueprintNativeCodeGen",
 				"OverlayEditor",
-				"AnimationModifiers",
 				"ClothPainter",
 				"Media",
 				"TimeManagementEditor",
diff --git a/Engine/Source/Runtime/Core/Private/Misc/Change.cpp b/Engine/Source/Runtime/Core/Private/Misc/Change.cpp
--- a/Engine/Source/Runtime/Core/Private/Misc/Change.cpp
+++ b/Engine/Source/Runtime/Core/Private/Misc/Change.cpp
@@ -54,7 +54,17 @@ FString FCompoundChange::ToString() const
 			}
 		}
 
-		Text += FString::Printf( TEXT( " (%i sub-change%s)" ), TotalValidSubchanges, TotalValidSubchanges == 1 ? TEXT( "s" ) : TEXT( "" ) );
+		Text += FString::Printf( TEXT( " (%i sub-change%s):" ), TotalValidSubchanges, TotalValidSubchanges == 1 ? TEXT( "s" ) : TEXT( "" ) );
+
+		for (int32 ChangeIndex = Input.Subchanges.Num() - 1; ChangeIndex >= 0; --ChangeIndex)
+		{
+			const TUniquePtr<FChange>& Subchange = Input.Subchanges[ChangeIndex];
+			if (Subchange != nullptr)
+			{
+				Text += TEXT("\n\t");
+				Text += Subchange->ToString();
+			}
+		}
 	}
 	else
 	{
diff --git a/Engine/Source/Runtime/Core/Public/Misc/Change.h b/Engine/Source/Runtime/Core/Public/Misc/Change.h
--- a/Engine/Source/Runtime/Core/Public/Misc/Change.h
+++ b/Engine/Source/Runtime/Core/Public/Misc/Change.h
@@ -81,19 +81,19 @@ class CORE_API FSwapChange : public FChange
 {
 
 public:
-	virtual EChangeStyle GetChangeType() override
+	virtual EChangeStyle GetChangeType() final
 	{
 		return FChange::EChangeStyle::InPlaceSwap;
 	}
 
 	/** Makes the change to the object */
-	virtual void Apply(UObject* Object) override
+	virtual void Apply(UObject* Object) final
 	{
 		check(false);
 	}
 
 	/** Reverts change to the object */
-	virtual void Revert(UObject* Object) override
+	virtual void Revert(UObject* Object) final
 	{
 		check(false);
 	}
@@ -109,12 +109,12 @@ class CORE_API FCommandChange : public FChange
 {
 
 public:
-	virtual EChangeStyle GetChangeType() override
+	virtual EChangeStyle GetChangeType() final
 	{
 		return FChange::EChangeStyle::CommandPattern;
 	}
 
-	virtual TUniquePtr<FChange> Execute(UObject* Object)
+	virtual TUniquePtr<FChange> Execute(UObject* Object) final
 	{
 		check(false);
 		return nullptr;
diff --git a/Engine/Source/Runtime/Core/Public/UObject/UE5MainStreamObjectVersion.h b/Engine/Source/Runtime/Core/Public/UObject/UE5MainStreamObjectVersion.h
--- a/Engine/Source/Runtime/Core/Public/UObject/UE5MainStreamObjectVersion.h
+++ b/Engine/Source/Runtime/Core/Public/UObject/UE5MainStreamObjectVersion.h
@@ -73,6 +73,9 @@ struct CORE_API FUE5MainStreamObjectVersion
 		// FoliageType with HLOD support
 		FoliageTypeIncludeInHLOD,
 
+		// Introducing UAnimDataModel sub-object for UAnimSequenceBase containing all animation source data
+		IntroducingAnimationDataModel,
+
 		// -----<new versions can be added above this line>-------------------------------------------------
 		VersionPlusOne,
 		LatestVersion = VersionPlusOne - 1
diff --git a/Engine/Source/Runtime/Engine/Classes/Animation/AnimComposite.h b/Engine/Source/Runtime/Engine/Classes/Animation/AnimComposite.h
--- a/Engine/Source/Runtime/Engine/Classes/Animation/AnimComposite.h
+++ b/Engine/Source/Runtime/Engine/Classes/Animation/AnimComposite.h
@@ -54,6 +54,7 @@ public:
 	//~ Begin UAnimCompositeBase Interface
 	virtual void InvalidateRecursiveAsset() override;
 	virtual bool ContainRecursive(TArray<UAnimCompositeBase*>& CurrentAccumulatedList) override;
+	virtual void SetCompositeLength(float InLength) override;
 	//~End UAnimCompositeBase Interface
 };
 
diff --git a/Engine/Source/Runtime/Engine/Classes/Animation/AnimCompositeBase.h b/Engine/Source/Runtime/Engine/Classes/Animation/AnimCompositeBase.h
--- a/Engine/Source/Runtime/Engine/Classes/Animation/AnimCompositeBase.h
+++ b/Engine/Source/Runtime/Engine/Classes/Animation/AnimCompositeBase.h
@@ -310,5 +310,7 @@ class UAnimCompositeBase : public UAnimSequenceBase
 	// and clear the reference if recursive is found. 
 	// We're going to remove the top reference if found
 	virtual bool ContainRecursive(TArray<UAnimCompositeBase*>& CurrentAccumulatedList) PURE_VIRTUAL(UAnimCompositeBase::ContainRecursive, return false; );
+
+	virtual void SetCompositeLength(float InLength) PURE_VIRTUAL(UAnimCompositeBase::SetCompositeLength, );
 };
 
