Anim Notifies get clamped to 0.0 when reimporting an FBX
#fix Rather than resetting the entire model remove only the bone animation tracks when (re-)importing an AnimSequence from FBX.

diff --git a/Engine/Source/Editor/UnrealEd/Private/SkeletalMeshEdit.cpp b/Engine/Source/Editor/UnrealEd/Private/SkeletalMeshEdit.cpp
--- a/Engine/Source/Editor/UnrealEd/Private/SkeletalMeshEdit.cpp
+++ b/Engine/Source/Editor/UnrealEd/Private/SkeletalMeshEdit.cpp
@@ -1416,7 +1416,7 @@ bool UnFbx::FFbxImporter::ImportAnimation(USkeleton* Skeleton, UAnimSequence * D
 	Controller->OpenBracket(LOCTEXT("ImportAnimation_Bracket", "Importing Animation"));
 
 	//This destroy all previously imported animation raw data
-	Controller->ResetModel();
+	Controller->RemoveAllBoneTracks();
 
 	// if you have one pose(thus 0.f duration), it still contains animation, so we'll need to consider that as MINIMUM_ANIMATION_LENGTH time length
 	Controller->SetPlayLength(FGenericPlatformMath::Max<float>(SequenceLength.GetSecondDouble(), MINIMUM_ANIMATION_LENGTH));
