Editor: Load WorldPartitionEditor module always

diff --git a/Engine/Source/Editor/UnrealEd/Private/EditorEngine.cpp b/Engine/Source/Editor/UnrealEd/Private/EditorEngine.cpp
--- a/Engine/Source/Editor/UnrealEd/Private/EditorEngine.cpp
+++ b/Engine/Source/Editor/UnrealEd/Private/EditorEngine.cpp
@@ -1020,7 +1020,8 @@ void UEditorEngine::Init(IEngineLoop* InEngineLoop)
 			TEXT("EditorInteractiveToolsFramework"),
 			TEXT("TraceInsights"),
 			TEXT("StaticMeshEditor"),
-			TEXT("EditorFramework")
+			TEXT("EditorFramework"),
+			TEXT("WorldPartitionEditor")
 		};
 
 		FScopedSlowTask ModuleSlowTask(UE_ARRAY_COUNT(ModuleNames));
