Add some Insight tags for world partition commandlet

diff --git a/Engine/Source/Editor/UnrealEd/Private/Commandlets/WorldPartitionConvertCommandlet.cpp b/Engine/Source/Editor/UnrealEd/Private/Commandlets/WorldPartitionConvertCommandlet.cpp
--- a/Engine/Source/Editor/UnrealEd/Private/Commandlets/WorldPartitionConvertCommandlet.cpp
+++ b/Engine/Source/Editor/UnrealEd/Private/Commandlets/WorldPartitionConvertCommandlet.cpp
@@ -161,6 +161,8 @@ UWorldPartitionConvertCommandlet::UWorldPartitionConvertCommandlet(const FObject
 
 UWorld* UWorldPartitionConvertCommandlet::LoadWorld(const FString& LevelToLoad)
 {
+	TRACE_CPUPROFILER_EVENT_SCOPE(UWorldPartitionConvertCommandlet::LoadWorld);
+
 	SET_WARN_COLOR(COLOR_WHITE);
 	UE_LOG(LogWorldPartitionConvertCommandlet, Log, TEXT("Loading level %s."), *LevelToLoad);
 	CLEAR_WARN_COLOR();
@@ -177,6 +179,8 @@ UWorld* UWorldPartitionConvertCommandlet::LoadWorld(const FString& LevelToLoad)
 
 ULevel* UWorldPartitionConvertCommandlet::InitWorld(UWorld* World)
 {
+	TRACE_CPUPROFILER_EVENT_SCOPE(UWorldPartitionConvertCommandlet::InitWorld);
+
 	SET_WARN_COLOR(COLOR_WHITE);
 	UE_LOG(LogWorldPartitionConvertCommandlet, Log, TEXT("Initializing level %s."), *World->GetName());
 	CLEAR_WARN_COLOR();
@@ -207,6 +211,8 @@ ULevel* UWorldPartitionConvertCommandlet::InitWorld(UWorld* World)
 
 UWorldPartition* UWorldPartitionConvertCommandlet::CreateWorldPartition(AWorldSettings* MainWorldSettings, UWorldComposition* WorldComposition) const
 {
+	TRACE_CPUPROFILER_EVENT_SCOPE(UWorldPartitionConvertCommandlet::CreateWorldPartition);
+
 	UWorldPartition* WorldPartition = NewObject<UWorldPartition>(MainWorldSettings);
 	MainWorldSettings->SetWorldPartition(WorldPartition);
 
@@ -234,6 +240,8 @@ UWorldPartition* UWorldPartitionConvertCommandlet::CreateWorldPartition(AWorldSe
 
 void UWorldPartitionConvertCommandlet::GatherAndPrepareSubLevelsToConvert(const UWorldPartition* WorldPartition, ULevel* Level, TArray<ULevel*>& SubLevels)
 {
+	TRACE_CPUPROFILER_EVENT_SCOPE(UWorldPartitionConvertCommandlet::GatherAndPrepareSubLevelsToConvert);
+
 	UWorld* World = Level->GetTypedOuter<UWorld>();	
 
 	// Set all streaming levels to be loaded/visible for next Flush
@@ -272,6 +280,8 @@ EActorGridPlacement UWorldPartitionConvertCommandlet::GetLevelGridPlacement(ULev
 
 bool UWorldPartitionConvertCommandlet::PrepareStreamingLevelForConversion(const UWorldPartition* WorldPartition, ULevelStreaming* StreamingLevel)
 {
+	TRACE_CPUPROFILER_EVENT_SCOPE(UWorldPartitionConvertCommandlet::PrepareStreamingLevelForConversion);
+
 	ULevel* SubLevel = StreamingLevel->GetLoadedLevel();
 	check(SubLevel);
 
@@ -406,6 +416,8 @@ void UWorldPartitionConvertCommandlet::ChangeObjectOuter(UObject* Object, UObjec
 
 void UWorldPartitionConvertCommandlet::FixupSoftObjectPaths(UPackage* OuterPackage)
 {
+	TRACE_CPUPROFILER_EVENT_SCOPE(UWorldPartitionConvertCommandlet::FixupSoftObjectPaths);
+
 	UE_SCOPED_TIMER(TEXT("FixupSoftObjectPaths"), LogWorldPartitionConvertCommandlet);
 
 	struct FSoftPathFixupSerializer : public FArchiveUObject
@@ -502,6 +514,8 @@ void UWorldPartitionConvertCommandlet::FixupSoftObjectPaths(UPackage* OuterPacka
 
 bool UWorldPartitionConvertCommandlet::DetachDependantLevelPackages(ULevel* Level)
 {
+	TRACE_CPUPROFILER_EVENT_SCOPE(UWorldPartitionConvertCommandlet::DetachDependantLevelPackages);
+
 	if (Level->MapBuildData && (Level->MapBuildData->GetPackage() != Level->GetPackage()))
 	{
 		PackagesToDelete.Add(Level->MapBuildData->GetPackage());
@@ -671,6 +685,8 @@ void UWorldPartitionConvertCommandlet::OnWorldLoaded(UWorld* World)
 
 void UWorldPartitionConvertCommandlet::CreateWorldMiniMapTexture(UWorld* World)
 {
+	TRACE_CPUPROFILER_EVENT_SCOPE(UWorldPartitionConvertCommandlet::CreateWorldMiniMapTexture);
+
 	UE_SCOPED_TIMER(TEXT("CreateWorldMiniMapTexture"), LogWorldPartitionConvertCommandlet);
 
 	AWorldPartitionMiniMap* WorldMiniMap = FWorldPartitionMiniMapHelper::GetWorldPartitionMiniMap(World, /*bCreateNewMiniMap*/true);
@@ -685,6 +701,8 @@ void UWorldPartitionConvertCommandlet::CreateWorldMiniMapTexture(UWorld* World)
 
 int32 UWorldPartitionConvertCommandlet::Main(const FString& Params)
 {
+	TRACE_CPUPROFILER_EVENT_SCOPE(UWorldPartitionConvertCommandlet::Main);
+
 	UE_SCOPED_TIMER(TEXT("Conversion"), LogWorldPartitionConvertCommandlet);
 
 	TArray<FString> Tokens, Switches;
@@ -834,6 +852,8 @@ int32 UWorldPartitionConvertCommandlet::Main(const FString& Params)
 	UActorPartitionSubsystem* ActorPartitionSubsystem = MainWorld->GetSubsystem<UActorPartitionSubsystem>();
 	auto PartitionFoliage = [this, MainWorld, ActorPartitionSubsystem](AInstancedFoliageActor* IFA)
 	{
+		TRACE_CPUPROFILER_EVENT_SCOPE(PartitionFoliage);
+
 		for (auto& Pair : IFA->FoliageInfos)
 		{
 			FFoliageInfo* FoliageInfo = &Pair.Value.Get();
@@ -870,6 +890,8 @@ int32 UWorldPartitionConvertCommandlet::Main(const FString& Params)
 
 	auto PartitionLandscape = [this, MainWorld](ULandscapeInfo* LandscapeInfo)
 	{
+		TRACE_CPUPROFILER_EVENT_SCOPE(PartitionLandscape);
+
 		// Handle Landscapes with missing LandscapeActor(s)
 		if (!LandscapeInfo->LandscapeActor.Get())
 		{
@@ -928,6 +950,8 @@ int32 UWorldPartitionConvertCommandlet::Main(const FString& Params)
 
 	auto PrepareLevelActors = [this, PartitionFoliage, PartitionLandscape, MainWorldDataLayers](ULevel* Level, bool bMainLevel, EActorGridPlacement DefaultGridPlacement)
 	{
+		TRACE_CPUPROFILER_EVENT_SCOPE(PrepareLevelActors);
+
 		const FBox WorldBounds(WorldOrigin - WorldExtent, WorldOrigin + WorldExtent);
 
 		TArray<AInstancedFoliageActor*> IFAs;
@@ -1073,6 +1097,8 @@ int32 UWorldPartitionConvertCommandlet::Main(const FString& Params)
 	TMap<UObject*, UObject*> PrivateRefsMap;
 	for(ULevel* SubLevel : SubLevelsToConvert)
 	{
+		TRACE_CPUPROFILER_EVENT_SCOPE(ConvertSubLevel);
+
 		UWorld* SubWorld = SubLevel->GetTypedOuter<UWorld>();
 		UPackage* SubPackage = SubLevel->GetPackage();
 
@@ -1254,6 +1280,8 @@ int32 UWorldPartitionConvertCommandlet::Main(const FString& Params)
 
 		if (bDeleteSourceLevels)
 		{
+			TRACE_CPUPROFILER_EVENT_SCOPE(DeleteSourceLevels);
+
 			for (UPackage* Package : PackagesToDelete)
 			{
 				if (!PackageHelper.Delete(Package))
@@ -1263,12 +1291,16 @@ int32 UWorldPartitionConvertCommandlet::Main(const FString& Params)
 			}
 		}
 
-		UE_LOG(LogWorldPartitionConvertCommandlet, Log, TEXT("Checking out %d actor packages."), PackagesToSave.Num());
-		for(UPackage* Package: PackagesToSave)
 		{
-			if (!PackageHelper.Checkout(Package))
+			TRACE_CPUPROFILER_EVENT_SCOPE(CheckoutPackages);
+
+			UE_LOG(LogWorldPartitionConvertCommandlet, Log, TEXT("Checking out %d actor packages."), PackagesToSave.Num());
+			for(UPackage* Package: PackagesToSave)
 			{
-				return 1;
+				if (!PackageHelper.Checkout(Package))
+				{
+					return 1;
+				}
 			}
 		}
 
@@ -1282,22 +1314,30 @@ int32 UWorldPartitionConvertCommandlet::Main(const FString& Params)
 		}
 	
 		// Save packages
-		UE_LOG(LogWorldPartitionConvertCommandlet, Log, TEXT("Saving %d packages."), PackagesToSave.Num());
-		for (UPackage* PackageToSave : PackagesToSave)
 		{
-			FString PackageFileName = SourceControlHelpers::PackageFilename(PackageToSave);
-			if (!UPackage::SavePackage(PackageToSave, nullptr, RF_Standalone, *PackageFileName, GError, nullptr, false, true, SAVE_Async))
+			TRACE_CPUPROFILER_EVENT_SCOPE(SavePackages);
+
+			UE_LOG(LogWorldPartitionConvertCommandlet, Log, TEXT("Saving %d packages."), PackagesToSave.Num());
+			for (UPackage* PackageToSave : PackagesToSave)
 			{
-				return 1;
+				FString PackageFileName = SourceControlHelpers::PackageFilename(PackageToSave);
+				if (!UPackage::SavePackage(PackageToSave, nullptr, RF_Standalone, *PackageFileName, GError, nullptr, false, true, SAVE_Async))
+				{
+					return 1;
+				}
 			}
 		}
 
-		// Add new packages to source control
-		for(UPackage* PackageToSave: PackagesToSave)
 		{
-			if(!PackageHelper.AddToSourceControl(PackageToSave))
+			TRACE_CPUPROFILER_EVENT_SCOPE(AddPackagesToSourceControl);
+
+			// Add new packages to source control
+			for(UPackage* PackageToSave: PackagesToSave)
 			{
-				return 1;
+				if(!PackageHelper.AddToSourceControl(PackageToSave))
+				{
+					return 1;
+				}
 			}
 		}
 
@@ -1310,10 +1350,14 @@ int32 UWorldPartitionConvertCommandlet::Main(const FString& Params)
 		}
 
 		// Save level
-		FString PackageFileName = SourceControlHelpers::PackageFilename(MainPackage);
-		if (!UPackage::SavePackage(MainPackage, nullptr, RF_Standalone, *PackageFileName, GError, nullptr, false, true, SAVE_Async))
 		{
-			return 1;
+			TRACE_CPUPROFILER_EVENT_SCOPE(SaveLevel);
+
+			FString PackageFileName = SourceControlHelpers::PackageFilename(MainPackage);
+			if (!UPackage::SavePackage(MainPackage, nullptr, RF_Standalone, *PackageFileName, GError, nullptr, false, true, SAVE_Async))
+			{
+				return 1;
+			}
 		}
 
 		if(bInitializedPhysicsSceneForSave)
diff --git a/Engine/Source/Runtime/CoreUObject/Private/UObject/SavePackage/SavePackageUtilities.cpp b/Engine/Source/Runtime/CoreUObject/Private/UObject/SavePackage/SavePackageUtilities.cpp
--- a/Engine/Source/Runtime/CoreUObject/Private/UObject/SavePackage/SavePackageUtilities.cpp
+++ b/Engine/Source/Runtime/CoreUObject/Private/UObject/SavePackage/SavePackageUtilities.cpp
@@ -2053,6 +2053,8 @@ void SaveWorldLevelInfo(UPackage* InOuter, FLinkerSave* Linker, FStructuredArchi
 
 void UPackage::WaitForAsyncFileWrites()
 {
+	TRACE_CPUPROFILER_EVENT_SCOPE(UPackage::WaitForAsyncFileWrites);
+
 	while (OutstandingAsyncWrites.GetValue())
 	{
 		FPlatformProcess::Sleep(0.0f);
diff --git a/Engine/Source/Runtime/Engine/Private/World.cpp b/Engine/Source/Runtime/Engine/Private/World.cpp
--- a/Engine/Source/Runtime/Engine/Private/World.cpp
+++ b/Engine/Source/Runtime/Engine/Private/World.cpp
@@ -3826,6 +3826,8 @@ void UWorld::UpdateStreamingLevelPriority(ULevelStreaming* StreamingLevel)
 
 void UWorld::FlushLevelStreaming(EFlushLevelStreamingType FlushType)
 {
+	TRACE_CPUPROFILER_EVENT_SCOPE(UWorld::FlushLevelStreaming);
+
 	if (FlushType == FlushLevelStreamingType
 		|| FlushLevelStreamingType == EFlushLevelStreamingType::Full
 		|| FlushType == EFlushLevelStreamingType::None)
