WorldPartition: Avoid Conversion prompt on Default map load

diff --git a/Engine/Source/Editor/UnrealEd/Private/FileHelpers.cpp b/Engine/Source/Editor/UnrealEd/Private/FileHelpers.cpp
--- a/Engine/Source/Editor/UnrealEd/Private/FileHelpers.cpp
+++ b/Engine/Source/Editor/UnrealEd/Private/FileHelpers.cpp
@@ -4077,6 +4077,14 @@ void FEditorFileUtils::LoadDefaultMapAtStartup()
 	const bool bIncludeReadOnlyRoots = true;
 	if ( FPackageName::IsValidLongPackageName(EditorStartupMap, bIncludeReadOnlyRoots) )
 	{
+		IWorldPartitionEditorModule& WorldPartitionEditorModule = FModuleManager::LoadModuleChecked<IWorldPartitionEditorModule>("WorldPartitionEditor");
+		bool bPreviousConversionPromptEnabled = WorldPartitionEditorModule.IsConversionPromptEnabled();
+		WorldPartitionEditorModule.SetConversionPromptEnabled(false);
+		ON_SCOPE_EXIT
+		{
+			WorldPartitionEditorModule.SetConversionPromptEnabled(bPreviousConversionPromptEnabled);
+		};
+		
 		FString MapFilenameToLoad = FPackageName::LongPackageNameToFilename( EditorStartupMap );
 
 		bIsLoadingDefaultStartupMap = true;
