Pass through hit actor to context menu

This should eventually be change to be an element too, however this will keep things working for now...

#rnx

diff --git a/Engine/Source/Editor/UnrealEd/Private/LevelViewportClickHandlers.cpp b/Engine/Source/Editor/UnrealEd/Private/LevelViewportClickHandlers.cpp
--- a/Engine/Source/Editor/UnrealEd/Private/LevelViewportClickHandlers.cpp
+++ b/Engine/Source/Editor/UnrealEd/Private/LevelViewportClickHandlers.cpp
@@ -31,6 +31,7 @@
 #include "LightMap.h"
 #include "Subsystems/BrushEditingSubsystem.h"
 #include "Elements/Framework/TypedElementSelectionSet.h"
+#include "Elements/Interfaces/TypedElementObjectInterface.h"
 
 #define LOCTEXT_NAMESPACE "ClickHandlers"
 
@@ -214,7 +215,9 @@ namespace LevelViewportClickHandlers
 
 				if (bIsRightClickSelection)
 				{
-					PrivateSummonContextMenu(ViewportClient);
+					TTypedElement<UTypedElementObjectInterface> ObjectElement = LevelEditorElementSelectionSet->GetElementList()->GetElement<UTypedElementObjectInterface>(ResolvedElement);
+					AActor* HitActor = ObjectElement ? Cast<AActor>(ObjectElement.GetObject()) : nullptr;
+					PrivateSummonContextMenu(ViewportClient, HitActor);
 				}
 			}
 		}
