Add missing pragma.

#rnx

diff --git a/Engine/Source/Editor/UnrealEd/Private/Factories/EditorStaticMeshFactory.h b/Engine/Source/Editor/UnrealEd/Private/Factories/EditorStaticMeshFactory.h
--- a/Engine/Source/Editor/UnrealEd/Private/Factories/EditorStaticMeshFactory.h
+++ b/Engine/Source/Editor/UnrealEd/Private/Factories/EditorStaticMeshFactory.h
@@ -1,5 +1,7 @@
 // Copyright Epic Games, Inc. All Rights Reserved.
 
+#pragma once
+
 #include "Factories/AssetFactoryInterface.h"
 #include "ActorFactories/ActorFactoryStaticMesh.h"
 
