Fix unused material and texture for the editor grid being loaded at startup causing additional caching. These assets were only used if an off by default cvar was turned on.

diff --git a/Engine/Source/Editor/UnrealEd/Private/EditorComponents.cpp b/Engine/Source/Editor/UnrealEd/Private/EditorComponents.cpp
--- a/Engine/Source/Editor/UnrealEd/Private/EditorComponents.cpp
+++ b/Engine/Source/Editor/UnrealEd/Private/EditorComponents.cpp
@@ -79,20 +79,41 @@ FGridWidget.
 ------------------------------------------------------------------------------*/
 
 FGridWidget::FGridWidget()
+	: GPUBasedGridMaterial(FSoftObjectPath(TEXT("/Engine/EditorMaterials/LevelGridMaterial.LevelGridMaterial")))
+	, TextureBasedLevelGridMaterial(FSoftObjectPath(TEXT("/Engine/EditorMaterials/LevelGridMaterial2.LevelGridMaterial2")))
+	, LevelGridMaterialInst(nullptr)
 {
-	LevelGridMaterial = (UMaterial*)StaticLoadObject( UMaterial::StaticClass(),NULL,TEXT("/Engine/EditorMaterials/LevelGridMaterial.LevelGridMaterial"),NULL,LOAD_None,NULL );
-	LevelGridMaterialInst = UMaterialInstanceDynamic::Create(LevelGridMaterial, NULL);
-
-	LevelGridMaterial2 = (UMaterial*)StaticLoadObject( UMaterial::StaticClass(),NULL,TEXT("/Engine/EditorMaterials/LevelGridMaterial2.LevelGridMaterial2"),NULL,LOAD_None,NULL );
-	LevelGridMaterialInst2 = UMaterialInstanceDynamic::Create(LevelGridMaterial2, NULL);
 }
 
 void FGridWidget::AddReferencedObjects( FReferenceCollector& Collector )
 {
-	Collector.AddReferencedObject( LevelGridMaterial );
 	Collector.AddReferencedObject( LevelGridMaterialInst );
-	Collector.AddReferencedObject( LevelGridMaterial2 );
-	Collector.AddReferencedObject( LevelGridMaterialInst2 );
+}
+
+UMaterial* FGridWidget::GetActiveLevelGridMaterial()
+{
+	bool bUseTextureSolution = CVarEditorNewLevelGrid.GetValueOnGameThread() > 1;
+
+	if (bUseTextureSolution)
+	{
+		return TextureBasedLevelGridMaterial.LoadSynchronous();
+	}
+	
+	return GPUBasedGridMaterial.LoadSynchronous();
+}
+
+UMaterialInstanceDynamic* FGridWidget::GetActiveLevelGridMID()
+{
+	bool bUseTextureSolution = CVarEditorNewLevelGrid.GetValueOnGameThread() > 1;
+
+	UMaterial* BaseGridMat = GetActiveLevelGridMaterial();
+
+	if (!LevelGridMaterialInst || LevelGridMaterialInst->Parent != BaseGridMat)
+	{
+		LevelGridMaterialInst = UMaterialInstanceDynamic::Create(BaseGridMat, nullptr);
+	}
+
+	return LevelGridMaterialInst;
 }
 
 static void GetAxisColors(FLinearColor Out[3], bool b3D)
@@ -120,9 +141,7 @@ static void GetAxisColors(FLinearColor Out[3], bool b3D)
 
 void FGridWidget::DrawNewGrid(const FSceneView* View, FPrimitiveDrawInterface* PDI)
 {
-	bool bUseTextureSolution = CVarEditorNewLevelGrid.GetValueOnGameThread() > 1;
-	UMaterial* GridMaterial = bUseTextureSolution ? LevelGridMaterial2 : LevelGridMaterial;
-	UMaterialInstanceDynamic* MaterialInst = bUseTextureSolution ? LevelGridMaterialInst2 : LevelGridMaterialInst;
+	UMaterial* GridMaterial = GetActiveLevelGridMaterial();
 
 	if (GridMaterial->IsCompilingOrHadCompileError(View->GetFeatureLevel()))
 	{
@@ -131,6 +150,8 @@ void FGridWidget::DrawNewGrid(const FSceneView* View, FPrimitiveDrawInterface* P
 		return;
 	}
 
+	UMaterialInstanceDynamic* MaterialInst = GetActiveLevelGridMID();
+
 	if(!MaterialInst)
 	{
 		return;
diff --git a/Engine/Source/Editor/UnrealEd/Public/EditorComponents.h b/Engine/Source/Editor/UnrealEd/Public/EditorComponents.h
--- a/Engine/Source/Editor/UnrealEd/Public/EditorComponents.h
+++ b/Engine/Source/Editor/UnrealEd/Public/EditorComponents.h
@@ -33,13 +33,16 @@ public:
 		return "FGridWidget";
 	}
 
+private:
+	UMaterial* GetActiveLevelGridMaterial();
+	UMaterialInstanceDynamic* GetActiveLevelGridMID();
 private:
 	// ALU based
-	UMaterial* LevelGridMaterial;
-	UMaterialInstanceDynamic* LevelGridMaterialInst;
+	TSoftObjectPtr<UMaterial> GPUBasedGridMaterial;
 	// texture based
-	UMaterial* LevelGridMaterial2;
-	UMaterialInstanceDynamic* LevelGridMaterialInst2;
+	TSoftObjectPtr<UMaterial> TextureBasedLevelGridMaterial;
+
+	UMaterialInstanceDynamic* LevelGridMaterialInst;
 };
 
 /**
