Fixed an issue where legacy widget handling UEdModes were calling the wrong UsesPropertyWidgets function, causing some property widgets to fail to render.

#rnx

diff --git a/Engine/Source/Editor/UnrealEd/Private/Tools/LegacyEdModeWidgetHelpers.cpp b/Engine/Source/Editor/UnrealEd/Private/Tools/LegacyEdModeWidgetHelpers.cpp
--- a/Engine/Source/Editor/UnrealEd/Private/Tools/LegacyEdModeWidgetHelpers.cpp
+++ b/Engine/Source/Editor/UnrealEd/Private/Tools/LegacyEdModeWidgetHelpers.cpp
@@ -481,6 +481,11 @@ EAxisList::Type FLegacyEdModeWidgetHelper::GetCurrentWidgetAxis() const
 
 bool FLegacyEdModeWidgetHelper::UsesPropertyWidgets() const
 {
+	if (ParentModeInterface)
+	{
+		return ParentModeInterface->UsesPropertyWidgets();
+	}
+
 	return false;
 }
 
@@ -764,6 +769,7 @@ void UBaseLegacyWidgetEdMode::Initialize()
 {
 	WidgetHelper = MakeShared<FLegacyEdModeWidgetHelper>();
 	WidgetHelper->Owner = this->Owner;
+	WidgetHelper->ParentModeInterface = this;
 
 	UEdMode::Initialize();
 }
@@ -859,7 +865,7 @@ EAxisList::Type UBaseLegacyWidgetEdMode::GetCurrentWidgetAxis() const
 
 bool UBaseLegacyWidgetEdMode::UsesPropertyWidgets() const
 {
-	return WidgetHelper->UsesPropertyWidgets();
+	return false;
 }
 
 bool UBaseLegacyWidgetEdMode::GetCustomDrawingCoordinateSystem(FMatrix& InMatrix, void* InData)
diff --git a/Engine/Source/Editor/UnrealEd/Public/Tools/LegacyEdModeWidgetHelpers.h b/Engine/Source/Editor/UnrealEd/Public/Tools/LegacyEdModeWidgetHelpers.h
--- a/Engine/Source/Editor/UnrealEd/Public/Tools/LegacyEdModeWidgetHelpers.h
+++ b/Engine/Source/Editor/UnrealEd/Public/Tools/LegacyEdModeWidgetHelpers.h
@@ -118,6 +118,8 @@ protected:
 
 	/** Pointer back to the mode tools that we are registered with */
 	FEditorModeTools* Owner;
+
+	TScriptInterface<ILegacyEdModeWidgetInterface> ParentModeInterface;
 };
 
 // This class is to aid transitioning from native FEdModes to UEdModes, in the case that the FEdMode used property widgets and/or transform widgets
