The property editor icons in World Settings > Game Mode > Selected GameMode array are no longer stretched vertically.

diff --git a/Engine/Source/Editor/UnrealEd/Public/GameModeInfoCustomizer.h b/Engine/Source/Editor/UnrealEd/Public/GameModeInfoCustomizer.h
--- a/Engine/Source/Editor/UnrealEd/Public/GameModeInfoCustomizer.h
+++ b/Engine/Source/Editor/UnrealEd/Public/GameModeInfoCustomizer.h
@@ -98,16 +98,19 @@ public:
 				]
 			]
 			+SHorizontalBox::Slot()
+			.VAlign(VAlign_Center)
 			.AutoWidth()
 			[
 				PropertyCustomizationHelpers::MakeUseSelectedButton(FSimpleDelegate::CreateSP(this, &FGameModeInfoCustomizer::OnMakeSelectedDefaultClassClicked, DefaultClassPropertyName))
 			]
 			+SHorizontalBox::Slot()
+			.VAlign(VAlign_Center)
 			.AutoWidth()
 			[
 				PropertyCustomizationHelpers::MakeBrowseButton(FSimpleDelegate::CreateSP(this, &FGameModeInfoCustomizer::OnBrowseDefaultClassClicked, DefaultClassPropertyName))
 			]
 			+SHorizontalBox::Slot()
+			.VAlign(VAlign_Center)
 			.AutoWidth()
 			[
 				PropertyCustomizationHelpers::MakeNewBlueprintButton(FSimpleDelegate::CreateSP(this, &FGameModeInfoCustomizer::OnMakeNewDefaultClassClicked, DefaultClassPropertyName))
