WorldPartition: Fix attached foliage Copy/Paste

diff --git a/Engine/Source/Editor/UnrealEd/Private/EditorObject.cpp b/Engine/Source/Editor/UnrealEd/Private/EditorObject.cpp
--- a/Engine/Source/Editor/UnrealEd/Private/EditorObject.cpp
+++ b/Engine/Source/Editor/UnrealEd/Private/EditorObject.cpp
@@ -31,6 +31,7 @@
 #include "InstancedFoliageActor.h"
 #include "InstancedFoliage.h"
 #include "Components/BrushComponent.h"
+#include "Algo/Transform.h"
 
 DEFINE_LOG_CATEGORY_STATIC(LogEditorObject, Log, All);
 
@@ -332,14 +333,12 @@ static const TCHAR* ImportProperties(
 
 				if (ActorComponent && ActorComponent->GetComponentLevel())
 				{
-					AInstancedFoliageActor* IFA = AInstancedFoliageActor::GetInstancedFoliageActorForLevel(ActorComponent->GetComponentLevel(), true);
-
-					FFoliageInfo* MeshInfo = nullptr;
-					UFoliageType* FoliageType = IFA->AddFoliageType(SourceFoliageType, &MeshInfo);
-
+					UWorld* World = ActorComponent->GetWorld();
+					TMap<AInstancedFoliageActor*, TArray<FFoliageInstance>> FoliageInstances;
+					
 					const TCHAR* StrPtr;
 					FString TextLine;
-					while (MeshInfo && FParse::Line(&SourceText, TextLine))
+					while (FParse::Line(&SourceText, TextLine))
 					{
 						StrPtr = *TextLine;
 						if (GetEND(&StrPtr, TEXT("Foliage")))
@@ -368,13 +367,26 @@ static const TCHAR* ImportProperties(
 						}
 						FParse::Value(StrPtr, TEXT("Flags="), Instance.Flags);
 
-						// Add the instance
-						MeshInfo->AddInstance(IFA, FoliageType, Instance, ActorComponent);
+						Instance.BaseComponent = ActorComponent;
+
+						if (AInstancedFoliageActor* IFA = AInstancedFoliageActor::Get(World, true, ActorComponent->GetComponentLevel(), Instance.Location))
+						{
+							FoliageInstances.FindOrAdd(IFA).Add(MoveTemp(Instance));
+						}
 					}
 
-					if (MeshInfo)
+					for (const auto& Pair : FoliageInstances)
 					{
-						MeshInfo->Refresh(IFA, true, true);
+						AInstancedFoliageActor* IFA = Pair.Key;
+						FFoliageInfo* MeshInfo = nullptr;
+						UFoliageType* FoliageType = IFA->AddFoliageType(SourceFoliageType, &MeshInfo);
+						TArray<const FFoliageInstance*> InstancePtrs;
+						InstancePtrs.Reserve(Pair.Value.Num());
+						Algo::Transform(Pair.Value, InstancePtrs, [](const FFoliageInstance& FoliageInstance) { return &FoliageInstance; });
+						if (MeshInfo)
+						{
+							MeshInfo->AddInstances(IFA, FoliageType, InstancePtrs);
+						}
 					}
 				}
 			}
