Added DistanceFieldResolutionScale to fbx import options

diff --git a/Engine/Source/Editor/UnrealEd/Classes/Factories/FbxStaticMeshImportData.h b/Engine/Source/Editor/UnrealEd/Classes/Factories/FbxStaticMeshImportData.h
--- a/Engine/Source/Editor/UnrealEd/Classes/Factories/FbxStaticMeshImportData.h
+++ b/Engine/Source/Editor/UnrealEd/Classes/Factories/FbxStaticMeshImportData.h
@@ -60,6 +60,10 @@ class UFbxStaticMeshImportData : public UFbxMeshImportData
 	UPROPERTY(EditAnywhere, BlueprintReadWrite, AdvancedDisplay, config, Category = Mesh, meta = (ToolTip = "If enabled, combines all meshes into a single mesh", ImportType = "StaticMesh"))
 	uint32 bCombineMeshes : 1;
 
+	/** Specify override color in the case that VertexColorImportOption is set to Override */
+	UPROPERTY(EditAnywhere, BlueprintReadWrite, AdvancedDisplay, config, Category= Mesh, meta=(ImportType="StaticMesh"))
+	float DistanceFieldResolutionScale;
+
 	/** Gets or creates fbx import data for the specified static mesh */
 	static UFbxStaticMeshImportData* GetImportDataForStaticMesh(UStaticMesh* StaticMesh, UFbxStaticMeshImportData* TemplateForCreation);
 
diff --git a/Engine/Source/Editor/UnrealEd/Private/Fbx/FbxMainImport.cpp b/Engine/Source/Editor/UnrealEd/Private/Fbx/FbxMainImport.cpp
--- a/Engine/Source/Editor/UnrealEd/Private/Fbx/FbxMainImport.cpp
+++ b/Engine/Source/Editor/UnrealEd/Private/Fbx/FbxMainImport.cpp
@@ -423,6 +423,7 @@ void ApplyImportUIToImportOptions(UFbxImportUI* ImportUI, FBXImportOptions& InOu
 		InOutImportOptions.VertexColorImportOption	= StaticMeshData->VertexColorImportOption;
 		InOutImportOptions.VertexOverrideColor		= StaticMeshData->VertexOverrideColor;
 		InOutImportOptions.bReorderMaterialToFbxOrder = StaticMeshData->bReorderMaterialToFbxOrder;
+		InOutImportOptions.DistanceFieldResolutionScale = StaticMeshData->DistanceFieldResolutionScale;
 	}
 	else if ( ImportUI->MeshTypeToImport == FBXIT_SkeletalMesh )
 	{
diff --git a/Engine/Source/Editor/UnrealEd/Private/Fbx/FbxStaticMeshImport.cpp b/Engine/Source/Editor/UnrealEd/Private/Fbx/FbxStaticMeshImport.cpp
--- a/Engine/Source/Editor/UnrealEd/Private/Fbx/FbxStaticMeshImport.cpp
+++ b/Engine/Source/Editor/UnrealEd/Private/Fbx/FbxStaticMeshImport.cpp
@@ -1971,6 +1971,7 @@ UStaticMesh* UnFbx::FFbxImporter::ImportStaticMeshAsSingle(UObject* InParent, TA
 		SrcModel.BuildSettings.bComputeWeightedNormals = ImportOptions->bComputeWeightedNormals;
 
 		StaticMesh->NaniteSettings.bEnabled = ImportOptions->bBuildNanite;
+		SrcModel.BuildSettings.DistanceFieldResolutionScale = ImportOptions->DistanceFieldResolutionScale;
 
 		if (ImportOptions->bGenerateLightmapUVs)
 		{
diff --git a/Engine/Source/Editor/UnrealEd/Private/Fbx/FbxStaticMeshImportData.cpp b/Engine/Source/Editor/UnrealEd/Private/Fbx/FbxStaticMeshImportData.cpp
--- a/Engine/Source/Editor/UnrealEd/Private/Fbx/FbxStaticMeshImportData.cpp
+++ b/Engine/Source/Editor/UnrealEd/Private/Fbx/FbxStaticMeshImportData.cpp
@@ -17,6 +17,7 @@ UFbxStaticMeshImportData::UFbxStaticMeshImportData(const FObjectInitializer& Obj
 	bTransformVertexToAbsolute = true;
 	bBakePivotInVertex = false;
 	VertexOverrideColor = FColor(255, 255, 255, 255);
+	DistanceFieldResolutionScale = 1.0f;
 }
 
 UFbxStaticMeshImportData* UFbxStaticMeshImportData::GetImportDataForStaticMesh(UStaticMesh* StaticMesh, UFbxStaticMeshImportData* TemplateForCreation)
diff --git a/Engine/Source/Editor/UnrealEd/Public/FbxImporter.h b/Engine/Source/Editor/UnrealEd/Public/FbxImporter.h
--- a/Engine/Source/Editor/UnrealEd/Public/FbxImporter.h
+++ b/Engine/Source/Editor/UnrealEd/Public/FbxImporter.h
@@ -153,6 +153,7 @@ struct FBXImportOptions
 	bool bCombineToSingle;
 	EVertexColorImportOption::Type VertexColorImportOption;
 	FColor VertexOverrideColor;
+	float DistanceFieldResolutionScale;
 	bool bRemoveDegenerates;
 	bool bBuildAdjacencyBuffer;
 	bool bBuildReversedIndexBuffer;
