Created SEditorHeaderButton to act as a simple way to get either an SComboButton or an SButton styled correctly for use in toolbars in the editor, most notably the Add component button.

diff --git a/Engine/Source/Editor/EditorStyle/Private/StarshipStyle.cpp b/Engine/Source/Editor/EditorStyle/Private/StarshipStyle.cpp
--- a/Engine/Source/Editor/EditorStyle/Private/StarshipStyle.cpp
+++ b/Engine/Source/Editor/EditorStyle/Private/StarshipStyle.cpp
@@ -699,6 +699,26 @@ void FStarshipEditorStyle::FStyle::SetupGeneralStyles()
 			.SetPadding(FMargin(0.0f));
 		Set("HoverOnlyHyperlink", HoverOnlyHyperlink);
 	}
+
+	FButtonStyle RoundedButtonStyle = FButtonStyle()
+		.SetNormal(FSlateRoundedBoxBrush(FStyleColors::Dropdown, 4))
+		.SetHovered(FSlateRoundedBoxBrush(FStyleColors::Hover, 4))
+		.SetPressed(FSlateRoundedBoxBrush(FStyleColors::Hover, 4))
+		.SetNormalForeground(FStyleColors::Foreground)
+		.SetHoveredForeground(FStyleColors::ForegroundHover)
+		.SetPressedForeground(FStyleColors::ForegroundHover)
+		.SetDisabledForeground(FStyleColors::White25)
+		.SetNormalPadding(FMargin(6, 6))
+		.SetPressedPadding(FMargin(6, 6));
+
+	Set("RoundedButton", RoundedButtonStyle);
+
+	FComboButtonStyle AddNewComboButtonStyle = FComboButtonStyle(FStarshipCoreStyle::GetCoreStyle().GetWidgetStyle<FComboButtonStyle>("SimpleComboButton"))
+		.SetButtonStyle(RoundedButtonStyle)
+		.SetDownArrowPadding(FMargin(1,0,0,0));
+
+	Set("RoundedComboButton", AddNewComboButtonStyle);
+
 #endif // WITH_EDITOR || (IS_PROGRAM && WITH_UNREAL_DEVELOPER_TOOLS)
 
 	// Expandable button
@@ -3249,7 +3269,7 @@ void FStarshipEditorStyle::FStyle::SetupPropertyEditorStyles()
 		Set( "DetailsView.GridLine", new FSlateColorBrush(FStyleColors::Recessed) );
 
 		Set( "DetailsView.AdvancedDropdownBorder",      new FSlateColorBrush(FStyleColors::Panel));
-		Set( "DetailsView.AdvancedDropdownBorder.Open", new FSlateColorBrush(FStyleColors::Panel));
+		Set( "DetailsView.AdvancedDropdownBorder.Open", new IMAGE_BRUSH("Common/ScrollBoxShadowTop", FVector2D(64,8)));
 
 		Set( "DetailsView.CategoryFontStyle", FStyleFonts::Get().NormalBold);
 		Set( "DetailsView.CategoryTextStyle", FTextBlockStyle(NormalText)
diff --git a/Engine/Source/Editor/UnrealEd/Private/SEditorHeaderButton.cpp b/Engine/Source/Editor/UnrealEd/Private/SEditorHeaderButton.cpp
new file mode 100644
--- /dev/null
+++ b/Engine/Source/Editor/UnrealEd/Private/SEditorHeaderButton.cpp
@@ -0,0 +1,88 @@
+// Copyright Epic Games, Inc. All Rights Reserved.
+
+#include "SEditorHeaderButton.h"
+#include "SComponentClassCombo.h"
+#include "Widgets/Images/SImage.h"
+#include "Widgets/Input/SButton.h"
+
+void SEditorHeaderButton::Construct(const FArguments& InArgs)
+{
+	check(InArgs._Icon.IsSet());
+
+	TAttribute<FText> Text = InArgs._Text;
+	
+	TSharedRef<SHorizontalBox> ButtonContent = SNew(SHorizontalBox)
+		+ SHorizontalBox::Slot()
+		.HAlign(HAlign_Center)
+		.VAlign(VAlign_Center)
+		.AutoWidth()
+		[
+			SNew(SImage)
+			.Image(InArgs._Icon)
+			.ColorAndOpacity(FSlateColor::UseForeground())
+		]
+		+ SHorizontalBox::Slot()
+		.Padding(FMargin(3, 0, 0, 0))
+		.VAlign(VAlign_Center)
+		.AutoWidth()
+		[
+			SNew(STextBlock)
+			.TextStyle(FAppStyle::Get(), "DialogButtonText")
+			.Text(InArgs._Text)
+			.Visibility_Lambda([Text]() { return Text.Get(FText::GetEmpty()).IsEmpty() ? EVisibility::Collapsed : EVisibility::Visible; })
+		];
+
+	if (InArgs._OnClicked.IsBound())
+	{
+		ChildSlot
+		[
+			SAssignNew(Button, SButton)
+			.ButtonStyle(FAppStyle::Get(), "RoundedButton")
+			.ForegroundColor(FSlateColor::UseStyle())
+			.IsEnabled(InArgs._IsEnabled)
+			.ToolTipText(InArgs._ToolTipText)
+			.HAlign(HAlign_Center)
+			.VAlign(VAlign_Center)
+			.ContentPadding(0)
+			[
+				ButtonContent
+			]
+		];
+	}
+	else
+	{
+		ChildSlot
+		[
+			SAssignNew(ComboButton, SComboButton)
+			.ComboButtonStyle(FAppStyle::Get(), "RoundedComboButton")
+			.ForegroundColor(FSlateColor::UseStyle())
+			.IsEnabled(InArgs._IsEnabled)
+			.ToolTipText(InArgs._ToolTipText)
+			.HAlign(HAlign_Center)
+			.VAlign(VAlign_Center)
+			.ButtonContent()
+			[
+				ButtonContent
+			]
+			.MenuContent()
+			[
+				InArgs._MenuContent.Widget
+			]
+			.OnGetMenuContent(InArgs._OnGetMenuContent)
+			.OnMenuOpenChanged(InArgs._OnMenuOpenChanged)
+			.OnComboBoxOpened(InArgs._OnComboBoxOpened)
+		];
+	}
+}
+
+void SEditorHeaderButton::SetMenuContentWidgetToFocus(TWeakPtr<SWidget> Widget)
+{
+	check(ComboButton.IsValid());
+	ComboButton->SetMenuContentWidgetToFocus(Widget);
+}
+
+void SEditorHeaderButton::SetIsMenuOpen(bool bIsOpen, bool bIsFocused)
+{
+	check(ComboButton.IsValid());
+	ComboButton->SetIsOpen(bIsOpen, bIsFocused);
+}
diff --git a/Engine/Source/Editor/UnrealEd/Public/SEditorHeaderButton.h b/Engine/Source/Editor/UnrealEd/Public/SEditorHeaderButton.h
new file mode 100644
--- /dev/null
+++ b/Engine/Source/Editor/UnrealEd/Public/SEditorHeaderButton.h
@@ -0,0 +1,48 @@
+// Copyright Epic Games, Inc. All Rights Reserved.
+
+#pragma once
+
+#include "Framework/SlateDelegates.h"
+#include "Widgets/DeclarativeSyntaxSupport.h"
+#include "Widgets/SCompoundWidget.h"
+#include "Widgets/Input/SComboButton.h"
+
+class UNREALED_API SEditorHeaderButton : public SCompoundWidget
+{
+public:
+
+	SLATE_BEGIN_ARGS(SEditorHeaderButton) :
+		_Icon(nullptr)
+		{}
+
+		/** The text to display in the button. */
+		SLATE_ATTRIBUTE(FText, Text)
+
+		SLATE_ATTRIBUTE(const FSlateBrush*, Icon)
+
+		/** The clicked handler. Note that if this is set, the button will behave as though it were just a button.
+		 * This means that OnGetMenuContent, OnComboBoxOpened and OnMenuOpenChanged will all be ignored, since there is no menu.
+		 */
+		SLATE_EVENT(FOnClicked, OnClicked)
+
+		/** The static menu content widget. */
+		SLATE_NAMED_SLOT(FArguments, MenuContent)
+
+		SLATE_EVENT(FOnGetContent, OnGetMenuContent)
+		SLATE_EVENT(FOnComboBoxOpened, OnComboBoxOpened)
+		SLATE_EVENT(FOnIsOpenChanged, OnMenuOpenChanged)
+
+	SLATE_END_ARGS()
+
+	SEditorHeaderButton() {}
+
+	void Construct(const FArguments& InArgs);
+
+	void SetMenuContentWidgetToFocus(TWeakPtr<SWidget> Widget);
+	void SetIsMenuOpen(bool bIsOpen, bool bIsFocused);
+
+private:
+
+	TSharedPtr<SComboButton> ComboButton;
+	TSharedPtr<class SButton> Button;
+};
\ No newline at end of file
