World Partition Convert Commandlet: don't convert existing levels that contains a level script, spawn a level instance instead and let the original level intact.
That way, Licensees with worlds with sublevels containing level scripts will still have a working version of their game projects.

diff --git a/Engine/Source/Editor/UnrealEd/Private/Commandlets/WorldPartitionConvertCommandlet.cpp b/Engine/Source/Editor/UnrealEd/Private/Commandlets/WorldPartitionConvertCommandlet.cpp
--- a/Engine/Source/Editor/UnrealEd/Private/Commandlets/WorldPartitionConvertCommandlet.cpp
+++ b/Engine/Source/Editor/UnrealEd/Private/Commandlets/WorldPartitionConvertCommandlet.cpp
@@ -21,6 +21,7 @@
 #include "WorldPartition/HLOD/HLODActor.h"
 #include "WorldPartition/WorldPartitionMiniMap.h"
 #include "WorldPartition/WorldPartitionMiniMapHelper.h"
+#include "LevelInstance/LevelInstanceActor.h"
 #include "GameFramework/WorldSettings.h"
 #include "UObject/UObjectHash.h"
 #include "PackageHelperFunctions.h"
@@ -1088,12 +1089,6 @@ int32 UWorldPartitionConvertCommandlet::Main(const FString& Params)
 		RemapSoftObjectPaths.Add(OldPackagePath, FSoftObjectPath(MainPackage).ToString());
 	}
 
-	for (ULevel* SubLevel : SubLevelsToConvert)
-	{
-		DetachDependantLevelPackages(SubLevel);
-		PrepareLevelActors(SubLevel, false, GetLevelGridPlacement(SubLevel, EActorGridPlacement::Bounds));
-	}
-
 	TMap<UObject*, UObject*> PrivateRefsMap;
 	for(ULevel* SubLevel : SubLevelsToConvert)
 	{
@@ -1102,87 +1097,123 @@ int32 UWorldPartitionConvertCommandlet::Main(const FString& Params)
 		UWorld* SubWorld = SubLevel->GetTypedOuter<UWorld>();
 		UPackage* SubPackage = SubLevel->GetPackage();
 
-		RemapSoftObjectPaths.Add(FSoftObjectPath(SubWorld).ToString(), FSoftObjectPath(MainWorld).ToString());
-		RemapSoftObjectPaths.Add(FSoftObjectPath(SubLevel).ToString(), FSoftObjectPath(MainLevel).ToString());
-		RemapSoftObjectPaths.Add(FSoftObjectPath(SubPackage).ToString(), FSoftObjectPath(MainPackage).ToString());
-
 		if (LevelHasLevelScriptBlueprint(SubLevel))
 		{
 			MapsWithLevelScriptsBPs.Add(SubPackage->GetLoadedPath().GetPackageName());
-		}
 
-		if (LevelHasMapBuildData(SubLevel))
-		{
-			MapsWithMapBuildData.Add(SubPackage->GetLoadedPath().GetPackageName());
-		}
+			// Spawn the Level Instance
+			ULevelStreaming* SubLevelStreaming = nullptr;
+			for (ULevelStreaming* LevelStreaming : MainLevel->GetWorld()->GetStreamingLevels())
+			{
+				if (LevelStreaming->GetLoadedLevel() == SubLevel)
+				{
+					SubLevelStreaming = LevelStreaming;
+					break;
+				}
+			}
+			check(SubLevelStreaming);
 
-		UE_LOG(LogWorldPartitionConvertCommandlet, Log, TEXT("Converting %s"), *SubWorld->GetName());
+			FActorSpawnParameters SpawnParams;
+			SpawnParams.OverrideLevel = MainLevel;
+			ALevelInstance* LevelInstanceActor = MainWorld->SpawnActor<ALevelInstance>(SpawnParams);
 
-		for(AActor* Actor: SubLevel->Actors)
+			UPackage* SubLevelPackage = SubLevel->GetPackage();
+
+			FTransform LevelTransform;
+			if (SubLevelPackage->WorldTileInfo)
+			{
+				LevelTransform = FTransform(FVector(SubLevelPackage->WorldTileInfo->Position));
+			}
+			else
+			{
+				LevelTransform = SubLevelStreaming->LevelTransform;
+			}
+
+			LevelInstanceActor->SetActorTransform(LevelTransform);
+			LevelInstanceActor->SetWorldAsset(SubLevel->GetTypedOuter<UWorld>());
+		}
+		else
 		{
-			if(Actor && !Actor->IsPendingKill())
+			RemapSoftObjectPaths.Add(FSoftObjectPath(SubWorld).ToString(), FSoftObjectPath(MainWorld).ToString());
+			RemapSoftObjectPaths.Add(FSoftObjectPath(SubLevel).ToString(), FSoftObjectPath(MainLevel).ToString());
+			RemapSoftObjectPaths.Add(FSoftObjectPath(SubPackage).ToString(), FSoftObjectPath(MainPackage).ToString());
+
+			if (LevelHasMapBuildData(SubLevel))
 			{
-				check(Actor->GetOuter() == SubLevel);
-				check(!ShouldDeleteActor(Actor, false));
-				
-				if (Actor->IsA(AGroupActor::StaticClass()))
-				{
-					GroupActors.Add(*Actor->GetFullName());
-				}
+				MapsWithMapBuildData.Add(SubPackage->GetLoadedPath().GetPackageName());
+			}
 
-				if (Actor->GroupActor)
-				{
-					ActorsInGroupActors.Add(*Actor->GetFullName());
-				}
+			DetachDependantLevelPackages(SubLevel);
+			PrepareLevelActors(SubLevel, false, GetLevelGridPlacement(SubLevel, EActorGridPlacement::Bounds));
 
-				TArray<AActor*> ChildActors;
-				Actor->GetAllChildActors(ChildActors, false);
+			UE_LOG(LogWorldPartitionConvertCommandlet, Log, TEXT("Converting %s"), *SubWorld->GetName());
 
-				if (ChildActors.Num())
+			for(AActor* Actor: SubLevel->Actors)
+			{
+				if(Actor && !Actor->IsPendingKill())
 				{
-					ActorsWithChildActors.Add(*Actor->GetFullName());
-				}
+					check(Actor->GetOuter() == SubLevel);
+					check(!ShouldDeleteActor(Actor, false));
+				
+					if (Actor->IsA(AGroupActor::StaticClass()))
+					{
+						GroupActors.Add(*Actor->GetFullName());
+					}
 
-				FArchiveGatherPrivateImports Ar(Actor, PrivateRefsMap, ActorsReferencesToActors);
-				Actor->Serialize(Ar);
+					if (Actor->GroupActor)
+					{
+						ActorsInGroupActors.Add(*Actor->GetFullName());
+					}
 
-				// Even after Foliage Partitioning it is possible some Actors still have a FoliageTag. Make sure it is removed.
-				if (FFoliageHelper::IsOwnedByFoliage(Actor))
-				{
-					FFoliageHelper::SetIsOwnedByFoliage(Actor, false);
-				}
+					TArray<AActor*> ChildActors;
+					Actor->GetAllChildActors(ChildActors, false);
 
-				ChangeObjectOuter(Actor, MainLevel);
+					if (ChildActors.Num())
+					{
+						ActorsWithChildActors.Add(*Actor->GetFullName());
+					}
 
-				// Migrate blueprint classes
-				UClass* ActorClass = Actor->GetClass();
-				if (!ActorClass->IsNative() && (ActorClass->GetPackage() == SubPackage))
-				{
-					ChangeObjectOuter(ActorClass, MainPackage);
-					UE_LOG(LogWorldPartitionConvertCommandlet, Log, TEXT("Extracted non-native class %s"), *ActorClass->GetName());
+					FArchiveGatherPrivateImports Ar(Actor, PrivateRefsMap, ActorsReferencesToActors);
+					Actor->Serialize(Ar);
+
+					// Even after Foliage Partitioning it is possible some Actors still have a FoliageTag. Make sure it is removed.
+					if (FFoliageHelper::IsOwnedByFoliage(Actor))
+					{
+						FFoliageHelper::SetIsOwnedByFoliage(Actor, false);
+					}
+
+					ChangeObjectOuter(Actor, MainLevel);
+
+					// Migrate blueprint classes
+					UClass* ActorClass = Actor->GetClass();
+					if (!ActorClass->IsNative() && (ActorClass->GetPackage() == SubPackage))
+					{
+						ChangeObjectOuter(ActorClass, MainPackage);
+						UE_LOG(LogWorldPartitionConvertCommandlet, Log, TEXT("Extracted non-native class %s"), *ActorClass->GetName());
+					}
 				}
 			}
-		}
 
-		if (!bReportOnly)
-		{
-			TArray<UObject*> ObjectsToRename;
-			ForEachObjectWithPackage(SubPackage, [&](UObject* Object)
+			if (!bReportOnly)
 			{
-				if(!Object->IsA<AActor>() && !Object->IsA<ULevel>() && !Object->IsA<UWorld>() && !Object->IsA<UMetaData>())
+				TArray<UObject*> ObjectsToRename;
+				ForEachObjectWithPackage(SubPackage, [&](UObject* Object)
+				{
+					if(!Object->IsA<AActor>() && !Object->IsA<ULevel>() && !Object->IsA<UWorld>() && !Object->IsA<UMetaData>())
+					{
+						ObjectsToRename.Add(Object);
+					}
+					return true;
+				}, /*bIncludeNestedObjects*/false);
+
+				for(UObject* ObjectToRename: ObjectsToRename)
 				{
-					ObjectsToRename.Add(Object);
+					ChangeObjectOuter(ObjectToRename, MainPackage);
+					UE_LOG(LogWorldPartitionConvertCommandlet, Warning, TEXT("Renamed orphan object %s"), *ObjectToRename->GetName());
 				}
-				return true;
-			}, /*bIncludeNestedObjects*/false);
 
-			for(UObject* ObjectToRename: ObjectsToRename)
-			{
-				ChangeObjectOuter(ObjectToRename, MainPackage);
-				UE_LOG(LogWorldPartitionConvertCommandlet, Warning, TEXT("Renamed orphan object %s"), *ObjectToRename->GetName());
+				PackagesToDelete.Add(SubLevel->GetPackage());
 			}
-
-			PackagesToDelete.Add(SubLevel->GetPackage());
 		}
 	}
 
