Avoid a syscall by only removing the readonly flags on readonly files

diff --git a/Engine/Source/Editor/UnrealEd/Private/PackageSourceControlHelper.cpp b/Engine/Source/Editor/UnrealEd/Private/PackageSourceControlHelper.cpp
--- a/Engine/Source/Editor/UnrealEd/Private/PackageSourceControlHelper.cpp
+++ b/Engine/Source/Editor/UnrealEd/Private/PackageSourceControlHelper.cpp
@@ -205,7 +205,7 @@ bool FPackageSourceControlHelper::Checkout(UPackage* Package) const
 	else
 	{
 		FString PackageFilename = SourceControlHelpers::PackageFilename(Package);
-		if (IPlatformFile::GetPlatformPhysical().FileExists(*PackageFilename))
+		if (IPlatformFile::GetPlatformPhysical().IsReadOnly(*PackageFilename))
 		{
 			if (!IPlatformFile::GetPlatformPhysical().SetReadOnly(*PackageFilename, false))
 			{
