Fix up raw GetData usage with typed element handles throughout asset placement mode.

diff --git a/Engine/Source/Editor/AssetPlacementEdMode/AssetPlacementEdMode.Build.cs b/Engine/Source/Editor/AssetPlacementEdMode/AssetPlacementEdMode.Build.cs
--- a/Engine/Source/Editor/AssetPlacementEdMode/AssetPlacementEdMode.Build.cs
+++ b/Engine/Source/Editor/AssetPlacementEdMode/AssetPlacementEdMode.Build.cs
@@ -22,6 +22,7 @@ public class AssetPlacementEdMode : ModuleRules
 				"Foliage",
 				"Landscape",
 				"TypedElementFramework",
+				"TypedElementInterfaces",
 				"PropertyEditor",
 			}
 		);
diff --git a/Engine/Source/Editor/AssetPlacementEdMode/Private/Tools/PlacementBrushToolBase.cpp b/Engine/Source/Editor/AssetPlacementEdMode/Private/Tools/PlacementBrushToolBase.cpp
--- a/Engine/Source/Editor/AssetPlacementEdMode/Private/Tools/PlacementBrushToolBase.cpp
+++ b/Engine/Source/Editor/AssetPlacementEdMode/Private/Tools/PlacementBrushToolBase.cpp
@@ -176,8 +176,7 @@ TArray<FTypedElementHandle> UPlacementBrushToolBase::GetElementsInBrushRadius()
 			continue;
 		}
 
-		const FActorInstanceHandle& HitObjectHandle = Hit.HitObjectHandle;
-		const AActor* HitActor = Hit.HitObjectHandle.FetchActor();
+		const AActor* HitActor = Hit.GetActor();
 		if (HitActor)
 		{
 			FTypedElementHandle ActorHandle = UEngineElementsLibrary::AcquireEditorActorElementHandle(HitActor);
diff --git a/Engine/Source/Editor/AssetPlacementEdMode/Private/Tools/PlacementEraseTool.cpp b/Engine/Source/Editor/AssetPlacementEdMode/Private/Tools/PlacementEraseTool.cpp
--- a/Engine/Source/Editor/AssetPlacementEdMode/Private/Tools/PlacementEraseTool.cpp
+++ b/Engine/Source/Editor/AssetPlacementEdMode/Private/Tools/PlacementEraseTool.cpp
@@ -7,7 +7,8 @@
 #include "Subsystems/EditorActorSubsystem.h"
 #include "Elements/Framework/EngineElementsLibrary.h"
 #include "Elements/Framework/TypedElementHandle.h"
-#include "Elements/Actor/ActorElementData.h"
+#include "Elements/Framework/TypedElementRegistry.h"
+#include "Elements/Interfaces/TypedElementObjectInterface.h"
 #include "Editor.h"
 #include "InstancedFoliageActor.h"
 #include "AssetPlacementEdMode.h"
@@ -45,11 +46,16 @@ void UPlacementModeEraseTool::OnTick(float DeltaTime)
 	TArray<FTypedElementHandle> HitElements = GetElementsInBrushRadius();
 	for (const FTypedElementHandle& HitElement : HitElements)
 	{
-		const FActorElementData* ActorData = HitElement.GetData<FActorElementData>();
-		if (ActorData && ActorData->Actor)
+		if (TTypedElement<UTypedElementObjectInterface> ObjectInterface = UTypedElementRegistry::GetInstance()->GetElement<UTypedElementObjectInterface>(HitElement))
 		{
+			AActor* Actor = ObjectInterface.GetObjectAs<AActor>();
+			if (!Actor)
+			{
+				continue;
+			}
+
 			// Since the foliage static mesh instances do not currently operate with element handles, we have to drill in manually here.
-			if (AInstancedFoliageActor* FoliageActor = Cast<AInstancedFoliageActor>(ActorData->Actor))
+			if (AInstancedFoliageActor* FoliageActor = Cast<AInstancedFoliageActor>(Actor))
 			{
 				for (auto& FoliageInfo : FoliageActor->FoliageInfos)
 				{
@@ -65,7 +71,7 @@ void UPlacementModeEraseTool::OnTick(float DeltaTime)
 			}
 			else if (ActorSubsystem)
 			{
-				ActorSubsystem->DestroyActor(ActorData->Actor);
+				ActorSubsystem->DestroyActor(Actor);
 			}
 		}
 	}
diff --git a/Engine/Source/Editor/UnrealEd/Private/AssetSelection.cpp b/Engine/Source/Editor/UnrealEd/Private/AssetSelection.cpp
--- a/Engine/Source/Editor/UnrealEd/Private/AssetSelection.cpp
+++ b/Engine/Source/Editor/UnrealEd/Private/AssetSelection.cpp
@@ -52,7 +52,8 @@
 #include "ISourceControlProvider.h"
 #include "Misc/MessageDialog.h"
 #include "Subsystems/PlacementSubsystem.h"
-#include "Elements/Actor/ActorElementData.h"
+#include "Elements/Framework/TypedElementRegistry.h"
+#include "Elements/Interfaces/TypedElementObjectInterface.h"
 
 namespace AssetSelectionUtils
 {
@@ -620,9 +621,14 @@ static AActor* PrivateAddActor( UObject* Asset, UActorFactory* Factory, bool Sel
 			TArray<FTypedElementHandle> PlacedElements = PlacementSubsystem->PlaceAsset(PlacementInfo, FPlacementOptions());
 			if (PlacedElements.Num())
 			{
-				const FActorElementData* ActorElementData = PlacedElements[0].GetData<FActorElementData>();
-				Actor = ActorElementData ? ActorElementData->Actor : nullptr;
-				Actor->SetFlags(ObjectFlags);
+				if (TTypedElement<UTypedElementObjectInterface> ObjectInterface = UTypedElementRegistry::GetInstance()->GetElement<UTypedElementObjectInterface>(PlacedElements[0]))
+				{
+					Actor = ObjectInterface.GetObjectAs<AActor>();
+					if (Actor)
+					{
+						Actor->SetFlags(ObjectFlags);
+					}
+				}
 			}
 		}
 		
diff --git a/Engine/Source/Editor/UnrealEd/Private/Elements/Actor/ActorElementEditorAssetDataInterface.cpp b/Engine/Source/Editor/UnrealEd/Private/Elements/Actor/ActorElementEditorAssetDataInterface.cpp
new file mode 100644
--- /dev/null
+++ b/Engine/Source/Editor/UnrealEd/Private/Elements/Actor/ActorElementEditorAssetDataInterface.cpp
@@ -0,0 +1,29 @@
+// Copyright Epic Games, Inc. All Rights Reserved.
+
+#include "Elements/Actor/ActorElementEditorAssetDataInterface.h"
+
+#include "Elements/Actor/ActorElementData.h"
+#include "GameFramework/Actor.h"
+
+TArray<FAssetData> UActorElementEditorAssetDataInterface::GetAllReferencedAssetDatas(const FTypedElementHandle& InElementHandle)
+{
+	TArray<FAssetData> AssetDatas;
+
+	if (AActor* RawActorPtr = ActorElementDataUtil::GetActorFromHandle(InElementHandle))
+	{
+		TArray<UObject*> ReferencedContentObjects;
+		RawActorPtr->GetReferencedContentObjects(ReferencedContentObjects);
+		for (const UObject* ContentObject : ReferencedContentObjects)
+		{
+			FAssetData ObjectAssetData = FAssetData(ContentObject);
+			if (ObjectAssetData.IsValid())
+			{
+				AssetDatas.Emplace(ObjectAssetData);
+			}
+		}
+
+		AssetDatas.Emplace(FAssetData(RawActorPtr));
+	}
+
+	return AssetDatas;
+}
diff --git a/Engine/Source/Editor/UnrealEd/Private/Elements/Actor/ActorElementEditorAssetDataInterface.h b/Engine/Source/Editor/UnrealEd/Private/Elements/Actor/ActorElementEditorAssetDataInterface.h
new file mode 100644
--- /dev/null
+++ b/Engine/Source/Editor/UnrealEd/Private/Elements/Actor/ActorElementEditorAssetDataInterface.h
@@ -0,0 +1,16 @@
+// Copyright Epic Games, Inc. All Rights Reserved.
+
+#pragma once
+
+#include "Elements/Actor/ActorElementAssetDataInterface.h"
+
+#include "ActorElementEditorAssetDataInterface.generated.h"
+
+UCLASS()
+class UActorElementEditorAssetDataInterface : public UActorElementAssetDataInterface
+{
+	GENERATED_BODY()
+
+public:
+	virtual TArray<FAssetData> GetAllReferencedAssetDatas(const FTypedElementHandle& InElementHandle) override;
+};
diff --git a/Engine/Source/Editor/UnrealEd/Private/Elements/Framework/EditorElements.cpp b/Engine/Source/Editor/UnrealEd/Private/Elements/Framework/EditorElements.cpp
--- a/Engine/Source/Editor/UnrealEd/Private/Elements/Framework/EditorElements.cpp
+++ b/Engine/Source/Editor/UnrealEd/Private/Elements/Framework/EditorElements.cpp
@@ -9,6 +9,7 @@
 #include "Elements/Actor/ActorElementDetailsInterface.h"
 #include "Elements/Actor/ActorElementEditorWorldInterface.h"
 #include "Elements/Actor/ActorElementEditorSelectionInterface.h"
+#include "Elements/Actor/ActorElementEditorAssetDataInterface.h"
 
 #include "Elements/Component/ComponentElementDetailsInterface.h"
 #include "Elements/Component/ComponentElementEditorWorldInterface.h"
@@ -31,6 +32,7 @@ void RegisterEditorActorElements()
 	Registry->RegisterElementInterface<UTypedElementDetailsInterface>(NAME_Actor, NewObject<UActorElementDetailsInterface>());
 	Registry->RegisterElementInterface<UTypedElementWorldInterface>(NAME_Actor, NewObject<UActorElementEditorWorldInterface>(), /*bAllowOverride*/true);
 	Registry->RegisterElementInterface<UTypedElementSelectionInterface>(NAME_Actor, NewObject<UActorElementEditorSelectionInterface>(), /*bAllowOverride*/true);
+	Registry->RegisterElementInterface<UTypedElementAssetDataInterface>(NAME_Actor, NewObject<UActorElementEditorAssetDataInterface>(), /*bAllowOverride*/true);
 }
 
 void RegisterEditorComponentElements()
diff --git a/Engine/Source/Editor/UnrealEd/Private/Factories/ActorFactory.cpp b/Engine/Source/Editor/UnrealEd/Private/Factories/ActorFactory.cpp
--- a/Engine/Source/Editor/UnrealEd/Private/Factories/ActorFactory.cpp
+++ b/Engine/Source/Editor/UnrealEd/Private/Factories/ActorFactory.cpp
@@ -131,8 +131,10 @@ ActorFactory.cpp:
 #include "LevelSequence.h"
 #include "LevelSequenceActor.h"
 #include "Factories/ActorFactoryMovieScene.h"
-#include "Elements/Actor/ActorElementData.h"
 #include "Elements/Framework/EngineElementsLibrary.h"
+#include "Elements/Framework/TypedElementRegistry.h"
+#include "Elements/Interfaces/TypedElementAssetDataInterface.h"
+#include "Elements/Interfaces/TypedElementObjectInterface.h"
 #include "Subsystems/PlacementSubsystem.h"
 
 #include "LevelEditorViewport.h"
@@ -267,6 +269,11 @@ FQuat UActorFactory::AlignObjectToSurfaceNormal(const FVector& InSurfaceNormal,
 
 bool UActorFactory::CanPlaceElementsFromAssetData(const FAssetData& InAssetData)
 {
+	if (!InAssetData.IsValid())
+	{
+		return false;
+	}
+
 	UClass* Class = InAssetData.GetClass();
 	if (Class && (Class == GetDefaultActorClass(InAssetData)))
 	{
@@ -297,52 +304,53 @@ void UActorFactory::PostPlaceAsset(TArrayView<const FTypedElementHandle> InEleme
 {
 	for (const FTypedElementHandle& PlacedElement : InElementHandles)
 	{
-		const FActorElementData* ActorElementData = PlacedElement.GetData<FActorElementData>();
-		if (ActorElementData && ActorElementData->Actor)
+		if (TTypedElement<UTypedElementObjectInterface> ObjectInterface = UTypedElementRegistry::GetInstance()->GetElement<UTypedElementObjectInterface>(PlacedElement))
 		{
-			PostSpawnActor(InPlacementInfo.AssetToPlace.GetAsset(), ActorElementData->Actor);
-
-			// Only do this if the actor wasn't already given a name
-			if (InPlacementInfo.NameOverride.IsNone())
+			if (AActor* CreatedActor = ObjectInterface.GetObjectAs<AActor>())
 			{
-				FActorLabelUtilities::SetActorLabelUnique(ActorElementData->Actor, InPlacementInfo.AssetToPlace.AssetName.ToString());
-			}
+				PostSpawnActor(InPlacementInfo.AssetToPlace.GetAsset(), CreatedActor);
+
+				// Only do this if the actor wasn't already given a name
+				if (InPlacementInfo.NameOverride.IsNone())
+				{
+					FActorLabelUtilities::SetActorLabelUnique(CreatedActor, InPlacementInfo.AssetToPlace.AssetName.ToString());
+				}
 
-			ActorElementData->Actor->PostEditChange();
-			ActorElementData->Actor->PostEditMove(true);
+				CreatedActor->PostEditChange();
+				CreatedActor->PostEditMove(true);
+			}
 		}
 	}
 }
 
 FAssetData UActorFactory::GetAssetDataFromElementHandle(const FTypedElementHandle& InHandle)
 {
-	const FActorElementData* ActorData = InHandle.GetData<FActorElementData>();
-	if (!ActorData || !ActorData->Actor)
+	if (!InHandle.IsSet())
 	{
 		return FAssetData();
 	}
 
-	AActor* ActorInstance = ActorData->Actor;
-	const UClass* DefaultActorClass = GetDefaultActorClass(FAssetData());
-
-	// See if the existing implementation of the actor factory has a valid asset data.
-	// First check that the actor instance derives from the default class
-	if (DefaultActorClass && ActorInstance->IsA(DefaultActorClass))
+	// Check if the handle is the type of actor created by this factory, and use the factory to find the wrapped asset data if possible.
+	if (TTypedElement<UTypedElementObjectInterface> ElementObjectInterface = UTypedElementRegistry::GetInstance()->GetElement<UTypedElementObjectInterface>(InHandle))
 	{
-		if (UObject* WrappedAssetObject = GetAssetFromActorInstance(ActorInstance))
+		if (AActor* RawElementActorPtr = ElementObjectInterface.GetObjectAs<AActor>(GetDefaultActorClass(FAssetData())))
 		{
-			return FAssetData(WrappedAssetObject);
+			if (UObject* WrappedAssetObject = GetAssetFromActorInstance(RawElementActorPtr))
+			{
+				return FAssetData(WrappedAssetObject);
+			}
 		}
+	}
 
-		// Query the actor itself for any contained content objects.
-		TArray<UObject*> ReferencedContentObjects;
-		ActorInstance->GetReferencedContentObjects(ReferencedContentObjects);
-		for (const UObject* ReferencedContentObject : ReferencedContentObjects)
+	// Check if any of the referenced content is created by this factory.
+	if (TTypedElement<UTypedElementAssetDataInterface> AssetDataInterface = UTypedElementRegistry::GetInstance()->GetElement<UTypedElementAssetDataInterface>(InHandle))
+	{
+		TArray<FAssetData> ContentAssetDatas = AssetDataInterface.GetAllReferencedAssetDatas();
+		for (FAssetData& AssetData : ContentAssetDatas)
 		{
-			FAssetData ReferencedObjectData = FAssetData(ReferencedContentObject);
-			if (CanPlaceElementsFromAssetData(ReferencedObjectData))
+			if (CanPlaceElementsFromAssetData(AssetData))
 			{
-				return ReferencedObjectData;
+				return MoveTemp(AssetData);
 			}
 		}
 	}
diff --git a/Engine/Source/Editor/UnrealEd/Private/Factories/EditorStaticMeshFactory.cpp b/Engine/Source/Editor/UnrealEd/Private/Factories/EditorStaticMeshFactory.cpp
--- a/Engine/Source/Editor/UnrealEd/Private/Factories/EditorStaticMeshFactory.cpp
+++ b/Engine/Source/Editor/UnrealEd/Private/Factories/EditorStaticMeshFactory.cpp
@@ -2,10 +2,10 @@
 
 #include "Factories/EditorStaticMeshFactory.h"
 
-#include "Elements/Actor/ActorElementData.h"
-#include "Elements/Component/ComponentElementData.h"
-#include "Elements/Object/ObjectElementData.h"
 #include "Elements/Framework/EngineElementsLibrary.h"
+#include "Elements/Framework/TypedElementRegistry.h"
+#include "Elements/Interfaces/TypedElementAssetDataInterface.h"
+#include "Elements/Interfaces/TypedElementObjectInterface.h"
 
 #include "ActorPartition/ActorPartitionSubsystem.h"
 #include "InstancedFoliageActor.h"
@@ -74,19 +74,16 @@ FAssetData UEditorStaticMeshFactory::GetAssetDataFromElementHandle(const FTypedE
 {
 	UInstancedStaticMeshComponent* ISMComponent = nullptr;
 
-	// Try to pull from component handle
-	const FComponentElementData* ComponentData = InHandle.GetData<FComponentElementData>();
-	if (ComponentData && ComponentData->Component)
+	if (TTypedElement<UTypedElementObjectInterface> ObjectInterface = UTypedElementRegistry::GetInstance()->GetElement<UTypedElementObjectInterface>(InHandle))
 	{
-		ISMComponent = Cast<UInstancedStaticMeshComponent>(ComponentData->Component);
-	}
-	else
-	{
-		// Try to pull from actor handle
-		const FActorElementData* ActorData = InHandle.GetData<FActorElementData>();
-		if (ActorData && ActorData->Actor)
+		// Try to pull from component handle
+		if (UInstancedStaticMeshComponent* RawComponentPtr = ObjectInterface.GetObjectAs<UInstancedStaticMeshComponent>())
+		{
+			ISMComponent = RawComponentPtr;
+		}
+		else if (AActor* RawActorPtr = ObjectInterface.GetObjectAs<AActor>())
 		{
-			ISMComponent = ActorData->Actor->FindComponentByClass<UInstancedStaticMeshComponent>();
+			ISMComponent = RawActorPtr->FindComponentByClass<UInstancedStaticMeshComponent>();
 		}
 	}
 
@@ -95,12 +92,9 @@ FAssetData UEditorStaticMeshFactory::GetAssetDataFromElementHandle(const FTypedE
 	{
 		FoundAssetData = FAssetData(ISMComponent->GetStaticMesh());
 	}
-	else if (const FObjectElementData* ObjectData = InHandle.GetData<FObjectElementData>())
+	else if (TTypedElement<UTypedElementAssetDataInterface> AssetDataInterface = UTypedElementRegistry::GetInstance()->GetElement<UTypedElementAssetDataInterface>(InHandle))
 	{
-		if (ObjectData && ObjectData->Object)
-		{
-			FoundAssetData = FAssetData(ObjectData->Object);
-		}
+		FoundAssetData = AssetDataInterface.GetAssetData();
 	}
 
 	if (CanPlaceElementsFromAssetData(FoundAssetData))
diff --git a/Engine/Source/Runtime/Engine/Private/Elements/Actor/ActorElementAssetDataInterface.cpp b/Engine/Source/Runtime/Engine/Private/Elements/Actor/ActorElementAssetDataInterface.cpp
new file mode 100644
--- /dev/null
+++ b/Engine/Source/Runtime/Engine/Private/Elements/Actor/ActorElementAssetDataInterface.cpp
@@ -0,0 +1,12 @@
+// Copyright Epic Games, Inc. All Rights Reserved.
+
+#include "Elements/Actor/ActorElementAssetDataInterface.h"
+
+#include "Elements/Actor/ActorElementData.h"
+#include "GameFramework/Actor.h"
+
+FAssetData UActorElementAssetDataInterface::GetAssetData(const FTypedElementHandle& InElementHandle)
+{
+	AActor* RawActorPtr = ActorElementDataUtil::GetActorFromHandle(InElementHandle);
+	return FAssetData(RawActorPtr);
+}
diff --git a/Engine/Source/Runtime/Engine/Private/Elements/Framework/EngineElements.cpp b/Engine/Source/Runtime/Engine/Private/Elements/Framework/EngineElements.cpp
--- a/Engine/Source/Runtime/Engine/Private/Elements/Framework/EngineElements.cpp
+++ b/Engine/Source/Runtime/Engine/Private/Elements/Framework/EngineElements.cpp
@@ -4,10 +4,12 @@
 #include "Elements/Framework/TypedElementRegistry.h"
 
 #include "Elements/Object/ObjectElementData.h"
+#include "Elements/Object/ObjectElementAssetDataInterface.h"
 #include "Elements/Object/ObjectElementObjectInterface.h"
 #include "Elements/Object/ObjectElementSelectionInterface.h"
 
 #include "Elements/Actor/ActorElementData.h"
+#include "Elements/Actor/ActorElementAssetDataInterface.h"
 #include "Elements/Actor/ActorElementObjectInterface.h"
 #include "Elements/Actor/ActorElementWorldInterface.h"
 #include "Elements/Actor/ActorElementSelectionInterface.h"
@@ -26,6 +28,7 @@ void RegisterEngineObjectElements()
 	UTypedElementRegistry* Registry = UTypedElementRegistry::GetInstance();
 
 	Registry->RegisterElementType<FObjectElementData>(NAME_Object);
+	Registry->RegisterElementInterface<UTypedElementAssetDataInterface>(NAME_Object, NewObject<UObjectElementAssetDataInterface>());
 	Registry->RegisterElementInterface<UTypedElementObjectInterface>(NAME_Object, NewObject<UObjectElementObjectInterface>());
 	Registry->RegisterElementInterface<UTypedElementSelectionInterface>(NAME_Object, NewObject<UObjectElementSelectionInterface>());
 }
@@ -35,6 +38,7 @@ void RegisterEngineActorElements()
 	UTypedElementRegistry* Registry = UTypedElementRegistry::GetInstance();
 
 	Registry->RegisterElementType<FActorElementData>(NAME_Actor);
+	Registry->RegisterElementInterface<UTypedElementAssetDataInterface>(NAME_Actor, NewObject<UActorElementAssetDataInterface>());
 	Registry->RegisterElementInterface<UTypedElementObjectInterface>(NAME_Actor, NewObject<UActorElementObjectInterface>());
 	Registry->RegisterElementInterface<UTypedElementWorldInterface>(NAME_Actor, NewObject<UActorElementWorldInterface>());
 	Registry->RegisterElementInterface<UTypedElementSelectionInterface>(NAME_Actor, NewObject<UActorElementSelectionInterface>());
diff --git a/Engine/Source/Runtime/Engine/Private/Elements/Object/ObjectElementAssetDataInterface.cpp b/Engine/Source/Runtime/Engine/Private/Elements/Object/ObjectElementAssetDataInterface.cpp
new file mode 100644
--- /dev/null
+++ b/Engine/Source/Runtime/Engine/Private/Elements/Object/ObjectElementAssetDataInterface.cpp
@@ -0,0 +1,12 @@
+// Copyright Epic Games, Inc. All Rights Reserved.
+
+#include "Elements/Object/ObjectElementAssetDataInterface.h"
+
+#include "Elements/Object/ObjectElementData.h"
+#include "UObject/Object.h"
+
+FAssetData UObjectElementAssetDataInterface::GetAssetData(const FTypedElementHandle& InElementHandle)
+{
+	UObject* RawObjectPtr = ObjectElementDataUtil::GetObjectFromHandle(InElementHandle);
+	return FAssetData(RawObjectPtr);
+}
diff --git a/Engine/Source/Runtime/Engine/Public/Elements/Actor/ActorElementAssetDataInterface.h b/Engine/Source/Runtime/Engine/Public/Elements/Actor/ActorElementAssetDataInterface.h
new file mode 100644
--- /dev/null
+++ b/Engine/Source/Runtime/Engine/Public/Elements/Actor/ActorElementAssetDataInterface.h
@@ -0,0 +1,16 @@
+// Copyright Epic Games, Inc. All Rights Reserved.
+
+#pragma once
+
+#include "Elements/Interfaces/TypedElementAssetDataInterface.h"
+
+#include "ActorElementAssetDataInterface.generated.h"
+
+UCLASS()
+class ENGINE_API UActorElementAssetDataInterface : public UTypedElementAssetDataInterface
+{
+	GENERATED_BODY()
+
+public:
+	virtual FAssetData GetAssetData(const FTypedElementHandle& InElementHandle) override;
+};
diff --git a/Engine/Source/Runtime/Engine/Public/Elements/Object/ObjectElementAssetDataInterface.h b/Engine/Source/Runtime/Engine/Public/Elements/Object/ObjectElementAssetDataInterface.h
new file mode 100644
--- /dev/null
+++ b/Engine/Source/Runtime/Engine/Public/Elements/Object/ObjectElementAssetDataInterface.h
@@ -0,0 +1,16 @@
+// Copyright Epic Games, Inc. All Rights Reserved.
+
+#pragma once
+
+#include "Elements/Interfaces/TypedElementAssetDataInterface.h"
+
+#include "ObjectElementAssetDataInterface.generated.h"
+
+UCLASS()
+class ENGINE_API UObjectElementAssetDataInterface : public UTypedElementAssetDataInterface
+{
+	GENERATED_BODY()
+
+public:
+	virtual FAssetData GetAssetData(const FTypedElementHandle& InElementHandle) override;
+};
diff --git a/Engine/Source/Runtime/TypedElementInterfaces/Private/Elements/Interfaces/TypedElementAssetDataInterface.cpp b/Engine/Source/Runtime/TypedElementInterfaces/Private/Elements/Interfaces/TypedElementAssetDataInterface.cpp
new file mode 100644
--- /dev/null
+++ b/Engine/Source/Runtime/TypedElementInterfaces/Private/Elements/Interfaces/TypedElementAssetDataInterface.cpp
@@ -0,0 +1,21 @@
+// Copyright Epic Games, Inc. All Rights Reserved.
+
+#include "Elements/Interfaces/TypedElementAssetDataInterface.h"
+
+TArray<FAssetData> UTypedElementAssetDataInterface::GetAllReferencedAssetDatas(const FTypedElementHandle& InElementHandle)
+{
+	TArray<FAssetData> AssetDatas;
+
+	FAssetData ElementAssetData = GetAssetData(InElementHandle);
+	if (ElementAssetData.IsValid())
+	{
+		AssetDatas.Emplace(ElementAssetData);
+	}
+
+	return AssetDatas;
+}
+
+FAssetData UTypedElementAssetDataInterface::GetAssetData(const FTypedElementHandle& InElementHandle)
+{
+	return FAssetData();
+}
diff --git a/Engine/Source/Runtime/TypedElementInterfaces/Private/Elements/Interfaces/TypedElementObjectInterface.cpp b/Engine/Source/Runtime/TypedElementInterfaces/Private/Elements/Interfaces/TypedElementObjectInterface.cpp
--- a/Engine/Source/Runtime/TypedElementInterfaces/Private/Elements/Interfaces/TypedElementObjectInterface.cpp
+++ b/Engine/Source/Runtime/TypedElementInterfaces/Private/Elements/Interfaces/TypedElementObjectInterface.cpp
@@ -1,3 +1,16 @@
 // Copyright Epic Games, Inc. All Rights Reserved.
 
 #include "Elements/Interfaces/TypedElementObjectInterface.h"
+
+#include "Elements/Framework/TypedElementRegistry.h"
+
+UObject* UTypedElementObjectInterface::GetObject(const FTypedElementHandle& InElementHandle)
+{
+	return nullptr;
+}
+
+UClass* UTypedElementObjectInterface::GetObjectClass(const FTypedElementHandle& InElementHandle)
+{
+	UObject* HandleAsObject = GetObject(InElementHandle);
+	return HandleAsObject ? HandleAsObject->GetClass() : nullptr;
+}
diff --git a/Engine/Source/Runtime/TypedElementInterfaces/Public/Elements/Interfaces/TypedElementAssetDataInterface.h b/Engine/Source/Runtime/TypedElementInterfaces/Public/Elements/Interfaces/TypedElementAssetDataInterface.h
new file mode 100644
--- /dev/null
+++ b/Engine/Source/Runtime/TypedElementInterfaces/Public/Elements/Interfaces/TypedElementAssetDataInterface.h
@@ -0,0 +1,37 @@
+// Copyright Epic Games, Inc. All Rights Reserved.
+
+#pragma once
+
+#include "Elements/Framework/TypedElementHandle.h"
+#include "AssetRegistry/AssetData.h"
+
+#include "TypedElementAssetDataInterface.generated.h"
+
+UCLASS(Abstract)
+class TYPEDELEMENTINTERFACES_API UTypedElementAssetDataInterface : public UTypedElementInterface
+{
+	GENERATED_BODY()
+
+public:
+	/**
+	 * Returns any asset datas for content objects referenced by handle.
+	 * If the given handle itself has valid asset data, it should be returned as the last element of the array.
+	 *
+	 * @returns An array of valid asset datas.
+	 */
+	UFUNCTION(BlueprintPure, Category = "TypedElementInterfaces|AssetData")
+	virtual TArray<FAssetData> GetAllReferencedAssetDatas(const FTypedElementHandle& InElementHandle);
+
+	/**
+	 * Returns the asset data for the given handle, if it exists.
+	 */
+	UFUNCTION(BlueprintPure, Category = "TypedElementInterfaces|AssetData")
+	virtual FAssetData GetAssetData(const FTypedElementHandle& InElementHandle);
+};
+
+template <>
+struct TTypedElement<UTypedElementAssetDataInterface> : public TTypedElementBase<UTypedElementAssetDataInterface>
+{
+	TArray<FAssetData> GetAllReferencedAssetDatas() const { return InterfacePtr->GetAllReferencedAssetDatas(*this); }
+	FAssetData GetAssetData() { return InterfacePtr->GetAssetData(*this); }
+};
diff --git a/Engine/Source/Runtime/TypedElementInterfaces/Public/Elements/Interfaces/TypedElementObjectInterface.h b/Engine/Source/Runtime/TypedElementInterfaces/Public/Elements/Interfaces/TypedElementObjectInterface.h
--- a/Engine/Source/Runtime/TypedElementInterfaces/Public/Elements/Interfaces/TypedElementObjectInterface.h
+++ b/Engine/Source/Runtime/TypedElementInterfaces/Public/Elements/Interfaces/TypedElementObjectInterface.h
@@ -15,11 +15,58 @@ public:
 	 * Get the object instance that this handle represents, if any.
 	 */
 	UFUNCTION(BlueprintPure, Category="TypedElementInterfaces|Object")
-	virtual UObject* GetObject(const FTypedElementHandle& InElementHandle) { return nullptr; }
+	virtual UObject* GetObject(const FTypedElementHandle& InElementHandle);
+
+	/**
+	 * Gets the object instance's class that the handle represents, if any. 
+	 */
+	UFUNCTION(BlueprintPure, Category = "TypedElementInterfaces|Object")
+	UClass* GetObjectClass(const FTypedElementHandle& InElementHandle);
+
+	/**
+	 * Attempts to cast the given handle to another class.
+	 *
+	 * @returns the casted object if successful, otherwise nullptr.
+	 */
+	template <class CastTo>
+	CastTo* GetObjectAs(const FTypedElementHandle& InElementHandle)
+	{
+		return Cast<CastTo>(GetObject(InElementHandle));
+	}
+
+	/**
+	 * Attempts to cast the given handle to another class, if it can also be casted to TargetClass.
+	 * This is intended for use in cases where the calling code may only need an actor pointer, but also wants to be sure it's a specific type of actor.
+	 *
+	 * @returns the casted object if successful, otherwise nullptr.
+	 */
+	template <class CastTo = UObject>
+	CastTo* GetObjectAs(const FTypedElementHandle& InElementHandle, TSubclassOf<CastTo> TargetClass)
+	{
+		if (!TargetClass)
+		{
+			return nullptr;
+		}
+
+		UObject* ObjectPtr = GetObject(InElementHandle);
+		if (ObjectPtr && ObjectPtr->IsA(TargetClass))
+		{
+			return (CastTo*)ObjectPtr;
+		}
+
+		return nullptr;
+	}
 };
 
 template <>
 struct TTypedElement<UTypedElementObjectInterface> : public TTypedElementBase<UTypedElementObjectInterface>
 {
 	UObject* GetObject() const { return InterfacePtr->GetObject(*this); }
+	UClass* GetObjectClass() { return InterfacePtr->GetObjectClass(*this); }
+
+	template <class CastTo>
+	CastTo* GetObjectAs() { return InterfacePtr->GetObjectAs<CastTo>(*this); }
+
+	template <class CastTo>
+	CastTo* GetObjectAs(TSubclassOf<CastTo> TargetClass) { return InterfacePtr->GetObjectAs<CastTo>(*this, TargetClass); }
 };
