Migrated the fix for UE-104718 to UE5. The name of a udim texture wasn't sanitized properly.

diff --git a/Engine/Source/Editor/UnrealEd/Private/Factories/EditorFactories.cpp b/Engine/Source/Editor/UnrealEd/Private/Factories/EditorFactories.cpp
--- a/Engine/Source/Editor/UnrealEd/Private/Factories/EditorFactories.cpp
+++ b/Engine/Source/Editor/UnrealEd/Private/Factories/EditorFactories.cpp
@@ -3913,7 +3913,8 @@ UObject* UTextureFactory::FactoryCreateBinary
 			// Exclude UDIM number from the name of the UE4 texture asset we create
 			// We do this even in the case where we're only importing a single 1001 image, which won't technically be a UDIM
 			// We still want to strip the UDIM suffix in this case, as industry standard still considers this a UDIM imageset
-			TextureName = *BaseUDIMName;
+			const FString ShortPackageName = ObjectTools::SanitizeInvalidChars(BaseUDIMName, INVALID_LONGPACKAGE_CHARACTERS);
+			TextureName = *ShortPackageName;
 
 			// Don't try to rename the package if its the transient package
 			if (InParent != GetTransientPackage())
@@ -3927,8 +3928,6 @@ UObject* UTextureFactory::FactoryCreateBinary
 				const int32 PackageUDIMIndex = UE::TextureUtilitiesCommon::ParseUDIMName(PackageName, UdimRegexPattern, PreUDIMName, PostUDIMName);
 				const FString PackageUDIMName = PreUDIMName + PostUDIMName;
 
-				const FString ShortPackageName = ObjectTools::SanitizeInvalidChars(BaseUDIMName, INVALID_LONGPACKAGE_CHARACTERS);
-
 				// If we're re-importing UDIM texture, the package will already be correctly named after the UDIM base name
 				// In this case we'll fail to parse the UDIM name, but the package should already have the proper name
 				// Note that the package name may not match the asset name in this case, if we're reimporting new files
