CookOnTheFlyServer: Add feedback for dlc cook missing the development asset registry file.
#rnx

diff --git a/Engine/Source/Editor/UnrealEd/Private/CookOnTheFlyServer.cpp b/Engine/Source/Editor/UnrealEd/Private/CookOnTheFlyServer.cpp
--- a/Engine/Source/Editor/UnrealEd/Private/CookOnTheFlyServer.cpp
+++ b/Engine/Source/Editor/UnrealEd/Private/CookOnTheFlyServer.cpp
@@ -7339,18 +7339,20 @@ void UCookOnTheFlyServer::StartCookByTheBook( const FCookByTheBookStartupOptions
 			if (!bSucceeded)
 			{
 				const PlatformInfo::FTargetPlatformInfo* PlatformInfo = PlatformInfo::FindPlatformInfo(*InPlatformName);
-				for (const PlatformInfo::FTargetPlatformInfo* PlatformFlavor : PlatformInfo->Flavors)
+				if (PlatformInfo)
 				{
-					OriginalSandboxRegistryFilename = GetBasedOnReleaseVersionAssetRegistryPath(BasedOnReleaseVersion, PlatformFlavor->Name.ToString()) / GetAssetRegistryFilename();
-					bSucceeded = GetAllPackageFilenamesFromAssetRegistry(OriginalSandboxRegistryFilename, bVerifyPackagesExist, OutPackageList);
-					if (bSucceeded)
+					for (const PlatformInfo::FTargetPlatformInfo* PlatformFlavor : PlatformInfo->Flavors)
 					{
-						break;
+						OriginalSandboxRegistryFilename = GetBasedOnReleaseVersionAssetRegistryPath(BasedOnReleaseVersion, PlatformFlavor->Name.ToString()) / GetAssetRegistryFilename();
+						bSucceeded = GetAllPackageFilenamesFromAssetRegistry(OriginalSandboxRegistryFilename, bVerifyPackagesExist, OutPackageList);
+						if (bSucceeded)
+						{
+							break;
+						}
 					}
 				}
 			}
-
-			checkf(bSucceeded, TEXT("Failed to load DevelopmentAssetRegistry for platform %s"), *InPlatformName);
+			return bSucceeded;
 		};
 
 		TArray<FName> OverridePackageList;
@@ -7359,8 +7361,15 @@ void UCookOnTheFlyServer::StartCookByTheBook( const FCookByTheBookStartupOptions
 		if (bUsingDevRegistryOverride)
 		{
 			// Read the contents of the asset registry for the overriden platform. We'll use this for all requested platforms so we can just keep one copy of it here
-			ReadDevelopmentAssetRegistry(OverridePackageList, *DevelopmentAssetRegistryPlatformOverride);
-			checkf(OverridePackageList.Num() != 0, TEXT("DevelopmentAssetRegistry platform override is empty! An override is expected to exist and contain some valid data"));
+			bool bReadSucceeded = ReadDevelopmentAssetRegistry(OverridePackageList, *DevelopmentAssetRegistryPlatformOverride);
+			if (!bReadSucceeded || OverridePackageList.Num() == 0)
+			{
+				UE_LOG(LogCook, Fatal, TEXT("%s based-on AssetRegistry file %s for DevelopmentAssetRegistryPlatformOverride %s. ")
+					TEXT("When cooking DLC, if DevelopmentAssetRegistryPlatformOverride is specified %s is expected to exist under Release/<override> and contain some valid data. Terminating the cook."),
+					!bReadSucceeded ? TEXT("Could not find") : TEXT("Empty"),
+					*(GetBasedOnReleaseVersionAssetRegistryPath(BasedOnReleaseVersion, DevelopmentAssetRegistryPlatformOverride) / TEXT("Metadata") / GetAssetRegistryFilename()),
+					*DevelopmentAssetRegistryPlatformOverride, *GetAssetRegistryFilename());
+			}
 		}
 
 		for ( const ITargetPlatform* TargetPlatform: TargetPlatforms )
@@ -7371,7 +7380,14 @@ void UCookOnTheFlyServer::StartCookByTheBook( const FCookByTheBookStartupOptions
 
 			if (!bUsingDevRegistryOverride)
 			{
-				ReadDevelopmentAssetRegistry(PackageList, PlatformNameString);
+				bool bReadSucceeded = ReadDevelopmentAssetRegistry(PackageList, PlatformNameString);
+				if (!bReadSucceeded)
+				{
+					UE_LOG(LogCook, Fatal, TEXT("Could not find based-on AssetRegistry file %s for platform %s. ")
+						TEXT("When cooking DLC, %s is expected to exist Release/<platform> for each platform being cooked. (Or use DevelopmentAssetRegistryPlatformOverride=<PlatformName> to specify an override platform that all platforms should use to find the %s file). Terminating the cook."),
+						*(GetBasedOnReleaseVersionAssetRegistryPath(BasedOnReleaseVersion, PlatformNameString) / TEXT("Metadata") / GetAssetRegistryFilename()),
+						*PlatformNameString, *GetAssetRegistryFilename(), *GetAssetRegistryFilename());
+				}
 			}
 
 			TArray<FName>& ActivePackageList = OverridePackageList.Num() > 0 ? OverridePackageList : PackageList;
