If given an empty name for FindUniqueKismetName, then default to a backup name instead of looping infinitely.

#rnx

diff --git a/Engine/Source/Editor/UnrealEd/Private/Kismet2/BlueprintEditorUtils.cpp b/Engine/Source/Editor/UnrealEd/Private/Kismet2/BlueprintEditorUtils.cpp
--- a/Engine/Source/Editor/UnrealEd/Private/Kismet2/BlueprintEditorUtils.cpp
+++ b/Engine/Source/Editor/UnrealEd/Private/Kismet2/BlueprintEditorUtils.cpp
@@ -7551,7 +7551,9 @@ FName FBlueprintEditorUtils::FindUniqueKismetName(const UBlueprint* InBlueprint,
 {
 	int32 Count = 0;
 	FString KismetName;
-	FString BaseName = InBaseName;
+	// If an empty string is given then we need to give a valid backup
+	static const FString BackupKismetName = TEXT("K2Name");
+	FString BaseName = InBaseName.IsEmpty() ? BackupKismetName : InBaseName;
 	TSharedPtr<FKismetNameValidator> NameValidator = MakeShareable(new FKismetNameValidator(InBlueprint, NAME_None, InScope));
 
 	// Clean up BaseName to not contain any invalid characters, which will mean we can never find a legal name no matter how many numbers we add
