Removed unnecessary VAlign_Centers now that Property Editor buttons are fixed size.

diff --git a/Engine/Source/Editor/UnrealEd/Public/GameModeInfoCustomizer.h b/Engine/Source/Editor/UnrealEd/Public/GameModeInfoCustomizer.h
--- a/Engine/Source/Editor/UnrealEd/Public/GameModeInfoCustomizer.h
+++ b/Engine/Source/Editor/UnrealEd/Public/GameModeInfoCustomizer.h
@@ -98,19 +98,16 @@ public:
 				]
 			]
 			+SHorizontalBox::Slot()
-			.VAlign(VAlign_Center)
 			.AutoWidth()
 			[
 				PropertyCustomizationHelpers::MakeUseSelectedButton(FSimpleDelegate::CreateSP(this, &FGameModeInfoCustomizer::OnMakeSelectedDefaultClassClicked, DefaultClassPropertyName))
 			]
 			+SHorizontalBox::Slot()
-			.VAlign(VAlign_Center)
 			.AutoWidth()
 			[
 				PropertyCustomizationHelpers::MakeBrowseButton(FSimpleDelegate::CreateSP(this, &FGameModeInfoCustomizer::OnBrowseDefaultClassClicked, DefaultClassPropertyName))
 			]
 			+SHorizontalBox::Slot()
-			.VAlign(VAlign_Center)
 			.AutoWidth()
 			[
 				PropertyCustomizationHelpers::MakeNewBlueprintButton(FSimpleDelegate::CreateSP(this, &FGameModeInfoCustomizer::OnMakeNewDefaultClassClicked, DefaultClassPropertyName))
