World Partition Convert Commandlet
- Don't remap soft object paths for newly created objects, as they are created in the proper world.
- Convert actors from sublevels containing a level script if they are not part of the level script reference tree.

diff --git a/Engine/Source/Editor/UnrealEd/Classes/Commandlets/WorldPartitionConvertCommandlet.h b/Engine/Source/Editor/UnrealEd/Classes/Commandlets/WorldPartitionConvertCommandlet.h
--- a/Engine/Source/Editor/UnrealEd/Classes/Commandlets/WorldPartitionConvertCommandlet.h
+++ b/Engine/Source/Editor/UnrealEd/Classes/Commandlets/WorldPartitionConvertCommandlet.h
@@ -100,6 +100,9 @@ protected:
 	UPROPERTY(Config)
 	TMap<FName, EActorGridPlacement> LevelsGridPlacement;
 
+	UPROPERTY(Config)
+	bool bConvertActorsNotReferencedByLevelScript;
+
 	UPROPERTY(Config)
 	FVector WorldOrigin;
 	
@@ -115,9 +118,9 @@ protected:
 	UPROPERTY(Config)
 	TArray<FHLODLayerActorMapping> HLODLayersForActorClasses;
 
-	UPROPERTY(Config)
-	uint32 LandscapeGridSize;
-
 	UPROPERTY(Transient)
 	TMap<FString, UHLODLayer*> HLODLayers;
+
+	UPROPERTY(Config)
+	uint32 LandscapeGridSize;
 };
diff --git a/Engine/Source/Editor/UnrealEd/Private/Commandlets/WorldPartitionConvertCommandlet.cpp b/Engine/Source/Editor/UnrealEd/Private/Commandlets/WorldPartitionConvertCommandlet.cpp
--- a/Engine/Source/Editor/UnrealEd/Private/Commandlets/WorldPartitionConvertCommandlet.cpp
+++ b/Engine/Source/Editor/UnrealEd/Private/Commandlets/WorldPartitionConvertCommandlet.cpp
@@ -16,6 +16,7 @@
 #include "Engine/LODActor.h"
 #include "Engine/LevelStreaming.h"
 #include "Engine/MapBuildDataRegistry.h"
+#include "Engine/Public/ActorReferencesUtils.h"
 #include "WorldPartition/WorldPartition.h"
 #include "WorldPartition/WorldPartitionSubsystem.h"
 #include "WorldPartition/HLOD/HLODActor.h"
@@ -155,6 +156,7 @@ UWorldPartitionConvertCommandlet::UWorldPartitionConvertCommandlet(const FObject
 	: Super(ObjectInitializer)
 	, bConversionSuffix(false)
 	, ConversionSuffix(TEXT("_WP"))
+	, bConvertActorsNotReferencedByLevelScript(true)
 	, WorldOrigin(FVector::ZeroVector)
 	, WorldExtent(HALF_WORLD_MAX)
 	, LandscapeGridSize(4)
@@ -508,7 +510,10 @@ void UWorldPartitionConvertCommandlet::FixupSoftObjectPaths(UPackage* OuterPacka
 
 	ForEachObjectWithPackage(OuterPackage, [&](UObject* Object)
 	{
-		Object->Serialize(FixupSerializer);
+		if (Object->HasAllFlags(RF_WasLoaded))
+		{
+			Object->Serialize(FixupSerializer);
+		}
 		return true;
 	}, true, RF_NoFlags, EInternalObjectFlags::PendingKill);
 }
@@ -949,7 +954,7 @@ int32 UWorldPartitionConvertCommandlet::Main(const FString& Params)
 		}
 	};
 
-	auto PrepareLevelActors = [this, PartitionFoliage, PartitionLandscape, MainWorldDataLayers](ULevel* Level, bool bMainLevel, EActorGridPlacement DefaultGridPlacement)
+	auto PrepareLevelActors = [this, PartitionFoliage, PartitionLandscape, MainWorldDataLayers](ULevel* Level, TArray<AActor*>& Actors, bool bMainLevel, EActorGridPlacement DefaultGridPlacement)
 	{
 		TRACE_CPUPROFILER_EVENT_SCOPE(PrepareLevelActors);
 
@@ -957,10 +962,12 @@ int32 UWorldPartitionConvertCommandlet::Main(const FString& Params)
 
 		TArray<AInstancedFoliageActor*> IFAs;
 		TSet<ULandscapeInfo*> LandscapeInfos;
-		for (AActor* Actor: Level->Actors)
+		for (AActor* Actor: Actors)
 		{
 			if (Actor && !Actor->IsPendingKill())
 			{
+				check(Actor->GetLevel() == Level);
+
 				if (ShouldDeleteActor(Actor, bMainLevel))
 				{
 					Level->GetWorld()->DestroyActor(Actor);
@@ -1071,7 +1078,8 @@ int32 UWorldPartitionConvertCommandlet::Main(const FString& Params)
 
 	// Prepare levels for conversion
 	DetachDependantLevelPackages(MainLevel);
-	PrepareLevelActors(MainLevel, true, GetLevelGridPlacement(MainLevel, SubLevelsToConvert.Num() ? EActorGridPlacement::AlwaysLoaded : EActorGridPlacement::Bounds));
+	PrepareLevelActors(MainLevel, MainLevel->Actors, true, GetLevelGridPlacement(MainLevel, SubLevelsToConvert.Num() ? EActorGridPlacement::AlwaysLoaded : EActorGridPlacement::Bounds));
+	PackagesToSave.Add(MainLevel->GetPackage());
 
 	if (bConversionSuffix)
 	{
@@ -1097,11 +1105,80 @@ int32 UWorldPartitionConvertCommandlet::Main(const FString& Params)
 		UWorld* SubWorld = SubLevel->GetTypedOuter<UWorld>();
 		UPackage* SubPackage = SubLevel->GetPackage();
 
+		RemapSoftObjectPaths.Add(FSoftObjectPath(SubWorld).ToString(), FSoftObjectPath(MainWorld).ToString());
+		RemapSoftObjectPaths.Add(FSoftObjectPath(SubLevel).ToString(), FSoftObjectPath(MainLevel).ToString());
+		RemapSoftObjectPaths.Add(FSoftObjectPath(SubPackage).ToString(), FSoftObjectPath(MainPackage).ToString());
+
+		TArray<AActor*> ActorsToConvert;
 		if (LevelHasLevelScriptBlueprint(SubLevel))
 		{
 			MapsWithLevelScriptsBPs.Add(SubPackage->GetLoadedPath().GetPackageName());
 
-			// Spawn the Level Instance
+			if (bConvertActorsNotReferencedByLevelScript)
+			{
+				// Gather the list of actors referenced by the level script blueprint
+				TSet<AActor*> LevelScriptActorReferences;
+
+				ALevelScriptActor* LevelScriptActor = SubLevel->GetLevelScriptActor();
+				LevelScriptActorReferences.Add(LevelScriptActor);
+
+				ULevelScriptBlueprint* LevelScriptBlueprint = SubLevel->GetLevelScriptBlueprint(true);
+				LevelScriptActorReferences.Append(ActorsReferencesUtils::GetActorReferences(LevelScriptBlueprint));
+
+				for(AActor* Actor: SubLevel->Actors)
+				{
+					if(Actor && !Actor->IsPendingKill())
+					{
+						TSet<AActor*> ActorReferences;
+						ActorReferences.Append(ActorsReferencesUtils::GetActorReferences(Actor));
+
+						for (AActor* ActorReference : ActorReferences)
+						{
+							if (LevelScriptActorReferences.Find(ActorReference))
+							{
+								LevelScriptActorReferences.Add(Actor);
+								LevelScriptActorReferences.Append(ActorReferences);
+								break;
+							}
+						}
+					}
+				}
+
+				for(AActor* Actor: SubLevel->Actors)
+				{
+					if(Actor && !Actor->IsPendingKill())
+					{
+						if (!LevelScriptActorReferences.Find(Actor))
+						{
+							ActorsToConvert.Add(Actor);
+						}
+					}
+				}
+			}
+
+			// Rename the world if requested
+			UWorld* SubLevelWorld = SubLevel->GetTypedOuter<UWorld>();
+			UPackage* SubLevelPackage = SubLevelWorld->GetPackage();
+
+			if (bConversionSuffix)
+			{
+				FString OldMainWorldPath = FSoftObjectPath(SubLevelWorld).ToString();
+				FString OldMainLevelPath = FSoftObjectPath(SubLevel).ToString();
+				FString OldPackagePath = FSoftObjectPath(SubLevelPackage).ToString();
+
+				if (!RenameWorldPackageWithSuffix(SubLevelWorld))
+				{
+					return 1;
+				}
+
+				RemapSoftObjectPaths.Add(OldMainWorldPath, FSoftObjectPath(SubLevelWorld).ToString());
+				RemapSoftObjectPaths.Add(OldMainLevelPath, FSoftObjectPath(SubLevel).ToString());
+				RemapSoftObjectPaths.Add(OldPackagePath, FSoftObjectPath(SubLevelPackage).ToString());
+			}
+			
+			PackagesToSave.Add(SubLevelPackage);
+
+			// Spawn the level instance actor
 			ULevelStreaming* SubLevelStreaming = nullptr;
 			for (ULevelStreaming* LevelStreaming : MainLevel->GetWorld()->GetStreamingLevels())
 			{
@@ -1117,8 +1194,6 @@ int32 UWorldPartitionConvertCommandlet::Main(const FString& Params)
 			SpawnParams.OverrideLevel = MainLevel;
 			ALevelInstance* LevelInstanceActor = MainWorld->SpawnActor<ALevelInstance>(SpawnParams);
 
-			UPackage* SubLevelPackage = SubLevel->GetPackage();
-
 			FTransform LevelTransform;
 			if (SubLevelPackage->WorldTileInfo)
 			{
@@ -1130,70 +1205,72 @@ int32 UWorldPartitionConvertCommandlet::Main(const FString& Params)
 			}
 
 			LevelInstanceActor->SetActorTransform(LevelTransform);
-			LevelInstanceActor->SetWorldAsset(SubLevel->GetTypedOuter<UWorld>());
+			LevelInstanceActor->SetWorldAsset(SubLevelWorld);
 		}
 		else
 		{
-			RemapSoftObjectPaths.Add(FSoftObjectPath(SubWorld).ToString(), FSoftObjectPath(MainWorld).ToString());
-			RemapSoftObjectPaths.Add(FSoftObjectPath(SubLevel).ToString(), FSoftObjectPath(MainLevel).ToString());
-			RemapSoftObjectPaths.Add(FSoftObjectPath(SubPackage).ToString(), FSoftObjectPath(MainPackage).ToString());
-
 			if (LevelHasMapBuildData(SubLevel))
 			{
 				MapsWithMapBuildData.Add(SubPackage->GetLoadedPath().GetPackageName());
 			}
 
 			DetachDependantLevelPackages(SubLevel);
-			PrepareLevelActors(SubLevel, false, GetLevelGridPlacement(SubLevel, EActorGridPlacement::Bounds));
 
-			UE_LOG(LogWorldPartitionConvertCommandlet, Log, TEXT("Converting %s"), *SubWorld->GetName());
+			ActorsToConvert = SubLevel->Actors;
+		}
+
+		UE_LOG(LogWorldPartitionConvertCommandlet, Log, TEXT("Converting %s"), *SubWorld->GetName());
+
+		PrepareLevelActors(SubLevel, ActorsToConvert, false, GetLevelGridPlacement(SubLevel, EActorGridPlacement::Bounds));
 
-			for(AActor* Actor: SubLevel->Actors)
+		for(AActor* Actor: ActorsToConvert)
+		{
+			if(Actor && !Actor->IsPendingKill())
 			{
-				if(Actor && !Actor->IsPendingKill())
-				{
-					check(Actor->GetOuter() == SubLevel);
-					check(!ShouldDeleteActor(Actor, false));
+				check(Actor->GetOuter() == SubLevel);
+				check(!ShouldDeleteActor(Actor, false));
 				
-					if (Actor->IsA(AGroupActor::StaticClass()))
-					{
-						GroupActors.Add(*Actor->GetFullName());
-					}
+				if (Actor->IsA(AGroupActor::StaticClass()))
+				{
+					GroupActors.Add(*Actor->GetFullName());
+				}
 
-					if (Actor->GroupActor)
-					{
-						ActorsInGroupActors.Add(*Actor->GetFullName());
-					}
+				if (Actor->GroupActor)
+				{
+					ActorsInGroupActors.Add(*Actor->GetFullName());
+				}
 
-					TArray<AActor*> ChildActors;
-					Actor->GetAllChildActors(ChildActors, false);
+				TArray<AActor*> ChildActors;
+				Actor->GetAllChildActors(ChildActors, false);
 
-					if (ChildActors.Num())
-					{
-						ActorsWithChildActors.Add(*Actor->GetFullName());
-					}
+				if (ChildActors.Num())
+				{
+					ActorsWithChildActors.Add(*Actor->GetFullName());
+				}
 
-					FArchiveGatherPrivateImports Ar(Actor, PrivateRefsMap, ActorsReferencesToActors);
-					Actor->Serialize(Ar);
+				FArchiveGatherPrivateImports Ar(Actor, PrivateRefsMap, ActorsReferencesToActors);
+				Actor->Serialize(Ar);
 
-					// Even after Foliage Partitioning it is possible some Actors still have a FoliageTag. Make sure it is removed.
-					if (FFoliageHelper::IsOwnedByFoliage(Actor))
-					{
-						FFoliageHelper::SetIsOwnedByFoliage(Actor, false);
-					}
+				// Even after Foliage Partitioning it is possible some Actors still have a FoliageTag. Make sure it is removed.
+				if (FFoliageHelper::IsOwnedByFoliage(Actor))
+				{
+					FFoliageHelper::SetIsOwnedByFoliage(Actor, false);
+				}
 
-					ChangeObjectOuter(Actor, MainLevel);
+				ChangeObjectOuter(Actor, MainLevel);
 
-					// Migrate blueprint classes
-					UClass* ActorClass = Actor->GetClass();
-					if (!ActorClass->IsNative() && (ActorClass->GetPackage() == SubPackage))
-					{
-						ChangeObjectOuter(ActorClass, MainPackage);
-						UE_LOG(LogWorldPartitionConvertCommandlet, Log, TEXT("Extracted non-native class %s"), *ActorClass->GetName());
-					}
+				// Migrate blueprint classes
+				UClass* ActorClass = Actor->GetClass();
+				if (!ActorClass->IsNative() && (ActorClass->GetPackage() == SubPackage))
+				{
+					ChangeObjectOuter(ActorClass, MainPackage);
+					UE_LOG(LogWorldPartitionConvertCommandlet, Log, TEXT("Extracted non-native class %s"), *ActorClass->GetName());
 				}
 			}
+		}
 
+		if (!LevelHasLevelScriptBlueprint(SubLevel))
+		{
 			if (!bReportOnly)
 			{
 				TArray<UObject*> ObjectsToRename;
@@ -1221,6 +1298,7 @@ int32 UWorldPartitionConvertCommandlet::Main(const FString& Params)
 	for (ULevelStreaming* LevelStreaming: MainLevel->GetWorld()->GetStreamingLevels())
 	{
 		LevelStreaming->MarkPendingKill();
+		ULevelStreaming::RemoveLevelAnnotation(LevelStreaming->GetLoadedLevel());
 	}
 	MainLevel->GetWorld()->ClearStreamingLevels();
 
@@ -1322,6 +1400,7 @@ int32 UWorldPartitionConvertCommandlet::Main(const FString& Params)
 			}
 		}
 
+		// Checkout packages
 		{
 			TRACE_CPUPROFILER_EVENT_SCOPE(CheckoutPackages);
 
@@ -1359,6 +1438,7 @@ int32 UWorldPartitionConvertCommandlet::Main(const FString& Params)
 			}
 		}
 
+		// Add packages
 		{
 			TRACE_CPUPROFILER_EVENT_SCOPE(AddPackagesToSourceControl);
 
@@ -1372,25 +1452,6 @@ int32 UWorldPartitionConvertCommandlet::Main(const FString& Params)
 			}
 		}
 
-		// Checkout level
-		UE_LOG(LogWorldPartitionConvertCommandlet, Log, TEXT("Saving %s."), *MainPackage->GetName());
-
-		if (!PackageHelper.Checkout(MainPackage))
-		{
-			return 1;
-		}
-
-		// Save level
-		{
-			TRACE_CPUPROFILER_EVENT_SCOPE(SaveLevel);
-
-			FString PackageFileName = SourceControlHelpers::PackageFilename(MainPackage);
-			if (!UPackage::SavePackage(MainPackage, nullptr, RF_Standalone, *PackageFileName, GError, nullptr, false, true, SAVE_Async))
-			{
-				return 1;
-			}
-		}
-
 		if(bInitializedPhysicsSceneForSave)
 		{
 			GEditor->CleanupPhysicsSceneThatWasInitializedForSave(MainWorld, bForceInitializeWorld);
diff --git a/Engine/Source/Runtime/Engine/Private/ActorReferencesUtils.cpp b/Engine/Source/Runtime/Engine/Private/ActorReferencesUtils.cpp
new file mode 100644
--- /dev/null
+++ b/Engine/Source/Runtime/Engine/Private/ActorReferencesUtils.cpp
@@ -0,0 +1,89 @@
+// Copyright Epic Games, Inc. All Rights Reserved.
+
+#include "ActorReferencesUtils.h"
+#include "Serialization/ArchiveUObject.h"
+#include "GameFramework/Actor.h"
+
+class FArchiveGatherExternalActorRefs : public FArchiveUObject
+{
+public:
+	FArchiveGatherExternalActorRefs(UObject* InRoot, TSet<AActor*>& InActorReferences, EObjectFlags InRequiredFlags)
+		: Root(InRoot)
+		, ActorReferences(InActorReferences)
+		, RequiredFlags(InRequiredFlags)
+	{
+		SetIsSaving(true);
+		SetIsPersistent(true);
+		ArIgnoreOuterRef = true;
+		ArIsObjectReferenceCollector = true;
+		ArShouldSkipBulkData = true;
+
+		SubObjects.Add(Root);
+
+		Root->Serialize(*this);
+	}
+
+	virtual FArchive& operator<<(UObject*& Obj) override
+	{
+		if (Obj && !Obj->IsTemplate() && !Obj->HasAnyFlags(RF_Transient))
+		{
+			bool bWasAlreadyInSet;
+			SubObjects.Add(Obj, &bWasAlreadyInSet);
+
+			if (!bWasAlreadyInSet)
+			{
+				HandleObjectReference(Obj);
+
+				if (Obj->IsInOuter(Root))
+				{
+					Obj->Serialize(*this);
+				}
+			}
+		}
+		return *this;
+	}
+
+private:
+	void HandleObjectReference(UObject* Obj)
+	{
+		AActor* Actor = Cast<AActor>(Obj);
+
+		if (!Actor)
+		{
+			Actor = Obj->GetTypedOuter<AActor>();
+		}
+
+		if(Actor)
+		{
+			AActor* TopParentActor = Actor;
+			while(TopParentActor->GetParentActor())
+			{
+				TopParentActor = TopParentActor->GetParentActor();
+			}
+
+			check(TopParentActor);
+
+			if (((RequiredFlags == RF_NoFlags) || TopParentActor->HasAnyFlags(RequiredFlags)) && (TopParentActor != Root))
+			{
+				ActorReferences.Add(TopParentActor);
+			}
+		}
+	}
+
+	UObject* Root;
+	TSet<AActor*>& ActorReferences;
+	TSet<UObject*> SubObjects;
+	EObjectFlags RequiredFlags;
+};
+
+TArray<AActor*> ActorsReferencesUtils::GetExternalActorReferences(UObject* Root)
+{
+	return GetActorReferences(Root, RF_HasExternalPackage);
+}
+
+TArray<AActor*> ActorsReferencesUtils::GetActorReferences(UObject* Root, EObjectFlags RequiredFlags)
+{
+	TSet<AActor*> Result;
+	FArchiveGatherExternalActorRefs Ar(Root, Result, RequiredFlags);
+	return Result.Array();
+}
\ No newline at end of file
diff --git a/Engine/Source/Runtime/Engine/Private/ExternalActorsUtils.cpp b/Engine/Source/Runtime/Engine/Private/ExternalActorsUtils.cpp
deleted file mode 100644
diff --git a/Engine/Source/Runtime/Engine/Private/WorldPartition/WorldPartitionActorDesc.cpp b/Engine/Source/Runtime/Engine/Private/WorldPartition/WorldPartitionActorDesc.cpp
--- a/Engine/Source/Runtime/Engine/Private/WorldPartition/WorldPartitionActorDesc.cpp
+++ b/Engine/Source/Runtime/Engine/Private/WorldPartition/WorldPartitionActorDesc.cpp
@@ -16,7 +16,7 @@
 #include "UObject/UE5MainStreamObjectVersion.h"
 #include "WorldPartition/WorldPartition.h"
 #include "WorldPartition/DataLayer/DataLayer.h"
-#include "Engine/Public/ExternalActorsUtils.h"
+#include "Engine/Public/ActorReferencesUtils.h"
 #endif
 
 #if WITH_EDITOR
@@ -66,7 +66,7 @@ void FWorldPartitionActorDesc::Init(const AActor* InActor)
 	ActorPackage = InActor->GetPackage()->GetFName();
 	ActorPath = *InActor->GetPathName();
 	
-	TArray<AActor*> ActorReferences = ExternalActorsUtils::GetExternalActorReferences((AActor*)InActor);
+	TArray<AActor*> ActorReferences = ActorsReferencesUtils::GetExternalActorReferences((AActor*)InActor);
 
 	if (ActorReferences.Num())
 	{
diff --git a/Engine/Source/Runtime/Engine/Private/WorldPartition/WorldPartitionRuntimeSpatialHash.cpp b/Engine/Source/Runtime/Engine/Private/WorldPartition/WorldPartitionRuntimeSpatialHash.cpp
--- a/Engine/Source/Runtime/Engine/Private/WorldPartition/WorldPartitionRuntimeSpatialHash.cpp
+++ b/Engine/Source/Runtime/Engine/Private/WorldPartition/WorldPartitionRuntimeSpatialHash.cpp
@@ -38,7 +38,7 @@
 #include "WorldPartition/HLOD/HLODActorDesc.h"
 #include "WorldPartition/NavigationData/NavigationDataChunkActor.h"
 #include "AI/NavigationSystemBase.h"
-#include "Engine/Public/ExternalActorsUtils.h"
+#include "Engine/Public/ActorReferencesUtils.h"
 #include "Engine/LevelScriptBlueprint.h"
 
 #include "AssetRegistryModule.h"
@@ -253,7 +253,7 @@ TArray<FActorCluster> CreateActorClustersImpl(UWorldPartition* WorldPartition, T
 	TSet<AActor*> LevelScriptExternalActorReferences;
 	if (ULevelScriptBlueprint* LevelScriptBlueprint = WorldPartition->GetWorld()->PersistentLevel->GetLevelScriptBlueprint(true))
 	{
-		LevelScriptExternalActorReferences.Append(ExternalActorsUtils::GetExternalActorReferences(LevelScriptBlueprint));
+		LevelScriptExternalActorReferences.Append(ActorsReferencesUtils::GetExternalActorReferences(LevelScriptBlueprint));
 	}
 
 	for (const auto& Pair : WorldPartition->Actors)
diff --git a/Engine/Source/Runtime/Engine/Public/ActorReferencesUtils.h b/Engine/Source/Runtime/Engine/Public/ActorReferencesUtils.h
new file mode 100644
--- /dev/null
+++ b/Engine/Source/Runtime/Engine/Public/ActorReferencesUtils.h
@@ -0,0 +1,20 @@
+// Copyright Epic Games, Inc. All Rights Reserved.
+
+#pragma once
+
+#include "CoreMinimal.h"
+
+class AActor;
+
+namespace ActorsReferencesUtils
+{
+	/**
+	 * Gather direct references to external actors from the root object.
+	 */
+	ENGINE_API TArray<AActor*> GetExternalActorReferences(UObject* Root);
+
+	/**
+	 * Gather direct references to actors from the root object.
+	 */
+	ENGINE_API TArray<AActor*> GetActorReferences(UObject* Root, EObjectFlags RequiredFlags = RF_NoFlags);
+}
\ No newline at end of file
diff --git a/Engine/Source/Runtime/Engine/Public/ExternalActorsUtils.h b/Engine/Source/Runtime/Engine/Public/ExternalActorsUtils.h
deleted file mode 100644
