World Partition Convert Commandlet: properly handle invisible streaming levels.

diff --git a/Engine/Source/Editor/UnrealEd/Private/Commandlets/WorldPartitionConvertCommandlet.cpp b/Engine/Source/Editor/UnrealEd/Private/Commandlets/WorldPartitionConvertCommandlet.cpp
--- a/Engine/Source/Editor/UnrealEd/Private/Commandlets/WorldPartitionConvertCommandlet.cpp
+++ b/Engine/Source/Editor/UnrealEd/Private/Commandlets/WorldPartitionConvertCommandlet.cpp
@@ -252,6 +252,7 @@ void UWorldPartitionConvertCommandlet::GatherAndPrepareSubLevelsToConvert(const
 	{
 		StreamingLevel->SetShouldBeLoaded(true);
 		StreamingLevel->SetShouldBeVisible(true);
+		StreamingLevel->SetShouldBeVisibleInEditor(true);
 	}
 
 	World->FlushLevelStreaming(EFlushLevelStreamingType::Full);
