Introduced a shared thumbnail pool for all editor thumbnails

Unless you are rendering a huge amount of thumbnails at once this shared pool should be used

diff --git a/Engine/Source/Editor/UnrealEd/Classes/ThumbnailRendering/ThumbnailManager.h b/Engine/Source/Editor/UnrealEd/Classes/ThumbnailRendering/ThumbnailManager.h
--- a/Engine/Source/Editor/UnrealEd/Classes/ThumbnailRendering/ThumbnailManager.h
+++ b/Engine/Source/Editor/UnrealEd/Classes/ThumbnailRendering/ThumbnailManager.h
@@ -17,6 +17,7 @@
 #include "ThumbnailManager.generated.h"
 
 class FViewport;
+class FAssetThumbnailPool;
 
 /**
  * Types of primitives for drawing thumbnails of resources.
@@ -82,32 +83,6 @@ UCLASS(config=Editor)
 class UThumbnailManager : public UObject
 {
 	GENERATED_UCLASS_BODY()
-
-protected:
-	/**
-	 * The array of thumbnail rendering information entries. Each type that supports
-	 * thumbnail rendering has an entry in here.
-	 */
-	UPROPERTY(config)
-	TArray<struct FThumbnailRenderingInfo> RenderableThumbnailTypes;
-
-	/**
-	 * Determines whether the initialization function is needed or not
-	 */
-	bool bIsInitialized;
-
-	// The following members are present for performance optimizations
-	
-	/**
-	 * Whether to update the map or not (GC usually causes this)
-	 */
-	bool bMapNeedsUpdate;
-
-	/**
-	 * This holds a map of object type to render info entries
-	 */
-	TMap<UClass*, FThumbnailRenderingInfo*> RenderInfoMap;
-
 public:
 	/**
 	 * The render info to share across all object types when the object doesn't
@@ -175,11 +150,25 @@ public:
 	 */
 	UNREALED_API virtual void UnregisterCustomRenderer(UClass* Class);
 
+	/** 
+	 * Returns the thumbnail pool which should be used for most thumbnails in the editor. 
+	 * Unless you are rendering a huge amount of thumbnails at once this shared pool should be used
+	 */
+	UNREALED_API TSharedPtr<FAssetThumbnailPool> GetSharedThumbnailPool() const { return SharedThumbnailPool; }
+protected:
+	/**
+	 * The array of thumbnail rendering information entries. Each type that supports
+	 * thumbnail rendering has an entry in here.
+	 */
+	UPROPERTY(config)
+	TArray<struct FThumbnailRenderingInfo> RenderableThumbnailTypes;
 
+	/**
+	 * This holds a map of object type to render info entries
+	 */
+	TMap<UClass*, FThumbnailRenderingInfo*> RenderInfoMap;
 
 
-
-protected:
 	/**
 	 * Holds the name of the thumbnail manager singleton class to instantiate
 	 */
@@ -191,6 +180,21 @@ protected:
 	 */
 	static class UThumbnailManager* ThumbnailManagerSingleton;
 
+	/** Thumbnail pool used for most thumbnails in the editor */
+	TSharedPtr<FAssetThumbnailPool> SharedThumbnailPool;
+
+	/**
+	 * Determines whether the initialization function is needed or not
+	 */
+	bool bIsInitialized;
+
+	// The following members are present for performance optimizations
+
+	/**
+	 * Whether to update the map or not (GC usually causes this)
+	 */
+	bool bMapNeedsUpdate;
+
 public:
 	/** Returns the thumbnail manager and creates it if missing */
 	UNREALED_API static UThumbnailManager& Get();
diff --git a/Engine/Source/Editor/UnrealEd/Private/ThumbnailManager.cpp b/Engine/Source/Editor/UnrealEd/Private/ThumbnailManager.cpp
--- a/Engine/Source/Editor/UnrealEd/Private/ThumbnailManager.cpp
+++ b/Engine/Source/Editor/UnrealEd/Private/ThumbnailManager.cpp
@@ -15,6 +15,7 @@
 #include "Engine/TextureCube.h"
 #include "Engine/Texture2DArray.h"
 #include "ImageUtils.h"
+#include "AssetThumbnail.h"
 
 
 DEFINE_LOG_CATEGORY_STATIC(LogThumbnailManager, Log, All);
@@ -70,6 +71,8 @@ void UThumbnailManager::Initialize(void)
 	{
 		InitializeRenderTypeArray(RenderableThumbnailTypes);
 
+		SharedThumbnailPool = MakeShared<FAssetThumbnailPool>(100);
+
 		bIsInitialized = true;
 	}
 }
