Deleting a struct from an array when that struct has a property with an EditCondition no longer crashes due to a missing null check.

Added test case to "testprops" console command.

Reported from UDN.

diff --git a/Engine/Source/Editor/PropertyEditor/Private/ObjectPropertyNode.h b/Engine/Source/Editor/PropertyEditor/Private/ObjectPropertyNode.h
--- a/Engine/Source/Editor/PropertyEditor/Private/ObjectPropertyNode.h
+++ b/Engine/Source/Editor/PropertyEditor/Private/ObjectPropertyNode.h
@@ -94,14 +94,24 @@ public:
 	}
 	virtual uint8* GetValuePtrOfInstance(int32 Index, const FProperty* InProperty, FPropertyNode* InParentNode) override
 	{
+		if (InParentNode == nullptr || InProperty == nullptr)
+		{
+			return nullptr;
+		}
+
 		const UObject* Obj = GetUObject(Index);
 		if (Obj == nullptr)
 		{
 			return nullptr;
 		}
 
-		uint8* ParentOffset = InParentNode ? InParentNode->GetValueAddressFromObject(Obj) : nullptr;
-		return InProperty ? InProperty->ContainerPtrToValuePtr<uint8>(ParentOffset) : nullptr;
+		uint8* ParentPtr = InParentNode->GetValueAddressFromObject(Obj);
+		if (ParentPtr == nullptr)
+		{
+			return nullptr;
+		}
+
+		return InProperty->ContainerPtrToValuePtr<uint8>(ParentPtr);
 	}
 	virtual TWeakObjectPtr<UObject> GetInstanceAsUObject(int32 Index) override
 	{
diff --git a/Engine/Source/Editor/UnrealEd/Classes/Editor/PropertyEditorTestObject.h b/Engine/Source/Editor/UnrealEd/Classes/Editor/PropertyEditorTestObject.h
--- a/Engine/Source/Editor/UnrealEd/Classes/Editor/PropertyEditorTestObject.h
+++ b/Engine/Source/Editor/UnrealEd/Classes/Editor/PropertyEditorTestObject.h
@@ -153,15 +153,21 @@ struct FPropertyEditorTestBasicStruct
 }; 
 
 USTRUCT()
-struct FPropertyEditorTestInlineEditCondition
+struct FPropertyEditorTestEditCondition
 {
 	GENERATED_BODY()
 
 	UPROPERTY(EditAnywhere, Category=Default, meta=(InlineEditConditionToggle))
 	bool InlineEditCondition;
 
-	UPROPERTY(EditAnywhere, Category=Default, meta=(EditCondition=InlineEditCondition))
+	UPROPERTY(EditAnywhere, Category=Default, meta=(EditCondition="InlineEditCondition"))
 	int32 HasInlineEditCondition;
+
+	UPROPERTY(EditAnywhere, Category=Default)
+	ETestEnumFlags Flags;
+
+	UPROPERTY(EditAnywhere, Category=Default, meta=(EditCondition="Flags == ETestEnumFlags::One", EditConditionHides))
+	int32 EnabledAndVisibleWhenOne;
 };
 
 UCLASS(EditInlineNew)
@@ -707,10 +713,10 @@ class UPropertyEditorTestObject : public UObject
 	float UsesSharedEditCondition2;
 
 	UPROPERTY(EditAnywhere, Category="Inline Edit Conditions")
-	FPropertyEditorTestInlineEditCondition StructWithInlineCondition;
+	FPropertyEditorTestEditCondition StructWithInlineCondition;
 
 	UPROPERTY(EditAnywhere, Category="Inline Edit Conditions")
-	TArray<FPropertyEditorTestInlineEditCondition> ArrayOfStructsWithInlineCondition;
+	TArray<FPropertyEditorTestEditCondition> ArrayOfStructsWithInlineCondition;
 
 	bool CanEditChange(const FProperty* InProperty) const;
 };
