Automated wrapped object pointer upgrade for Engine + ShooterGame + ShooterGame referenced plugins

diff --git a/Engine/Source/Developer/GameplayDebugger/Public/GameplayDebuggerLocalController.h b/Engine/Source/Developer/GameplayDebugger/Public/GameplayDebuggerLocalController.h
--- a/Engine/Source/Developer/GameplayDebugger/Public/GameplayDebuggerLocalController.h
+++ b/Engine/Source/Developer/GameplayDebugger/Public/GameplayDebuggerLocalController.h
@@ -42,13 +42,13 @@ protected:
 	friend struct FGameplayDebuggerConsoleCommands;
 
 	UPROPERTY()
-	AGameplayDebuggerCategoryReplicator* CachedReplicator;
+	TObjectPtr<AGameplayDebuggerCategoryReplicator> CachedReplicator;
 
 	UPROPERTY()
-	AGameplayDebuggerPlayerManager* CachedPlayerManager;
+	TObjectPtr<AGameplayDebuggerPlayerManager> CachedPlayerManager;
 
 	UPROPERTY()
-	AActor* DebugActorCandidate;
+	TObjectPtr<AActor> DebugActorCandidate;
 
 	TArray<TArray<int32> > DataPackMap;
 	TArray<TArray<int32> > SlotCategoryIds;
diff --git a/Engine/Source/Developer/GameplayDebugger/Public/GameplayDebuggerPlayerManager.h b/Engine/Source/Developer/GameplayDebugger/Public/GameplayDebuggerPlayerManager.h
--- a/Engine/Source/Developer/GameplayDebugger/Public/GameplayDebuggerPlayerManager.h
+++ b/Engine/Source/Developer/GameplayDebugger/Public/GameplayDebuggerPlayerManager.h
@@ -19,13 +19,13 @@ struct FGameplayDebuggerPlayerData
 	GENERATED_USTRUCT_BODY()
 
 	UPROPERTY()
-	UGameplayDebuggerLocalController* Controller = nullptr;
+	TObjectPtr<UGameplayDebuggerLocalController> Controller = nullptr;
 
 	UPROPERTY()
-	UInputComponent* InputComponent = nullptr;
+	TObjectPtr<UInputComponent> InputComponent = nullptr;
 
 	UPROPERTY()
-	AGameplayDebuggerCategoryReplicator* Replicator = nullptr;
+	TObjectPtr<AGameplayDebuggerCategoryReplicator> Replicator = nullptr;
 };
 
 UCLASS(NotBlueprintable, NotBlueprintType, notplaceable, noteditinlinenew, hidedropdown, Transient)
@@ -75,7 +75,7 @@ protected:
 	TArray<FGameplayDebuggerPlayerData> PlayerData;
 
 	UPROPERTY()
-	TArray<AGameplayDebuggerCategoryReplicator*> PendingRegistrations;
+	TArray<TObjectPtr<AGameplayDebuggerCategoryReplicator>> PendingRegistrations;
 
 #if WITH_EDITORONLY_DATA 
 	UPROPERTY()
diff --git a/Engine/Source/Developer/Localization/Public/LocalizationSettings.h b/Engine/Source/Developer/Localization/Public/LocalizationSettings.h
--- a/Engine/Source/Developer/Localization/Public/LocalizationSettings.h
+++ b/Engine/Source/Developer/Localization/Public/LocalizationSettings.h
@@ -22,13 +22,13 @@ public:
 
 private:
 	UPROPERTY()
-	ULocalizationTargetSet* EngineTargetSet;
+	TObjectPtr<ULocalizationTargetSet> EngineTargetSet;
 
 	UPROPERTY(config)
 	TArray<FLocalizationTargetSettings> EngineTargetsSettings;
 
 	UPROPERTY()
-	ULocalizationTargetSet* GameTargetSet;
+	TObjectPtr<ULocalizationTargetSet> GameTargetSet;
 
 	UPROPERTY(config)
 	TArray<FLocalizationTargetSettings> GameTargetsSettings;
diff --git a/Engine/Source/Developer/Localization/Public/LocalizationTargetTypes.h b/Engine/Source/Developer/Localization/Public/LocalizationTargetTypes.h
--- a/Engine/Source/Developer/Localization/Public/LocalizationTargetTypes.h
+++ b/Engine/Source/Developer/Localization/Public/LocalizationTargetTypes.h
@@ -496,7 +496,7 @@ class LOCALIZATION_API ULocalizationTargetSet : public UObject
 
 public:
 	UPROPERTY(EditAnywhere, Category = "Targets")
-	TArray<ULocalizationTarget*> TargetObjects;
+	TArray<TObjectPtr<ULocalizationTarget>> TargetObjects;
 
 public:
 #if WITH_EDITOR
diff --git a/Engine/Source/Developer/LogVisualizer/Private/VisualLoggerCameraController.h b/Engine/Source/Developer/LogVisualizer/Private/VisualLoggerCameraController.h
--- a/Engine/Source/Developer/LogVisualizer/Private/VisualLoggerCameraController.h
+++ b/Engine/Source/Developer/LogVisualizer/Private/VisualLoggerCameraController.h
@@ -18,7 +18,7 @@ public:
 	GENERATED_UCLASS_BODY()
 
 	UPROPERTY()
-	class AActor* PickedActor;
+	TObjectPtr<class AActor> PickedActor;
 
 	virtual void PostInitializeComponents() override;
 
diff --git a/Engine/Source/Developer/LogVisualizer/Public/LogVisualizerSettings.h b/Engine/Source/Developer/LogVisualizer/Public/LogVisualizerSettings.h
--- a/Engine/Source/Developer/LogVisualizer/Public/LogVisualizerSettings.h
+++ b/Engine/Source/Developer/LogVisualizer/Public/LogVisualizerSettings.h
@@ -188,7 +188,7 @@ protected:
 
 	/** A material used to render debug meshes with kind of flat shading, mostly used by Visual Logger tool. */
 	UPROPERTY()
-	class UMaterial* DebugMeshMaterialFakeLight;
+	TObjectPtr<class UMaterial> DebugMeshMaterialFakeLight;
 
 	/** @todo document */
 	UPROPERTY(config)
diff --git a/Engine/Source/Developer/ToolMenus/Public/ToolMenu.h b/Engine/Source/Developer/ToolMenus/Public/ToolMenu.h
--- a/Engine/Source/Developer/ToolMenus/Public/ToolMenu.h
+++ b/Engine/Source/Developer/ToolMenus/Public/ToolMenu.h
@@ -157,7 +157,7 @@ public:
 	TArray<FToolMenuSection> Sections;
 
 	UPROPERTY()
-	const UToolMenu* SubMenuParent;
+	TObjectPtr<const UToolMenu> SubMenuParent;
 
 	UPROPERTY()
 	FName SubMenuSourceEntryName;
diff --git a/Engine/Source/Developer/ToolMenus/Public/ToolMenuContext.h b/Engine/Source/Developer/ToolMenus/Public/ToolMenuContext.h
--- a/Engine/Source/Developer/ToolMenus/Public/ToolMenuContext.h
+++ b/Engine/Source/Developer/ToolMenus/Public/ToolMenuContext.h
@@ -93,7 +93,7 @@ private:
 	void Empty();
 
 	UPROPERTY()
-	TArray<UObject*> ContextObjects;
+	TArray<TObjectPtr<UObject>> ContextObjects;
 
 	TArray<TSharedPtr<FUICommandList>> CommandLists;
 
diff --git a/Engine/Source/Developer/ToolMenus/Public/ToolMenuEntry.h b/Engine/Source/Developer/ToolMenus/Public/ToolMenuEntry.h
--- a/Engine/Source/Developer/ToolMenus/Public/ToolMenuEntry.h
+++ b/Engine/Source/Developer/ToolMenus/Public/ToolMenuEntry.h
@@ -163,7 +163,7 @@ public:
 	bool bShouldCloseWindowAfterMenuSelection;
 
 	UPROPERTY(EditAnywhere, BlueprintReadWrite, Category = "Tool Menus")
-	UToolMenuEntryScript* ScriptObject;
+	TObjectPtr<UToolMenuEntryScript> ScriptObject;
 
 	UPROPERTY(EditAnywhere, BlueprintReadWrite, Category = "Tool Menus")
 	FName StyleNameOverride;
diff --git a/Engine/Source/Developer/ToolMenus/Public/ToolMenuSection.h b/Engine/Source/Developer/ToolMenus/Public/ToolMenuSection.h
--- a/Engine/Source/Developer/ToolMenus/Public/ToolMenuSection.h
+++ b/Engine/Source/Developer/ToolMenus/Public/ToolMenuSection.h
@@ -95,7 +95,7 @@ public:
 	FToolMenuContext Context;
 
 	UPROPERTY(EditAnywhere, BlueprintReadWrite, Category = "Tool Menus")
-	UToolMenuSectionDynamic* ToolMenuSectionDynamic;
+	TObjectPtr<UToolMenuSectionDynamic> ToolMenuSectionDynamic;
 
 	TAttribute<FText> Label;
 
diff --git a/Engine/Source/Developer/ToolMenus/Public/ToolMenus.h b/Engine/Source/Developer/ToolMenus/Public/ToolMenus.h
--- a/Engine/Source/Developer/ToolMenus/Public/ToolMenus.h
+++ b/Engine/Source/Developer/ToolMenus/Public/ToolMenus.h
@@ -382,7 +382,7 @@ private:
 	TMap<FName, FName> MenuSubstitutionsDuringGenerate;
 
 	UPROPERTY()
-	TMap<FName, UToolMenu*> Menus;
+	TMap<FName, TObjectPtr<UToolMenu>> Menus;
 
 	TMap<FName, FGeneratedToolMenuWidgets> GeneratedMenuWidgets;
 
diff --git a/Engine/Source/Editor/AIGraph/Classes/AIGraphNode.h b/Engine/Source/Editor/AIGraph/Classes/AIGraphNode.h
--- a/Engine/Source/Editor/AIGraph/Classes/AIGraphNode.h
+++ b/Engine/Source/Editor/AIGraph/Classes/AIGraphNode.h
@@ -24,13 +24,13 @@ class AIGRAPH_API UAIGraphNode : public UEdGraphNode
 	struct FGraphNodeClassData ClassData;
 
 	UPROPERTY()
-	UObject* NodeInstance;
+	TObjectPtr<UObject> NodeInstance;
 
 	UPROPERTY(transient)
-	UAIGraphNode* ParentNode;
+	TObjectPtr<UAIGraphNode> ParentNode;
 
 	UPROPERTY()
-	TArray<UAIGraphNode*> SubNodes;
+	TArray<TObjectPtr<UAIGraphNode>> SubNodes;
 
 	/** subnode index assigned during copy operation to connect nodes again on paste */
 	UPROPERTY()
diff --git a/Engine/Source/Editor/AIGraph/Classes/AIGraphSchema.h b/Engine/Source/Editor/AIGraph/Classes/AIGraphSchema.h
--- a/Engine/Source/Editor/AIGraph/Classes/AIGraphSchema.h
+++ b/Engine/Source/Editor/AIGraph/Classes/AIGraphSchema.h
@@ -36,7 +36,7 @@ struct AIGRAPH_API FAISchemaAction_NewNode : public FEdGraphSchemaAction
 
 	/** Template of node we want to create */
 	UPROPERTY()
-	class UAIGraphNode* NodeTemplate;
+	TObjectPtr<class UAIGraphNode> NodeTemplate;
 
 	FAISchemaAction_NewNode()
 		: FEdGraphSchemaAction()
@@ -72,11 +72,11 @@ struct AIGRAPH_API FAISchemaAction_NewSubNode : public FEdGraphSchemaAction
 
 	/** Template of node we want to create */
 	UPROPERTY()
-	class UAIGraphNode* NodeTemplate;
+	TObjectPtr<class UAIGraphNode> NodeTemplate;
 
 	/** parent node */
 	UPROPERTY()
-	class UAIGraphNode* ParentNode;
+	TObjectPtr<class UAIGraphNode> ParentNode;
 
 	FAISchemaAction_NewSubNode()
 		: FEdGraphSchemaAction()
diff --git a/Engine/Source/Editor/AnimGraph/Public/AnimGraphNode_AnimDynamics.h b/Engine/Source/Editor/AnimGraph/Public/AnimGraphNode_AnimDynamics.h
--- a/Engine/Source/Editor/AnimGraph/Public/AnimGraphNode_AnimDynamics.h
+++ b/Engine/Source/Editor/AnimGraph/Public/AnimGraphNode_AnimDynamics.h
@@ -60,7 +60,7 @@ class UAnimGraphNode_AnimDynamics : public UAnimGraphNode_SkeletalControlBase
 	bool bShowCollisionSpheres;
 
 	UPROPERTY(Transient)
-	mutable USkeletalMeshComponent* LastPreviewComponent;
+	mutable TObjectPtr<USkeletalMeshComponent> LastPreviewComponent;
 
 public:
 
diff --git a/Engine/Source/Editor/AnimGraph/Public/AnimGraphNode_BlendListByEnum.h b/Engine/Source/Editor/AnimGraph/Public/AnimGraphNode_BlendListByEnum.h
--- a/Engine/Source/Editor/AnimGraph/Public/AnimGraphNode_BlendListByEnum.h
+++ b/Engine/Source/Editor/AnimGraph/Public/AnimGraphNode_BlendListByEnum.h
@@ -23,7 +23,7 @@ class UAnimGraphNode_BlendListByEnum : public UAnimGraphNode_BlendListBase, publ
 protected:
 	/** Name of the enum being switched on */
 	UPROPERTY()
-	UEnum* BoundEnum;
+	TObjectPtr<UEnum> BoundEnum;
 
 	UPROPERTY()
 	TArray<FName> VisibleEnumEntries;
diff --git a/Engine/Source/Editor/AnimGraph/Public/AnimGraphNode_BlendSpaceGraphBase.h b/Engine/Source/Editor/AnimGraph/Public/AnimGraphNode_BlendSpaceGraphBase.h
--- a/Engine/Source/Editor/AnimGraph/Public/AnimGraphNode_BlendSpaceGraphBase.h
+++ b/Engine/Source/Editor/AnimGraph/Public/AnimGraphNode_BlendSpaceGraphBase.h
@@ -62,7 +62,7 @@ protected:
 
 	// Internal blendspace
 	UPROPERTY()
-	UBlendSpaceBase* BlendSpace;
+	TObjectPtr<UBlendSpaceBase> BlendSpace;
 
 	// Blendspace class, for template nodes
 	UPROPERTY()
@@ -70,11 +70,11 @@ protected:
 
 	// Dummy blendspace graph (used for navigation only)
 	UPROPERTY()
-	UBlendSpaceGraph* BlendSpaceGraph;
+	TObjectPtr<UBlendSpaceGraph> BlendSpaceGraph;
 
 	// Linked animation graphs for sample points
 	UPROPERTY()
-	TArray<UEdGraph*> Graphs;
+	TArray<TObjectPtr<UEdGraph>> Graphs;
 
 protected:
 	// UEdGraphNode interface
diff --git a/Engine/Source/Editor/AnimGraph/Public/AnimGraphNode_PoseDriver.h b/Engine/Source/Editor/AnimGraph/Public/AnimGraphNode_PoseDriver.h
--- a/Engine/Source/Editor/AnimGraph/Public/AnimGraphNode_PoseDriver.h
+++ b/Engine/Source/Editor/AnimGraph/Public/AnimGraphNode_PoseDriver.h
@@ -58,7 +58,7 @@ public:
 
 	/** Used to refer back to preview instance in anim tools */
 	UPROPERTY(Transient)
-	USkeletalMeshComponent* LastPreviewComponent;
+	TObjectPtr<USkeletalMeshComponent> LastPreviewComponent;
 
 	// Begin UObject Interface.
 	virtual void PostLoad() override;
diff --git a/Engine/Source/Editor/AnimGraph/Public/AnimGraphNode_StateMachineBase.h b/Engine/Source/Editor/AnimGraph/Public/AnimGraphNode_StateMachineBase.h
--- a/Engine/Source/Editor/AnimGraph/Public/AnimGraphNode_StateMachineBase.h
+++ b/Engine/Source/Editor/AnimGraph/Public/AnimGraphNode_StateMachineBase.h
@@ -18,7 +18,7 @@ class ANIMGRAPH_API UAnimGraphNode_StateMachineBase : public UAnimGraphNode_Base
 
 	// Editor state machine representation
 	UPROPERTY()
-	class UAnimationStateMachineGraph* EditorStateMachineGraph;
+	TObjectPtr<class UAnimationStateMachineGraph> EditorStateMachineGraph;
 
 	// UEdGraphNode interface
 	virtual FLinearColor GetNodeTitleColor() const override;
diff --git a/Engine/Source/Editor/AnimGraph/Public/AnimStateConduitNode.h b/Engine/Source/Editor/AnimGraph/Public/AnimStateConduitNode.h
--- a/Engine/Source/Editor/AnimGraph/Public/AnimStateConduitNode.h
+++ b/Engine/Source/Editor/AnimGraph/Public/AnimStateConduitNode.h
@@ -18,7 +18,7 @@ public:
 
 	// The transition graph for this conduit; it's a logic graph, not an animation graph
 	UPROPERTY()
-	class UEdGraph* BoundGraph;
+	TObjectPtr<class UEdGraph> BoundGraph;
 
 	//~ Begin UEdGraphNode Interface
 	virtual void AllocateDefaultPins() override;
diff --git a/Engine/Source/Editor/AnimGraph/Public/AnimStateNode.h b/Engine/Source/Editor/AnimGraph/Public/AnimStateNode.h
--- a/Engine/Source/Editor/AnimGraph/Public/AnimStateNode.h
+++ b/Engine/Source/Editor/AnimGraph/Public/AnimStateNode.h
@@ -27,7 +27,7 @@ public:
 
 	// The animation graph for this state
 	UPROPERTY()
-	class UEdGraph* BoundGraph;
+	TObjectPtr<class UEdGraph> BoundGraph;
 
 	// The type of the contents of this state
 	UPROPERTY()
diff --git a/Engine/Source/Editor/AnimGraph/Public/AnimStateTransitionNode.h b/Engine/Source/Editor/AnimGraph/Public/AnimStateTransitionNode.h
--- a/Engine/Source/Editor/AnimGraph/Public/AnimStateTransitionNode.h
+++ b/Engine/Source/Editor/AnimGraph/Public/AnimStateTransitionNode.h
@@ -23,11 +23,11 @@ class UAnimStateTransitionNode : public UAnimStateNodeBase
 
 	// The transition logic graph for this transition (returning a boolean)
 	UPROPERTY()
-	class UEdGraph* BoundGraph;
+	TObjectPtr<class UEdGraph> BoundGraph;
 
 	// The animation graph for this transition if it uses custom blending (returning a pose)
 	UPROPERTY()
-	class UEdGraph* CustomTransitionGraph;
+	TObjectPtr<class UEdGraph> CustomTransitionGraph;
 
 	// The priority order of this transition. If multiple transitions out of a state go
 	// true at the same time, the one with the smallest priority order will take precedent
@@ -46,11 +46,11 @@ class UAnimStateTransitionNode : public UAnimStateNodeBase
 	EAlphaBlendOption BlendMode;
 
 	UPROPERTY(EditAnywhere, Category=Transition)
-	UCurveFloat* CustomBlendCurve;
+	TObjectPtr<UCurveFloat> CustomBlendCurve;
 
 	// The blend profile to use to evaluate this transition per-bone
 	UPROPERTY(EditAnywhere, Category=Transition, meta=(UseAsBlendProfile=true))
-	UBlendProfile* BlendProfile;
+	TObjectPtr<UBlendProfile> BlendProfile;
 
 	// Try setting the rule automatically based on most relevant player node's remaining time and the CrossfadeDuration of the transition, ignoring the internal time
 	UPROPERTY(EditAnywhere, Category=Transition)
diff --git a/Engine/Source/Editor/AnimGraph/Public/AnimationBlendSpaceSampleGraph.h b/Engine/Source/Editor/AnimGraph/Public/AnimationBlendSpaceSampleGraph.h
--- a/Engine/Source/Editor/AnimGraph/Public/AnimationBlendSpaceSampleGraph.h
+++ b/Engine/Source/Editor/AnimGraph/Public/AnimationBlendSpaceSampleGraph.h
@@ -17,5 +17,5 @@ class UAnimationBlendSpaceSampleGraph : public UAnimationGraph
 public:
 	// Result node within the state's animation graph
 	UPROPERTY()
-	UAnimGraphNode_BlendSpaceSampleResult* ResultNode;
+	TObjectPtr<UAnimGraphNode_BlendSpaceSampleResult> ResultNode;
 };
diff --git a/Engine/Source/Editor/AnimGraph/Public/AnimationCustomTransitionGraph.h b/Engine/Source/Editor/AnimGraph/Public/AnimationCustomTransitionGraph.h
--- a/Engine/Source/Editor/AnimGraph/Public/AnimationCustomTransitionGraph.h
+++ b/Engine/Source/Editor/AnimGraph/Public/AnimationCustomTransitionGraph.h
@@ -14,7 +14,7 @@ class UAnimationCustomTransitionGraph : public UAnimationGraph
 
 	// Result node within the state's animation graph
 	UPROPERTY()
-	class UAnimGraphNode_CustomTransitionResult* MyResultNode;
+	TObjectPtr<class UAnimGraphNode_CustomTransitionResult> MyResultNode;
 
 	//@TODO: Document
 	ANIMGRAPH_API class UAnimGraphNode_CustomTransitionResult* GetResultNode();
diff --git a/Engine/Source/Editor/AnimGraph/Public/AnimationStateGraph.h b/Engine/Source/Editor/AnimGraph/Public/AnimationStateGraph.h
--- a/Engine/Source/Editor/AnimGraph/Public/AnimationStateGraph.h
+++ b/Engine/Source/Editor/AnimGraph/Public/AnimationStateGraph.h
@@ -14,7 +14,7 @@ class UAnimationStateGraph : public UAnimationGraph
 
 	// Result node within the state's animation graph
 	UPROPERTY()
-	class UAnimGraphNode_StateResult* MyResultNode;
+	TObjectPtr<class UAnimGraphNode_StateResult> MyResultNode;
 
 	ANIMGRAPH_API class UAnimGraphNode_StateResult* GetResultNode();
 };
diff --git a/Engine/Source/Editor/AnimGraph/Public/AnimationStateMachineGraph.h b/Engine/Source/Editor/AnimGraph/Public/AnimationStateMachineGraph.h
--- a/Engine/Source/Editor/AnimGraph/Public/AnimationStateMachineGraph.h
+++ b/Engine/Source/Editor/AnimGraph/Public/AnimationStateMachineGraph.h
@@ -14,10 +14,10 @@ class UAnimationStateMachineGraph : public UEdGraph
 
 	// Entry node within the state machine
 	UPROPERTY()
-	class UAnimStateEntryNode* EntryNode;
+	TObjectPtr<class UAnimStateEntryNode> EntryNode;
 
 	// Parent instance node
 	UPROPERTY()
-	class UAnimGraphNode_StateMachineBase* OwnerAnimGraphNode;
+	TObjectPtr<class UAnimGraphNode_StateMachineBase> OwnerAnimGraphNode;
 };
 
diff --git a/Engine/Source/Editor/AnimGraph/Public/AnimationTransitionGraph.h b/Engine/Source/Editor/AnimGraph/Public/AnimationTransitionGraph.h
--- a/Engine/Source/Editor/AnimGraph/Public/AnimationTransitionGraph.h
+++ b/Engine/Source/Editor/AnimGraph/Public/AnimationTransitionGraph.h
@@ -13,7 +13,7 @@ class UAnimationTransitionGraph : public UAnimationGraph
 	GENERATED_UCLASS_BODY()
 
 	UPROPERTY()
-	class UAnimGraphNode_TransitionResult* MyResultNode;
+	TObjectPtr<class UAnimGraphNode_TransitionResult> MyResultNode;
 
 	ANIMGRAPH_API class UAnimGraphNode_TransitionResult* GetResultNode();
 };
diff --git a/Engine/Source/Editor/AnimGraph/Public/BlendSpaceGraph.h b/Engine/Source/Editor/AnimGraph/Public/BlendSpaceGraph.h
--- a/Engine/Source/Editor/AnimGraph/Public/BlendSpaceGraph.h
+++ b/Engine/Source/Editor/AnimGraph/Public/BlendSpaceGraph.h
@@ -17,5 +17,5 @@ class UBlendSpaceGraph : public UEdGraph
 public:
 	// Blendspace that we wrap
 	UPROPERTY()
-	UBlendSpaceBase* BlendSpace;
+	TObjectPtr<UBlendSpaceBase> BlendSpace;
 };
\ No newline at end of file
diff --git a/Engine/Source/Editor/AnimGraph/Public/K2Node_AnimGetter.h b/Engine/Source/Editor/AnimGraph/Public/K2Node_AnimGetter.h
--- a/Engine/Source/Editor/AnimGraph/Public/K2Node_AnimGetter.h
+++ b/Engine/Source/Editor/AnimGraph/Public/K2Node_AnimGetter.h
@@ -26,23 +26,23 @@ struct FNodeSpawnData
 
 	// The node the spawned getter accesses
 	UPROPERTY()
-	UAnimGraphNode_Base* SourceNode;
+	TObjectPtr<UAnimGraphNode_Base> SourceNode;
 
 	// The state node the spawned getter accesses
 	UPROPERTY()
-	UAnimStateNodeBase* SourceStateNode;
+	TObjectPtr<UAnimStateNodeBase> SourceStateNode;
 
 	// The instance class the spawned getter is defined on
 	UPROPERTY()
-	UClass* AnimInstanceClass;
+	TObjectPtr<UClass> AnimInstanceClass;
 
 	// The blueprint the getter is valid within
 	UPROPERTY()
-	const UAnimBlueprint* SourceBlueprint;
+	TObjectPtr<const UAnimBlueprint> SourceBlueprint;
 
 	// The UFunction (as a UField) 
 	UPROPERTY()
-	UField* Getter;
+	TObjectPtr<UField> Getter;
 
 	// String of combined valid contexts for the spawned getter
 	UPROPERTY()
@@ -73,20 +73,20 @@ public:
 
 	// The node that is required for the getter
 	UPROPERTY()
-	UAnimGraphNode_Base* SourceNode;
+	TObjectPtr<UAnimGraphNode_Base> SourceNode;
 
 	// UAnimStateNode doesn't use the same hierarchy so we need to have a seperate property here to handle
 	// those.
 	UPROPERTY()
-	UAnimStateNodeBase* SourceStateNode;
+	TObjectPtr<UAnimStateNodeBase> SourceStateNode;
 
 	// The UAnimInstance derived class that implements the getter we are running
 	UPROPERTY()
-	UClass* GetterClass;
+	TObjectPtr<UClass> GetterClass;
 
 	// The anim blueprint that generated this getter
 	UPROPERTY()
-	const UAnimBlueprint* SourceAnimBlueprint;
+	TObjectPtr<const UAnimBlueprint> SourceAnimBlueprint;
 
 	// Cached node title
 	UPROPERTY()
diff --git a/Engine/Source/Editor/AnimGraph/Public/K2Node_TransitionRuleGetter.h b/Engine/Source/Editor/AnimGraph/Public/K2Node_TransitionRuleGetter.h
--- a/Engine/Source/Editor/AnimGraph/Public/K2Node_TransitionRuleGetter.h
+++ b/Engine/Source/Editor/AnimGraph/Public/K2Node_TransitionRuleGetter.h
@@ -43,10 +43,10 @@ class UK2Node_TransitionRuleGetter : public UK2Node
 	TEnumAsByte<ETransitionGetter::Type> GetterType;
 
 	UPROPERTY()
-	UAnimGraphNode_Base* AssociatedAnimAssetPlayerNode;
+	TObjectPtr<UAnimGraphNode_Base> AssociatedAnimAssetPlayerNode;
 
 	UPROPERTY()
-	UAnimStateNode* AssociatedStateNode;
+	TObjectPtr<UAnimStateNode> AssociatedStateNode;
 
 	// UEdGraphNode interface
 	virtual void AllocateDefaultPins() override;
diff --git a/Engine/Source/Editor/AnimationModifiers/Public/AnimationModifier.h b/Engine/Source/Editor/AnimationModifiers/Public/AnimationModifier.h
--- a/Engine/Source/Editor/AnimationModifiers/Public/AnimationModifier.h
+++ b/Engine/Source/Editor/AnimationModifiers/Public/AnimationModifier.h
@@ -78,5 +78,5 @@ private:
 
 	/** Serialized version of the modifier that has been previously applied to the Animation Asset */
 	UPROPERTY()
-	UAnimationModifier* PreviouslyAppliedModifier;
+	TObjectPtr<UAnimationModifier> PreviouslyAppliedModifier;
 };
\ No newline at end of file
diff --git a/Engine/Source/Editor/AnimationModifiers/Public/AnimationModifiersAssetUserData.h b/Engine/Source/Editor/AnimationModifiers/Public/AnimationModifiersAssetUserData.h
--- a/Engine/Source/Editor/AnimationModifiers/Public/AnimationModifiersAssetUserData.h
+++ b/Engine/Source/Editor/AnimationModifiers/Public/AnimationModifiersAssetUserData.h
@@ -32,5 +32,5 @@ private:
 	void RemoveInvalidModifiers();
 protected:
 	UPROPERTY()
-	TArray<UAnimationModifier*> AnimationModifierInstances;
+	TArray<TObjectPtr<UAnimationModifier>> AnimationModifierInstances;
 };
diff --git a/Engine/Source/Editor/AssetPlacementEdMode/Private/Tools/PlacementBrushToolBase.h b/Engine/Source/Editor/AssetPlacementEdMode/Private/Tools/PlacementBrushToolBase.h
--- a/Engine/Source/Editor/AssetPlacementEdMode/Private/Tools/PlacementBrushToolBase.h
+++ b/Engine/Source/Editor/AssetPlacementEdMode/Private/Tools/PlacementBrushToolBase.h
@@ -21,7 +21,7 @@ public:
 	virtual UInteractiveTool* BuildTool(const FToolBuilderState& SceneState) const override;
 
 	UPROPERTY()
-	UAssetPlacementSettings* PlacementSettings;
+	TObjectPtr<UAssetPlacementSettings> PlacementSettings;
 
 protected:
 	virtual UPlacementBrushToolBase* FactoryToolInstance(UObject* Outer) const PURE_VIRTUAL(UPlacementToolBuilderBase::FactoryToolInstance, return nullptr; );
diff --git a/Engine/Source/Editor/AudioEditor/Classes/Factories/DialogueWaveFactory.h b/Engine/Source/Editor/AudioEditor/Classes/Factories/DialogueWaveFactory.h
--- a/Engine/Source/Editor/AudioEditor/Classes/Factories/DialogueWaveFactory.h
+++ b/Engine/Source/Editor/AudioEditor/Classes/Factories/DialogueWaveFactory.h
@@ -24,11 +24,11 @@ class UDialogueWaveFactory : public UFactory
 
 	/** An initial sound wave to place in the newly created dialogue wave */
 	UPROPERTY()
-	class USoundWave* InitialSoundWave;
+	TObjectPtr<class USoundWave> InitialSoundWave;
 
 	/** An initial speaking dialogue voice to place in the newly created dialogue wave */
 	UPROPERTY()
-	class UDialogueVoice* InitialSpeakerVoice;
+	TObjectPtr<class UDialogueVoice> InitialSpeakerVoice;
 
 	/** Whether an initial target dialogue voice should be set */
 	UPROPERTY()
@@ -36,7 +36,7 @@ class UDialogueWaveFactory : public UFactory
 
 	/** An initial target dialogue voices to place in the newly created dialogue wave */
 	UPROPERTY()
-	TArray<UDialogueVoice*> InitialTargetVoices;
+	TArray<TObjectPtr<UDialogueVoice>> InitialTargetVoices;
 };
 
 
diff --git a/Engine/Source/Editor/AudioEditor/Classes/Factories/SoundCueFactoryNew.h b/Engine/Source/Editor/AudioEditor/Classes/Factories/SoundCueFactoryNew.h
--- a/Engine/Source/Editor/AudioEditor/Classes/Factories/SoundCueFactoryNew.h
+++ b/Engine/Source/Editor/AudioEditor/Classes/Factories/SoundCueFactoryNew.h
@@ -23,11 +23,11 @@ class USoundCueFactoryNew : public UFactory
 
 	/** An initial sound wave to place in the newly created cue */
 	UPROPERTY()
-	class USoundWave* InitialSoundWave;
+	TObjectPtr<class USoundWave> InitialSoundWave;
 
 	/** An initial dialogue wave to place in the newly created cue */
 	UPROPERTY()
-	class UDialogueWave* InitialDialogueWave;
+	TObjectPtr<class UDialogueWave> InitialDialogueWave;
 };
 
 
diff --git a/Engine/Source/Editor/AudioEditor/Classes/SoundClassGraph/SoundClassGraphNode.h b/Engine/Source/Editor/AudioEditor/Classes/SoundClassGraph/SoundClassGraphNode.h
--- a/Engine/Source/Editor/AudioEditor/Classes/SoundClassGraph/SoundClassGraphNode.h
+++ b/Engine/Source/Editor/AudioEditor/Classes/SoundClassGraph/SoundClassGraphNode.h
@@ -18,7 +18,7 @@ class USoundClassGraphNode : public UEdGraphNode
 
 	/** The SoundNode this represents */
 	UPROPERTY(VisibleAnywhere, instanced, Category=Sound)
-	USoundClass*		SoundClass;
+	TObjectPtr<USoundClass>		SoundClass;
 
 	/** Get the Pin that connects to all children */
 	UEdGraphPin* GetChildPin() const { return ChildPin; }
diff --git a/Engine/Source/Editor/AudioEditor/Classes/SoundCueGraph/SoundCueGraphNode.h b/Engine/Source/Editor/AudioEditor/Classes/SoundCueGraph/SoundCueGraphNode.h
--- a/Engine/Source/Editor/AudioEditor/Classes/SoundCueGraph/SoundCueGraphNode.h
+++ b/Engine/Source/Editor/AudioEditor/Classes/SoundCueGraph/SoundCueGraphNode.h
@@ -17,7 +17,7 @@ class USoundCueGraphNode : public USoundCueGraphNode_Base
 
 	/** The SoundNode this represents */
 	UPROPERTY(VisibleAnywhere, instanced, Category=Sound)
-	USoundNode* SoundNode;
+	TObjectPtr<USoundNode> SoundNode;
 
 	/** Set the SoundNode this represents (also assigns this to the SoundNode in Editor)*/
 	AUDIOEDITOR_API void SetSoundNode(USoundNode* InSoundNode);
diff --git a/Engine/Source/Editor/AudioEditor/Classes/SoundCueGraph/SoundCueGraphSchema.h b/Engine/Source/Editor/AudioEditor/Classes/SoundCueGraph/SoundCueGraphSchema.h
--- a/Engine/Source/Editor/AudioEditor/Classes/SoundCueGraph/SoundCueGraphSchema.h
+++ b/Engine/Source/Editor/AudioEditor/Classes/SoundCueGraph/SoundCueGraphSchema.h
@@ -19,7 +19,7 @@ struct AUDIOEDITOR_API FSoundCueGraphSchemaAction_NewNode : public FEdGraphSchem
 
 	/** Class of node we want to create */
 	UPROPERTY()
-	class UClass* SoundNodeClass;
+	TObjectPtr<class UClass> SoundNodeClass;
 
 
 	FSoundCueGraphSchemaAction_NewNode() 
diff --git a/Engine/Source/Editor/AudioEditor/Classes/SoundSubmixGraph/SoundSubmixGraph.h b/Engine/Source/Editor/AudioEditor/Classes/SoundSubmixGraph/SoundSubmixGraph.h
--- a/Engine/Source/Editor/AudioEditor/Classes/SoundSubmixGraph/SoundSubmixGraph.h
+++ b/Engine/Source/Editor/AudioEditor/Classes/SoundSubmixGraph/SoundSubmixGraph.h
@@ -127,9 +127,9 @@ private:
 private:
 	/** SoundSubmix which forms the root of this graph */
 	UPROPERTY(Transient)
-	USoundSubmixBase* RootSoundSubmix = nullptr;
+	TObjectPtr<USoundSubmixBase> RootSoundSubmix = nullptr;
 
 	UPROPERTY(Transient)
-	TArray<USoundSubmixBase*> StaleRoots;
+	TArray<TObjectPtr<USoundSubmixBase>> StaleRoots;
 };
 
diff --git a/Engine/Source/Editor/AudioEditor/Classes/SoundSubmixGraph/SoundSubmixGraphNode.h b/Engine/Source/Editor/AudioEditor/Classes/SoundSubmixGraph/SoundSubmixGraphNode.h
--- a/Engine/Source/Editor/AudioEditor/Classes/SoundSubmixGraph/SoundSubmixGraphNode.h
+++ b/Engine/Source/Editor/AudioEditor/Classes/SoundSubmixGraph/SoundSubmixGraphNode.h
@@ -48,11 +48,11 @@ class USoundSubmixGraphNode : public UEdGraphNode
 
 	/** The SoundSubmix this represents */
 	UPROPERTY(VisibleAnywhere, instanced, Category=Sound)
-	USoundSubmixBase* SoundSubmix;
+	TObjectPtr<USoundSubmixBase> SoundSubmix;
 
 	/** A user widget to use to represent the graph node */
 	UPROPERTY(Transient)
-	UUserWidget* SubmixNodeUserWidget;
+	TObjectPtr<UUserWidget> SubmixNodeUserWidget;
 
 	/** Get the Pin that connects to all children */
 	UEdGraphPin* GetChildPin() const { return ChildPin; }
diff --git a/Engine/Source/Editor/BehaviorTreeEditor/Classes/BehaviorTreeDecoratorGraphNode_Decorator.h b/Engine/Source/Editor/BehaviorTreeEditor/Classes/BehaviorTreeDecoratorGraphNode_Decorator.h
--- a/Engine/Source/Editor/BehaviorTreeEditor/Classes/BehaviorTreeDecoratorGraphNode_Decorator.h
+++ b/Engine/Source/Editor/BehaviorTreeEditor/Classes/BehaviorTreeDecoratorGraphNode_Decorator.h
@@ -16,7 +16,7 @@ class UBehaviorTreeDecoratorGraphNode_Decorator : public UBehaviorTreeDecoratorG
 	GENERATED_UCLASS_BODY()
 
 	UPROPERTY()
-	UObject* NodeInstance;
+	TObjectPtr<UObject> NodeInstance;
 
 	UPROPERTY()
 	FGraphNodeClassData ClassData;
diff --git a/Engine/Source/Editor/BehaviorTreeEditor/Classes/BehaviorTreeGraphNode.h b/Engine/Source/Editor/BehaviorTreeEditor/Classes/BehaviorTreeGraphNode.h
--- a/Engine/Source/Editor/BehaviorTreeEditor/Classes/BehaviorTreeGraphNode.h
+++ b/Engine/Source/Editor/BehaviorTreeEditor/Classes/BehaviorTreeGraphNode.h
@@ -17,11 +17,11 @@ class UBehaviorTreeGraphNode : public UAIGraphNode
 
 	/** only some of behavior tree nodes support decorators */
 	UPROPERTY()
-	TArray<UBehaviorTreeGraphNode*> Decorators;
+	TArray<TObjectPtr<UBehaviorTreeGraphNode>> Decorators;
 
 	/** only some of behavior tree nodes support services */
 	UPROPERTY()
-	TArray<UBehaviorTreeGraphNode*> Services;
+	TArray<TObjectPtr<UBehaviorTreeGraphNode>> Services;
 
 	//~ Begin UEdGraphNode Interface
 	virtual class UBehaviorTreeGraph* GetBehaviorTreeGraph();
diff --git a/Engine/Source/Editor/BehaviorTreeEditor/Classes/BehaviorTreeGraphNode_CompositeDecorator.h b/Engine/Source/Editor/BehaviorTreeEditor/Classes/BehaviorTreeGraphNode_CompositeDecorator.h
--- a/Engine/Source/Editor/BehaviorTreeEditor/Classes/BehaviorTreeGraphNode_CompositeDecorator.h
+++ b/Engine/Source/Editor/BehaviorTreeEditor/Classes/BehaviorTreeGraphNode_CompositeDecorator.h
@@ -16,7 +16,7 @@ class UBehaviorTreeGraphNode_CompositeDecorator : public UBehaviorTreeGraphNode
 
 	// The logic graph for this decorator (returning a boolean)
 	UPROPERTY()
-	class UEdGraph* BoundGraph;
+	TObjectPtr<class UEdGraph> BoundGraph;
 
 	UPROPERTY(EditAnywhere, Category=Description)
 	FString CompositeName;
@@ -68,7 +68,7 @@ protected:
 	virtual void ResetNodeOwner() override;
 
 	UPROPERTY()
-	class UBTCompositeNode* ParentNodeInstance;
+	TObjectPtr<class UBTCompositeNode> ParentNodeInstance;
 
 	uint8 ChildIndex;
 
diff --git a/Engine/Source/Editor/BehaviorTreeEditor/Classes/BehaviorTreeGraphNode_Root.h b/Engine/Source/Editor/BehaviorTreeEditor/Classes/BehaviorTreeGraphNode_Root.h
--- a/Engine/Source/Editor/BehaviorTreeEditor/Classes/BehaviorTreeGraphNode_Root.h
+++ b/Engine/Source/Editor/BehaviorTreeEditor/Classes/BehaviorTreeGraphNode_Root.h
@@ -14,7 +14,7 @@ class UBehaviorTreeGraphNode_Root : public UBehaviorTreeGraphNode
 	GENERATED_UCLASS_BODY()
 
 	UPROPERTY(EditAnywhere, Category="AI|BehaviorTree")
-	class UBlackboardData* BlackboardAsset;
+	TObjectPtr<class UBlackboardData> BlackboardAsset;
 
 	virtual void PostPlacedNewNode() override;
 	virtual void AllocateDefaultPins() override;
diff --git a/Engine/Source/Editor/BehaviorTreeEditor/Classes/EdGraphSchema_BehaviorTreeDecorator.h b/Engine/Source/Editor/BehaviorTreeEditor/Classes/EdGraphSchema_BehaviorTreeDecorator.h
--- a/Engine/Source/Editor/BehaviorTreeEditor/Classes/EdGraphSchema_BehaviorTreeDecorator.h
+++ b/Engine/Source/Editor/BehaviorTreeEditor/Classes/EdGraphSchema_BehaviorTreeDecorator.h
@@ -17,7 +17,7 @@ struct FDecoratorSchemaAction_NewNode : public FEdGraphSchemaAction
 
 	/** Template of node we want to create */
 	UPROPERTY()
-	class UBehaviorTreeDecoratorGraphNode* NodeTemplate;
+	TObjectPtr<class UBehaviorTreeDecoratorGraphNode> NodeTemplate;
 
 	FDecoratorSchemaAction_NewNode() 
 		: FEdGraphSchemaAction()
diff --git a/Engine/Source/Editor/BlueprintGraph/Classes/BlueprintEventNodeSpawner.h b/Engine/Source/Editor/BlueprintGraph/Classes/BlueprintEventNodeSpawner.h
--- a/Engine/Source/Editor/BlueprintGraph/Classes/BlueprintEventNodeSpawner.h
+++ b/Engine/Source/Editor/BlueprintGraph/Classes/BlueprintEventNodeSpawner.h
@@ -79,7 +79,7 @@ public:
 private:
 	/** The function to configure new nodes with. */
 	UPROPERTY()
-    UFunction const* EventFunc;
+    TObjectPtr<UFunction const> EventFunc;
 
 	/** The custom name to configure new event nodes with. */
 	UPROPERTY()
diff --git a/Engine/Source/Editor/BlueprintGraph/Classes/BlueprintFieldNodeSpawner.h b/Engine/Source/Editor/BlueprintGraph/Classes/BlueprintFieldNodeSpawner.h
--- a/Engine/Source/Editor/BlueprintGraph/Classes/BlueprintFieldNodeSpawner.h
+++ b/Engine/Source/Editor/BlueprintGraph/Classes/BlueprintFieldNodeSpawner.h
@@ -63,12 +63,12 @@ protected:
 
 	/** The owning class to configure new nodes with. */
 	UPROPERTY()
-	UClass const* OwnerClass;
+	TObjectPtr<UClass const> OwnerClass;
 
 private:
 	/** The field to configure new nodes with. */
 	UPROPERTY()
-	UField* Field;
+	TObjectPtr<UField> Field;
 
 	UPROPERTY()
 	TFieldPath<FProperty> Property;
diff --git a/Engine/Source/Editor/BlueprintGraph/Classes/BlueprintVariableNodeSpawner.h b/Engine/Source/Editor/BlueprintGraph/Classes/BlueprintVariableNodeSpawner.h
--- a/Engine/Source/Editor/BlueprintGraph/Classes/BlueprintVariableNodeSpawner.h
+++ b/Engine/Source/Editor/BlueprintGraph/Classes/BlueprintVariableNodeSpawner.h
@@ -113,7 +113,7 @@ private:
 
 	/** The graph that the local variable belongs to (if this is a local variable spawner). */
 	UPROPERTY()
-	UEdGraph* LocalVarOuter;
+	TObjectPtr<UEdGraph> LocalVarOuter;
 
 	/** */
 	UPROPERTY()
diff --git a/Engine/Source/Editor/BlueprintGraph/Classes/EdGraphSchema_K2_Actions.h b/Engine/Source/Editor/BlueprintGraph/Classes/EdGraphSchema_K2_Actions.h
--- a/Engine/Source/Editor/BlueprintGraph/Classes/EdGraphSchema_K2_Actions.h
+++ b/Engine/Source/Editor/BlueprintGraph/Classes/EdGraphSchema_K2_Actions.h
@@ -41,7 +41,7 @@ struct BLUEPRINTGRAPH_API FEdGraphSchemaAction_K2NewNode : public FEdGraphSchema
 
 	/** Template of node we want to create */
 	UPROPERTY()
-	class UK2Node* NodeTemplate;
+	TObjectPtr<class UK2Node> NodeTemplate;
 
 	UPROPERTY()
 	bool bGotoNode;
@@ -120,7 +120,7 @@ struct BLUEPRINTGRAPH_API FEdGraphSchemaAction_K2ViewNode : public FEdGraphSchem
 
 	/** node we want to view */
 	UPROPERTY()
-	const UK2Node* NodePtr;
+	TObjectPtr<const UK2Node> NodePtr;
 
 	FEdGraphSchemaAction_K2ViewNode() 
 		: FEdGraphSchemaAction()
@@ -181,7 +181,7 @@ struct BLUEPRINTGRAPH_API FEdGraphSchemaAction_EventFromFunction : public FEdGra
 	virtual FName GetTypeId() const override { return StaticGetTypeId(); } 
 
 	UPROPERTY()
-	class UFunction* SignatureFunction;
+	TObjectPtr<class UFunction> SignatureFunction;
 
 	FEdGraphSchemaAction_EventFromFunction() 
 		:FEdGraphSchemaAction()
@@ -220,7 +220,7 @@ struct FEdGraphSchemaAction_K2AddComponent : public FEdGraphSchemaAction_K2NewNo
 
 	/** Option asset to assign to newly created component */
 	UPROPERTY()
-	class UObject* ComponentAsset;
+	TObjectPtr<class UObject> ComponentAsset;
 
 	FEdGraphSchemaAction_K2AddComponent()
 		: FEdGraphSchemaAction_K2NewNode()
@@ -313,7 +313,7 @@ struct FEdGraphSchemaAction_K2AddCallOnActor : public FEdGraphSchemaAction_K2New
 
 	/** Pointer to actors in level we want to call function on */
 	UPROPERTY()
-	TArray<class AActor*> LevelActors;
+	TArray<TObjectPtr<class AActor>> LevelActors;
 
 	FEdGraphSchemaAction_K2AddCallOnActor()
 		: FEdGraphSchemaAction_K2NewNode()
diff --git a/Engine/Source/Editor/BlueprintGraph/Classes/K2Node_ActorBoundEvent.h b/Engine/Source/Editor/BlueprintGraph/Classes/K2Node_ActorBoundEvent.h
--- a/Engine/Source/Editor/BlueprintGraph/Classes/K2Node_ActorBoundEvent.h
+++ b/Engine/Source/Editor/BlueprintGraph/Classes/K2Node_ActorBoundEvent.h
@@ -25,11 +25,11 @@ class UK2Node_ActorBoundEvent : public UK2Node_Event
 
 	/** Delegate property's owner class that this event is associated with */
 	UPROPERTY()
-	UClass* DelegateOwnerClass;
+	TObjectPtr<UClass> DelegateOwnerClass;
 
 	/** The event that this event is bound to */
 	UPROPERTY()
-	class AActor* EventOwner;
+	TObjectPtr<class AActor> EventOwner;
 
 	//~ Begin UObject Interface
 	virtual void Serialize(FArchive& Ar) override;
diff --git a/Engine/Source/Editor/BlueprintGraph/Classes/K2Node_AddComponent.h b/Engine/Source/Editor/BlueprintGraph/Classes/K2Node_AddComponent.h
--- a/Engine/Source/Editor/BlueprintGraph/Classes/K2Node_AddComponent.h
+++ b/Engine/Source/Editor/BlueprintGraph/Classes/K2Node_AddComponent.h
@@ -27,7 +27,7 @@ class UK2Node_AddComponent : public UK2Node_CallFunction
 	FString TemplateBlueprint;
 
 	UPROPERTY()
-	UClass* TemplateType;
+	TObjectPtr<UClass> TemplateType;
 
 	//~ Begin UObject Interface
 	virtual void Serialize(FArchive& Ar) override;
diff --git a/Engine/Source/Editor/BlueprintGraph/Classes/K2Node_BaseAsyncTask.h b/Engine/Source/Editor/BlueprintGraph/Classes/K2Node_BaseAsyncTask.h
--- a/Engine/Source/Editor/BlueprintGraph/Classes/K2Node_BaseAsyncTask.h
+++ b/Engine/Source/Editor/BlueprintGraph/Classes/K2Node_BaseAsyncTask.h
@@ -68,11 +68,11 @@ protected:
 
 	// The class containing the proxy object functions
 	UPROPERTY()
-	UClass* ProxyFactoryClass;
+	TObjectPtr<UClass> ProxyFactoryClass;
 
 	// The type of proxy object that will be created
 	UPROPERTY()
-	UClass* ProxyClass;
+	TObjectPtr<UClass> ProxyClass;
 
 	// The name of the 'go' function on the proxy object that will be called after delegates are in place, can be NAME_None
 	UPROPERTY()
diff --git a/Engine/Source/Editor/BlueprintGraph/Classes/K2Node_BitmaskLiteral.h b/Engine/Source/Editor/BlueprintGraph/Classes/K2Node_BitmaskLiteral.h
--- a/Engine/Source/Editor/BlueprintGraph/Classes/K2Node_BitmaskLiteral.h
+++ b/Engine/Source/Editor/BlueprintGraph/Classes/K2Node_BitmaskLiteral.h
@@ -17,7 +17,7 @@ class UK2Node_BitmaskLiteral : public UK2Node, public INodeDependingOnEnumInterf
 	GENERATED_UCLASS_BODY()
 
 	UPROPERTY()
-	UEnum* BitflagsEnum;
+	TObjectPtr<UEnum> BitflagsEnum;
 
 	//~ Begin UObject Interface
 	virtual void Serialize(FArchive& Ar) override;
diff --git a/Engine/Source/Editor/BlueprintGraph/Classes/K2Node_CastByteToEnum.h b/Engine/Source/Editor/BlueprintGraph/Classes/K2Node_CastByteToEnum.h
--- a/Engine/Source/Editor/BlueprintGraph/Classes/K2Node_CastByteToEnum.h
+++ b/Engine/Source/Editor/BlueprintGraph/Classes/K2Node_CastByteToEnum.h
@@ -20,7 +20,7 @@ class UK2Node_CastByteToEnum : public UK2Node
 	GENERATED_UCLASS_BODY()
 
 	UPROPERTY()
-	UEnum* Enum;
+	TObjectPtr<UEnum> Enum;
 
 	/* if true, the node returns always a valid value */
 	UPROPERTY()
diff --git a/Engine/Source/Editor/BlueprintGraph/Classes/K2Node_ComponentBoundEvent.h b/Engine/Source/Editor/BlueprintGraph/Classes/K2Node_ComponentBoundEvent.h
--- a/Engine/Source/Editor/BlueprintGraph/Classes/K2Node_ComponentBoundEvent.h
+++ b/Engine/Source/Editor/BlueprintGraph/Classes/K2Node_ComponentBoundEvent.h
@@ -24,7 +24,7 @@ class UK2Node_ComponentBoundEvent : public UK2Node_Event
 
 	/** Delegate property's owner class that this event is associated with */
 	UPROPERTY()
-	UClass* DelegateOwnerClass;
+	TObjectPtr<UClass> DelegateOwnerClass;
 
 	/** Name of property in Blueprint class that pointer to component we want to bind to */
 	UPROPERTY()
diff --git a/Engine/Source/Editor/BlueprintGraph/Classes/K2Node_DoOnceMultiInput.h b/Engine/Source/Editor/BlueprintGraph/Classes/K2Node_DoOnceMultiInput.h
--- a/Engine/Source/Editor/BlueprintGraph/Classes/K2Node_DoOnceMultiInput.h
+++ b/Engine/Source/Editor/BlueprintGraph/Classes/K2Node_DoOnceMultiInput.h
@@ -25,7 +25,7 @@ class UK2Node_DoOnceMultiInput : public UK2Node, public IK2Node_AddPinInterface
 
 	/** Reference to the integer that contains */
 	UPROPERTY(transient)
-	class UK2Node_TemporaryVariable* DataNode;
+	TObjectPtr<class UK2Node_TemporaryVariable> DataNode;
 
 	virtual ERedirectType DoPinsMatchForReconstruction(const UEdGraphPin* NewPin, int32 NewPinIndex, const UEdGraphPin* OldPin, int32 OldPinIndex) const override;
 
diff --git a/Engine/Source/Editor/BlueprintGraph/Classes/K2Node_GetClassDefaults.h b/Engine/Source/Editor/BlueprintGraph/Classes/K2Node_GetClassDefaults.h
--- a/Engine/Source/Editor/BlueprintGraph/Classes/K2Node_GetClassDefaults.h
+++ b/Engine/Source/Editor/BlueprintGraph/Classes/K2Node_GetClassDefaults.h
@@ -89,7 +89,7 @@ private:
 
 	/** Blueprint that we subscribed OnBlueprintChangedDelegate and OnBlueprintCompiledDelegate to */
 	UPROPERTY()
-	UBlueprint* BlueprintSubscribedTo;
+	TObjectPtr<UBlueprint> BlueprintSubscribedTo;
 
 	/** Blueprint.OnChanged delegate handle */
 	FDelegateHandle OnBlueprintChangedDelegate;
diff --git a/Engine/Source/Editor/BlueprintGraph/Classes/K2Node_Literal.h b/Engine/Source/Editor/BlueprintGraph/Classes/K2Node_Literal.h
--- a/Engine/Source/Editor/BlueprintGraph/Classes/K2Node_Literal.h
+++ b/Engine/Source/Editor/BlueprintGraph/Classes/K2Node_Literal.h
@@ -21,7 +21,7 @@ class UK2Node_Literal : public UK2Node
 private:
 	/** If this is an object reference literal, keep a reference here so that it can be updated as objects move around */
 	UPROPERTY()
-	class UObject* ObjectRef;
+	TObjectPtr<class UObject> ObjectRef;
 
 public:
 
diff --git a/Engine/Source/Editor/BlueprintGraph/Classes/K2Node_MacroInstance.h b/Engine/Source/Editor/BlueprintGraph/Classes/K2Node_MacroInstance.h
--- a/Engine/Source/Editor/BlueprintGraph/Classes/K2Node_MacroInstance.h
+++ b/Engine/Source/Editor/BlueprintGraph/Classes/K2Node_MacroInstance.h
@@ -23,7 +23,7 @@ private:
 	/** A macro is like a composite node, except that the associated graph lives
 	  * in another blueprint, and can be instanced multiple times. */
 	UPROPERTY()
-	class UEdGraph* MacroGraph_DEPRECATED;
+	TObjectPtr<class UEdGraph> MacroGraph_DEPRECATED;
 
 	UPROPERTY()
 	FGraphReference MacroGraphReference;
diff --git a/Engine/Source/Editor/BlueprintGraph/Classes/K2Node_MatineeController.h b/Engine/Source/Editor/BlueprintGraph/Classes/K2Node_MatineeController.h
--- a/Engine/Source/Editor/BlueprintGraph/Classes/K2Node_MatineeController.h
+++ b/Engine/Source/Editor/BlueprintGraph/Classes/K2Node_MatineeController.h
@@ -20,7 +20,7 @@ class UK2Node_MatineeController : public UK2Node
 
 	/** The matinee actor in the level that this node controls */
 	UPROPERTY(EditAnywhere, Category=K2Node_MatineeController)
-	class AMatineeActor* MatineeActor;
+	TObjectPtr<class AMatineeActor> MatineeActor;
 
 	//~ Begin UEdGraphNode Interface
 	virtual void AllocateDefaultPins() override;
diff --git a/Engine/Source/Editor/BlueprintGraph/Classes/K2Node_MultiGate.h b/Engine/Source/Editor/BlueprintGraph/Classes/K2Node_MultiGate.h
--- a/Engine/Source/Editor/BlueprintGraph/Classes/K2Node_MultiGate.h
+++ b/Engine/Source/Editor/BlueprintGraph/Classes/K2Node_MultiGate.h
@@ -18,7 +18,7 @@ class UK2Node_MultiGate : public UK2Node_ExecutionSequence
 
 	/** Reference to the integer that contains */
 	UPROPERTY(transient)
-	class UK2Node_TemporaryVariable* DataNode;
+	TObjectPtr<class UK2Node_TemporaryVariable> DataNode;
 
 	//~ Begin UEdGraphNode Interface
 	virtual void AllocateDefaultPins() override;
diff --git a/Engine/Source/Editor/BlueprintGraph/Classes/K2Node_StructOperation.h b/Engine/Source/Editor/BlueprintGraph/Classes/K2Node_StructOperation.h
--- a/Engine/Source/Editor/BlueprintGraph/Classes/K2Node_StructOperation.h
+++ b/Engine/Source/Editor/BlueprintGraph/Classes/K2Node_StructOperation.h
@@ -16,7 +16,7 @@ class UK2Node_StructOperation : public UK2Node_Variable
 
 	/** Class that this variable is defined in.  */
 	UPROPERTY()
-	UScriptStruct* StructType;
+	TObjectPtr<UScriptStruct> StructType;
 
 	//~ Begin UEdGraphNode Interface
 	virtual FString GetPinMetaData(FName InPinName, FName InKey) override;
diff --git a/Engine/Source/Editor/BlueprintGraph/Classes/K2Node_Tunnel.h b/Engine/Source/Editor/BlueprintGraph/Classes/K2Node_Tunnel.h
--- a/Engine/Source/Editor/BlueprintGraph/Classes/K2Node_Tunnel.h
+++ b/Engine/Source/Editor/BlueprintGraph/Classes/K2Node_Tunnel.h
@@ -18,11 +18,11 @@ class BLUEPRINTGRAPH_API UK2Node_Tunnel : public UK2Node_EditablePinBase
 	
 	// The output pins of this tunnel node came from the input pins of OutputSourceNode
 	UPROPERTY()
-	UK2Node_Tunnel* OutputSourceNode;
+	TObjectPtr<UK2Node_Tunnel> OutputSourceNode;
 
 	// The input pins of this tunnel go to the output pins of InputSinkNode
 	UPROPERTY()
-	UK2Node_Tunnel* InputSinkNode;
+	TObjectPtr<UK2Node_Tunnel> InputSinkNode;
 
 	// Whether this node is allowed to have inputs
 	UPROPERTY()
diff --git a/Engine/Source/Editor/Blutility/Classes/EditorUtilityTask.h b/Engine/Source/Editor/Blutility/Classes/EditorUtilityTask.h
--- a/Engine/Source/Editor/Blutility/Classes/EditorUtilityTask.h
+++ b/Engine/Source/Editor/Blutility/Classes/EditorUtilityTask.h
@@ -52,7 +52,7 @@ private:
 
 private:
 	UPROPERTY(Transient)
-	UEditorUtilitySubsystem* MyTaskManager;
+	TObjectPtr<UEditorUtilitySubsystem> MyTaskManager;
 
 	UPROPERTY(Transient)
 	bool bCancelRequested = false;
diff --git a/Engine/Source/Editor/Blutility/Classes/EditorUtilityWidgetBlueprint.h b/Engine/Source/Editor/Blutility/Classes/EditorUtilityWidgetBlueprint.h
--- a/Engine/Source/Editor/Blutility/Classes/EditorUtilityWidgetBlueprint.h
+++ b/Engine/Source/Editor/Blutility/Classes/EditorUtilityWidgetBlueprint.h
@@ -64,6 +64,6 @@ private:
 	TWeakPtr<SDockTab> CreatedTab;
 
 	UPROPERTY(Transient)
-	UEditorUtilityWidget* CreatedUMGWidget;
+	TObjectPtr<UEditorUtilityWidget> CreatedUMGWidget;
 
 };
diff --git a/Engine/Source/Editor/Blutility/Public/EditorUtilitySubsystem.h b/Engine/Source/Editor/Blutility/Public/EditorUtilitySubsystem.h
--- a/Engine/Source/Editor/Blutility/Public/EditorUtilitySubsystem.h
+++ b/Engine/Source/Editor/Blutility/Public/EditorUtilitySubsystem.h
@@ -94,19 +94,19 @@ private:
 	IConsoleObject* CancelAllTasksCommandObject = nullptr;
 	
 	UPROPERTY()
-	TMap<UObject* /*Asset*/, UObject* /*Instance*/> ObjectInstances;
+	TMap<TObjectPtr<UObject> /*Asset*/, TObjectPtr<UObject> /*Instance*/> ObjectInstances;
 
 	TQueue< TArray<FString> > RunTaskCommandBuffer;
 
 	UPROPERTY(Transient)
-	TArray<UEditorUtilityTask*> PendingTasks;
+	TArray<TObjectPtr<UEditorUtilityTask>> PendingTasks;
 
 	UPROPERTY(Transient)
-	UEditorUtilityTask* ActiveTask;
+	TObjectPtr<UEditorUtilityTask> ActiveTask;
 
 	FDelegateHandle TickerHandle;
 
 	/** List of objects that are being kept alive by this subsystem. */
 	UPROPERTY()
-	TSet<UObject*> ReferencedObjects;
+	TSet<TObjectPtr<UObject>> ReferencedObjects;
 };
diff --git a/Engine/Source/Editor/ClothPainter/Private/ClothPaintSettings.h b/Engine/Source/Editor/ClothPainter/Private/ClothPaintSettings.h
--- a/Engine/Source/Editor/ClothPainter/Private/ClothPaintSettings.h
+++ b/Engine/Source/Editor/ClothPainter/Private/ClothPaintSettings.h
@@ -84,7 +84,7 @@ public:
 
 	/** Array of Clothing assets */
 	UPROPERTY()
-	TArray<UClothingAssetCommon*> ClothingAssets;
+	TArray<TObjectPtr<UClothingAssetCommon>> ClothingAssets;
 
 	/** Whether to flip normals on the mesh preview */
 	UPROPERTY(EditAnywhere, AdvancedDisplay, Category = View)
diff --git a/Engine/Source/Editor/ComponentVisualizers/Public/Manipulator.h b/Engine/Source/Editor/ComponentVisualizers/Public/Manipulator.h
--- a/Engine/Source/Editor/ComponentVisualizers/Public/Manipulator.h
+++ b/Engine/Source/Editor/ComponentVisualizers/Public/Manipulator.h
@@ -38,10 +38,10 @@ private:
 
 	/** The component to transform when this manipulator was moved. */
 	UPROPERTY()
-	USceneComponent* AssociatedComponent;
+	TObjectPtr<USceneComponent> AssociatedComponent;
 
 	/** Visual representation of this manipulator. */
 	UPROPERTY()
-	class UStaticMeshComponent* StaticMeshComponent;
+	TObjectPtr<class UStaticMeshComponent> StaticMeshComponent;
 
 };
\ No newline at end of file
diff --git a/Engine/Source/Editor/ConfigEditor/Public/ConfigPropertyHelper.h b/Engine/Source/Editor/ConfigEditor/Public/ConfigPropertyHelper.h
--- a/Engine/Source/Editor/ConfigEditor/Public/ConfigPropertyHelper.h
+++ b/Engine/Source/Editor/ConfigEditor/Public/ConfigPropertyHelper.h
@@ -47,5 +47,5 @@ public:
 	TFieldPath<FProperty> EditProperty;
 
 	UPROPERTY(Transient, Category = Helper, EditAnywhere)
-	TArray<UPropertyConfigFileDisplayRow*> ConfigFilePropertyObjects;
+	TArray<TObjectPtr<UPropertyConfigFileDisplayRow>> ConfigFilePropertyObjects;
 };
diff --git a/Engine/Source/Editor/ContentBrowser/Public/ContentBrowserMenuContexts.h b/Engine/Source/Editor/ContentBrowser/Public/ContentBrowserMenuContexts.h
--- a/Engine/Source/Editor/ContentBrowser/Public/ContentBrowserMenuContexts.h
+++ b/Engine/Source/Editor/ContentBrowser/Public/ContentBrowserMenuContexts.h
@@ -29,7 +29,7 @@ public:
 	TArray<TWeakObjectPtr<UObject>> SelectedObjects;
 
 	UPROPERTY()
-	UClass* CommonClass;
+	TObjectPtr<UClass> CommonClass;
 
 	UPROPERTY()
 	bool bCanBeModified;
diff --git a/Engine/Source/Editor/CurveEditor/Private/CurveEditorCopyBuffer.h b/Engine/Source/Editor/CurveEditor/Private/CurveEditorCopyBuffer.h
--- a/Engine/Source/Editor/CurveEditor/Private/CurveEditorCopyBuffer.h
+++ b/Engine/Source/Editor/CurveEditor/Private/CurveEditorCopyBuffer.h
@@ -38,7 +38,7 @@ public:
 	GENERATED_BODY()
 
 	UPROPERTY(Instanced)
-	TArray<UCurveEditorCopyableCurveKeys*> Curves;
+	TArray<TObjectPtr<UCurveEditorCopyableCurveKeys>> Curves;
 
 	UPROPERTY()
 	double TimeOffset;
diff --git a/Engine/Source/Editor/EditorConfig/Private/Tests/EditorConfigTests.h b/Engine/Source/Editor/EditorConfig/Private/Tests/EditorConfigTests.h
--- a/Engine/Source/Editor/EditorConfig/Private/Tests/EditorConfigTests.h
+++ b/Engine/Source/Editor/EditorConfig/Private/Tests/EditorConfigTests.h
@@ -99,7 +99,7 @@ class UEditorConfigTestObject : public UObject
 	GENERATED_BODY()
 
 	UPROPERTY()
-	UObject* Object;
+	TObjectPtr<UObject> Object;
 
 	UPROPERTY()
 	FSoftObjectPath SoftObjectPath;
diff --git a/Engine/Source/Editor/Experimental/EditorInteractiveToolsFramework/Public/EdModeInteractiveToolsContext.h b/Engine/Source/Editor/Experimental/EditorInteractiveToolsFramework/Public/EdModeInteractiveToolsContext.h
--- a/Engine/Source/Editor/Experimental/EditorInteractiveToolsFramework/Public/EdModeInteractiveToolsContext.h
+++ b/Engine/Source/Editor/Experimental/EditorInteractiveToolsFramework/Public/EdModeInteractiveToolsContext.h
@@ -105,7 +105,7 @@ protected:
 
 public:
 	UPROPERTY()
-	UMaterialInterface* StandardVertexColorMaterial;
+	TObjectPtr<UMaterialInterface> StandardVertexColorMaterial;
 
 protected:
 	// called when PIE is about to start, shuts down active tools
diff --git a/Engine/Source/Editor/GameProjectGeneration/Public/ClassTemplateEditorSubsystem.h b/Engine/Source/Editor/GameProjectGeneration/Public/ClassTemplateEditorSubsystem.h
--- a/Engine/Source/Editor/GameProjectGeneration/Public/ClassTemplateEditorSubsystem.h
+++ b/Engine/Source/Editor/GameProjectGeneration/Public/ClassTemplateEditorSubsystem.h
@@ -53,7 +53,7 @@ protected:
 private:
 	// Base UClass of which template class corresponds.
 	UPROPERTY(Transient)
-	UClass* GeneratedBaseClass;
+	TObjectPtr<UClass> GeneratedBaseClass;
 };
 
 UCLASS(Abstract)
diff --git a/Engine/Source/Editor/IntroTutorials/Public/EditorTutorial.h b/Engine/Source/Editor/IntroTutorials/Public/EditorTutorial.h
--- a/Engine/Source/Editor/IntroTutorials/Public/EditorTutorial.h
+++ b/Engine/Source/Editor/IntroTutorials/Public/EditorTutorial.h
@@ -251,7 +251,7 @@ class INTROTUTORIALS_API UEditorTutorial : public UObject
 
 	/** Texture for this tutorial, used when presented to the user in the tutorial browser. */
 	UPROPERTY(EditAnywhere, Category="Tutorial")
-	UTexture2D* Texture;
+	TObjectPtr<UTexture2D> Texture;
 
 	/** Category of this tutorial, used to organize tutorials when presented to the user */
 	UPROPERTY(EditAnywhere, Category="Tutorial")
diff --git a/Engine/Source/Editor/Kismet/Public/BlueprintCompilerExtension.h b/Engine/Source/Editor/Kismet/Public/BlueprintCompilerExtension.h
--- a/Engine/Source/Editor/Kismet/Public/BlueprintCompilerExtension.h
+++ b/Engine/Source/Editor/Kismet/Public/BlueprintCompilerExtension.h
@@ -11,7 +11,7 @@ struct FBlueprintCompiledData
 	GENERATED_BODY()
 
 	UPROPERTY()
-	TArray<UEdGraph*> IntermediateGraphs;
+	TArray<TObjectPtr<UEdGraph>> IntermediateGraphs;
 };
 
 UCLASS(Abstract)
diff --git a/Engine/Source/Editor/LandscapeEditor/Public/LandscapeEditorObject.h b/Engine/Source/Editor/LandscapeEditor/Public/LandscapeEditorObject.h
--- a/Engine/Source/Editor/LandscapeEditor/Public/LandscapeEditorObject.h
+++ b/Engine/Source/Editor/LandscapeEditor/Public/LandscapeEditorObject.h
@@ -202,7 +202,7 @@ struct FLandscapeImportLayer : public FLandscapeImportLayerInfo
 	GENERATED_USTRUCT_BODY()
 
 	UPROPERTY(Category="Import", VisibleAnywhere)
-	ULandscapeMaterialInstanceConstant* ThumbnailMIC;
+	TObjectPtr<ULandscapeMaterialInstanceConstant> ThumbnailMIC;
 
 	UPROPERTY(Category = "Import", VisibleAnywhere)
 	ELandscapeImportResult ImportResult;
@@ -626,7 +626,7 @@ public:
 
 	// Mask texture to use
 	UPROPERTY(Category="Brush Settings", EditAnywhere, NonTransactional, meta=(DisplayName="Texture", ShowForBrushes="BrushSet_Alpha,BrushSet_Pattern"))
-	UTexture2D* AlphaTexture;
+	TObjectPtr<UTexture2D> AlphaTexture;
 
 	// Channel of Mask Texture to use
 	UPROPERTY(Category="Brush Settings", EditAnywhere, NonTransactional, meta=(DisplayName="Texture Channel", ShowForBrushes="BrushSet_Alpha,BrushSet_Pattern"))
diff --git a/Engine/Source/Editor/LevelAssetEditor/Public/ULevelAssetEditor.h b/Engine/Source/Editor/LevelAssetEditor/Public/ULevelAssetEditor.h
--- a/Engine/Source/Editor/LevelAssetEditor/Public/ULevelAssetEditor.h
+++ b/Engine/Source/Editor/LevelAssetEditor/Public/ULevelAssetEditor.h
@@ -20,5 +20,5 @@ public:
 
 protected:
 	UPROPERTY()
-	UInteractiveToolsContext* InteractiveToolsContext;
+	TObjectPtr<UInteractiveToolsContext> InteractiveToolsContext;
 };
diff --git a/Engine/Source/Editor/LevelEditor/Public/LevelEditorMenuContext.h b/Engine/Source/Editor/LevelEditor/Public/LevelEditorMenuContext.h
--- a/Engine/Source/Editor/LevelEditor/Public/LevelEditorMenuContext.h
+++ b/Engine/Source/Editor/LevelEditor/Public/LevelEditorMenuContext.h
@@ -43,11 +43,11 @@ public:
 	ELevelEditorMenuContext ContextType;
 
 	UPROPERTY()
-	TArray<UActorComponent*> SelectedComponents;
+	TArray<TObjectPtr<UActorComponent>> SelectedComponents;
 
 	/** If the ContextType is Viewport this property can be set to the HitProxy actor that triggered the ContextMenu. */
 	UPROPERTY()
-	AActor* HitProxyActor = nullptr;
+	TObjectPtr<AActor> HitProxyActor = nullptr;
 };
 
 UCLASS()
@@ -68,5 +68,5 @@ class LEVELEDITOR_API UQuickActionMenuContext : public UObject
 	GENERATED_BODY()
 public:
 	UPROPERTY(Transient)
-	const UTypedElementSelectionSet* CurrentSelection;
+	TObjectPtr<const UTypedElementSelectionSet> CurrentSelection;
 };
diff --git a/Engine/Source/Editor/MaterialEditor/Public/MaterialPropertyHelpers.h b/Engine/Source/Editor/MaterialEditor/Public/MaterialPropertyHelpers.h
--- a/Engine/Source/Editor/MaterialEditor/Public/MaterialPropertyHelpers.h
+++ b/Engine/Source/Editor/MaterialEditor/Public/MaterialPropertyHelpers.h
@@ -43,7 +43,7 @@ public:
 	EStackDataType StackDataType;
 
 	UPROPERTY(Transient)
-	class UDEditorParameterValue* Parameter = nullptr;
+	TObjectPtr<class UDEditorParameterValue> Parameter = nullptr;
 
 	FName PropertyName;
 
@@ -65,7 +65,7 @@ struct FUnsortedParamData
 {
 	GENERATED_USTRUCT_BODY()
 	UPROPERTY(Transient)
-	class UDEditorParameterValue* Parameter = nullptr;
+	TObjectPtr<class UDEditorParameterValue> Parameter = nullptr;
 	FEditorParameterGroup ParameterGroup;
 	TSharedPtr<class IDetailTreeNode> ParameterNode;
 	FName UnsortedName;
diff --git a/Engine/Source/Editor/MovieSceneTools/Private/MovieSceneEventBlueprintExtension.h b/Engine/Source/Editor/MovieSceneTools/Private/MovieSceneEventBlueprintExtension.h
--- a/Engine/Source/Editor/MovieSceneTools/Private/MovieSceneEventBlueprintExtension.h
+++ b/Engine/Source/Editor/MovieSceneTools/Private/MovieSceneEventBlueprintExtension.h
@@ -31,5 +31,5 @@ private:
 
 	/** List of event sections that are bound to the blueprint */
 	UPROPERTY()
-	TArray<UMovieSceneEventSectionBase*> EventSections;
+	TArray<TObjectPtr<UMovieSceneEventSectionBase>> EventSections;
 };
\ No newline at end of file
diff --git a/Engine/Source/Editor/MovieSceneTools/Public/AutomatedLevelSequenceCapture.h b/Engine/Source/Editor/MovieSceneTools/Public/AutomatedLevelSequenceCapture.h
--- a/Engine/Source/Editor/MovieSceneTools/Public/AutomatedLevelSequenceCapture.h
+++ b/Engine/Source/Editor/MovieSceneTools/Public/AutomatedLevelSequenceCapture.h
@@ -68,7 +68,7 @@ public:
 	float DelayEveryFrame;
 
 	UPROPERTY(Transient, EditAnywhere, BlueprintReadWrite, Category=CaptureSettings, AdvancedDisplay, meta=(EditInline))
-	ULevelSequenceBurnInOptions* BurnInOptions;
+	TObjectPtr<ULevelSequenceBurnInOptions> BurnInOptions;
 
 	/** Whether to write edit decision lists (EDLs) if the sequence contains shots */
 	UPROPERTY(config, EditAnywhere, BlueprintReadWrite, Category=Sequence)
diff --git a/Engine/Source/Editor/Persona/Private/PersonaPreviewSceneDescription.h b/Engine/Source/Editor/Persona/Private/PersonaPreviewSceneDescription.h
--- a/Engine/Source/Editor/Persona/Private/PersonaPreviewSceneDescription.h
+++ b/Engine/Source/Editor/Persona/Private/PersonaPreviewSceneDescription.h
@@ -30,10 +30,10 @@ public:
 	class TSubclassOf<UPersonaPreviewSceneController> PreviewController;
 
 	UPROPERTY()
-	UPersonaPreviewSceneController* PreviewControllerInstance;
+	TObjectPtr<UPersonaPreviewSceneController> PreviewControllerInstance;
 
 	UPROPERTY()
-	TArray<UPersonaPreviewSceneController*> PreviewControllerInstances;
+	TArray<TObjectPtr<UPersonaPreviewSceneController>> PreviewControllerInstances;
 
 	/** The preview mesh to use */
 	UPROPERTY(EditAnywhere, Category = "Mesh", meta=(DisplayThumbnail=true))
@@ -55,7 +55,7 @@ public:
 	TSoftObjectPtr<UDataAsset> AdditionalMeshes;
 
 	UPROPERTY()
-	UPreviewMeshCollection* DefaultAdditionalMeshes;
+	TObjectPtr<UPreviewMeshCollection> DefaultAdditionalMeshes;
 
 	// Sets the current preview controller for the scene (handles uninitializing and initializing controllers)
 	bool SetPreviewController(UClass* PreviewControllerClass, class IPersonaPreviewScene* PreviewScene)
diff --git a/Engine/Source/Editor/PhysicsAssetEditor/Private/PhysicsAssetEditorSkeletalMeshComponent.h b/Engine/Source/Editor/PhysicsAssetEditor/Private/PhysicsAssetEditorSkeletalMeshComponent.h
--- a/Engine/Source/Editor/PhysicsAssetEditor/Private/PhysicsAssetEditorSkeletalMeshComponent.h
+++ b/Engine/Source/Editor/PhysicsAssetEditor/Private/PhysicsAssetEditorSkeletalMeshComponent.h
@@ -36,15 +36,15 @@ class UPhysicsAssetEditorSkeletalMeshComponent : public UDebugSkelMeshComponent
 	// Materials
 
 	UPROPERTY(transient)
-	UMaterialInstanceDynamic* ElemSelectedMaterial;
+	TObjectPtr<UMaterialInstanceDynamic> ElemSelectedMaterial;
 	UPROPERTY(transient)
-	UMaterialInstanceDynamic* BoneSelectedMaterial;
+	TObjectPtr<UMaterialInstanceDynamic> BoneSelectedMaterial;
 	UPROPERTY(transient)
-	UMaterialInstanceDynamic* BoneUnselectedMaterial;
+	TObjectPtr<UMaterialInstanceDynamic> BoneUnselectedMaterial;
 	UPROPERTY(transient)
-	UMaterialInterface* BoneMaterialHit;
+	TObjectPtr<UMaterialInterface> BoneMaterialHit;
 	UPROPERTY(transient)
-	UMaterialInstanceDynamic* BoneNoCollisionMaterial;
+	TObjectPtr<UMaterialInstanceDynamic> BoneNoCollisionMaterial;
 
 	/** Mesh-space matrices showing state of just animation (ie before physics) - useful for debugging! */
 	TArray<FTransform> AnimationSpaceBases;
diff --git a/Engine/Source/Editor/PropertyEditor/Private/EditConditionParserTests.h b/Engine/Source/Editor/PropertyEditor/Private/EditConditionParserTests.h
--- a/Engine/Source/Editor/PropertyEditor/Private/EditConditionParserTests.h
+++ b/Engine/Source/Editor/PropertyEditor/Private/EditConditionParserTests.h
@@ -43,7 +43,7 @@ class UEditConditionTestObject : public UObject
 	uint8 UintBitfieldProperty : 1;
 
 	UPROPERTY(EditAnywhere, Category=Test)
-	UObject* UObjectPtr;
+	TObjectPtr<UObject> UObjectPtr;
 
 	UPROPERTY(EditAnywhere, Category=Test)
 	TSoftClassPtr<UObject> SoftClassPtr;
diff --git a/Engine/Source/Editor/SequenceRecorder/Public/ActorRecording.h b/Engine/Source/Editor/SequenceRecorder/Public/ActorRecording.h
--- a/Engine/Source/Editor/SequenceRecorder/Public/ActorRecording.h
+++ b/Engine/Source/Editor/SequenceRecorder/Public/ActorRecording.h
@@ -112,7 +112,7 @@ public:
 
 	/** The level sequence to record into */
 	UPROPERTY(EditAnywhere, Category = "Actor Recording", meta=(EditCondition = "bCreateLevelSequence"))
-	class ULevelSequence* TargetLevelSequence;
+	TObjectPtr<class ULevelSequence> TargetLevelSequence;
 
 	/** Optional target name to record to. If not specified, the actor label will be used */
 	UPROPERTY(EditAnywhere, Category = "Actor Recording", meta=(EditCondition = "bCreateLevelSequence"))
@@ -128,7 +128,7 @@ public:
 
 	/** The target animation we want to record to */
 	UPROPERTY(EditAnywhere, Category = "Animation Recording", meta=(EditCondition = "bSpecifyTargetAnimation"))
-	class UAnimSequence* TargetAnimation;
+	TObjectPtr<class UAnimSequence> TargetAnimation;
 
 	/** The settings to apply to this actor's animation */
 	UPROPERTY(EditAnywhere, Category = "Animation Recording")
diff --git a/Engine/Source/Editor/SequenceRecorder/Public/ActorRecordingSettings.h b/Engine/Source/Editor/SequenceRecorder/Public/ActorRecordingSettings.h
--- a/Engine/Source/Editor/SequenceRecorder/Public/ActorRecordingSettings.h
+++ b/Engine/Source/Editor/SequenceRecorder/Public/ActorRecordingSettings.h
@@ -35,7 +35,7 @@ struct SEQUENCERECORDER_API FActorRecordingSettings
 private:
 	/** External settings objects for recorders that supply them. Displayed via a details customization  */
 	UPROPERTY(EditAnywhere, Category = "Actor Recording")
-	TArray<UObject*> Settings;
+	TArray<TObjectPtr<UObject>> Settings;
 
 	/** An optional outer that settings objects should be created with. */
 	TWeakObjectPtr<UObject> Outer;
diff --git a/Engine/Source/Editor/SequenceRecorder/Public/SequenceRecorderActorGroup.h b/Engine/Source/Editor/SequenceRecorder/Public/SequenceRecorderActorGroup.h
--- a/Engine/Source/Editor/SequenceRecorder/Public/SequenceRecorderActorGroup.h
+++ b/Engine/Source/Editor/SequenceRecorder/Public/SequenceRecorderActorGroup.h
@@ -41,7 +41,7 @@ public:
 
 	/** The level sequence to record into */
 	UPROPERTY(EditAnywhere, AdvancedDisplay, Category = "Recording Groups", meta=(EditCondition = "bSpecifyTargetLevelSequence"))
-	class ULevelSequence* TargetLevelSequence;
+	TObjectPtr<class ULevelSequence> TargetLevelSequence;
 
 	/** Whether we should duplicate the target level sequence and record into the duplicate */
 	UPROPERTY(EditAnywhere, AdvancedDisplay, Category = "Recording Groups", meta=(EditCondition = "bSpecifyTargetLevelSequence"))
@@ -53,7 +53,7 @@ public:
 
 	/** A list of actor recordings in this group which contains both the actors to record as well as settings for each one. */
 	UPROPERTY(VisibleAnywhere, Category = "Recording Groups")
-	TArray<UActorRecording*> RecordedActors;
+	TArray<TObjectPtr<UActorRecording>> RecordedActors;
 
 	void PostEditChangeChainProperty(struct FPropertyChangedChainEvent& PropertyChangedEvent) override;
 };
@@ -95,5 +95,5 @@ public:
 public:
 
 	UPROPERTY()
-	TArray<USequenceRecorderActorGroup*> ActorGroups;
+	TArray<TObjectPtr<USequenceRecorderActorGroup>> ActorGroups;
 };
diff --git a/Engine/Source/Editor/Sequencer/Private/MovieSceneCopyableBinding.h b/Engine/Source/Editor/Sequencer/Private/MovieSceneCopyableBinding.h
--- a/Engine/Source/Editor/Sequencer/Private/MovieSceneCopyableBinding.h
+++ b/Engine/Source/Editor/Sequencer/Private/MovieSceneCopyableBinding.h
@@ -24,7 +24,7 @@ public:
 	* the rest of the text as it'll fall back to the same reference issue. Marking this as TextExportTransient solves this.
 	*/
 	UPROPERTY(TextExportTransient)
-	UObject* SpawnableObjectTemplate;
+	TObjectPtr<UObject> SpawnableObjectTemplate;
 
 	/**
 	 * Tracks are also owned by the owning Movie Sequence. We manually copy the tracks out of a binding when we copy,
@@ -32,7 +32,7 @@ public:
 	 * and then restore them upon paste to re-create the tracks with the correct owner.
 	 */
 	UPROPERTY()
-	TArray<UMovieSceneTrack*> Tracks;
+	TArray<TObjectPtr<UMovieSceneTrack>> Tracks;
 
 	UPROPERTY()
 	FMovieSceneBinding Binding;
diff --git a/Engine/Source/Editor/Sequencer/Private/MovieSceneCopyableTrack.h b/Engine/Source/Editor/Sequencer/Private/MovieSceneCopyableTrack.h
--- a/Engine/Source/Editor/Sequencer/Private/MovieSceneCopyableTrack.h
+++ b/Engine/Source/Editor/Sequencer/Private/MovieSceneCopyableTrack.h
@@ -15,7 +15,7 @@ public:
 	GENERATED_BODY()
 
 	UPROPERTY()
-	UMovieSceneTrack* Track;
+	TObjectPtr<UMovieSceneTrack> Track;
 
 	UPROPERTY()
 	bool bIsAMasterTrack;
diff --git a/Engine/Source/Editor/Sequencer/Private/SequencerKeyActor.h b/Engine/Source/Editor/Sequencer/Private/SequencerKeyActor.h
--- a/Engine/Source/Editor/Sequencer/Private/SequencerKeyActor.h
+++ b/Engine/Source/Editor/Sequencer/Private/SequencerKeyActor.h
@@ -46,15 +46,15 @@ private:
 
 	/** The key mesh */
 	UPROPERTY()
-	class UStaticMeshComponent* KeyMeshComponent;
+	TObjectPtr<class UStaticMeshComponent> KeyMeshComponent;
 
 	/** The actor whose transform was used to build this key */
 	UPROPERTY()
-	AActor*	AssociatedActor;
+	TObjectPtr<AActor>	AssociatedActor;
 
 	/** The track section this key resides on */
 	UPROPERTY()
-	class UMovieScene3DTransformSection* TrackSection;
+	TObjectPtr<class UMovieScene3DTransformSection> TrackSection;
 
 	/** The time this key is associated with in Sequencer */
 	UPROPERTY()
diff --git a/Engine/Source/Editor/Sequencer/Public/SequencerExportTask.h b/Engine/Source/Editor/Sequencer/Public/SequencerExportTask.h
--- a/Engine/Source/Editor/Sequencer/Public/SequencerExportTask.h
+++ b/Engine/Source/Editor/Sequencer/Public/SequencerExportTask.h
@@ -20,5 +20,5 @@ public:
 
 	/* A UWorld for LevelSequences, UUserWidget for WidgetAnimations, or AActor for Actor Sequences, etc... */
 	UPROPERTY(BlueprintReadWrite, Category = Miscellaneous)
-	UObject* SequencerContext;
+	TObjectPtr<UObject> SequencerContext;
 };
diff --git a/Engine/Source/Editor/ToolMenusEditor/Public/ToolMenusEditor.h b/Engine/Source/Editor/ToolMenusEditor/Public/ToolMenusEditor.h
--- a/Engine/Source/Editor/ToolMenusEditor/Public/ToolMenusEditor.h
+++ b/Engine/Source/Editor/ToolMenusEditor/Public/ToolMenusEditor.h
@@ -41,7 +41,7 @@ public:
 	FName Name;
 
 	UPROPERTY()
-	UToolMenu* Menu;
+	TObjectPtr<UToolMenu> Menu;
 };
 
 UCLASS()
@@ -61,7 +61,7 @@ public:
 	ECustomizedToolMenuVisibility Visibility;
 
 	UPROPERTY()
-	UToolMenu* Menu;
+	TObjectPtr<UToolMenu> Menu;
 };
 
 UCLASS()
diff --git a/Engine/Source/Editor/UMGEditor/Classes/K2Node_WidgetAnimationEvent.h b/Engine/Source/Editor/UMGEditor/Classes/K2Node_WidgetAnimationEvent.h
--- a/Engine/Source/Editor/UMGEditor/Classes/K2Node_WidgetAnimationEvent.h
+++ b/Engine/Source/Editor/UMGEditor/Classes/K2Node_WidgetAnimationEvent.h
@@ -36,7 +36,7 @@ public:
 	FName UserTag;
 
 	UPROPERTY()
-	const UWidgetBlueprint* SourceWidgetBlueprint;
+	TObjectPtr<const UWidgetBlueprint> SourceWidgetBlueprint;
 
 	virtual void PostDuplicate(bool bDuplicateForPIE) override;
 
diff --git a/Engine/Source/Editor/UMGEditor/Classes/WidgetBlueprintFactory.h b/Engine/Source/Editor/UMGEditor/Classes/WidgetBlueprintFactory.h
--- a/Engine/Source/Editor/UMGEditor/Classes/WidgetBlueprintFactory.h
+++ b/Engine/Source/Editor/UMGEditor/Classes/WidgetBlueprintFactory.h
@@ -31,5 +31,5 @@ class UWidgetBlueprintFactory : public UFactory
 
 private:
 	UPROPERTY(Transient)
-	UClass* RootWidgetClass;
+	TObjectPtr<UClass> RootWidgetClass;
 };
diff --git a/Engine/Source/Editor/UMGEditor/Public/Settings/WidgetDesignerSettings.h b/Engine/Source/Editor/UMGEditor/Public/Settings/WidgetDesignerSettings.h
--- a/Engine/Source/Editor/UMGEditor/Public/Settings/WidgetDesignerSettings.h
+++ b/Engine/Source/Editor/UMGEditor/Public/Settings/WidgetDesignerSettings.h
@@ -67,6 +67,6 @@ public:
 	 * List of Favorites widgets used to populate the Favorites Palette
 	 */
 	UPROPERTY()
-	UWidgetPaletteFavorites* Favorites;
+	TObjectPtr<UWidgetPaletteFavorites> Favorites;
 
 };
diff --git a/Engine/Source/Editor/UMGEditor/Public/WidgetBlueprint.h b/Engine/Source/Editor/UMGEditor/Public/WidgetBlueprint.h
--- a/Engine/Source/Editor/UMGEditor/Public/WidgetBlueprint.h
+++ b/Engine/Source/Editor/UMGEditor/Public/WidgetBlueprint.h
@@ -63,7 +63,7 @@ private:
 
 	/** The owner of the path segment (ie. What class or structure was this property from) */
 	UPROPERTY()
-	UStruct* Struct;
+	TObjectPtr<UStruct> Struct;
 
 	/** The member name in the structure this segment represents. */
 	UPROPERTY()
@@ -175,7 +175,7 @@ struct FWidgetAnimation_DEPRECATED
 	GENERATED_USTRUCT_BODY()
 
 	UPROPERTY()
-	UMovieScene* MovieScene = nullptr;
+	TObjectPtr<UMovieScene> MovieScene = nullptr;
 
 	UPROPERTY()
 	TArray<FWidgetAnimationBinding> AnimationBindings;
@@ -238,7 +238,7 @@ public:
 	TArray<FWidgetAnimation_DEPRECATED> AnimationData_DEPRECATED;
 
 	UPROPERTY()
-	TArray<UWidgetAnimation*> Animations;
+	TArray<TObjectPtr<UWidgetAnimation>> Animations;
 
 	/**
 	 * Don't directly modify this property to change the palette category.  The actual value is stored 
diff --git a/Engine/Source/Editor/UnrealEd/Classes/Animation/DebugSkelMeshComponent.h b/Engine/Source/Editor/UnrealEd/Classes/Animation/DebugSkelMeshComponent.h
--- a/Engine/Source/Editor/UnrealEd/Classes/Animation/DebugSkelMeshComponent.h
+++ b/Engine/Source/Editor/UnrealEd/Classes/Animation/DebugSkelMeshComponent.h
@@ -256,17 +256,17 @@ class UNREALED_API UDebugSkelMeshComponent : public USkeletalMeshComponent
 
 	/** Array of morphtargets to render verts for */
 	UPROPERTY(transient)
-	TArray<class UMorphTarget*> MorphTargetOfInterests;
+	TArray<TObjectPtr<class UMorphTarget>> MorphTargetOfInterests;
 
 	/** Array of materials to restore when not rendering blend weights */
 	UPROPERTY(transient)
-	TArray<class UMaterialInterface*> SkelMaterials;
+	TArray<TObjectPtr<class UMaterialInterface>> SkelMaterials;
 	
 	UPROPERTY(transient, NonTransactional)
-	class UAnimPreviewInstance* PreviewInstance;
+	TObjectPtr<class UAnimPreviewInstance> PreviewInstance;
 
 	UPROPERTY(transient)
-	class UAnimInstance* SavedAnimScriptInstance;
+	TObjectPtr<class UAnimInstance> SavedAnimScriptInstance;
 
 	/** Does this component use in game bounds or does it use bounds calculated from bones */
 	UPROPERTY(transient)
diff --git a/Engine/Source/Editor/UnrealEd/Classes/Commandlets/DiffPackagesCommandlet.h b/Engine/Source/Editor/UnrealEd/Classes/Commandlets/DiffPackagesCommandlet.h
--- a/Engine/Source/Editor/UnrealEd/Classes/Commandlets/DiffPackagesCommandlet.h
+++ b/Engine/Source/Editor/UnrealEd/Classes/Commandlets/DiffPackagesCommandlet.h
@@ -38,7 +38,7 @@ class UDEPRECATED_DiffPackagesCommandlet : public UCommandlet // No longer maint
 	GENERATED_UCLASS_BODY()
 
 	UPROPERTY()
-	UPackage* Packages[3];
+	TObjectPtr<UPackage> Packages[3];
 
 	bool bDiffNonEditProps;
 	bool bDiffAllProps;
diff --git a/Engine/Source/Editor/UnrealEd/Classes/Commandlets/WorldPartitionBuilderCommandlet.h b/Engine/Source/Editor/UnrealEd/Classes/Commandlets/WorldPartitionBuilderCommandlet.h
--- a/Engine/Source/Editor/UnrealEd/Classes/Commandlets/WorldPartitionBuilderCommandlet.h
+++ b/Engine/Source/Editor/UnrealEd/Classes/Commandlets/WorldPartitionBuilderCommandlet.h
@@ -19,5 +19,5 @@ class UWorldPartitionBuilderCommandlet : public UCommandlet, public FPackageSour
 
 private:
 	UPROPERTY()
-	TArray<UWorldPartitionBuilder*> Builders;
+	TArray<TObjectPtr<UWorldPartitionBuilder>> Builders;
 };
\ No newline at end of file
diff --git a/Engine/Source/Editor/UnrealEd/Classes/Commandlets/WorldPartitionConvertCommandlet.h b/Engine/Source/Editor/UnrealEd/Classes/Commandlets/WorldPartitionConvertCommandlet.h
--- a/Engine/Source/Editor/UnrealEd/Classes/Commandlets/WorldPartitionConvertCommandlet.h
+++ b/Engine/Source/Editor/UnrealEd/Classes/Commandlets/WorldPartitionConvertCommandlet.h
@@ -119,7 +119,7 @@ protected:
 	TArray<FHLODLayerActorMapping> HLODLayersForActorClasses;
 
 	UPROPERTY(Transient)
-	TMap<FString, UHLODLayer*> HLODLayers;
+	TMap<FString, TObjectPtr<UHLODLayer>> HLODLayers;
 
 	UPROPERTY(Config)
 	uint32 LandscapeGridSize;
diff --git a/Engine/Source/Editor/UnrealEd/Classes/Editor/EditorEngine.h b/Engine/Source/Editor/UnrealEd/Classes/Editor/EditorEngine.h
--- a/Engine/Source/Editor/UnrealEd/Classes/Editor/EditorEngine.h
+++ b/Engine/Source/Editor/UnrealEd/Classes/Editor/EditorEngine.h
@@ -281,41 +281,41 @@ public:
 
 	// Objects.
 	UPROPERTY()
-	class UModel* TempModel;
+	TObjectPtr<class UModel> TempModel;
 
 	UPROPERTY()
-	class UModel* ConversionTempModel;
+	TObjectPtr<class UModel> ConversionTempModel;
 
 	UPROPERTY()
-	class UTransactor* Trans;
+	TObjectPtr<class UTransactor> Trans;
 
 	// Textures.
 	UPROPERTY()
-	class UTexture2D* Bad;
+	TObjectPtr<class UTexture2D> Bad;
 
 	// Font used by Canvas-based editors
 	UPROPERTY()
-	class UFont* EditorFont;
+	TObjectPtr<class UFont> EditorFont;
 
 	// Audio
 	UPROPERTY(transient)
-	class USoundCue* PreviewSoundCue;
+	TObjectPtr<class USoundCue> PreviewSoundCue;
 
 	UPROPERTY(transient)
-	class UAudioComponent* PreviewAudioComponent;
+	TObjectPtr<class UAudioComponent> PreviewAudioComponent;
 
 	// Used in UnrealEd for showing materials
 	UPROPERTY()
-	class UStaticMesh* EditorCube;
+	TObjectPtr<class UStaticMesh> EditorCube;
 
 	UPROPERTY()
-	class UStaticMesh* EditorSphere;
+	TObjectPtr<class UStaticMesh> EditorSphere;
 
 	UPROPERTY()
-	class UStaticMesh* EditorPlane;
+	TObjectPtr<class UStaticMesh> EditorPlane;
 
 	UPROPERTY()
-	class UStaticMesh* EditorCylinder;
+	TObjectPtr<class UStaticMesh> EditorCylinder;
 
 	// Toggles.
 	UPROPERTY()
@@ -329,7 +329,7 @@ public:
 	uint32 ClickFlags;
 
 	UPROPERTY()
-	class UPackage* ParentContext;
+	TObjectPtr<class UPackage> ParentContext;
 
 	UPROPERTY()
 	FVector UnsnappedClickLocation;
@@ -382,7 +382,7 @@ public:
 
 	/** Array of actor factories created at editor startup and used by context menu etc. */
 	UPROPERTY()
-	TArray<class UActorFactory*> ActorFactories;
+	TArray<TObjectPtr<class UActorFactory>> ActorFactories;
 
 	/** The name of the file currently being opened in the editor. "" if no file is being opened. */
 	UPROPERTY()
@@ -397,7 +397,7 @@ public:
 
 	/** A pointer to a UWorld that is the duplicated/saved-loaded to be played in with "Play From Here" 								*/
 	UPROPERTY()
-	class UWorld* PlayWorld;
+	TObjectPtr<class UWorld> PlayWorld;
 
 
 
@@ -427,7 +427,7 @@ public:
 
 	/** When Simulating In Editor, a pointer to the original (non-simulating) editor world */
 	UPROPERTY()
-	class UWorld* EditorWorld;
+	TObjectPtr<class UWorld> EditorWorld;
 	
 	/** When Simulating In Editor, an array of all actors that were selected when it began*/
 	UPROPERTY()
@@ -459,20 +459,20 @@ public:
 
 	/** Temporary render target that can be used by the editor. */
 	UPROPERTY(transient)
-	class UTextureRenderTarget2D* ScratchRenderTarget2048;
+	TObjectPtr<class UTextureRenderTarget2D> ScratchRenderTarget2048;
 
 	UPROPERTY(transient)
-	class UTextureRenderTarget2D* ScratchRenderTarget1024;
+	TObjectPtr<class UTextureRenderTarget2D> ScratchRenderTarget1024;
 
 	UPROPERTY(transient)
-	class UTextureRenderTarget2D* ScratchRenderTarget512;
+	TObjectPtr<class UTextureRenderTarget2D> ScratchRenderTarget512;
 
 	UPROPERTY(transient)
-	class UTextureRenderTarget2D* ScratchRenderTarget256;
+	TObjectPtr<class UTextureRenderTarget2D> ScratchRenderTarget256;
 
 	/** A mesh component used to preview in editor without spawning a static mesh actor. */
 	UPROPERTY(transient)
-	class UStaticMeshComponent* PreviewMeshComp;
+	TObjectPtr<class UStaticMeshComponent> PreviewMeshComp;
 
 	/** The index of the mesh to use from the list of preview meshes. */
 	UPROPERTY()
@@ -500,7 +500,7 @@ public:
 
 	/** Brush builders that have been created in the editor */
 	UPROPERTY(transient)
-	TArray<class UBrushBuilder*> BrushBuilders;	
+	TArray<TObjectPtr<class UBrushBuilder>> BrushBuilders;	
 
 	/** Whether or not to recheck the current actor selection for lock actors the next time HasLockActors is called */
 	bool bCheckForLockActors;
@@ -535,7 +535,7 @@ private:
 
 	/** Manager that holds all extensions paired with a world */
 	UPROPERTY()
-	UEditorWorldExtensionManager* EditorWorldExtensionsManager;
+	TObjectPtr<UEditorWorldExtensionManager> EditorWorldExtensionsManager;
 
 public:
 
@@ -3185,7 +3185,7 @@ protected:
 	FSoftClassPath ActorGroupingUtilsClassName;
 
 	UPROPERTY()
-	class UActorGroupingUtils* ActorGroupingUtils;
+	TObjectPtr<class UActorGroupingUtils> ActorGroupingUtils;
 private:
 	FTimerHandle CleanupPIEOnlineSessionsTimerHandle;
 
diff --git a/Engine/Source/Editor/UnrealEd/Classes/Editor/EditorPerProjectUserSettings.h b/Engine/Source/Editor/UnrealEd/Classes/Editor/EditorPerProjectUserSettings.h
--- a/Engine/Source/Editor/UnrealEd/Classes/Editor/EditorPerProjectUserSettings.h
+++ b/Engine/Source/Editor/UnrealEd/Classes/Editor/EditorPerProjectUserSettings.h
@@ -144,7 +144,7 @@ class UEditorPerProjectUserSettings : public UObject
 	uint32 bAllowSelectTranslucent:1;
 
 	UPROPERTY()
-	class UBlueprintPaletteFavorites* BlueprintFavorites;
+	TObjectPtr<class UBlueprintPaletteFavorites> BlueprintFavorites;
 	
 public:
 	// Per project user settings for which asset viewer profile should be used
diff --git a/Engine/Source/Editor/UnrealEd/Classes/Editor/GroupActor.h b/Engine/Source/Editor/UnrealEd/Classes/Editor/GroupActor.h
--- a/Engine/Source/Editor/UnrealEd/Classes/Editor/GroupActor.h
+++ b/Engine/Source/Editor/UnrealEd/Classes/Editor/GroupActor.h
@@ -25,10 +25,10 @@ class AGroupActor : public AActor
 	uint32 bLocked:1;
 
 	UPROPERTY()
-	TArray<class AActor*> GroupActors;
+	TArray<TObjectPtr<class AActor>> GroupActors;
 
 	UPROPERTY()
-	TArray<class AGroupActor*> SubGroups;
+	TArray<TObjectPtr<class AGroupActor>> SubGroups;
 
 #endif // WITH_EDITORONLY_DATA
 
diff --git a/Engine/Source/Editor/UnrealEd/Classes/Editor/PropertyEditorTestObject.h b/Engine/Source/Editor/UnrealEd/Classes/Editor/PropertyEditorTestObject.h
--- a/Engine/Source/Editor/UnrealEd/Classes/Editor/PropertyEditorTestObject.h
+++ b/Engine/Source/Editor/UnrealEd/Classes/Editor/PropertyEditorTestObject.h
@@ -146,7 +146,7 @@ struct FPropertyEditorTestBasicStruct
 	float FloatPropertyInsideAStruct;
 
 	UPROPERTY(EditAnywhere, Category=InnerStructCategoryWithPushedOutProps)
-	UObject* ObjectPropertyInsideAStruct;
+	TObjectPtr<UObject> ObjectPropertyInsideAStruct;
 
 	UPROPERTY(EditAnywhere, Category=InnerStructCategoryWithPushedOutProps)
 	FPropertyEditorTestSubStruct InnerStruct;
@@ -185,7 +185,7 @@ struct FPropertyEditorTestInstancedStruct
 	GENERATED_BODY()
 
 	UPROPERTY(EditAnywhere, Instanced, Category=Default)
-	UPropertyEditorTestInstancedObject* Object { nullptr };
+	TObjectPtr<UPropertyEditorTestInstancedObject> Object { nullptr };
 };
 
 UCLASS(transient)
@@ -251,7 +251,7 @@ class UPropertyEditorTestObject : public UObject
 	FRotator RotatorProperty;
 
 	UPROPERTY(EditAnywhere, Category=BasicProperties)
-	UObject* ObjectProperty;
+	TObjectPtr<UObject> ObjectProperty;
 
 	UPROPERTY(EditAnywhere, Category=BasicProperties)
 	FLinearColor LinearColorProperty;
@@ -272,13 +272,13 @@ class UPropertyEditorTestObject : public UObject
 	FTransform TransformProperty;
 
 	UPROPERTY(EditAnywhere, Category=Classes)
-	UClass* ClassProperty;
+	TObjectPtr<UClass> ClassProperty;
 
 	UPROPERTY(EditAnywhere, Category=Classes, meta=(AllowedClasses="Texture2D"))
-	UClass* ClassPropertyWithAllowed;
+	TObjectPtr<UClass> ClassPropertyWithAllowed;
 
 	UPROPERTY(EditAnywhere, Category=Classes, meta=(DisallowedClasses="Texture2D"))
-	UClass* ClassPropertyWithDisallowed;
+	TObjectPtr<UClass> ClassPropertyWithDisallowed;
 
 	UPROPERTY(EditAnywhere, Category=Classes)
 	TSubclassOf<UTexture> SubclassOfTexture;
@@ -328,10 +328,10 @@ class UPropertyEditorTestObject : public UObject
 	TArray<FRotator> RotatorPropertyArray;
 
 	UPROPERTY(EditAnywhere, Category=ArraysOfProperties)
-	TArray<UObject*> ObjectPropertyArray;
+	TArray<TObjectPtr<UObject>> ObjectPropertyArray;
 
 	UPROPERTY(EditAnywhere, Category=ArraysOfProperties)
-	TArray<AActor*> ActorPropertyArray;
+	TArray<TObjectPtr<AActor>> ActorPropertyArray;
 
 	UPROPERTY(EditAnywhere, Category=ArraysOfProperties)
 	TArray<FLinearColor> LinearColorPropertyArray;
@@ -352,7 +352,7 @@ class UPropertyEditorTestObject : public UObject
 	TArray<FPropertyEditorTestInstancedStruct> InstancedStructArray;
 
 	UPROPERTY(EditAnywhere, Instanced, Category=ArraysOfProperties)
-	TArray<UPropertyEditorTestInstancedObject*> InstancedUObjectArray;
+	TArray<TObjectPtr<UPropertyEditorTestInstancedObject>> InstancedUObjectArray;
 
 	UPROPERTY(EditAnywhere, editfixedsize, Category=ArraysOfProperties)
 	TArray<int32> FixedArrayOfInts;
@@ -380,7 +380,7 @@ class UPropertyEditorTestObject : public UObject
 	FString StringThatCannotBeChanged;
 
 	UPROPERTY(VisibleAnywhere, Category=AdvancedProperties)
-	UPrimitiveComponent* ObjectThatCannotBeChanged;
+	TObjectPtr<UPrimitiveComponent> ObjectThatCannotBeChanged;
 
 	UPROPERTY(EditAnywhere, Category=AdvancedProperties, meta=(Bitmask, BitmaskEnum="PropertyEditorTestBitflags"))
 	int32 EnumBitflags=0;
@@ -455,25 +455,25 @@ class UPropertyEditorTestObject : public UObject
 	FPropertyEditTestTextStruct Struct;
 
 	UPROPERTY(EditAnywhere, Category=EditInlineProps)
-	UStaticMeshComponent* EditInlineNewStaticMeshComponent;
+	TObjectPtr<UStaticMeshComponent> EditInlineNewStaticMeshComponent;
 
 	UPROPERTY(EditAnywhere, Category=EditInlineProps)
-	TArray<UStaticMeshComponent*> ArrayOfEditInlineNewSMCs;
+	TArray<TObjectPtr<UStaticMeshComponent>> ArrayOfEditInlineNewSMCs;
 
 	UPROPERTY(EditAnywhere, Category=AssetPropertyTests)
-	UTexture* TextureProp;
+	TObjectPtr<UTexture> TextureProp;
 
 	UPROPERTY(EditAnywhere, Category=AssetPropertyTests)
-	UStaticMesh* StaticMeshProp;
+	TObjectPtr<UStaticMesh> StaticMeshProp;
 
 	UPROPERTY(EditAnywhere, Category=AssetPropertyTests)
-	UMaterialInterface* AnyMaterialInterface;
+	TObjectPtr<UMaterialInterface> AnyMaterialInterface;
 
 	UPROPERTY(EditAnywhere, Category = AssetPropertyTests, meta=(DisplayThumbnail=false))
-	UMaterialInterface* MaterialNoThumbnail;
+	TObjectPtr<UMaterialInterface> MaterialNoThumbnail;
 
 	UPROPERTY(EditAnywhere, Category=AssetPropertyTests)
-	AActor* OnlyActorsAllowed;
+	TObjectPtr<AActor> OnlyActorsAllowed;
 
 	UPROPERTY(EditAnywhere, Category="TSet Tests")
 	TSet<int32> Int32Set;
@@ -485,10 +485,10 @@ class UPropertyEditorTestObject : public UObject
 	TSet<FString> StringSet;
 
 	UPROPERTY(EditAnywhere, Category="TSet Tests")
-	TSet<UObject*> ObjectSet;
+	TSet<TObjectPtr<UObject>> ObjectSet;
 
 	UPROPERTY(EditAnywhere, Category="TSet Tests")
-	TSet<AActor*> ActorSet;
+	TSet<TObjectPtr<AActor>> ActorSet;
 
 	UPROPERTY(EditAnywhere, Category="TSet Tests")
 	TSet<EditColor> EditColorSet;
@@ -509,16 +509,16 @@ class UPropertyEditorTestObject : public UObject
 	TMap<FString, float> StringToFloatMap;
 
 	UPROPERTY(EditAnywhere, Category="TMap Tests")
-	TMap<FString, UObject*> StringToObjectMap;
+	TMap<FString, TObjectPtr<UObject>> StringToObjectMap;
 
 	UPROPERTY(EditAnywhere, Category="TMap Tests")
-	TMap<FString, AActor*> StringToActorMap;
+	TMap<FString, TObjectPtr<AActor>> StringToActorMap;
 
 	UPROPERTY(EditAnywhere, Category="TMap Tests")
-	TMap<UObject*, int32> ObjectToInt32Map;
+	TMap<TObjectPtr<UObject>, int32> ObjectToInt32Map;
 
 	UPROPERTY(EditAnywhere, Category="TMap Tests")
-	TMap<UObject*, FLinearColor> ObjectToColorMap;
+	TMap<TObjectPtr<UObject>, FLinearColor> ObjectToColorMap;
 
 	UPROPERTY(EditAnywhere, Category="TMap Tests")
 	TMap<int32, TEnumAsByte<PropertyEditorTestEnum> > IntToEnumMap;
@@ -527,7 +527,7 @@ class UPropertyEditorTestObject : public UObject
 	TMap<FName, FName> NameToNameMap;
 
 	UPROPERTY(EditAnywhere, Category="TMap Tests")
-	TMap<FName, UObject*> NameToObjectMap;
+	TMap<FName, TObjectPtr<UObject>> NameToObjectMap;
 
 	UPROPERTY(EditAnywhere, Category="TMap Tests")
 	TMap<FName, FPropertyEditorTestBasicStruct> NameToCustomMap;
@@ -569,7 +569,7 @@ class UPropertyEditorTestObject : public UObject
 	// Allows either an object that's derived from UTexture or IBlendableInterface, to ensure that Object Property handles know how to
 	// filter for AllowedClasses correctly.
 	UPROPERTY(EditAnywhere, Category=ObjectPropertyAllowedClasses, meta=(AllowedClasses="Texture,BlendableInterface"))
-	UObject* TextureOrBlendableInterface;
+	TObjectPtr<UObject> TextureOrBlendableInterface;
 
 	UPROPERTY(EditAnywhere, Category="Subcategory")
 	bool bSubcategory;
@@ -638,7 +638,7 @@ class UPropertyEditorTestObject : public UObject
 	bool bEditConditionForArrays;
 
 	UPROPERTY(EditAnywhere, Category=EditCondition, meta=(EditCondition="bEditConditionForArrays"))
-	TArray<UTexture2D*> ArrayWithEditCondition;
+	TArray<TObjectPtr<UTexture2D>> ArrayWithEditCondition;
 
 	UPROPERTY(EditAnywhere, Category=EditCondition, meta=(EditCondition="bEditConditionForArrays"))
 	TArray<FPropertyEditorTestBasicStruct> ArrayOfStructsWithEditCondition;
@@ -739,7 +739,7 @@ class UBlueprintPropertyContainerTestObject : public UObject
 	GENERATED_BODY()
 public:
 	UPROPERTY(Instanced, EditAnywhere, Category="Default", meta=(ShowOnlyInnerProperties))
-	TArray<UBlueprintPropertyTestObject*> Array;
+	TArray<TObjectPtr<UBlueprintPropertyTestObject>> Array;
 };
 
 UCLASS(Abstract, BlueprintType, SparseClassDataType=TestSparseClassDataStorage)
diff --git a/Engine/Source/Editor/UnrealEd/Classes/Editor/UnrealEdEngine.h b/Engine/Source/Editor/UnrealEd/Classes/Editor/UnrealEdEngine.h
--- a/Engine/Source/Editor/UnrealEd/Classes/Editor/UnrealEdEngine.h
+++ b/Engine/Source/Editor/UnrealEd/Classes/Editor/UnrealEdEngine.h
@@ -87,7 +87,7 @@ struct FTemplateMapInfo
 
 	/** The Texture2D associated with this map template */
 	UPROPERTY()
-	UTexture2D* ThumbnailTexture;
+	TObjectPtr<UTexture2D> ThumbnailTexture;
 
 	/** The object path to the template map */
 	UPROPERTY(config)
@@ -111,29 +111,29 @@ public:
 
 	/** Global instance of the editor options class. */
 	UPROPERTY()
-	class UUnrealEdOptions* EditorOptionsInst;
+	TObjectPtr<class UUnrealEdOptions> EditorOptionsInst;
 
 	/**
 	 * Manager responsible for configuring auto reimport
 	 */
 	UPROPERTY()
-	class UAutoReimportManager* AutoReimportManager;
+	TObjectPtr<class UAutoReimportManager> AutoReimportManager;
 
 	/** A buffer for implementing material expression copy/paste. */
 	UPROPERTY()
-	class UMaterial* MaterialCopyPasteBuffer;
+	TObjectPtr<class UMaterial> MaterialCopyPasteBuffer;
 
 	/** A buffer for implementing matinee track/group copy/paste. */
 	UPROPERTY()
-	TArray<class UObject*> MatineeCopyPasteBuffer;
+	TArray<TObjectPtr<class UObject>> MatineeCopyPasteBuffer;
 
 	/** A buffer for implementing sound cue nodes copy/paste. */
 	UPROPERTY()
-	class USoundCue* SoundCueCopyPasteBuffer;
+	TObjectPtr<class USoundCue> SoundCueCopyPasteBuffer;
 
 	/** Global list of instanced animation compression algorithms. */
 	UPROPERTY()
-	TArray<class UAnimCompress*> AnimationCompressionAlgorithms;
+	TArray<TObjectPtr<class UAnimCompress>> AnimationCompressionAlgorithms;
 
 	/** Array of packages to be fully loaded at Editor startup. */
 	UPROPERTY(config)
@@ -141,7 +141,7 @@ public:
 
 	/** Current target for LOD parenting operations (actors will use this as the replacement) */
 	UPROPERTY()
-	class AActor* CurrentLODParentActor;
+	TObjectPtr<class AActor> CurrentLODParentActor;
 
 	/** Whether the user needs to be prompted about a package being saved with an engine version newer than the current one or not */
 	UPROPERTY()
@@ -160,7 +160,7 @@ public:
 
 	/** Cooker server incase we want to cook on the side while editing... */
 	UPROPERTY()
-	class UCookOnTheFlyServer* CookServer;
+	TObjectPtr<class UCookOnTheFlyServer> CookServer;
 
 	/** A list of packages dirtied this tick */
 	TArray<TWeakObjectPtr<UPackage>> PackagesDirtiedThisTick;
diff --git a/Engine/Source/Editor/UnrealEd/Classes/Factories/AnimBlueprintFactory.h b/Engine/Source/Editor/UnrealEd/Classes/Factories/AnimBlueprintFactory.h
--- a/Engine/Source/Editor/UnrealEd/Classes/Factories/AnimBlueprintFactory.h
+++ b/Engine/Source/Editor/UnrealEd/Classes/Factories/AnimBlueprintFactory.h
@@ -28,11 +28,11 @@ class UAnimBlueprintFactory : public UFactory
 
 	// The kind of skeleton that animation graphs compiled from the blueprint will animate
 	UPROPERTY(EditAnywhere, Category=AnimBlueprintFactory)
-	class USkeleton* TargetSkeleton;
+	TObjectPtr<class USkeleton> TargetSkeleton;
 
 	// The preview mesh to use with this animation blueprint
 	UPROPERTY(EditAnywhere, Category=AnimBlueprintFactory)
-	class USkeletalMesh* PreviewSkeletalMesh;
+	TObjectPtr<class USkeletalMesh> PreviewSkeletalMesh;
 
 	//~ Begin UFactory Interface
 	virtual bool ConfigureProperties() override;
diff --git a/Engine/Source/Editor/UnrealEd/Classes/Factories/AnimCompositeFactory.h b/Engine/Source/Editor/UnrealEd/Classes/Factories/AnimCompositeFactory.h
--- a/Engine/Source/Editor/UnrealEd/Classes/Factories/AnimCompositeFactory.h
+++ b/Engine/Source/Editor/UnrealEd/Classes/Factories/AnimCompositeFactory.h
@@ -20,15 +20,15 @@ class UAnimCompositeFactory : public UFactory
 	GENERATED_UCLASS_BODY()
 
 	UPROPERTY()
-	class USkeleton* TargetSkeleton;
+	TObjectPtr<class USkeleton> TargetSkeleton;
 
 	/* Used when creating a composite from an AnimSequence, becomes the only AnimSequence contained */
 	UPROPERTY()
-	class UAnimSequence* SourceAnimation;
+	TObjectPtr<class UAnimSequence> SourceAnimation;
 
 	/** The preview mesh to use with this animation */
 	UPROPERTY()
-	class USkeletalMesh* PreviewSkeletalMesh;
+	TObjectPtr<class USkeletalMesh> PreviewSkeletalMesh;
 
 	//~ Begin UFactory Interface
 	virtual bool ConfigureProperties() override;
diff --git a/Engine/Source/Editor/UnrealEd/Classes/Factories/AnimMontageFactory.h b/Engine/Source/Editor/UnrealEd/Classes/Factories/AnimMontageFactory.h
--- a/Engine/Source/Editor/UnrealEd/Classes/Factories/AnimMontageFactory.h
+++ b/Engine/Source/Editor/UnrealEd/Classes/Factories/AnimMontageFactory.h
@@ -21,15 +21,15 @@ class UAnimMontageFactory : public UFactory
 	GENERATED_UCLASS_BODY()
 
 	UPROPERTY()
-	class USkeleton* TargetSkeleton;
+	TObjectPtr<class USkeleton> TargetSkeleton;
 
 	/* Used when creating a montage from an AnimSequence, becomes the only AnimSequence contained */
 	UPROPERTY()
-	class UAnimSequence* SourceAnimation;
+	TObjectPtr<class UAnimSequence> SourceAnimation;
 
 	/** The preview mesh to use with this animation */
 	UPROPERTY()
-	class USkeletalMesh* PreviewSkeletalMesh;
+	TObjectPtr<class USkeletalMesh> PreviewSkeletalMesh;
 
 	//~ Begin UFactory Interface
 	virtual bool ConfigureProperties() override;
diff --git a/Engine/Source/Editor/UnrealEd/Classes/Factories/AnimSequenceFactory.h b/Engine/Source/Editor/UnrealEd/Classes/Factories/AnimSequenceFactory.h
--- a/Engine/Source/Editor/UnrealEd/Classes/Factories/AnimSequenceFactory.h
+++ b/Engine/Source/Editor/UnrealEd/Classes/Factories/AnimSequenceFactory.h
@@ -20,11 +20,11 @@ class UNREALED_API UAnimSequenceFactory : public UFactory
 	GENERATED_UCLASS_BODY()
 
 	UPROPERTY()
-	class USkeleton* TargetSkeleton;
+	TObjectPtr<class USkeleton> TargetSkeleton;
 
 	/** The preview mesh to use with this animation */
 	UPROPERTY()
-	class USkeletalMesh* PreviewSkeletalMesh;
+	TObjectPtr<class USkeletalMesh> PreviewSkeletalMesh;
 
 	//~ Begin UFactory Interface
 	virtual bool ConfigureProperties() override;
diff --git a/Engine/Source/Editor/UnrealEd/Classes/Factories/AnimStreamableFactory.h b/Engine/Source/Editor/UnrealEd/Classes/Factories/AnimStreamableFactory.h
--- a/Engine/Source/Editor/UnrealEd/Classes/Factories/AnimStreamableFactory.h
+++ b/Engine/Source/Editor/UnrealEd/Classes/Factories/AnimStreamableFactory.h
@@ -20,11 +20,11 @@ class UAnimStreamableFactory : public UFactory
 	GENERATED_UCLASS_BODY()
 
 	UPROPERTY()
-	class USkeleton* TargetSkeleton;
+	TObjectPtr<class USkeleton> TargetSkeleton;
 
 	/* Used when creating a composite from an AnimSequence, becomes the only AnimSequence contained */
 	UPROPERTY()
-	class UAnimSequence* SourceAnimation;
+	TObjectPtr<class UAnimSequence> SourceAnimation;
 
 	//~ Begin UFactory Interface
 	virtual bool ConfigureProperties() override;
diff --git a/Engine/Source/Editor/UnrealEd/Classes/Factories/BlendSpaceFactory1D.h b/Engine/Source/Editor/UnrealEd/Classes/Factories/BlendSpaceFactory1D.h
--- a/Engine/Source/Editor/UnrealEd/Classes/Factories/BlendSpaceFactory1D.h
+++ b/Engine/Source/Editor/UnrealEd/Classes/Factories/BlendSpaceFactory1D.h
@@ -20,11 +20,11 @@ class UBlendSpaceFactory1D : public UFactory
 	GENERATED_UCLASS_BODY()
 
 	UPROPERTY()
-	class USkeleton*	TargetSkeleton;
+	TObjectPtr<class USkeleton>	TargetSkeleton;
 
 	/** The preview mesh to use with this animation */
 	UPROPERTY()
-	class USkeletalMesh* PreviewSkeletalMesh;
+	TObjectPtr<class USkeletalMesh> PreviewSkeletalMesh;
 
 	//~ Begin UFactory Interface
 	virtual bool ConfigureProperties() override;
diff --git a/Engine/Source/Editor/UnrealEd/Classes/Factories/BlendSpaceFactoryNew.h b/Engine/Source/Editor/UnrealEd/Classes/Factories/BlendSpaceFactoryNew.h
--- a/Engine/Source/Editor/UnrealEd/Classes/Factories/BlendSpaceFactoryNew.h
+++ b/Engine/Source/Editor/UnrealEd/Classes/Factories/BlendSpaceFactoryNew.h
@@ -20,11 +20,11 @@ class UBlendSpaceFactoryNew : public UFactory
 	GENERATED_UCLASS_BODY()
 
 	UPROPERTY()
-	class USkeleton*	TargetSkeleton;
+	TObjectPtr<class USkeleton>	TargetSkeleton;
 
 	/** The preview mesh to use with this animation */
 	UPROPERTY()
-	class USkeletalMesh* PreviewSkeletalMesh;
+	TObjectPtr<class USkeletalMesh> PreviewSkeletalMesh;
 
 	//~ Begin UFactory Interface
 	virtual bool ConfigureProperties() override;
diff --git a/Engine/Source/Editor/UnrealEd/Classes/Factories/CSVImportFactory.h b/Engine/Source/Editor/UnrealEd/Classes/Factories/CSVImportFactory.h
--- a/Engine/Source/Editor/UnrealEd/Classes/Factories/CSVImportFactory.h
+++ b/Engine/Source/Editor/UnrealEd/Classes/Factories/CSVImportFactory.h
@@ -40,7 +40,7 @@ struct FCSVImportSettings
 	FCSVImportSettings();
 
 	UPROPERTY(BlueprintReadWrite, Category="Misc")
-	UScriptStruct* ImportRowStruct;
+	TObjectPtr<UScriptStruct> ImportRowStruct;
 
 	UPROPERTY(BlueprintReadWrite, Category="Misc")
 	ECSVImportType ImportType;
@@ -91,6 +91,6 @@ public:
 
 	/** Temporary data table to use to display import options */
 	UPROPERTY()
-	UDataTable* DataTableImportOptions;
+	TObjectPtr<UDataTable> DataTableImportOptions;
 };
 
diff --git a/Engine/Source/Editor/UnrealEd/Classes/Factories/DataTableFactory.h b/Engine/Source/Editor/UnrealEd/Classes/Factories/DataTableFactory.h
--- a/Engine/Source/Editor/UnrealEd/Classes/Factories/DataTableFactory.h
+++ b/Engine/Source/Editor/UnrealEd/Classes/Factories/DataTableFactory.h
@@ -15,7 +15,7 @@ class UNREALED_API UDataTableFactory : public UFactory
 	GENERATED_UCLASS_BODY()
 
 	UPROPERTY(BlueprintReadWrite, Category = "Data Table Factory")
-	const class UScriptStruct* Struct;
+	TObjectPtr<const class UScriptStruct> Struct;
 
 	//~ Begin UFactory Interface
 	virtual bool ConfigureProperties() override;
diff --git a/Engine/Source/Editor/UnrealEd/Classes/Factories/Factory.h b/Engine/Source/Editor/UnrealEd/Classes/Factories/Factory.h
--- a/Engine/Source/Editor/UnrealEd/Classes/Factories/Factory.h
+++ b/Engine/Source/Editor/UnrealEd/Classes/Factories/Factory.h
@@ -386,11 +386,11 @@ public:
 
 	/** Data for how to import files via the automated command line importing interface */
 	UPROPERTY(BlueprintReadWrite, Category=Misc)
-	const class UAutomatedAssetImportData* AutomatedImportData;
+	TObjectPtr<const class UAutomatedAssetImportData> AutomatedImportData;
 
 	/** Task for importing file via script interfaces */
 	UPROPERTY(BlueprintReadWrite, Category=Misc)
-	class UAssetImportTask* AssetImportTask;
+	TObjectPtr<class UAssetImportTask> AssetImportTask;
 
 protected:
 
diff --git a/Engine/Source/Editor/UnrealEd/Classes/Factories/FbxAssetImportData.h b/Engine/Source/Editor/UnrealEd/Classes/Factories/FbxAssetImportData.h
--- a/Engine/Source/Editor/UnrealEd/Classes/Factories/FbxAssetImportData.h
+++ b/Engine/Source/Editor/UnrealEd/Classes/Factories/FbxAssetImportData.h
@@ -44,5 +44,5 @@ class UNREALED_API UFbxAssetImportData : public UAssetImportData
 
 	/* Use by the reimport factory to answer CanReimport, if true only factory for scene reimport will return true */
 	UPROPERTY()
-	UFbxSceneImportData* FbxSceneImportDataReference;
+	TObjectPtr<UFbxSceneImportData> FbxSceneImportDataReference;
 };
diff --git a/Engine/Source/Editor/UnrealEd/Classes/Factories/FbxFactory.h b/Engine/Source/Editor/UnrealEd/Classes/Factories/FbxFactory.h
--- a/Engine/Source/Editor/UnrealEd/Classes/Factories/FbxFactory.h
+++ b/Engine/Source/Editor/UnrealEd/Classes/Factories/FbxFactory.h
@@ -19,11 +19,11 @@ class UNREALED_API UFbxFactory : public UFactory
 	GENERATED_UCLASS_BODY()
 
 	UPROPERTY()
-	class UFbxImportUI* ImportUI;
+	TObjectPtr<class UFbxImportUI> ImportUI;
 
 	/** Prevent garbage collection of original when overriding ImportUI property */
 	UPROPERTY()
-	class UFbxImportUI* OriginalImportUI;
+	TObjectPtr<class UFbxImportUI> OriginalImportUI;
 
 	/**  Set import batch **/
 	void EnableShowOption() { bShowOption = true; }
