CIS: Fix a msvc static analysis warning

- It looks like the compiler is getting confused with the scope of the local variable Attrs declared  in both ::RegisterCookStats and the template method FCookStatsManager::CreateKeyValueArray which is calls.
- The local in RegisterCookStats was not being used anyway so we can just remove this to fix the warning.
- The warning only showed up on older tool chains (14.24.28314) and looks like a false positive that was fixed with later releases.

#rnx

diff --git a/Engine/Source/Editor/UnrealEd/Private/Commandlets/CookCommandlet.cpp b/Engine/Source/Editor/UnrealEd/Private/Commandlets/CookCommandlet.cpp
--- a/Engine/Source/Editor/UnrealEd/Private/Commandlets/CookCommandlet.cpp
+++ b/Engine/Source/Editor/UnrealEd/Private/Commandlets/CookCommandlet.cpp
@@ -79,7 +79,6 @@ namespace DetailedCookStats
 	FCookStatsManager::FAutoRegisterCallback RegisterCookStats([](FCookStatsManager::AddStatFuncRef AddStat)
 	{
 		const FString StatName(TEXT("Cook.Profile"));
-		TArray<FCookStatsManager::StringKeyValue> Attrs;
 		#define ADD_COOK_STAT_FLT(Path, Name) AddStat(StatName, FCookStatsManager::CreateKeyValueArray(TEXT("Path"), TEXT(Path), TEXT(#Name), Name))
 		ADD_COOK_STAT_FLT(" 0", CookWallTimeSec);
 		ADD_COOK_STAT_FLT(" 0. 0", StartupWallTimeSec);
