Implemented a function to copy a list of actors from one level to another.

diff --git a/Engine/Source/Editor/UnrealEd/Private/EditorLevelUtils.cpp b/Engine/Source/Editor/UnrealEd/Private/EditorLevelUtils.cpp
--- a/Engine/Source/Editor/UnrealEd/Private/EditorLevelUtils.cpp
+++ b/Engine/Source/Editor/UnrealEd/Private/EditorLevelUtils.cpp
@@ -65,6 +65,20 @@ int32 UEditorLevelUtils::MoveActorsToLevel(const TArray<AActor*>& ActorsToMove,
 }
 
 int32 UEditorLevelUtils::MoveActorsToLevel(const TArray<AActor*>& ActorsToMove, ULevel* DestLevel, bool bWarnAboutReferences, bool bWarnAboutRenaming, bool bMoveAllOrFail)
+{
+	const bool bCutActors = true;
+	const bool bMoveActors = true;
+	return CopyOrMoveActorsToLevel(ActorsToMove, DestLevel, bCutActors, bMoveActors, bWarnAboutReferences, bWarnAboutRenaming, bMoveAllOrFail);
+}
+
+int32 UEditorLevelUtils::CopyActorsToLevel(const TArray<AActor*>& ActorsToMove, ULevel* DestLevel, bool bWarnAboutReferences, bool bWarnAboutRenaming, bool bMoveAllOrFail)
+{
+	const bool bCutActors = false;
+	const bool bMoveActors = false;
+	return CopyOrMoveActorsToLevel(ActorsToMove, DestLevel, bCutActors, bMoveActors, bWarnAboutReferences, bWarnAboutRenaming, bMoveAllOrFail);
+}
+
+int32 UEditorLevelUtils::CopyOrMoveActorsToLevel(const TArray<AActor*>& ActorsToMove, ULevel* DestLevel, bool bCutActors, bool bMoveActors, bool bWarnAboutReferences /*= true*/, bool bWarnAboutRenaming /*= true*/, bool bMoveAllOrFail /*= false*/)
 {
 	int32 NumMovedActors = 0;
 
@@ -144,11 +158,8 @@ int32 UEditorLevelUtils::MoveActorsToLevel(const TArray<AActor*>& ActorsToMove,
 				// Cache the old level
 				ULevel* OldCurrentLevel = OwningWorld->GetCurrentLevel();
 
-				// We are moving the actors so cut them to remove them from the existing level
-				const bool bShoudCut = true;
-				const bool bIsMove = true;
 				FString DestinationData;
-				GEditor->CopySelectedActorsToClipboard(OwningWorld, bShoudCut, bIsMove, bWarnAboutReferences, &DestinationData);
+				GEditor->CopySelectedActorsToClipboard(OwningWorld, bCutActors, bMoveActors, bWarnAboutReferences, &DestinationData);
 
 				const bool bLevelVisible = DestLevel->bIsVisible;
 				if (!bLevelVisible)
@@ -161,7 +172,7 @@ int32 UEditorLevelUtils::MoveActorsToLevel(const TArray<AActor*>& ActorsToMove,
 					// Set the new level and force it visible while we do the paste
 					FLevelPartitionOperationScope LevelPartitionScope(DestLevel);
 					OwningWorld->SetCurrentLevel(LevelPartitionScope.GetLevel());
-										
+
 					const bool bDuplicate = false;
 					const bool bOffsetLocations = false;
 					const bool bWarnIfHidden = false;
@@ -229,10 +240,10 @@ int32 UEditorLevelUtils::MoveActorsToLevel(const TArray<AActor*>& ActorsToMove,
 						RenameData.Add(FAssetRenameData(Pair.Key, Pair.Value, true));
 					}
 				}
-					
+
 				if (RenameData.Num() > 0)
 				{
-					if(bWarnAboutRenaming)
+					if (bWarnAboutRenaming)
 					{
 						AssetToolsModule.Get().RenameAssetsWithDialog(RenameData);
 					}
diff --git a/Engine/Source/Editor/UnrealEd/Public/EditorLevelUtils.h b/Engine/Source/Editor/UnrealEd/Public/EditorLevelUtils.h
--- a/Engine/Source/Editor/UnrealEd/Public/EditorLevelUtils.h
+++ b/Engine/Source/Editor/UnrealEd/Public/EditorLevelUtils.h
@@ -88,6 +88,8 @@ public:
 	
 	static UNREALED_API int32 MoveActorsToLevel(const TArray<AActor*>& ActorsToMove, ULevel* DestLevel, bool bWarnAboutReferences = true, bool bWarnAboutRenaming = true, bool bMoveAllOrFail = false);
 
+	static UNREALED_API int32 CopyActorsToLevel(const TArray<AActor*>& ActorsToMove, ULevel* DestLevel, bool bWarnAboutReferences = true, bool bWarnAboutRenaming = true, bool bMoveAllOrFail = false);
+
 	static UNREALED_API int32 MoveSelectedActorsToLevel(ULevel* DestLevel, bool bWarnAboutReferences = true);
 
 	/** 
@@ -167,6 +169,8 @@ private:
 
 	static UNREALED_API ULevelStreaming* AddLevelToWorld_Internal(UWorld* InWorld, const TCHAR* LevelPackageName, TSubclassOf<ULevelStreaming> LevelStreamingClass, const FTransform& LevelTransform = FTransform::Identity);
 
+	static UNREALED_API int32 CopyOrMoveActorsToLevel(const TArray<AActor*>& ActorsToMove, ULevel* DestLevel, bool bCutActors, bool bMoveActors, bool bWarnAboutReferences = true, bool bWarnAboutRenaming = true, bool bMoveAllOrFail = false);
+
 public:
 	/** Sets the LevelStreamingClass for the specified Level 
 	  * @param	InLevel				The level for which to change the streaming class
diff --git a/Engine/Source/Runtime/Engine/Private/LevelInstance/LevelInstanceSubsystem.cpp b/Engine/Source/Runtime/Engine/Private/LevelInstance/LevelInstanceSubsystem.cpp
--- a/Engine/Source/Runtime/Engine/Private/LevelInstance/LevelInstanceSubsystem.cpp
+++ b/Engine/Source/Runtime/Engine/Private/LevelInstance/LevelInstanceSubsystem.cpp
@@ -809,7 +809,7 @@ bool ULevelInstanceSubsystem::BreakLevelInstance(ALevelInstance* LevelInstanceAc
 		const bool bWarnAboutReferences = true;
 		const bool bWarnAboutRenaming = true;
 		const bool bMoveAllOrFail = true;
-		if (!EditorLevelUtils::MoveActorsToLevel(ActorsToMove, DestinationLevel, bWarnAboutReferences, bWarnAboutRenaming, bMoveAllOrFail))
+		if (!EditorLevelUtils::CopyActorsToLevel(ActorsToMove, DestinationLevel, bWarnAboutReferences, bWarnAboutRenaming, bMoveAllOrFail))
 		{
 			UE_LOG(LogLevelInstance, Warning, TEXT("Failed to break Level Instance because not all actors could be moved"));
 			return false;
