When reinstancing an external actor, make sure to set the old actor internal so the newwly created actor can reuse the old external package to avoid having to submit a new file each time we recompile a BP.

diff --git a/Engine/Source/Editor/UnrealEd/Private/Kismet2/KismetReinstanceUtilities.cpp b/Engine/Source/Editor/UnrealEd/Private/Kismet2/KismetReinstanceUtilities.cpp
--- a/Engine/Source/Editor/UnrealEd/Private/Kismet2/KismetReinstanceUtilities.cpp
+++ b/Engine/Source/Editor/UnrealEd/Private/Kismet2/KismetReinstanceUtilities.cpp
@@ -1928,14 +1928,20 @@ static void ReplaceActorHelper(AActor* OldActor, UClass* OldClass, UObject*& New
 	SpawnInfo.SpawnCollisionHandlingOverride = ESpawnActorCollisionHandlingMethod::AlwaysSpawn;
 	SpawnInfo.bDeferConstruction = true;
 	SpawnInfo.Name = OldActor->GetFName();
-	SpawnInfo.bCreateActorPackage = false;
 
 	if (!OldActor->IsListedInSceneOutliner())
 	{
 		SpawnInfo.bHideFromSceneOutliner = true;
 	}
 
+	// Make sure to reuse the same external package if any
+	SpawnInfo.bCreateActorPackage = false;
 	SpawnInfo.OverridePackage = OldActor->GetExternalPackage();
+	if (SpawnInfo.OverridePackage)
+	{
+		OldActor->SetPackageExternal(false, false);
+	}
+
 	SpawnInfo.OverrideActorGuid = OldActor->GetActorGuid();
 
 	OldActor->Rename(nullptr, OldActor->GetOuter(), REN_DoNotDirty | REN_DontCreateRedirectors | REN_ForceNoResetLoaders);
