Viewport Clients that use LVT_None no longer crash when used in conjunction with mode tools and ITF.

diff --git a/Engine/Source/Editor/UnrealEd/Private/EditorViewportClient.cpp b/Engine/Source/Editor/UnrealEd/Private/EditorViewportClient.cpp
--- a/Engine/Source/Editor/UnrealEd/Private/EditorViewportClient.cpp
+++ b/Engine/Source/Editor/UnrealEd/Private/EditorViewportClient.cpp
@@ -1649,6 +1649,10 @@ bool FEditorViewportClient::IsAspectRatioConstrained() const
 ELevelViewportType FEditorViewportClient::GetViewportType() const
 {
 	ELevelViewportType EffectiveViewportType = ViewportType;
+	if (EffectiveViewportType == LVT_None)
+	{
+		EffectiveViewportType = LVT_Perspective;
+	}
 	if (bUseControllingActorViewInfo)
 	{
 		EffectiveViewportType = (ControllingActorViewInfo.ProjectionMode == ECameraProjectionMode::Perspective) ? LVT_Perspective : LVT_OrthoFreelook;
@@ -4767,7 +4771,7 @@ void FEditorViewportClient::UpdateRequiredCursorVisibility()
 		}
 	}
 
-	if (GetViewportType() == LVT_None)
+	if (ViewportType == LVT_None)
 	{
 		SetRequiredCursor(true, false);
 		return;
