Moving / removing ticking calls for MediaFramework to ensure proper execution ordering when movie queue rendering is active
#rnx

diff --git a/Engine/Source/Editor/UnrealEd/Private/EditorEngine.cpp b/Engine/Source/Editor/UnrealEd/Private/EditorEngine.cpp
--- a/Engine/Source/Editor/UnrealEd/Private/EditorEngine.cpp
+++ b/Engine/Source/Editor/UnrealEd/Private/EditorEngine.cpp
@@ -1501,7 +1501,13 @@ void UEditorEngine::Tick( float DeltaSeconds, bool bIdleMode )
 		DirectoryWatcherModule.Get()->Tick(DeltaSeconds);
 	}
 
+#if !UE_SERVER
+	static const FName MediaModuleName(TEXT("Media"));
+	IMediaModule* MediaModule = FModuleManager::LoadModulePtr<IMediaModule>(MediaModuleName);
+#endif
+
 	bool bAWorldTicked = false;
+	bool bMediaModulePreEngineTickDone = false;
 	ELevelTick TickType = IsRealtime ? LEVELTICK_ViewportsOnly : LEVELTICK_TimeOnly;
 
 	if( bShouldTickEditorWorld )
@@ -1514,6 +1520,9 @@ void UEditorEngine::Tick( float DeltaSeconds, bool bIdleMode )
 			EditorContext.World()->Tick(TickType, DeltaSeconds);
 			bAWorldTicked = true;
 			FKismetDebugUtilities::NotifyDebuggerOfEndOfGameFrame(EditorContext.World());
+
+			// Track if the editor context has an active movie sequence tick associated -> if so it will have executed the MediaModule pre-engine tick already
+			bMediaModulePreEngineTickDone = (MediaModule != nullptr) && EditorContext.World()->IsMovieSceneSequenceTickHandlerBound();
 		}
 	}
 
@@ -1629,16 +1638,26 @@ void UEditorEngine::Tick( float DeltaSeconds, bool bIdleMode )
 	// if we have the side-by-side world for "Play From Here", tick it unless we are ensuring slate is responsive
 	if( FSlateThrottleManager::Get().IsAllowingExpensiveTasks() )
 	{
-		// Determine number of PIE worlds that should tick.
+		// Determine number of PIE worlds that should tick and if they feature an active movie sequence tick.
+		bool bMovieSequenceTickWillBeHandled = false;
 		TArray<FWorldContext*> LocalPieContextPtrs;
 		for (FWorldContext& PieContext : WorldList)
 		{
 			if (PieContext.WorldType == EWorldType::PIE && PieContext.World() != nullptr && PieContext.World()->ShouldTick())
 			{
 				LocalPieContextPtrs.Add(&PieContext);
+				bMovieSequenceTickWillBeHandled |= (MediaModule != nullptr) && PieContext.World()->IsMovieSceneSequenceTickHandlerBound();
 			}
 		}
 
+#if !UE_SERVER
+		// If we do not have any active movie sequence ticking, we will issue the pre-engine tick for the MediaModule as soon as we can
+		if ((MediaModule != nullptr) && !bMediaModulePreEngineTickDone && !bMovieSequenceTickWillBeHandled)
+		{
+			MediaModule->TickPreEngine();
+		}
+#endif
+
 		// Note: WorldList can change size within this loop during PIE when stopped at a breakpoint. In that case, we are
 		// running in a nested tick loop within this loop, where a new editor window with a preview viewport can be opened.
 		// So we iterate on a local list here instead.
@@ -1726,20 +1745,30 @@ void UEditorEngine::Tick( float DeltaSeconds, bool bIdleMode )
 			RestoreEditorWorld( OldGWorld );
 		}
 	}
+	else
+	{
+#if !UE_SERVER
+	// If we do not tick anything of the above, we need to still make sure to issue the pre-tick to MediaModule
+		// (unless the Editor context had a movie sequence tick active and hence already did so)
+		if ((MediaModule != nullptr) && !bMediaModulePreEngineTickDone)
+		{
+			MediaModule->TickPreEngine();
+		}
+#endif
+	}
 
 	if (bAWorldTicked)
 	{
 		FTickableGameObject::TickObjects(nullptr, TickType, false, DeltaSeconds);
 	}
 
-	// tick media framework
-	static const FName MediaModuleName(TEXT("Media"));
-	IMediaModule* MediaModule = FModuleManager::LoadModulePtr<IMediaModule>(MediaModuleName);
-
+#if !UE_SERVER
+	// tick media framework post engine ticks
 	if (MediaModule != nullptr)
 	{
 		MediaModule->TickPostEngine();
 	}
+#endif
 
 	if (bFirstTick)
 	{
diff --git a/Engine/Source/Runtime/Engine/Classes/Engine/World.h b/Engine/Source/Runtime/Engine/Classes/Engine/World.h
--- a/Engine/Source/Runtime/Engine/Classes/Engine/World.h
+++ b/Engine/Source/Runtime/Engine/Classes/Engine/World.h
@@ -2588,6 +2588,8 @@ public:
 	FDelegateHandle AddMovieSceneSequenceTickHandler(const FOnMovieSceneSequenceTick::FDelegate& InHandler);
 	/** Removes a tick handler for sequences */
 	void RemoveMovieSceneSequenceTickHandler(FDelegateHandle InHandle);
+	/** Check if movie sequences tick handler is bound at all */
+	bool IsMovieSceneSequenceTickHandlerBound() const;
 
 	//~ Begin UObject Interface
 	virtual void Serialize( FArchive& Ar ) override;
diff --git a/Engine/Source/Runtime/Engine/Private/LevelTick.cpp b/Engine/Source/Runtime/Engine/Private/LevelTick.cpp
--- a/Engine/Source/Runtime/Engine/Private/LevelTick.cpp
+++ b/Engine/Source/Runtime/Engine/Private/LevelTick.cpp
@@ -44,6 +44,9 @@
 #if ENABLE_COLLISION_ANALYZER
 #include "PhysicsEngine/CollisionAnalyzerCapture.h"
 #endif
+#if !UE_SERVER
+#include "IMediaModule.h"
+#endif
 
 //#include "SoundDefinitions.h"
 #include "FXSystem.h"
@@ -1452,6 +1455,20 @@ void UWorld::Tick( ELevelTick TickType, float DeltaSeconds )
 		MovieSceneSequenceTick.Broadcast(DeltaSeconds);
 	}
 
+#if !UE_SERVER
+	if (MovieSceneSequenceTick.IsBound())
+	{
+		// tick media framework pre-engine
+		// (needs to be delayed post movie scene tick delegates handling, sif they are used - otherwise is triggered elsewhere)
+		static const FName MediaModuleName(TEXT("Media"));
+		IMediaModule* MediaModule = FModuleManager::LoadModulePtr<IMediaModule>(MediaModuleName);
+		if (MediaModule != nullptr)
+		{
+			MediaModule->TickPreEngine();
+		}
+	}
+#endif
+
 	// If only the DynamicLevel collection has entries, we can skip the validation and tick all levels.
 	bool bValidateLevelList = false;
 
diff --git a/Engine/Source/Runtime/Engine/Private/World.cpp b/Engine/Source/Runtime/Engine/Private/World.cpp
--- a/Engine/Source/Runtime/Engine/Private/World.cpp
+++ b/Engine/Source/Runtime/Engine/Private/World.cpp
@@ -4972,6 +4972,11 @@ void UWorld::RemoveMovieSceneSequenceTickHandler(FDelegateHandle InHandle)
 	MovieSceneSequenceTick.Remove(InHandle);
 }
 
+bool UWorld::IsMovieSceneSequenceTickHandlerBound() const
+{
+	return MovieSceneSequenceTick.IsBound();
+}
+
 ABrush* UWorld::GetDefaultBrush() const
 {
 	check(PersistentLevel);
diff --git a/Engine/Source/Runtime/Launch/Private/LaunchEngineLoop.cpp b/Engine/Source/Runtime/Launch/Private/LaunchEngineLoop.cpp
--- a/Engine/Source/Runtime/Launch/Private/LaunchEngineLoop.cpp
+++ b/Engine/Source/Runtime/Launch/Private/LaunchEngineLoop.cpp
@@ -5042,17 +5042,6 @@ void FEngineLoop::Tick()
             }
 		}
 
-#if !UE_SERVER
-		// tick media framework
-		static const FName MediaModuleName(TEXT("Media"));
-		IMediaModule* MediaModule = FModuleManager::LoadModulePtr<IMediaModule>(MediaModuleName);
-
-		if (MediaModule != nullptr)
-		{
-			MediaModule->TickPreEngine();
-		}
-#endif
-
 		// main game engine tick (world, game objects, etc.)
 		GEngine->Tick(FApp::GetDeltaTime(), bIdleMode);
 
@@ -5114,14 +5103,6 @@ void FEngineLoop::Tick()
 			FAssetCompilingManager::Get().ProcessAsyncTasks(true);
 #endif
 		}
-#if !UE_SERVER
-		// tick media framework
-		if (MediaModule != nullptr)
-		{
-			MediaModule->TickPreSlate();
-		}
-#endif
-
 		// Tick the platform and input portion of Slate application, we need to do this before we run things
 		// concurrent with networking.
 		if (FSlateApplication::IsInitialized() && !bIdleMode)
@@ -5288,6 +5269,8 @@ void FEngineLoop::Tick()
 
 #if !UE_SERVER
 		// tick media framework
+		static const FName MediaModuleName(TEXT("Media"));
+		IMediaModule* MediaModule = FModuleManager::LoadModulePtr<IMediaModule>(MediaModuleName);
 		if (MediaModule != nullptr)
 		{
 			QUICK_SCOPE_CYCLE_COUNTER(STAT_FEngineLoop_MediaTickPostRender);
