Split edactDeleteSelectedActors to avoid dependency on the global selection

diff --git a/Engine/Source/Editor/UnrealEd/Classes/Editor/UnrealEdEngine.h b/Engine/Source/Editor/UnrealEd/Classes/Editor/UnrealEdEngine.h
--- a/Engine/Source/Editor/UnrealEd/Classes/Editor/UnrealEdEngine.h
+++ b/Engine/Source/Editor/UnrealEd/Classes/Editor/UnrealEdEngine.h
@@ -412,10 +412,68 @@ public:
 	 */
 	virtual void edactSelectRelevantLights( UWorld* InWorld );
 
+	/**
+	 * Can the given component be deleted?
+	 * 
+	 * @param InComponent				Component to check
+	 * @param OutReason					Optional value to fill with the reason the component cannot be deleted, if any
+	 */
+	bool CanDeleteComponent(const UActorComponent* InComponent, FText* OutReason = nullptr) const;
+
+	/**
+	 * Can the given actor be deleted?
+	 *
+	 * @param InActor					Actor to check
+	 * @param OutReason					Optional value to fill with the reason the actor cannot be deleted, if any
+	 */
+	bool CanDeleteActor(const AActor* InActor, FText* OutReason = nullptr) const;
+
+	/**
+	 * Should the deletion of the given components be outright aborted?
+	 *
+	 * @param InComponentsToDelete		Components to check
+	 * @param OutReason					Optional value to fill with the reason the component deletion was aborted, if any
+	 */
+	bool ShouldAbortComponentDeletion(const TArray<UActorComponent*>& InComponentsToDelete, FText* OutReason = nullptr) const;
+
+	/**
+	 * Should the deletion of the given actors be outright aborted?
+	 *
+	 * @param InActorsToDelete			Actors to check
+	 * @param OutReason					Optional value to fill with the reason the actor deletion was aborted, if any
+	 */
+	bool ShouldAbortActorDeletion(const TArray<AActor*>& InActorsToDelete, FText* OutReason = nullptr) const;
+
+	/**
+	 * Delete the given components.
+	 *
+	 * @param	InComponentsToDelete		Array of components to delete
+	 * @param	InSelectionSet				The selection set potentially containing to components that are being deleted
+	 * @param	OutSuggestedNewSelection	Array to fill with suitable components to select post-delete
+	 * @param	bVerifyDeletionCanHappen	If true (default), verify that deletion can be performed
+	 * 
+	 * @return								true unless the delete operation was aborted.
+	 */
+	bool DeleteComponents(const TArray<UActorComponent*>& InComponentsToDelete, UTypedElementSelectionSet* InSelectionSet, const bool bVerifyDeletionCanHappen = true);
+
+	/**
+	 * Deletes the given actors.
+	 *
+	 * @param	InActorsToDelete			Array of actors to delete
+	 * @param	InWorld						World context
+	 * @param	InSelectionSet				The selection set potentially containing to actors that are being deleted
+	 * @param	bVerifyDeletionCanHappen	If true (default), verify that deletion can be performed
+	 * @param	bWarnAboutReferences		If true (default), we prompt the user about referenced actors they are about to delete
+	 * @param	bWarnAboutSoftReferences	If true (default), we prompt the user about soft references to actors they are about to delete
+	 * 
+	 * @return								true unless the delete operation was aborted.
+	 */
+	bool DeleteActors(const TArray<AActor*>& InActorsToDelete, UWorld* InWorld, UTypedElementSelectionSet* InSelectionSet, const bool bVerifyDeletionCanHappen = true, const bool bWarnAboutReferences = true, const bool bWarnAboutSoftReferences = true);
+
 	/**
 	 * Deletes all selected actors
 	 *
-	 * @param	InWorld				World context
+	 * @param	InWorld						World context
 	 * @param	bVerifyDeletionCanHappen	[opt] If true (default), verify that deletion can be performed.
 	 * @param	bWarnAboutReferences		[opt] If true (default), we prompt the user about referenced actors they are about to delete
 	 * @param	bWarnAboutSoftReferences	[opt] If true (default), we prompt the user about soft references to actors they are about to delete
diff --git a/Engine/Source/Editor/UnrealEd/Private/EditorActor.cpp b/Engine/Source/Editor/UnrealEd/Private/EditorActor.cpp
--- a/Engine/Source/Editor/UnrealEd/Private/EditorActor.cpp
+++ b/Engine/Source/Editor/UnrealEd/Private/EditorActor.cpp
@@ -66,6 +66,10 @@
 #include "AssetToolsModule.h"
 #include "AssetSelection.h"
 #include "Framework/Application/SlateApplication.h"
+#include "Framework/Notifications/NotificationManager.h"
+#include "Widgets/Notifications/SNotificationList.h"
+#include "Elements/Framework/EngineElementsLibrary.h"
+#include "Elements/Framework/TypedElementSelectionSet.h"
 #include "EdMode.h"
 #include "Subsystems/BrushEditingSubsystem.h"
 #include "Misc/ScopedSlowTask.h"
@@ -613,14 +617,6 @@ void UUnrealEdEngine::DuplicateActors(const TArray<AActor*>& InActorsToDuplicate
 
 bool UUnrealEdEngine::CanDeleteSelectedActors( const UWorld* InWorld, const bool bStopAtFirst, const bool bLogUndeletable, TArray<AActor*>* OutDeletableActors ) const
 {
-	// Iterate over all levels and create a list of world infos.
-	TArray<AWorldSettings*> WorldSettingss;
-	for ( int32 LevelIndex = 0 ; LevelIndex < InWorld->GetNumLevels() ; ++LevelIndex )
-	{
-		ULevel* Level = InWorld->GetLevel( LevelIndex );
-		WorldSettingss.Add( Level->GetWorldSettings() );
-	}
-
 	// Iterate over selected actors and assemble a list of actors to delete.
 	bool bContainsDeletable = false;
 	for ( FSelectionIterator It( GetSelectedActorIterator() ) ; It ; ++It )
@@ -628,24 +624,12 @@ bool UUnrealEdEngine::CanDeleteSelectedActors( const UWorld* InWorld, const bool
 		AActor* Actor		= static_cast<AActor*>( *It );
 		checkSlow( Actor->IsA(AActor::StaticClass()) );
 
-		// Only delete transactional actors that aren't a level's builder brush or worldsettings.
-		bool bDeletable	= false;
+		bool bDeletable = false;
 		FText CannotDeleteReason;
-		if (Actor->HasAllFlags(RF_Transactional) && Actor->CanDeleteSelectedActor(CannotDeleteReason))
+		if (CanDeleteActor(Actor, &CannotDeleteReason))
 		{
-			// TODO: The Brush and WorldSettings logic should be moved to use the CanDeleteSelectedActor virtual
-			ABrush* Brush = Cast< ABrush >( Actor );
-			const bool bIsDefaultBrush = Brush && FActorEditorUtils::IsABuilderBrush(Brush);
-			if ( !bIsDefaultBrush )
-			{
-				const bool bIsWorldSettings =
-					Actor->IsA( AWorldSettings::StaticClass() ) && WorldSettingss.Contains( static_cast<AWorldSettings*>(Actor) );
-				if ( !bIsWorldSettings )
-				{
-					bContainsDeletable = true;
-					bDeletable = true;
-				}
-			}
+			bContainsDeletable = true;
+			bDeletable = true;
 		}
 
 		// Can this actor be deleted
@@ -668,7 +652,7 @@ bool UUnrealEdEngine::CanDeleteSelectedActors( const UWorld* InWorld, const bool
 			FText LogText;
 			if (CannotDeleteReason.IsEmpty())
 			{
-				LogText = FText::Format(LOCTEXT("CannotDeleteSpecialActor", "Cannot delete special actor {Name}"), Arguments);
+				LogText = FText::Format(LOCTEXT("CannotDeleteActor", "Cannot delete actor {Name}"), Arguments);
 			}
 			else
 			{
@@ -681,146 +665,326 @@ bool UUnrealEdEngine::CanDeleteSelectedActors( const UWorld* InWorld, const bool
 	return bContainsDeletable;
 }
 
-bool UUnrealEdEngine::edactDeleteSelected( UWorld* InWorld, bool bVerifyDeletionCanHappen, bool bWarnAboutReferences, bool bWarnAboutSoftReferences)
+bool UUnrealEdEngine::CanDeleteComponent(const UActorComponent* InComponent, FText* OutReason) const
 {
-	if ( bVerifyDeletionCanHappen )
+	if (!InComponent->IsEditableWhenInherited())
 	{
-		// Provide the option to abort the delete
-		if ( ShouldAbortActorDeletion() )
+		if (OutReason)
 		{
-			return false;
+			*OutReason = LOCTEXT("CanDeleteComponent_Error_ComponentNotEditableWhenInherited", "Can't delete non-editable component.");
 		}
+		return false;
 	}
 
-	const double StartSeconds = FPlatformTime::Seconds();
+	if (!FComponentEditorUtils::CanDeleteComponent(InComponent))
+	{
+		if (OutReason)
+		{
+			*OutReason = LOCTEXT("CanDeleteComponent_Error_InvalidComponent", "Can't delete non-instanced components or the scene root.");
+		}
+		return false;
+	}
 
-	// Aggregate the time we waited on user input to remove it from the total time
-	double DialogWaitingSeconds = 0;
+	return true;
+}
 
-	FSlateApplication::Get().CancelDragDrop();
+bool UUnrealEdEngine::CanDeleteActor(const AActor* InActor, FText* OutReason) const
+{
+	if (!InActor->HasAllFlags(RF_Transactional))
+	{
+		if (OutReason)
+		{
+			*OutReason = LOCTEXT("CanDeleteActor_Error_NotTransactional", "Can't delete non-transactional actors.");
+		}
+		return false;
+	}
 
-	if (GetSelectedComponentCount() > 0)
 	{
-		TArray<UActorComponent*> SelectedEditableComponents;
-		for (FSelectedEditableComponentIterator It(GetSelectedEditableComponentIterator()); It; ++It)
+		FText TmpReason;
+		if (!InActor->CanDeleteSelectedActor(TmpReason))
 		{
-			SelectedEditableComponents.Add(CastChecked<UActorComponent>(*It));
+			if (OutReason)
+			{
+				*OutReason = TmpReason;
+			}
+			return false;
 		}
+	}
+
+	// TODO: The Brush and WorldSettings logic below should be moved to use the CanDeleteSelectedActor virtual
 
-		if (SelectedEditableComponents.Num() > 0)
+	if (const ABrush* Brush = Cast<ABrush>(InActor))
+	{
+		if (FActorEditorUtils::IsABuilderBrush(Brush))
 		{
-			// Modify the actor that owns the selected components
-			check(GetSelectedActorCount() == 1);
-			(*GetSelectedActorIterator())->Modify();
-			
-			// Delete the selected components
-			UActorComponent* ComponentToSelect = nullptr;
-			int32 NumDeletedComponents = FComponentEditorUtils::DeleteComponents(SelectedEditableComponents, ComponentToSelect);
+			if (OutReason)
+			{
+				*OutReason = LOCTEXT("CanDeleteActor_Error_BuilderBrush", "Can't delete a builder brush.");
+			}
+			return false;
+		}
+	}
 
-			if (NumDeletedComponents > 0)
+	if (const AWorldSettings* WorldSettings = Cast<AWorldSettings>(InActor))
+	{
+		const ULevel* OwnerLevel = WorldSettings->GetLevel();
+		if (OwnerLevel && OwnerLevel->GetWorldSettings() == WorldSettings)
+		{
+			if (OutReason)
 			{
-				// Make sure all the SCS trees have a chance to rebuild
-				FLevelEditorModule& LevelEditor = FModuleManager::LoadModuleChecked<FLevelEditorModule>("LevelEditor");
-				LevelEditor.BroadcastComponentsEdited();
+				*OutReason = LOCTEXT("CanDeleteActor_Error_WorldSettings", "Can't delete a level's world settings.");
+			}
+			return false;
+		}
+	}
 
-				// Update the editor component selection if possible
-				if (ComponentToSelect)
-				{
-					USelection* ComponentSelection = GetSelectedComponents();
-					ComponentSelection->Modify(false);
-					ComponentSelection->BeginBatchSelectOperation();
-					ComponentSelection->DeselectAll();
+	return true;
+}
 
-					GEditor->SelectComponent(ComponentToSelect, true, false);
+bool UUnrealEdEngine::ShouldAbortComponentDeletion(const TArray<UActorComponent*>& InComponentsToDelete, FText* OutReason) const
+{
+	if (!GLevelEditorModeTools().EnsureNotInMode(FBuiltinEditorModes::EM_InterpEdit))
+	{
+		if (OutReason)
+		{
+			*OutReason = LOCTEXT("ShouldAbortComponentDeletion_Error_WrongModeForComponentDeletion", "Cannot delete components while Matinee is open.");
+		}
+		return true;
+	}
 
-					// Make sure the selection changed event fires so the SCS trees can update their selection
-					ComponentSelection->ForceBatchDirty();
-					ComponentSelection->EndBatchSelectOperation(true);
+	for (UActorComponent* Component : InComponentsToDelete)
+	{
+		AActor* OwnerActor = Component->GetOwner();
+		if (OwnerActor && FLevelUtils::IsLevelLocked(OwnerActor))
+		{
+			if (OutReason)
+			{
+				*OutReason = FText::Format(LOCTEXT("ShouldAbortComponentDeletion_Error_LevelLockedDuringComponentDeletion", "Cannot delete component '{0}' because its owner level is locked."), FText::FromString(Component->GetPathName()));
+			}
+			return true;
+		}
+	}
 
-					// Notify the level editor of the new component selection
-					NoteSelectionChange();
-				}
+	return false;
+}
 
-				UE_LOG(LogEditorActor, Log, TEXT("Deleted %d Components (%3.3f secs)"), NumDeletedComponents, FPlatformTime::Seconds() - StartSeconds);
-				return true;
+bool UUnrealEdEngine::ShouldAbortActorDeletion(const TArray<AActor*>& InActorsToDelete, FText* OutReason) const
+{
+	if (!GLevelEditorModeTools().EnsureNotInMode(FBuiltinEditorModes::EM_InterpEdit))
+	{
+		if (OutReason)
+		{
+			*OutReason = LOCTEXT("ShouldAbortActorDeletion_WrongModeForActorDeletion", "Cannot delete actors while Matinee is open.");
+		}
+		return true;
+	}
+
+	for (AActor* Actor : InActorsToDelete)
+	{
+		if (FLevelUtils::IsLevelLocked(Actor))
+		{
+			if (OutReason)
+			{
+				*OutReason = FText::Format(LOCTEXT("ShouldAbortActorDeletion_LevelLockedDuringActorDeletion", "Cannot delete actor '{0}' because its owner level is locked."), FText::FromString(Actor->GetPathName()));
 			}
+			return true;
+		}
+	}
+
+	return false;
+}
+
+bool UUnrealEdEngine::DeleteComponents(const TArray<UActorComponent*>& InComponentsToDelete, UTypedElementSelectionSet* InSelectionSet, const bool bVerifyDeletionCanHappen)
+{
+	if (bVerifyDeletionCanHappen)
+	{
+		FText ErrorMsg;
+		if (ShouldAbortComponentDeletion(InComponentsToDelete, &ErrorMsg))
+		{
+			FSlateNotificationManager::Get().AddNotification(FNotificationInfo(ErrorMsg));
+			return false;
+		}
+	}
+
+	const double StartSeconds = FPlatformTime::Seconds();
+
+	FSlateApplication::Get().CancelDragDrop();
+
+	// Get a list of all the deletable components
+	TArray<UActorComponent*> EditableComponentsToDelete;
+	EditableComponentsToDelete.Reserve(InComponentsToDelete.Num());
+	for (UActorComponent* ComponentToDelete : InComponentsToDelete)
+	{
+		FText ErrorMsg;
+		if (!CanDeleteComponent(ComponentToDelete, &ErrorMsg))
+		{
+			UE_LOG(LogEditorActor, Log, TEXT("Cannot delete component %s: %s"), *ComponentToDelete->GetFullName(), *ErrorMsg.ToString());
+			continue;
+		}
+
+		// Modify the actor that owns the component
+		if (AActor* OwnerActor = ComponentToDelete->GetOwner())
+		{
+			OwnerActor->Modify();
 		}
 
+		EditableComponentsToDelete.Add(ComponentToDelete);
+	}
+
+	if (EditableComponentsToDelete.Num() == 0)
+	{
 		return false;
 	}
 
-	GetSelectedActors()->Modify();
+	TUniquePtr<FTypedElementListLegacySyncScopedBatch> LegacySyncBatch = MakeUnique<FTypedElementListLegacySyncScopedBatch>(InSelectionSet->GetElementList());
 
-	// Fire ULevel::LevelDirtiedEvent when falling out of scope.
-	FScopedLevelDirtied			LevelDirtyCallback;
+	const FTypedElementSelectionOptions SelectionOptions = FTypedElementSelectionOptions()
+		.SetAllowHidden(true)
+		.SetAllowGroups(false)
+		.SetWarnIfLocked(false);
 
-	// Get a list of all the deletable actors in the selection
-	TArray<AActor*> ActorsToDelete;
-	CanDeleteSelectedActors( InWorld, false, true, &ActorsToDelete );
+	// Clear the selection of any components we may delete
+	for (UActorComponent* ComponentToDelete : EditableComponentsToDelete)
+	{
+		if (FTypedElementHandle ComponentHandle = UEngineElementsLibrary::AcquireEditorComponentElementHandle(ComponentToDelete, /*bAllowCreate*/false))
+		{
+			InSelectionSet->DeselectElement(ComponentHandle, SelectionOptions);
+		}
+	}
 
-	// Maintain a list of levels that have already been Modify()'d so that each level
-	// is modify'd only once.
-	TArray<ULevel*> LevelsAlreadyModified;
-	// A list of levels that will need their Bsp updated after the deletion is complete
-	TSet<ULevel*> LevelsToRebuildBSP;
-	TSet<ULevel*> LevelsToRebuildNavigation;
+	// Delete the components
+	UActorComponent* ComponentToSelect = nullptr;
+	const int32 NumDeletedComponents = FComponentEditorUtils::DeleteComponents(EditableComponentsToDelete, ComponentToSelect);
 
+	if (NumDeletedComponents > 0)
+	{
+		// Make sure all the SCS trees have a chance to rebuild
+		FLevelEditorModule& LevelEditor = FModuleManager::LoadModuleChecked<FLevelEditorModule>("LevelEditor");
+		LevelEditor.BroadcastComponentsEdited();
 
-	bool	bRequestedDeleteAllByLevel = false;
-	bool	bRequestedDeleteAllByActor = false;
-	bool	bRequestedDeleteAllBySoftReference = false;
-	EAppMsgType::Type MessageType = ActorsToDelete.Num() > 1 ? EAppMsgType::YesNoYesAllNoAll : EAppMsgType::YesNo;
-	int32		DeleteCount = 0;
+		// Remove all references to destroyed components once at the end, instead of once for each component destroyed
+		CollectGarbage(GARBAGE_COLLECTION_KEEPFLAGS);
+	}
 
-	USelection* SelectedActors = GetSelectedActors();
-	TMap<AActor*, TArray<AActor*> > ReferencingActorsMap;
-	TMap<AActor*, TArray<UObject*> > SoftReferencingObjectsMap;
-	TArray<UClass*> ClassTypesToIgnore;
-	ClassTypesToIgnore.Add(ALevelScriptActor::StaticClass());
-	// The delete warning is meant for actor references that affect gameplay.  Group actors do not affect gameplay and should not show up as a warning.
-	ClassTypesToIgnore.Add(AGroupActor::StaticClass());
+	// Update the editor component selection if possible
+	if (ComponentToSelect)
+	{
+		InSelectionSet->SelectElement(UEngineElementsLibrary::AcquireEditorComponentElementHandle(ComponentToSelect), SelectionOptions);
+	}
+
+	// Make sure the selection changed event fires so the SCS trees can update their selection
+	LegacySyncBatch->ForceDirty();
+	LegacySyncBatch.Reset();
+	NoteSelectionChange();
+
+	UE_LOG(LogEditorActor, Log, TEXT("Deleted %d Components (%3.3f secs)"), NumDeletedComponents, FPlatformTime::Seconds() - StartSeconds);
+	return NumDeletedComponents > 0;
+}
 
-	// If we want to warn about references to the actors to be deleted, it is a lot more efficient to query
-	// the world first and build a map of actors referenced by other actors. We can then quickly look this up later on in the loop.
-	if (bWarnAboutReferences)
+bool UUnrealEdEngine::DeleteActors(const TArray<AActor*>& InActorsToDelete, UWorld* InWorld, UTypedElementSelectionSet* InSelectionSet, const bool bVerifyDeletionCanHappen, const bool bWarnAboutReferences, const bool bWarnAboutSoftReferences)
+{
+	if (bVerifyDeletionCanHappen)
 	{
-		FBlueprintEditorUtils::GetActorReferenceMap(InWorld, ClassTypesToIgnore, ReferencingActorsMap);
+		// TODO: Bubble up error message?
+		FText ErrorMsg;
+		if (ShouldAbortActorDeletion(InActorsToDelete, &ErrorMsg))
+		{
+			FSlateNotificationManager::Get().AddNotification(FNotificationInfo(ErrorMsg));
+			return false;
+		}
+	}
 
-		if (bWarnAboutSoftReferences)
+	const double StartSeconds = FPlatformTime::Seconds();
+
+	// Aggregate the time we waited on user input to remove it from the total time
+	double DialogWaitingSeconds = 0;
+
+	FSlateApplication::Get().CancelDragDrop();
+
+	// Fire ULevel::LevelDirtiedEvent when falling out of scope.
+	FScopedLevelDirtied LevelDirtyCallback;
+
+	// Get a list of all the deletable actors
+	TArray<AActor*> ActorsToDelete;
+	ActorsToDelete.Reserve(InActorsToDelete.Num());
+	for (AActor* ActorToDelete : InActorsToDelete)
+	{
+		FText ErrorMsg;
+		if (!CanDeleteActor(ActorToDelete, &ErrorMsg))
 		{
-			FScopedSlowTask SlowTask(ActorsToDelete.Num(), LOCTEXT("ComputeActorSoftReferences", "Computing References"));
-			SlowTask.MakeDialogDelayed(1.0f);
+			UE_LOG(LogEditorActor, Log, TEXT("Cannot delete actor %s: %s"), *ActorToDelete->GetFullName(), *ErrorMsg.ToString());
+			continue;
+		}
 
-			FAssetToolsModule& AssetToolsModule = FModuleManager::GetModuleChecked<FAssetToolsModule>(TEXT("AssetTools"));
+		ActorsToDelete.Add(ActorToDelete);
+	}
 
-			for (int32 ActorIndex = 0; ActorIndex < ActorsToDelete.Num(); ++ActorIndex)
-			{
-				SlowTask.EnterProgressFrame();
+	TMap<AActor*, TArray<AActor*>> ReferencingActorsMap;
+	TMap<AActor*, TArray<UObject*>> SoftReferencingObjectsMap;
+	{
+		TArray<UClass*> ClassTypesToIgnore;
+		ClassTypesToIgnore.Add(ALevelScriptActor::StaticClass());
+		// The delete warning is meant for actor references that affect gameplay.  Group actors do not affect gameplay and should not show up as a warning.
+		ClassTypesToIgnore.Add(AGroupActor::StaticClass());
 
-				TArray<UObject*> SoftReferencingObjects;
-				AActor* Actor = ActorsToDelete[ActorIndex];
+		// If we want to warn about references to the actors to be deleted, it is a lot more efficient to query
+		// the world first and build a map of actors referenced by other actors. We can then quickly look this up later on in the loop.
+		if (bWarnAboutReferences)
+		{
+			FBlueprintEditorUtils::GetActorReferenceMap(InWorld, ClassTypesToIgnore, ReferencingActorsMap);
 
-				AssetToolsModule.Get().FindSoftReferencesToObject(Actor, SoftReferencingObjects);
+			if (bWarnAboutSoftReferences)
+			{
+				FScopedSlowTask SlowTask(ActorsToDelete.Num(), LOCTEXT("ComputeActorSoftReferences", "Computing References"));
+				SlowTask.MakeDialogDelayed(1.0f);
+
+				FAssetToolsModule& AssetToolsModule = FModuleManager::GetModuleChecked<FAssetToolsModule>(TEXT("AssetTools"));
 
-				if (SoftReferencingObjects.Num() > 0)
+				for (int32 ActorIndex = 0; ActorIndex < ActorsToDelete.Num(); ++ActorIndex)
 				{
-					SoftReferencingObjectsMap.Add(Actor, SoftReferencingObjects);
+					SlowTask.EnterProgressFrame();
+
+					TArray<UObject*> SoftReferencingObjects;
+					AActor* Actor = ActorsToDelete[ActorIndex];
+
+					AssetToolsModule.Get().FindSoftReferencesToObject(Actor, SoftReferencingObjects);
+
+					if (SoftReferencingObjects.Num() > 0)
+					{
+						SoftReferencingObjectsMap.Add(Actor, SoftReferencingObjects);
+					}
 				}
 			}
 		}
 	}
 
+	// Maintain a list of levels that have already been Modify()'d so that each level is modified only once
+	TArray<ULevel*> LevelsAlreadyModified;
+	// A list of levels that will need their BSP updated after the deletion is complete
+	TSet<ULevel*> LevelsToRebuildBSP;
+	TSet<ULevel*> LevelsToRebuildNavigation;
+
+	bool bRequestedDeleteAllByLevel = false;
+	bool bRequestedDeleteAllByActor = false;
+	bool bRequestedDeleteAllBySoftReference = false;
+	int32 DeleteCount = 0;
+
+	TUniquePtr<FTypedElementListLegacySyncScopedBatch> LegacySyncBatch = MakeUnique<FTypedElementListLegacySyncScopedBatch>(InSelectionSet->GetElementList());
+
+	const FTypedElementSelectionOptions SelectionOptions = FTypedElementSelectionOptions()
+		.SetAllowHidden(true)
+		.SetAllowGroups(false)
+		.SetWarnIfLocked(false);
+
 	ULayersSubsystem* LayersSubsystem = GEditor->GetEditorSubsystem<ULayersSubsystem>();
-	for ( int32 ActorIndex = 0 ; ActorIndex < ActorsToDelete.Num() ; ++ActorIndex )
+	for (AActor* Actor : ActorsToDelete)
 	{
-		AActor* Actor = ActorsToDelete[ ActorIndex ];
-
 		//If actor is referenced by script, ask user if they really want to delete
 		ULevelScriptBlueprint* LSB = Actor->GetLevel()->GetLevelScriptBlueprint(true);
 
 		// Get the array of actors that reference this actor from the cached map we built above.
 		TArray<AActor*>* ReferencingActors = nullptr;
-		if( bWarnAboutReferences )
+		if (bWarnAboutReferences)
 		{
 			ReferencingActors = ReferencingActorsMap.Find(Actor);
 		}
@@ -847,7 +1011,7 @@ bool UUnrealEdEngine::edactDeleteSelected( UWorld* InWorld, bool bVerifyDeletion
 		// If there are any referencing actors, make sure that they are reference types that we care about.
 		if (ReferencingActors != nullptr)
 		{
-			for (AActor* ReferencingActor : (*ReferencingActors))
+			for (AActor* ReferencingActor : *ReferencingActors)
 			{
 				// Skip to next if we are referencing ourselves
 				if (ReferencingActor == Actor)
@@ -888,9 +1052,9 @@ bool UUnrealEdEngine::edactDeleteSelected( UWorld* InWorld, bool bVerifyDeletion
 
 				FString LevelScriptReferenceString;
 
-				for (int32 i = 0; i < ReferencedToActorsFromLevelScriptArray.Num(); ++i)
+				for (UK2Node* Node : ReferencedToActorsFromLevelScriptArray)
 				{
-					LevelScriptReferenceString += ReferencedToActorsFromLevelScriptArray[i]->GetFindReferenceSearchString();
+					LevelScriptReferenceString += Node->GetFindReferenceSearchString();
 
 					if (bReferencedByLevelScript && bReferencedByActor)
 					{
@@ -906,9 +1070,9 @@ bool UUnrealEdEngine::edactDeleteSelected( UWorld* InWorld, bool bVerifyDeletion
 
 				if (ReferencingActors != nullptr)
 				{
-					for (int32 i = 0; i < ReferencingActors->Num(); ++i)
+					for (AActor* ReferencingActor : *ReferencingActors)
 					{
-						ActorReferenceString += (*ReferencingActors)[i]->GetActorLabel();
+						ActorReferenceString += ReferencingActor->GetActorLabel();
 
 						if (bReferencedByLevelScript && bReferencedByActor)
 						{
@@ -954,6 +1118,7 @@ bool UUnrealEdEngine::edactDeleteSelected( UWorld* InWorld, bool bVerifyDeletion
 
 				const double DialogStartSeconds = FPlatformTime::Seconds();
 
+				const EAppMsgType::Type MessageType = ActorsToDelete.Num() > 1 ? EAppMsgType::YesNoYesAllNoAll : EAppMsgType::YesNo;
 				int32 Result = FMessageDialog::Open(MessageType, ConfirmDelete);
 
 				DialogWaitingSeconds += FPlatformTime::Seconds() - DialogStartSeconds;
@@ -982,9 +1147,8 @@ bool UUnrealEdEngine::edactDeleteSelected( UWorld* InWorld, bool bVerifyDeletion
 			if (bReferencedByActor || bReferencedByLODActor)
 			{
 				check(ReferencingActors != nullptr);
-				for (int32 ReferencingActorIndex = 0; ReferencingActorIndex < ReferencingActors->Num(); ReferencingActorIndex++)
+				for (AActor* ReferencingActor : *ReferencingActors)
 				{
-					AActor* ReferencingActor = (*ReferencingActors)[ReferencingActorIndex];
 					ReferencingActor->Modify();
 
 					ALODActor* LODActor = Cast<ALODActor>(ReferencingActor);
@@ -1002,29 +1166,29 @@ bool UUnrealEdEngine::edactDeleteSelected( UWorld* InWorld, bool bVerifyDeletion
 				}
 			}
 		}
-	
+
 		bool bRebuildNavigation = false;
 
-		ABrush* Brush = Cast< ABrush >( Actor );
+		ABrush* Brush = Cast< ABrush >(Actor);
 		if (Brush && !FActorEditorUtils::IsABuilderBrush(Brush)) // Track whether or not a brush actor was deleted.
 		{
 			ULevel* BrushLevel = Actor->GetLevel();
-			if (BrushLevel && !Brush->IsVolumeBrush() )
+			if (BrushLevel && !Brush->IsVolumeBrush())
 			{
 				BrushLevel->Model->Modify();
-				LevelsToRebuildBSP.Add( BrushLevel );
+				LevelsToRebuildBSP.Add(BrushLevel);
 				// Rebuilding bsp will also take care of navigation
-				LevelsToRebuildNavigation.Remove( BrushLevel );
+				LevelsToRebuildNavigation.Remove(BrushLevel);
 			}
-			else if( BrushLevel && !LevelsToRebuildBSP.Contains( BrushLevel ) )
+			else if (BrushLevel && !LevelsToRebuildBSP.Contains(BrushLevel))
 			{
-				LevelsToRebuildNavigation.Add( BrushLevel );
+				LevelsToRebuildNavigation.Add(BrushLevel);
 			}
 		}
 
 		// If the actor about to be deleted is in a group, be sure to remove it from the group
 		AGroupActor* ActorParentGroup = AGroupActor::GetParentForActor(Actor);
-		if(ActorParentGroup)
+		if (ActorParentGroup)
 		{
 			ActorParentGroup->Remove(*Actor);
 		}
@@ -1037,46 +1201,50 @@ bool UUnrealEdEngine::edactDeleteSelected( UWorld* InWorld, bool bVerifyDeletion
 		LevelDirtyCallback.Request();
 
 		// Deselect the Actor.
-		SelectedActors->Deselect(Actor);
+		if (FTypedElementHandle ActorHandle = UEngineElementsLibrary::AcquireEditorActorElementHandle(Actor, /*bAllowCreate*/false))
+		{
+			InSelectionSet->DeselectElement(ActorHandle, SelectionOptions);
+		}
 
 		// Modify the level.  Each level is modified only once.
 		// @todo DB: Shouldn't this be calling UWorld::ModifyLevel?
 		ULevel* Level = Actor->GetLevel();
-		if ( LevelsAlreadyModified.Find( Level ) == INDEX_NONE )
+		if (LevelsAlreadyModified.Find(Level) == INDEX_NONE)
 		{
-			LevelsAlreadyModified.Add( Level );
+			LevelsAlreadyModified.Add(Level);
 			// Don't mark the level dirty when deleting external actors and the level is in `use external actors` mode.
 			bool bShouldDirty = !(Actor->IsPackageExternal() && Level->IsUsingExternalActors());
 			Level->Modify(bShouldDirty);
 		}
 
-		UE_LOG(LogEditorActor, Log,  TEXT("Deleted Actor: %s"), *Actor->GetClass()->GetName() );
+		UE_LOG(LogEditorActor, Log, TEXT("Deleted Actor: %s"), *Actor->GetClass()->GetName());
 
 		// Destroy actor and clear references.
-		LayersSubsystem->DisassociateActorFromLayers( Actor );
-		bool WasDestroyed = Actor->GetWorld()->EditorDestroyActor( Actor, false );
-		checkf( WasDestroyed,TEXT( "Failed to destroy Actor %s (%s)"), *Actor->GetClass()->GetName(), *Actor->GetActorLabel() );
+		LayersSubsystem->DisassociateActorFromLayers(Actor);
+		bool WasDestroyed = Actor->GetWorld()->EditorDestroyActor(Actor, false);
+		checkf(WasDestroyed, TEXT("Failed to destroy Actor %s (%s)"), *Actor->GetClass()->GetName(), *Actor->GetActorLabel());
 
 		DeleteCount++;
 	}
 
-	// Remove all references to destroyed actors once at the end, instead of once for each Actor destroyed..
-	CollectGarbage( GARBAGE_COLLECTION_KEEPFLAGS );
+	// Remove all references to destroyed actors once at the end, instead of once for each Actor destroyed
+	CollectGarbage(GARBAGE_COLLECTION_KEEPFLAGS);
+
+	LegacySyncBatch.Reset();
+	NoteSelectionChange();
 
 	// If any brush actors were modified, update the Bsp in the appropriate levels
 	if (LevelsToRebuildBSP.Num())
 	{
 		FlushRenderingCommands();
 
-		for ( ULevel* Level : LevelsToRebuildBSP )
+		for (ULevel* Level : LevelsToRebuildBSP)
 		{
 			GEditor->RebuildLevel(*Level);
 		}
 	}
 
-	NoteSelectionChange();
-
-	if( LevelsToRebuildNavigation.Num() )
+	if (LevelsToRebuildNavigation.Num())
 	{
 		for (ULevel* Level : LevelsToRebuildNavigation)
 		{
@@ -1087,49 +1255,49 @@ bool UUnrealEdEngine::edactDeleteSelected( UWorld* InWorld, bool bVerifyDeletion
 		}
 	}
 
-	if( LevelsToRebuildBSP.Num() || LevelsToRebuildNavigation.Num() )
+	if (LevelsToRebuildBSP.Num() || LevelsToRebuildNavigation.Num())
 	{
 		RedrawLevelEditingViewports();
 		ULevel::LevelDirtiedEvent.Broadcast();
 	}
 
-	UE_LOG(LogEditorActor, Log,  TEXT("Deleted %d Actors (%3.3f secs)"), DeleteCount, (FPlatformTime::Seconds() - StartSeconds) - DialogWaitingSeconds);
-
+	UE_LOG(LogEditorActor, Log, TEXT("Deleted %d Actors (%3.3f secs)"), DeleteCount, (FPlatformTime::Seconds() - StartSeconds) - DialogWaitingSeconds);
 	return true;
 }
 
-bool UUnrealEdEngine::ShouldAbortActorDeletion() const
+bool UUnrealEdEngine::edactDeleteSelected( UWorld* InWorld, bool bVerifyDeletionCanHappen, bool bWarnAboutReferences, bool bWarnAboutSoftReferences)
 {
-	bool bResult = false;
-
-	// Can't delete actors if Matinee is open.
-	const FText ErrorMsg = NSLOCTEXT("UnrealEd", "Error_WrongModeForActorDeletion", "Cannot delete actor while Matinee is open" );
-	if ( !GLevelEditorModeTools().EnsureNotInMode( FBuiltinEditorModes::EM_InterpEdit, ErrorMsg, true ) )
+	if (GetSelectedComponentCount() > 0)
 	{
-		bResult = true;
+		// Delete selected components
+		TArray<UActorComponent*> SelectedComponents;
+		GetSelectedComponents()->GetSelectedObjects<UActorComponent>(SelectedComponents);
+		return DeleteComponents(SelectedComponents, GetSelectedComponents()->GetElementSelectionSet(), bVerifyDeletionCanHappen);
 	}
-
-	if ( !bResult )
+	else
 	{
-		for ( FSelectionIterator It( GetSelectedActorIterator() ) ; It ; ++It )
-		{
-			AActor* Actor = static_cast<AActor*>( *It );
-			checkSlow( Actor->IsA(AActor::StaticClass()) );
+		// Delete selected actors
+		TArray<AActor*> SelectedActors;
+		GetSelectedActors()->GetSelectedObjects<AActor>(SelectedActors);
+		return DeleteActors(SelectedActors, InWorld, GetSelectedActors()->GetElementSelectionSet(), bVerifyDeletionCanHappen, bWarnAboutReferences, bWarnAboutSoftReferences);
+	}
+}
 
-			ULevel* ActorLevel = Actor->GetLevel();
-			if ( FLevelUtils::IsLevelLocked(ActorLevel) )
-			{
-				UE_LOG(LogEditorActor, Warning, TEXT("Cannot perform action on actor %s because the actor's level is locked"), *Actor->GetName());
-				bResult = true;
-				break;
-			}
-		}
+bool UUnrealEdEngine::ShouldAbortActorDeletion() const
+{
+	TArray<AActor*> SelectedActors;
+	GetSelectedActors()->GetSelectedObjects<AActor>(SelectedActors);
+
+	FText ErrorMsg;
+	if (ShouldAbortActorDeletion(SelectedActors, &ErrorMsg))
+	{
+		FSlateNotificationManager::Get().AddNotification(FNotificationInfo(ErrorMsg));
+		return true;
 	}
 
-	return bResult;
+	return false;
 }
 
-
 void UUnrealEdEngine::edactReplaceSelectedBrush( UWorld* InWorld )
 {
 	// Make a list of brush actors to replace.
diff --git a/Engine/Source/Editor/UnrealEd/Private/Kismet2/ComponentEditorUtils.cpp b/Engine/Source/Editor/UnrealEd/Private/Kismet2/ComponentEditorUtils.cpp
--- a/Engine/Source/Editor/UnrealEd/Private/Kismet2/ComponentEditorUtils.cpp
+++ b/Engine/Source/Editor/UnrealEd/Private/Kismet2/ComponentEditorUtils.cpp
@@ -648,21 +648,23 @@ void FComponentEditorUtils::GetComponentsFromClipboard(TMap<FName, FName>& OutPa
 	OutNewObjectMap = MoveTemp(Factory->NewObjectMap);
 }
 
+bool FComponentEditorUtils::CanDeleteComponent(const UActorComponent* ComponentToDelete)
+{
+	// We can't delete non-instance components or the default scene root
+	return ComponentToDelete->CreationMethod == EComponentCreationMethod::Instance
+		&& ComponentToDelete->GetFName() != USceneComponent::GetDefaultSceneRootVariableName();
+}
+
 bool FComponentEditorUtils::CanDeleteComponents(const TArray<UActorComponent*>& ComponentsToDelete)
 {
-	bool bCanDelete = true;
-	for (UActorComponent* ComponentToDelete : ComponentsToDelete)
+	for (const UActorComponent* ComponentToDelete : ComponentsToDelete)
 	{
-		// We can't delete non-instance components or the default scene root
-		if (ComponentToDelete->CreationMethod != EComponentCreationMethod::Instance 
-			|| ComponentToDelete->GetFName() == USceneComponent::GetDefaultSceneRootVariableName())
+		if (!CanDeleteComponent(ComponentToDelete))
 		{
-			bCanDelete = false;
-			break;
+			return false;
 		}
 	}
-
-	return bCanDelete;
+	return true;
 }
 
 int32 FComponentEditorUtils::DeleteComponents(const TArray<UActorComponent*>& ComponentsToDelete, UActorComponent*& OutComponentToSelect)
diff --git a/Engine/Source/Editor/UnrealEd/Public/Kismet2/ComponentEditorUtils.h b/Engine/Source/Editor/UnrealEd/Public/Kismet2/ComponentEditorUtils.h
--- a/Engine/Source/Editor/UnrealEd/Public/Kismet2/ComponentEditorUtils.h
+++ b/Engine/Source/Editor/UnrealEd/Public/Kismet2/ComponentEditorUtils.h
@@ -87,6 +87,13 @@ public:
 	 */
 	static void GetComponentsFromClipboard(TMap<FName, FName>& OutParentMap, TMap<FName, UActorComponent*>& OutNewObjectMap, bool bGetComponentsAsArchetypes);
 
+	/**
+	 * Determines whether the indicated component can be deleted
+	 * @param ComponentToDelete The component to determine can be deleted
+	 * @return Whether the indicated component can be deleted
+	 */
+	static bool CanDeleteComponent(const UActorComponent* ComponentToDelete);
+
 	/**
 	 * Determines whether the indicated components can be deleted
 	 * @param ComponentsToDelete The list of components to determine can be deleted
diff --git a/Engine/Source/Runtime/TypedElementFramework/Private/Elements/Framework/TypedElementList.cpp b/Engine/Source/Runtime/TypedElementFramework/Private/Elements/Framework/TypedElementList.cpp
--- a/Engine/Source/Runtime/TypedElementFramework/Private/Elements/Framework/TypedElementList.cpp
+++ b/Engine/Source/Runtime/TypedElementFramework/Private/Elements/Framework/TypedElementList.cpp
@@ -91,6 +91,21 @@ FTypedElementListLegacySyncScopedBatch::~FTypedElementListLegacySyncScopedBatch(
 	}
 }
 
+bool FTypedElementListLegacySyncScopedBatch::IsDirty() const
+{
+	return ElementListLegacySync
+		&& ElementListLegacySync->IsBatchOperationDirty();
+}
+
+void FTypedElementListLegacySyncScopedBatch::ForceDirty()
+{
+	if (ElementListLegacySync)
+	{
+		ElementListLegacySync->ForceBatchOperationDirty();
+	}
+}
+
+
 UTypedElementList* UTypedElementList::Private_CreateElementList(UTypedElementRegistry* InRegistry)
 {
 	UTypedElementList* ElementList = NewObject<UTypedElementList>();
diff --git a/Engine/Source/Runtime/TypedElementFramework/Public/Elements/Framework/TypedElementList.h b/Engine/Source/Runtime/TypedElementFramework/Public/Elements/Framework/TypedElementList.h
--- a/Engine/Source/Runtime/TypedElementFramework/Public/Elements/Framework/TypedElementList.h
+++ b/Engine/Source/Runtime/TypedElementFramework/Public/Elements/Framework/TypedElementList.h
@@ -116,6 +116,9 @@ public:
 	FTypedElementListLegacySyncScopedBatch(FTypedElementListLegacySyncScopedBatch&&) = delete;
 	FTypedElementListLegacySyncScopedBatch& operator=(FTypedElementListLegacySyncScopedBatch&&) = delete;
 
+	bool IsDirty() const;
+	void ForceDirty();
+
 private:
 	FTypedElementListLegacySync* ElementListLegacySync = nullptr;
 	bool bNotify = true;
