Updated the level editor pivot to use typed elements

diff --git a/Engine/Source/Editor/LevelEditor/Private/Elements/Actor/ActorElementLevelEditorSelectionCustomization.cpp b/Engine/Source/Editor/LevelEditor/Private/Elements/Actor/ActorElementLevelEditorSelectionCustomization.cpp
--- a/Engine/Source/Editor/LevelEditor/Private/Elements/Actor/ActorElementLevelEditorSelectionCustomization.cpp
+++ b/Engine/Source/Editor/LevelEditor/Private/Elements/Actor/ActorElementLevelEditorSelectionCustomization.cpp
@@ -177,7 +177,7 @@ bool FActorElementLevelEditorSelectionCustomization::SelectActorElement(const TT
 		}
 	}
 
-	UE_LOG(LogActorLevelEditorSelection, Verbose, TEXT("Selected Actor: %s"), *Actor->GetClass()->GetName());
+	UE_LOG(LogActorLevelEditorSelection, Verbose, TEXT("Selected Actor: %s (%s)"), *Actor->GetActorLabel(), *Actor->GetClass()->GetName());
 
 	// Update the annotation state
 	GSelectedActorAnnotation.Set(Actor);
@@ -239,7 +239,7 @@ bool FActorElementLevelEditorSelectionCustomization::DeselectActorElement(const
 		}
 	}
 	
-	UE_LOG(LogActorLevelEditorSelection, Verbose, TEXT("Deselected Actor: %s"), *Actor->GetClass()->GetName());
+	UE_LOG(LogActorLevelEditorSelection, Verbose, TEXT("Deselected Actor: %s (%s)"), *Actor->GetActorLabel(), *Actor->GetClass()->GetName());
 	
 	// Update the annotation state
 	GSelectedActorAnnotation.Clear(Actor);
diff --git a/Engine/Source/Editor/LevelEditor/Private/Elements/Component/ComponentElementLevelEditorSelectionCustomization.cpp b/Engine/Source/Editor/LevelEditor/Private/Elements/Component/ComponentElementLevelEditorSelectionCustomization.cpp
--- a/Engine/Source/Editor/LevelEditor/Private/Elements/Component/ComponentElementLevelEditorSelectionCustomization.cpp
+++ b/Engine/Source/Editor/LevelEditor/Private/Elements/Component/ComponentElementLevelEditorSelectionCustomization.cpp
@@ -109,7 +109,7 @@ bool FComponentElementLevelEditorSelectionCustomization::SelectComponentElement(
 		return false;
 	}
 
-	UE_LOG(LogComponentLevelEditorSelection, Verbose, TEXT("Selected Component: %s"), *Component->GetClass()->GetName());
+	UE_LOG(LogComponentLevelEditorSelection, Verbose, TEXT("Selected Component: %s (%s)"), *Component->GetPathName(), *Component->GetClass()->GetName());
 	
 	// Update the annotation state
 	GSelectedComponentAnnotation.Set(Component);
@@ -148,7 +148,7 @@ bool FComponentElementLevelEditorSelectionCustomization::DeselectComponentElemen
 		return false;
 	}
 
-	UE_LOG(LogComponentLevelEditorSelection, Verbose, TEXT("Deselected Component: %s"), *Component->GetClass()->GetName());
+	UE_LOG(LogComponentLevelEditorSelection, Verbose, TEXT("Deselected Component: %s (%s)"), *Component->GetPathName(), *Component->GetClass()->GetName());
 
 	// Update the annotation state
 	GSelectedComponentAnnotation.Clear(Component);
diff --git a/Engine/Source/Editor/LevelEditor/Private/Elements/Component/ComponentElementLevelEditorViewportInteractionCustomization.cpp b/Engine/Source/Editor/LevelEditor/Private/Elements/Component/ComponentElementLevelEditorViewportInteractionCustomization.cpp
--- a/Engine/Source/Editor/LevelEditor/Private/Elements/Component/ComponentElementLevelEditorViewportInteractionCustomization.cpp
+++ b/Engine/Source/Editor/LevelEditor/Private/Elements/Component/ComponentElementLevelEditorViewportInteractionCustomization.cpp
@@ -20,6 +20,14 @@ void FComponentElementLevelEditorViewportInteractionCustomization::GetElementsTo
 	{
 		if (AActor* ComponentOwner = SceneComponent->GetOwner())
 		{
+#if DO_CHECK
+			{
+				FTypedElementHandle ComponentOwnerElementHandle = UEngineElementsLibrary::AcquireEditorActorElementHandle(ComponentOwner, /*bAllowCreate*/false);
+				const bool bIsOwnerSelected = ComponentOwnerElementHandle && InSelectionSet->IsElementSelected(ComponentOwnerElementHandle, FTypedElementIsSelectedOptions());
+				ensureMsgf(bIsOwnerSelected, TEXT("Owner(%s) of %s is not selected"), *ComponentOwner->GetFullName(), *Component->GetFullName());
+			}
+#endif	// DO_CHECK
+
 			if (FActorElementLevelEditorViewportInteractionCustomization::CanMoveActorInViewport(ComponentOwner, InWorldType))
 			{
 				if (ComponentOwner->GetRootComponent() == SceneComponent)
diff --git a/Engine/Source/Editor/UnrealEd/Private/EditorSelectUtils.cpp b/Engine/Source/Editor/UnrealEd/Private/EditorSelectUtils.cpp
--- a/Engine/Source/Editor/UnrealEd/Private/EditorSelectUtils.cpp
+++ b/Engine/Source/Editor/UnrealEd/Private/EditorSelectUtils.cpp
@@ -22,14 +22,18 @@
 #include "ScopedTransaction.h"
 #include "Engine/LevelStreaming.h"
 #include "LevelUtils.h"
+#include "LevelEditorViewport.h"
 #include "StatsViewerModule.h"
 #include "SnappingUtils.h"
 #include "Logging/MessageLog.h"
 #include "ActorGroupingUtils.h"
 #include "Subsystems/BrushEditingSubsystem.h"
 #include "Elements/Framework/TypedElementList.h"
+#include "Elements/Framework/TypedElementRegistry.h"
 #include "Elements/Framework/TypedElementSelectionSet.h"
 #include "Elements/Framework/EngineElementsLibrary.h"
+#include "Elements/Interfaces/TypedElementWorldInterface.h"
+#include "Elements/Interfaces/TypedElementObjectInterface.h"
 
 #define LOCTEXT_NAMESPACE "EditorSelectUtils"
 
@@ -164,6 +168,11 @@ FVector UUnrealEdEngine::GetPivotLocation()
 
 void UUnrealEdEngine::SetPivot( FVector NewPivot, bool bSnapPivotToGrid, bool bIgnoreAxis, bool bAssignPivot/*=false*/ )
 {
+	if (!GCurrentLevelEditingViewportClient)
+	{
+		return;
+	}
+
 	FEditorModeTools& EditorModeTools = GLevelEditorModeTools();
 
 	if( !bIgnoreAxis )
@@ -185,64 +194,77 @@ void UUnrealEdEngine::SetPivot( FVector NewPivot, bool bSnapPivotToGrid, bool bI
 		EditorModeTools.PivotLocation = EditorModeTools.SnappedLocation;
 	}
 
-	// Check all actors.
-	int32 Count=0, SnapCount=0;
+	// Check all elements.
+	int32 NumElements = 0;
 
 	//default to using the x axis for the translate rotate widget
 	EditorModeTools.TranslateRotateXAxisAngle = 0.0f;
 	EditorModeTools.TranslateRotate2DAngle = 0.0f;
-	FVector TranslateRotateWidgetWorldXAxis;
+	
+	FVector TranslateRotateWidgetWorldXAxis = FVector::ZeroVector;
+	FVector Widget2DWorldXAxis = FVector::ZeroVector;
 
-	FVector Widget2DWorldXAxis;
+	UTypedElementRegistry* Registry = UTypedElementRegistry::GetInstance();
 
-	AActor* LastSelectedActor = NULL;
-	for ( FSelectionIterator It( GetSelectedActorIterator() ) ; It ; ++It )
+	// Pick a new common pivot, or not.
+	TTypedElement<UTypedElementWorldInterface> SingleWorldElement;
 	{
-		AActor* Actor = static_cast<AActor*>( *It );
-		checkSlow( Actor->IsA(AActor::StaticClass()) );
-
-		if (Count==0)
+		const TArray<FTypedElementHandle> ElementsToManipulate = GCurrentLevelEditingViewportClient->GetElementsToManipulate();
+		for (const FTypedElementHandle& Element : ElementsToManipulate)
 		{
-			TranslateRotateWidgetWorldXAxis = Actor->ActorToWorld().TransformVector(FVector(1.0f, 0.0f, 0.0f));
-			//get the xy plane project of this vector
-			TranslateRotateWidgetWorldXAxis.Z = 0.0f;
-			if (!TranslateRotateWidgetWorldXAxis.Normalize())
+			if (TTypedElement<UTypedElementWorldInterface> WorldElement = Registry->GetElement<UTypedElementWorldInterface>(Element))
 			{
-				TranslateRotateWidgetWorldXAxis = FVector(1.0f, 0.0f, 0.0f);
-			}
+				if (NumElements == 0)
+				{
+					FTransform ElementWorldTransform;
+					WorldElement.GetWorldTransform(ElementWorldTransform);
 
-			Widget2DWorldXAxis = Actor->ActorToWorld().TransformVector(FVector(1, 0, 0));
-			Widget2DWorldXAxis.Y = 0;
-			if (!Widget2DWorldXAxis.Normalize())
-			{
-				Widget2DWorldXAxis = FVector(1, 0, 0);
+					TranslateRotateWidgetWorldXAxis = ElementWorldTransform.TransformVector(FVector(1.0f, 0.0f, 0.0f));
+					//get the xy plane project of this vector
+					TranslateRotateWidgetWorldXAxis.Z = 0.0f;
+					if (!TranslateRotateWidgetWorldXAxis.Normalize())
+					{
+						TranslateRotateWidgetWorldXAxis = FVector(1.0f, 0.0f, 0.0f);
+					}
+
+					Widget2DWorldXAxis = ElementWorldTransform.TransformVector(FVector(1, 0, 0));
+					Widget2DWorldXAxis.Y = 0;
+					if (!Widget2DWorldXAxis.Normalize())
+					{
+						Widget2DWorldXAxis = FVector(1, 0, 0);
+					}
+				}
+
+				SingleWorldElement = MoveTemp(WorldElement);
+				++NumElements;
 			}
 		}
-
-		LastSelectedActor = Actor;
-		++Count;
-		++SnapCount;
 	}
 	
-	if( bAssignPivot && LastSelectedActor && UActorGroupingUtils::IsGroupingActive())
+	if (bAssignPivot && UActorGroupingUtils::IsGroupingActive())
 	{
-		// set group pivot for the root-most group
-		AGroupActor* ActorGroupRoot = AGroupActor::GetRootForActor(LastSelectedActor, true, true);
-		if(ActorGroupRoot)
+		if (TTypedElement<UTypedElementObjectInterface> ObjectElement = Registry->GetElement<UTypedElementObjectInterface>(SingleWorldElement))
 		{
-			ActorGroupRoot->SetActorLocation( EditorModeTools.PivotLocation, false );
+			if (AActor* SingleActor = Cast<AActor>(ObjectElement.GetObject()))
+			{
+				// set group pivot for the root-most group
+				if (AGroupActor* ActorGroupRoot = AGroupActor::GetRootForActor(SingleActor, true, true))
+				{
+					ActorGroupRoot->SetActorLocation(EditorModeTools.PivotLocation, false);
+				}
+			}
 		}
 	}
 
-	//if there are multiple actors selected, just use the x-axis for the "translate/rotate" or 2D widgets
-	if (Count == 1)
+	//if there are multiple elements selected, just use the x-axis for the "translate/rotate" or 2D widgets
+	if (NumElements == 1)
 	{
 		EditorModeTools.TranslateRotateXAxisAngle = TranslateRotateWidgetWorldXAxis.Rotation().Yaw;
 		EditorModeTools.TranslateRotate2DAngle = FMath::RadiansToDegrees(FMath::Atan2(Widget2DWorldXAxis.Z, Widget2DWorldXAxis.X));
 	}
 
 	// Update showing.
-	EditorModeTools.PivotShown = SnapCount>0 || Count>1;
+	EditorModeTools.PivotShown = NumElements > 0;
 }
 
 
@@ -304,78 +326,69 @@ bool UUnrealEdEngine::IsPivotMovedIndependently() const
 
 void UUnrealEdEngine::UpdatePivotLocationForSelection( bool bOnChange )
 {
-	// Pick a new common pivot, or not.
-	AActor* SingleActor = nullptr;
-	USceneComponent* SingleComponent = nullptr;
+	if (!GCurrentLevelEditingViewportClient)
+	{
+		return;
+	}
 
-	if (GetSelectedComponentCount() > 0)
+	UTypedElementRegistry* Registry = UTypedElementRegistry::GetInstance();
+
+	// Pick a new common pivot, or not.
+	TTypedElement<UTypedElementWorldInterface> SingleWorldElement;
 	{
-		for (FSelectedEditableComponentIterator It(*GetSelectedComponents()); It; ++It)
+		const TArray<FTypedElementHandle> ElementsToManipulate = GCurrentLevelEditingViewportClient->GetElementsToManipulate();
+		for (const FTypedElementHandle& Element : ElementsToManipulate)
 		{
-			UActorComponent* Component = CastChecked<UActorComponent>(*It);
-			AActor* ComponentOwner = Component->GetOwner();
-
-			if (ComponentOwner != nullptr)
+			if (TTypedElement<UTypedElementWorldInterface> WorldElement = Registry->GetElement<UTypedElementWorldInterface>(Element))
 			{
-				USelection* SelectedActors = GetSelectedActors();
-				const bool bIsOwnerSelected = SelectedActors->IsSelected(ComponentOwner);
-				ensureMsgf(bIsOwnerSelected, TEXT("Owner(%s) of %s is not selected"), *ComponentOwner->GetFullName(), *Component->GetFullName());
-
-				if (ComponentOwner->GetWorld() == GWorld)
+#if DO_CHECK
 				{
-					SingleActor = ComponentOwner;
-					if (Component->IsA<USceneComponent>())
-					{
-						SingleComponent = CastChecked<USceneComponent>(Component);
-					}
-
-					const bool IsTemplate = ComponentOwner->IsTemplate();
-					const bool LevelLocked = !FLevelUtils::IsLevelLocked(ComponentOwner->GetLevel());
-					check(IsTemplate || LevelLocked);
+					ULevel* OwnerLevel = WorldElement.GetOwnerLevel();
+					const bool bIsTemplate = WorldElement.IsTemplateElement();
+					const bool bLevelLocked = FLevelUtils::IsLevelLocked(OwnerLevel);
+					check(bIsTemplate || !bLevelLocked);
 				}
-			}
-		}
-	}
-	else
-	{
-		for (FSelectionIterator It(GetSelectedActorIterator()); It; ++It)
-		{
-			AActor* Actor = static_cast<AActor*>(*It);
-			checkSlow(Actor->IsA(AActor::StaticClass()));
+#endif	// DO_CHECK
 
-			const bool IsTemplate = Actor->IsTemplate();
-			const bool LevelLocked = !FLevelUtils::IsLevelLocked(Actor->GetLevel());
-			check(IsTemplate || LevelLocked);
-
-			SingleActor = Actor;
+				SingleWorldElement = MoveTemp(WorldElement);
+			}
 		}
 	}
 	
-	if (SingleComponent != NULL)
-	{
-		SetPivot(SingleComponent->GetComponentLocation(), false, true);
-	}
-	else if( SingleActor != NULL ) 
+	if (SingleWorldElement)
 	{
 		UBrushEditingSubsystem* BrushSubsystem = GEditor->GetEditorSubsystem<UBrushEditingSubsystem>();
-		const bool bGeometryMode = BrushSubsystem ? BrushSubsystem->IsGeometryEditorModeActive() : false;
+		const bool bGeometryMode = BrushSubsystem && BrushSubsystem->IsGeometryEditorModeActive();
 
 		// For geometry mode use current pivot location as it's set to selected face, not actor
-		if (!bGeometryMode || bOnChange == true)
+		// TODO: If geometry used elements for face selection, then this could work via the world interface and this special case could be removed
+		if (!bGeometryMode || bOnChange)
 		{
-			// Set pivot point to the actor's location, accounting for any set pivot offset
-			FVector PivotPoint = SingleActor->GetTransform().TransformPosition(SingleActor->GetPivotOffset());
+			FTransform ElementWorldTransform;
+			SingleWorldElement.GetWorldTransform(ElementWorldTransform);
+
+			FVector ElementPivotOffset = FVector::ZeroVector;
+			SingleWorldElement.GetPivotOffset(ElementPivotOffset);
+
+			// Set pivot point to the element location, accounting for any set pivot offset
+			FVector PivotPoint = ElementWorldTransform.TransformPosition(ElementPivotOffset);
 
-			// If grouping is active, see if this actor is part of a locked group and use that pivot instead
-			if(UActorGroupingUtils::IsGroupingActive())
+			// If grouping is active, see if this element is an actor that's part of a locked group and use that pivot instead
+			if (UActorGroupingUtils::IsGroupingActive())
 			{
-				AGroupActor* ActorGroupRoot = AGroupActor::GetRootForActor(SingleActor, true, true);
-				if(ActorGroupRoot)
+				if (TTypedElement<UTypedElementObjectInterface> ObjectElement = Registry->GetElement<UTypedElementObjectInterface>(SingleWorldElement))
 				{
-					PivotPoint = ActorGroupRoot->GetActorLocation();
+					if (AActor* SingleActor = Cast<AActor>(ObjectElement.GetObject()))
+					{
+						if (AGroupActor* ActorGroupRoot = AGroupActor::GetRootForActor(SingleActor, true, true))
+						{
+							PivotPoint = ActorGroupRoot->GetActorLocation();
+						}
+					}
 				}
 			}
-			SetPivot( PivotPoint, false, true );
+
+			SetPivot(PivotPoint, false, true);
 		}
 	}
 	else
@@ -386,8 +399,6 @@ void UUnrealEdEngine::UpdatePivotLocationForSelection( bool bOnChange )
 	SetPivotMovedIndependently(false);
 }
 
-
-
 void UUnrealEdEngine::NoteSelectionChange(bool bNotify)
 {
 	// The selection changed, so make sure the pivot (widget) is located in the right place
diff --git a/Engine/Source/Editor/UnrealEd/Private/Elements/Actor/ActorElementEditorSelectionInterface.cpp b/Engine/Source/Editor/UnrealEd/Private/Elements/Actor/ActorElementEditorSelectionInterface.cpp
--- a/Engine/Source/Editor/UnrealEd/Private/Elements/Actor/ActorElementEditorSelectionInterface.cpp
+++ b/Engine/Source/Editor/UnrealEd/Private/Elements/Actor/ActorElementEditorSelectionInterface.cpp
@@ -43,16 +43,22 @@ bool UActorElementEditorSelectionInterface::IsActorSelected(const AActor* InActo
 		return false;
 	}
 
-	if (InSelectionSet->Contains(UEngineElementsLibrary::AcquireEditorActorElementHandle(InActor)))
+	if (FTypedElementHandle ActorElement = UEngineElementsLibrary::AcquireEditorActorElementHandle(InActor, /*bAllowCreate*/false))
 	{
-		return true;
+		if (InSelectionSet->Contains(ActorElement))
+		{
+			return true;
+		}
 	}
 
 	if (InSelectionOptions.AllowIndirect())
 	{
 		if (const AActor* RootSelectionActor = InActor->GetRootSelectionParent())
 		{
-			return InSelectionSet->Contains(UEngineElementsLibrary::AcquireEditorActorElementHandle(RootSelectionActor));
+			if (FTypedElementHandle RootSelectionElement = UEngineElementsLibrary::AcquireEditorActorElementHandle(RootSelectionActor, /*bAllowCreate*/false))
+			{
+				return InSelectionSet->Contains(RootSelectionElement);
+			}
 		}
 	}
 
diff --git a/Engine/Source/Editor/UnrealEd/Private/Elements/Actor/ActorElementEditorWorldInterface.cpp b/Engine/Source/Editor/UnrealEd/Private/Elements/Actor/ActorElementEditorWorldInterface.cpp
--- a/Engine/Source/Editor/UnrealEd/Private/Elements/Actor/ActorElementEditorWorldInterface.cpp
+++ b/Engine/Source/Editor/UnrealEd/Private/Elements/Actor/ActorElementEditorWorldInterface.cpp
@@ -10,6 +10,28 @@
 #include "UnrealEdGlobals.h"
 #include "Editor/UnrealEdEngine.h"
 
+bool UActorElementEditorWorldInterface::GetPivotOffset(const FTypedElementHandle& InElementHandle, FVector& OutPivotOffset)
+{
+	if (AActor* Actor = ActorElementDataUtil::GetActorFromHandle(InElementHandle))
+	{
+		OutPivotOffset = Actor->GetPivotOffset();
+		return true;
+	}
+
+	return false;
+}
+
+bool UActorElementEditorWorldInterface::SetPivotOffset(const FTypedElementHandle& InElementHandle, const FVector& InPivotOffset)
+{
+	if (AActor* Actor = ActorElementDataUtil::GetActorFromHandle(InElementHandle))
+	{
+		Actor->SetPivotOffset(InPivotOffset);
+		return true;
+	}
+
+	return false;
+}
+
 void UActorElementEditorWorldInterface::NotifyMovementStarted(const FTypedElementHandle& InElementHandle)
 {
 	if (AActor* Actor = ActorElementDataUtil::GetActorFromHandle(InElementHandle))
diff --git a/Engine/Source/Editor/UnrealEd/Private/Elements/Actor/ActorElementEditorWorldInterface.h b/Engine/Source/Editor/UnrealEd/Private/Elements/Actor/ActorElementEditorWorldInterface.h
--- a/Engine/Source/Editor/UnrealEd/Private/Elements/Actor/ActorElementEditorWorldInterface.h
+++ b/Engine/Source/Editor/UnrealEd/Private/Elements/Actor/ActorElementEditorWorldInterface.h
@@ -11,6 +11,8 @@ class UActorElementEditorWorldInterface : public UActorElementWorldInterface
 	GENERATED_BODY()
 
 public:
+	virtual bool GetPivotOffset(const FTypedElementHandle& InElementHandle, FVector& OutPivotOffset) override;
+	virtual bool SetPivotOffset(const FTypedElementHandle& InElementHandle, const FVector& InPivotOffset) override;
 	virtual void NotifyMovementStarted(const FTypedElementHandle& InElementHandle) override;
 	virtual void NotifyMovementOngoing(const FTypedElementHandle& InElementHandle) override;
 	virtual void NotifyMovementEnded(const FTypedElementHandle& InElementHandle) override;
diff --git a/Engine/Source/Editor/UnrealEd/Private/Elements/Component/ComponentElementEditorSelectionInterface.cpp b/Engine/Source/Editor/UnrealEd/Private/Elements/Component/ComponentElementEditorSelectionInterface.cpp
--- a/Engine/Source/Editor/UnrealEd/Private/Elements/Component/ComponentElementEditorSelectionInterface.cpp
+++ b/Engine/Source/Editor/UnrealEd/Private/Elements/Component/ComponentElementEditorSelectionInterface.cpp
@@ -11,13 +11,13 @@
 
 bool UComponentElementEditorSelectionInterface::IsElementSelected(const FTypedElementHandle& InElementHandle, const UTypedElementList* InSelectionSet, const FTypedElementIsSelectedOptions& InSelectionOptions)
 {
-	const UActorComponent* Component = ComponentElementDataUtil::GetComponentFromHandleChecked(InElementHandle);
+	const UActorComponent* Component = ComponentElementDataUtil::GetComponentFromHandle(InElementHandle);
 	return Component && IsComponentSelected(Component, InSelectionSet, InSelectionOptions);
 }
 
 bool UComponentElementEditorSelectionInterface::ShouldPreventTransactions(const FTypedElementHandle& InElementHandle)
 {
-	const UActorComponent* Component = ComponentElementDataUtil::GetComponentFromHandleChecked(InElementHandle);
+	const UActorComponent* Component = ComponentElementDataUtil::GetComponentFromHandle(InElementHandle);
 	return Component && UObjectElementEditorSelectionInterface::ShouldObjectPreventTransactions(Component);
 }
 
@@ -44,9 +44,12 @@ bool UComponentElementEditorSelectionInterface::IsComponentSelected(const UActor
 		return false;
 	}
 
-	if (InSelectionSet->Contains(UEngineElementsLibrary::AcquireEditorComponentElementHandle(InComponent)))
+	if (FTypedElementHandle ComponentElement = UEngineElementsLibrary::AcquireEditorComponentElementHandle(InComponent, /*bAllowCreate*/false))
 	{
-		return true;
+		if (InSelectionSet->Contains(ComponentElement))
+		{
+			return true;
+		}
 	}
 
 	if (InSelectionOptions.AllowIndirect())
@@ -69,7 +72,10 @@ bool UComponentElementEditorSelectionInterface::IsComponentSelected(const UActor
 
 		if (ConsideredComponent)
 		{
-			return InSelectionSet->Contains(UEngineElementsLibrary::AcquireEditorComponentElementHandle(ConsideredComponent));
+			if (FTypedElementHandle ConsideredElement = UEngineElementsLibrary::AcquireEditorComponentElementHandle(ConsideredComponent, /*bAllowCreate*/false))
+			{
+				return InSelectionSet->Contains(ConsideredElement);
+			}
 		}
 	}
 
diff --git a/Engine/Source/Editor/UnrealEd/Private/LevelEditorViewport.cpp b/Engine/Source/Editor/UnrealEd/Private/LevelEditorViewport.cpp
--- a/Engine/Source/Editor/UnrealEd/Private/LevelEditorViewport.cpp
+++ b/Engine/Source/Editor/UnrealEd/Private/LevelEditorViewport.cpp
@@ -3422,6 +3422,24 @@ bool FLevelEditorViewportClient::HaveSelectedObjectsBeenChanged() const
 	return (TrackingTransaction.TransCount > 0 || TrackingTransaction.IsActive()) && (MouseDeltaTracker->HasReceivedDelta() || MouseDeltaTracker->WasExternalMovement());
 }
 
+bool FLevelEditorViewportClient::HasElementsToManipulate() const
+{
+	FLevelEditorViewportClient* MutableThis = const_cast<FLevelEditorViewportClient*>(this);
+	MutableThis->CacheElementsToManipulate();
+	const bool bHasElements = ElementsToManipulate->Num() > 0;
+	MutableThis->ElementsToManipulate->Reset();
+	return bHasElements;
+}
+
+TArray<FTypedElementHandle> FLevelEditorViewportClient::GetElementsToManipulate() const
+{
+	FLevelEditorViewportClient* MutableThis = const_cast<FLevelEditorViewportClient*>(this);
+	MutableThis->CacheElementsToManipulate();
+	TArray<FTypedElementHandle> ElementsToManipulateArray = ElementsToManipulate->GetElementHandles();
+	MutableThis->ElementsToManipulate->Reset();
+	return ElementsToManipulateArray;
+}
+
 void FLevelEditorViewportClient::CacheElementsToManipulate()
 {
 	const UTypedElementSelectionSet* SelectionSet = GetSelectionSet();
diff --git a/Engine/Source/Editor/UnrealEd/Private/Tools/DefaultEdMode.cpp b/Engine/Source/Editor/UnrealEd/Private/Tools/DefaultEdMode.cpp
--- a/Engine/Source/Editor/UnrealEd/Private/Tools/DefaultEdMode.cpp
+++ b/Engine/Source/Editor/UnrealEd/Private/Tools/DefaultEdMode.cpp
@@ -5,6 +5,17 @@
 #include "EditorModes.h"
 #include "EditorStyleSet.h"
 #include "Textures/SlateIcon.h"
+#include "LevelEditorViewport.h"
+
+class FLevelEditorSelectModeWidgetHelper : public FLegacyEdModeWidgetHelper
+{
+public:
+	virtual bool ShouldDrawWidget() const override
+	{
+		return GCurrentLevelEditingViewportClient
+			&& GCurrentLevelEditingViewportClient->HasElementsToManipulate();
+	}
+};
 
 UEdModeDefault::UEdModeDefault()
 {
@@ -24,3 +35,8 @@ bool UEdModeDefault::UsesToolkits() const
 {
 	return false;
 }
+
+TSharedRef<FLegacyEdModeWidgetHelper> UEdModeDefault::CreateWidgetHelper()
+{
+	return MakeShared<FLevelEditorSelectModeWidgetHelper>();
+}
diff --git a/Engine/Source/Editor/UnrealEd/Private/Tools/DefaultEdMode.h b/Engine/Source/Editor/UnrealEd/Private/Tools/DefaultEdMode.h
--- a/Engine/Source/Editor/UnrealEd/Private/Tools/DefaultEdMode.h
+++ b/Engine/Source/Editor/UnrealEd/Private/Tools/DefaultEdMode.h
@@ -16,4 +16,7 @@ public:
 
 	virtual bool UsesPropertyWidgets() const override;
 	virtual bool UsesToolkits() const override;
+
+protected:
+	virtual TSharedRef<FLegacyEdModeWidgetHelper> CreateWidgetHelper() override;
 };
diff --git a/Engine/Source/Editor/UnrealEd/Private/Tools/LegacyEdModeWidgetHelpers.cpp b/Engine/Source/Editor/UnrealEd/Private/Tools/LegacyEdModeWidgetHelpers.cpp
--- a/Engine/Source/Editor/UnrealEd/Private/Tools/LegacyEdModeWidgetHelpers.cpp
+++ b/Engine/Source/Editor/UnrealEd/Private/Tools/LegacyEdModeWidgetHelpers.cpp
@@ -767,13 +767,18 @@ UObject* FLegacyEdModeWidgetHelper::GetItemToTryDisplayingWidgetsFor(FTransform&
 
 void UBaseLegacyWidgetEdMode::Initialize()
 {
-	WidgetHelper = MakeShared<FLegacyEdModeWidgetHelper>();
+	WidgetHelper = CreateWidgetHelper();
 	WidgetHelper->Owner = this->Owner;
 	WidgetHelper->ParentModeInterface = this;
 
 	UEdMode::Initialize();
 }
 
+TSharedRef<FLegacyEdModeWidgetHelper> UBaseLegacyWidgetEdMode::CreateWidgetHelper()
+{
+	return MakeShared<FLegacyEdModeWidgetHelper>();
+}
+
 bool UBaseLegacyWidgetEdMode::InputDelta(FEditorViewportClient* InViewportClient, FViewport* InViewport, FVector& InDrag, FRotator& InRot, FVector& InScale)
 {
 	if (WidgetHelper->InputDelta(InViewportClient, InViewport, InDrag, InRot, InScale))
diff --git a/Engine/Source/Editor/UnrealEd/Public/LevelEditorViewport.h b/Engine/Source/Editor/UnrealEd/Public/LevelEditorViewport.h
--- a/Engine/Source/Editor/UnrealEd/Public/LevelEditorViewport.h
+++ b/Engine/Source/Editor/UnrealEd/Public/LevelEditorViewport.h
@@ -325,6 +325,16 @@ public:
 	void MirrorSelectedActors(const FVector& InMirrorScale);
 	void MirrorSelectedElements(const FVector& InMirrorScale);
 
+	/**
+	 * Are there any elements (from the current selection set) that this viewport can manipulate (eg, via the transform gizmo).
+	 */
+	bool HasElementsToManipulate() const;
+
+	/**
+	 * Get the array of elements (from the current selection set) that this viewport can manipulate (eg, via the transform gizmo).
+	 */
+	TArray<FTypedElementHandle> GetElementsToManipulate() const;
+
 	virtual void SetIsSimulateInEditorViewport( bool bInIsSimulateInEditorViewport ) override;
 
 	/**
diff --git a/Engine/Source/Editor/UnrealEd/Public/Tools/LegacyEdModeWidgetHelpers.h b/Engine/Source/Editor/UnrealEd/Public/Tools/LegacyEdModeWidgetHelpers.h
--- a/Engine/Source/Editor/UnrealEd/Public/Tools/LegacyEdModeWidgetHelpers.h
+++ b/Engine/Source/Editor/UnrealEd/Public/Tools/LegacyEdModeWidgetHelpers.h
@@ -124,7 +124,7 @@ protected:
 
 // This class is to aid transitioning from native FEdModes to UEdModes, in the case that the FEdMode used property widgets and/or transform widgets
 // To use this class:
-//   1. Sublcass FLegacyEdModeWidgetHelper, and override the methods needed for your native FEdMode implementation
+//   1. Subclass FLegacyEdModeWidgetHelper, and override the methods needed for your native FEdMode implementation
 //   2. Transition your native FEdMode to a UObject, by inheriting from this class
 //   3. Override CreateWidgetHelper function to return a SharedRef to the class you created in step 1.
 UCLASS(Abstract)
@@ -161,5 +161,7 @@ public:
 	virtual bool AllowsViewportDragTool() const override;
 
 protected:
+	virtual TSharedRef<FLegacyEdModeWidgetHelper> CreateWidgetHelper();
+
 	TSharedPtr<FLegacyEdModeWidgetHelper> WidgetHelper;
 };
diff --git a/Engine/Source/Runtime/Engine/Private/Elements/Actor/ActorElementWorldInterface.cpp b/Engine/Source/Runtime/Engine/Private/Elements/Actor/ActorElementWorldInterface.cpp
--- a/Engine/Source/Runtime/Engine/Private/Elements/Actor/ActorElementWorldInterface.cpp
+++ b/Engine/Source/Runtime/Engine/Private/Elements/Actor/ActorElementWorldInterface.cpp
@@ -7,6 +7,18 @@
 #include "Elements/Component/ComponentElementData.h"
 #include "Components/PrimitiveComponent.h"
 
+bool UActorElementWorldInterface::IsTemplateElement(const FTypedElementHandle& InElementHandle)
+{
+	const AActor* Actor = ActorElementDataUtil::GetActorFromHandle(InElementHandle);
+	return Actor && Actor->IsTemplate();
+}
+
+ULevel* UActorElementWorldInterface::GetOwnerLevel(const FTypedElementHandle& InElementHandle)
+{
+	const AActor* Actor = ActorElementDataUtil::GetActorFromHandle(InElementHandle);
+	return Actor ? Actor->GetLevel() : nullptr;
+}
+
 UWorld* UActorElementWorldInterface::GetOwnerWorld(const FTypedElementHandle& InElementHandle)
 {
 	const AActor* Actor = ActorElementDataUtil::GetActorFromHandle(InElementHandle);
diff --git a/Engine/Source/Runtime/Engine/Private/Elements/Component/ComponentElementWorldInterface.cpp b/Engine/Source/Runtime/Engine/Private/Elements/Component/ComponentElementWorldInterface.cpp
--- a/Engine/Source/Runtime/Engine/Private/Elements/Component/ComponentElementWorldInterface.cpp
+++ b/Engine/Source/Runtime/Engine/Private/Elements/Component/ComponentElementWorldInterface.cpp
@@ -8,19 +8,38 @@
 
 bool UComponentElementWorldInterface::CanEditElement(const FTypedElementHandle& InElementHandle)
 {
-	const UActorComponent* Component = ComponentElementDataUtil::GetComponentFromHandleChecked(InElementHandle);
+	const UActorComponent* Component = ComponentElementDataUtil::GetComponentFromHandle(InElementHandle);
 	return Component && Component->IsEditableWhenInherited();
 }
 
+bool UComponentElementWorldInterface::IsTemplateElement(const FTypedElementHandle& InElementHandle)
+{
+	const UActorComponent* Component = ComponentElementDataUtil::GetComponentFromHandle(InElementHandle);
+	return Component && Component->IsTemplate();
+}
+
+ULevel* UComponentElementWorldInterface::GetOwnerLevel(const FTypedElementHandle& InElementHandle)
+{
+	if (const UActorComponent* Component = ComponentElementDataUtil::GetComponentFromHandle(InElementHandle))
+	{
+		if (const AActor* ComponentOwner = Component->GetOwner())
+		{
+			return ComponentOwner->GetLevel();
+		}
+	}
+
+	return nullptr;
+}
+
 UWorld* UComponentElementWorldInterface::GetOwnerWorld(const FTypedElementHandle& InElementHandle)
 {
-	const UActorComponent* Component = ComponentElementDataUtil::GetComponentFromHandleChecked(InElementHandle);
+	const UActorComponent* Component = ComponentElementDataUtil::GetComponentFromHandle(InElementHandle);
 	return Component ? Component->GetWorld() : nullptr;
 }
 
 bool UComponentElementWorldInterface::GetBounds(const FTypedElementHandle& InElementHandle, FBoxSphereBounds& OutBounds)
 {
-	if (const UActorComponent* Component = ComponentElementDataUtil::GetComponentFromHandleChecked(InElementHandle))
+	if (const UActorComponent* Component = ComponentElementDataUtil::GetComponentFromHandle(InElementHandle))
 	{
 		if (const USceneComponent* SceneComponent = Cast<USceneComponent>(Component))
 		{
@@ -34,7 +53,7 @@ bool UComponentElementWorldInterface::GetBounds(const FTypedElementHandle& InEle
 
 bool UComponentElementWorldInterface::GetWorldTransform(const FTypedElementHandle& InElementHandle, FTransform& OutTransform)
 {
-	if (const UActorComponent* Component = ComponentElementDataUtil::GetComponentFromHandleChecked(InElementHandle))
+	if (const UActorComponent* Component = ComponentElementDataUtil::GetComponentFromHandle(InElementHandle))
 	{
 		if (const USceneComponent* SceneComponent = Cast<USceneComponent>(Component))
 		{
@@ -48,7 +67,7 @@ bool UComponentElementWorldInterface::GetWorldTransform(const FTypedElementHandl
 
 bool UComponentElementWorldInterface::SetWorldTransform(const FTypedElementHandle& InElementHandle, const FTransform& InTransform)
 {
-	if (UActorComponent* Component = ComponentElementDataUtil::GetComponentFromHandleChecked(InElementHandle))
+	if (UActorComponent* Component = ComponentElementDataUtil::GetComponentFromHandle(InElementHandle))
 	{
 		if (USceneComponent* SceneComponent = Cast<USceneComponent>(Component))
 		{
@@ -63,7 +82,7 @@ bool UComponentElementWorldInterface::SetWorldTransform(const FTypedElementHandl
 
 bool UComponentElementWorldInterface::GetRelativeTransform(const FTypedElementHandle& InElementHandle, FTransform& OutTransform)
 {
-	if (const UActorComponent* Component = ComponentElementDataUtil::GetComponentFromHandleChecked(InElementHandle))
+	if (const UActorComponent* Component = ComponentElementDataUtil::GetComponentFromHandle(InElementHandle))
 	{
 		if (const USceneComponent* SceneComponent = Cast<USceneComponent>(Component))
 		{
@@ -77,7 +96,7 @@ bool UComponentElementWorldInterface::GetRelativeTransform(const FTypedElementHa
 
 bool UComponentElementWorldInterface::SetRelativeTransform(const FTypedElementHandle& InElementHandle, const FTransform& InTransform)
 {
-	if (UActorComponent* Component = ComponentElementDataUtil::GetComponentFromHandleChecked(InElementHandle))
+	if (UActorComponent* Component = ComponentElementDataUtil::GetComponentFromHandle(InElementHandle))
 	{
 		if (USceneComponent* SceneComponent = Cast<USceneComponent>(Component))
 		{
@@ -92,7 +111,7 @@ bool UComponentElementWorldInterface::SetRelativeTransform(const FTypedElementHa
 
 bool UComponentElementWorldInterface::FindSuitableTransformAlongPath(const FTypedElementHandle& InElementHandle, const FVector& InPathStart, const FVector& InPathEnd, const FCollisionShape& InTestShape, TArrayView<const FTypedElementHandle> InElementsToIgnore, FTransform& OutSuitableTransform)
 {
-	if (const UActorComponent* Component = ComponentElementDataUtil::GetComponentFromHandleChecked(InElementHandle))
+	if (const UActorComponent* Component = ComponentElementDataUtil::GetComponentFromHandle(InElementHandle))
 	{
 		if (const UWorld* World = Component->GetWorld())
 		{
diff --git a/Engine/Source/Runtime/Engine/Public/Elements/Actor/ActorElementWorldInterface.h b/Engine/Source/Runtime/Engine/Public/Elements/Actor/ActorElementWorldInterface.h
--- a/Engine/Source/Runtime/Engine/Public/Elements/Actor/ActorElementWorldInterface.h
+++ b/Engine/Source/Runtime/Engine/Public/Elements/Actor/ActorElementWorldInterface.h
@@ -13,6 +13,8 @@ class ENGINE_API UActorElementWorldInterface : public UTypedElementWorldInterfac
 	GENERATED_BODY()
 
 public:
+	virtual bool IsTemplateElement(const FTypedElementHandle& InElementHandle) override;
+	virtual ULevel* GetOwnerLevel(const FTypedElementHandle& InElementHandle) override;
 	virtual UWorld* GetOwnerWorld(const FTypedElementHandle& InElementHandle) override;
 	virtual bool GetBounds(const FTypedElementHandle& InElementHandle, FBoxSphereBounds& OutBounds) override;
 	virtual bool GetWorldTransform(const FTypedElementHandle& InElementHandle, FTransform& OutTransform) override;
diff --git a/Engine/Source/Runtime/Engine/Public/Elements/Component/ComponentElementWorldInterface.h b/Engine/Source/Runtime/Engine/Public/Elements/Component/ComponentElementWorldInterface.h
--- a/Engine/Source/Runtime/Engine/Public/Elements/Component/ComponentElementWorldInterface.h
+++ b/Engine/Source/Runtime/Engine/Public/Elements/Component/ComponentElementWorldInterface.h
@@ -12,6 +12,8 @@ class ENGINE_API UComponentElementWorldInterface : public UTypedElementWorldInte
 
 public:
 	virtual bool CanEditElement(const FTypedElementHandle& InElementHandle) override;
+	virtual bool IsTemplateElement(const FTypedElementHandle& InElementHandle) override;
+	virtual ULevel* GetOwnerLevel(const FTypedElementHandle& InElementHandle) override;
 	virtual UWorld* GetOwnerWorld(const FTypedElementHandle& InElementHandle) override;
 	virtual bool GetBounds(const FTypedElementHandle& InElementHandle, FBoxSphereBounds& OutBounds) override;
 	virtual bool GetWorldTransform(const FTypedElementHandle& InElementHandle, FTransform& OutTransform) override;
diff --git a/Engine/Source/Runtime/Engine/Public/Elements/Interfaces/TypedElementWorldInterface.h b/Engine/Source/Runtime/Engine/Public/Elements/Interfaces/TypedElementWorldInterface.h
--- a/Engine/Source/Runtime/Engine/Public/Elements/Interfaces/TypedElementWorldInterface.h
+++ b/Engine/Source/Runtime/Engine/Public/Elements/Interfaces/TypedElementWorldInterface.h
@@ -5,6 +5,7 @@
 #include "Elements/Framework/TypedElementHandle.h"
 #include "TypedElementWorldInterface.generated.h"
 
+class ULevel;
 class UWorld;
 struct FCollisionShape;
 
@@ -14,12 +15,24 @@ class ENGINE_API UTypedElementWorldInterface : public UTypedElementInterface
 	GENERATED_BODY()
 
 public:
+	/**
+	 * Is this element considered a template within its world (eg, a CDO or archetype).
+	 */
+	UFUNCTION(BlueprintPure, Category="TypedElementInterfaces|World")
+	virtual bool IsTemplateElement(const FTypedElementHandle& InElementHandle) { return false; }
+
 	/**
 	 * Can this element actually be edited in the world?
 	 */
 	UFUNCTION(BlueprintPure, Category="TypedElementInterfaces|World")
 	virtual bool CanEditElement(const FTypedElementHandle& InElementHandle) { return true; }
 
+	/**
+	 * Get the owner level associated with this element, if any.
+	 */
+	UFUNCTION(BlueprintPure, Category="TypedElementInterfaces|World")
+	virtual ULevel* GetOwnerLevel(const FTypedElementHandle& InElementHandle) { return nullptr; }
+
 	/**
 	 * Get the owner world associated with this element, if any.
 	 */
@@ -56,6 +69,18 @@ public:
 	UFUNCTION(BlueprintCallable, Category="TypedElementInterfaces|World")
 	virtual bool SetRelativeTransform(const FTypedElementHandle& InElementHandle, const FTransform& InTransform) { return SetWorldTransform(InElementHandle, InTransform); }
 
+	/**
+	 * Get the local space offset of this element that should be added to its pivot location, if any.
+	 */
+	UFUNCTION(BlueprintPure, Category="TypedElementInterfaces|World")
+	virtual bool GetPivotOffset(const FTypedElementHandle& InElementHandle, FVector& OutPivotOffset) { return false; }
+
+	/**
+	 * Attempt to set the local space offset of this element that should be added to its pivot location.
+	 */
+	UFUNCTION(BlueprintCallable, Category="TypedElementInterfaces|World")
+	virtual bool SetPivotOffset(const FTypedElementHandle& InElementHandle, const FVector& InPivotOffset) { return false; }
+
 	/**
 	 * Notify that this element is about to be moved.
 	 */
@@ -115,13 +140,17 @@ public:
 template <>
 struct TTypedElement<UTypedElementWorldInterface> : public TTypedElementBase<UTypedElementWorldInterface>
 {
+	bool IsTemplateElement() const { return InterfacePtr->IsTemplateElement(*this); }
 	bool CanEditElement() const { return InterfacePtr->CanEditElement(*this); }
+	ULevel* GetOwnerLevel() const { return InterfacePtr->GetOwnerLevel(*this); }
 	UWorld* GetOwnerWorld() const { return InterfacePtr->GetOwnerWorld(*this); }
 	bool GetBounds(FBoxSphereBounds& OutBounds) const { return InterfacePtr->GetBounds(*this, OutBounds); }
 	bool GetWorldTransform(FTransform& OutTransform) const { return InterfacePtr->GetWorldTransform(*this, OutTransform); }
 	bool SetWorldTransform(const FTransform& InTransform) const { return InterfacePtr->SetWorldTransform(*this, InTransform); }
 	bool GetRelativeTransform(FTransform& OutTransform) const { return InterfacePtr->GetRelativeTransform(*this, OutTransform); }
 	bool SetRelativeTransform(const FTransform& InTransform) const { return InterfacePtr->SetRelativeTransform(*this, InTransform); }
+	bool GetPivotOffset(FVector& OutPivotOffset) const { return InterfacePtr->GetPivotOffset(*this, OutPivotOffset); }
+	bool SetPivotOffset(const FVector& InPivotOffset) const { return InterfacePtr->SetPivotOffset(*this, InPivotOffset); }
 	void NotifyMovementStarted() const { InterfacePtr->NotifyMovementStarted(*this); }
 	void NotifyMovementOngoing() const { InterfacePtr->NotifyMovementOngoing(*this); }
 	void NotifyMovementEnded() const { InterfacePtr->NotifyMovementEnded(*this); }
