Added missing override

#rnx

diff --git a/Engine/Source/Editor/UnrealEd/Private/Elements/Actor/ActorElementEditorWorldInterface.cpp b/Engine/Source/Editor/UnrealEd/Private/Elements/Actor/ActorElementEditorWorldInterface.cpp
--- a/Engine/Source/Editor/UnrealEd/Private/Elements/Actor/ActorElementEditorWorldInterface.cpp
+++ b/Engine/Source/Editor/UnrealEd/Private/Elements/Actor/ActorElementEditorWorldInterface.cpp
@@ -61,6 +61,12 @@ void UActorElementEditorWorldInterface::NotifyMovementEnded(const FTypedElementH
 	}
 }
 
+bool UActorElementEditorWorldInterface::CanDeleteElement(const FTypedElementHandle& InElementHandle)
+{
+	AActor* Actor = ActorElementDataUtil::GetActorFromHandle(InElementHandle);
+	return Actor && GUnrealEd->CanDeleteActor(Actor);
+}
+
 bool UActorElementEditorWorldInterface::DeleteElements(TArrayView<const FTypedElementHandle> InElementHandles, UWorld* InWorld, UTypedElementSelectionSet* InSelectionSet, const FTypedElementDeletionOptions& InDeletionOptions)
 {
 	const TArray<AActor*> ActorsToDelete = ActorElementDataUtil::GetActorsFromHandles(InElementHandles);
diff --git a/Engine/Source/Editor/UnrealEd/Private/Elements/Actor/ActorElementEditorWorldInterface.h b/Engine/Source/Editor/UnrealEd/Private/Elements/Actor/ActorElementEditorWorldInterface.h
--- a/Engine/Source/Editor/UnrealEd/Private/Elements/Actor/ActorElementEditorWorldInterface.h
+++ b/Engine/Source/Editor/UnrealEd/Private/Elements/Actor/ActorElementEditorWorldInterface.h
@@ -16,6 +16,7 @@ public:
 	virtual void NotifyMovementStarted(const FTypedElementHandle& InElementHandle) override;
 	virtual void NotifyMovementOngoing(const FTypedElementHandle& InElementHandle) override;
 	virtual void NotifyMovementEnded(const FTypedElementHandle& InElementHandle) override;
+	virtual bool CanDeleteElement(const FTypedElementHandle& InElementHandle) override;
 	virtual bool DeleteElements(TArrayView<const FTypedElementHandle> InElementHandles, UWorld* InWorld, UTypedElementSelectionSet* InSelectionSet, const FTypedElementDeletionOptions& InDeletionOptions) override;
 	virtual void DuplicateElements(TArrayView<const FTypedElementHandle> InElementHandles, UWorld* InWorld, bool bOffsetLocations, TArray<FTypedElementHandle>& OutNewElements) override;
 };
diff --git a/Engine/Source/Editor/UnrealEd/Private/Elements/Component/ComponentElementEditorWorldInterface.cpp b/Engine/Source/Editor/UnrealEd/Private/Elements/Component/ComponentElementEditorWorldInterface.cpp
--- a/Engine/Source/Editor/UnrealEd/Private/Elements/Component/ComponentElementEditorWorldInterface.cpp
+++ b/Engine/Source/Editor/UnrealEd/Private/Elements/Component/ComponentElementEditorWorldInterface.cpp
@@ -44,6 +44,12 @@ void UComponentElementEditorWorldInterface::NotifyMovementEnded(const FTypedElem
 	}
 }
 
+bool UComponentElementEditorWorldInterface::CanDeleteElement(const FTypedElementHandle& InElementHandle)
+{
+	UActorComponent* Component = ComponentElementDataUtil::GetComponentFromHandle(InElementHandle);
+	return Component && GUnrealEd->CanDeleteComponent(Component);
+}
+
 bool UComponentElementEditorWorldInterface::DeleteElements(TArrayView<const FTypedElementHandle> InElementHandles, UWorld* InWorld, UTypedElementSelectionSet* InSelectionSet, const FTypedElementDeletionOptions& InDeletionOptions)
 {
 	const TArray<UActorComponent*> ComponentsToDelete = ComponentElementDataUtil::GetComponentsFromHandles(InElementHandles);
diff --git a/Engine/Source/Editor/UnrealEd/Private/Elements/Component/ComponentElementEditorWorldInterface.h b/Engine/Source/Editor/UnrealEd/Private/Elements/Component/ComponentElementEditorWorldInterface.h
--- a/Engine/Source/Editor/UnrealEd/Private/Elements/Component/ComponentElementEditorWorldInterface.h
+++ b/Engine/Source/Editor/UnrealEd/Private/Elements/Component/ComponentElementEditorWorldInterface.h
@@ -14,6 +14,7 @@ public:
 	virtual void NotifyMovementStarted(const FTypedElementHandle& InElementHandle) override;
 	virtual void NotifyMovementOngoing(const FTypedElementHandle& InElementHandle) override;
 	virtual void NotifyMovementEnded(const FTypedElementHandle& InElementHandle) override;
+	virtual bool CanDeleteElement(const FTypedElementHandle& InElementHandle) override;
 	virtual bool DeleteElements(TArrayView<const FTypedElementHandle> InElementHandles, UWorld* InWorld, UTypedElementSelectionSet* InSelectionSet, const FTypedElementDeletionOptions& InDeletionOptions) override;
 	virtual void DuplicateElements(TArrayView<const FTypedElementHandle> InElementHandles, UWorld* InWorld, bool bOffsetLocations, TArray<FTypedElementHandle>& OutNewElements) override;
 };
