Added missing transaction

#rnx

diff --git a/Engine/Source/Editor/UnrealEd/Private/UnrealEdSrv.cpp b/Engine/Source/Editor/UnrealEd/Private/UnrealEdSrv.cpp
--- a/Engine/Source/Editor/UnrealEd/Private/UnrealEdSrv.cpp
+++ b/Engine/Source/Editor/UnrealEd/Private/UnrealEdSrv.cpp
@@ -2498,8 +2498,12 @@ bool UUnrealEdEngine::Exec_Actor( UWorld* InWorld, const TCHAR* Str, FOutputDevi
 
 				TStrongObjectPtr<UTypedElementList> ElementsToDelete(UTypedElementRegistry::GetInstance()->CreateElementList());
 				CommonActions->GetSelectedElementsToDelete(SelectionSet, ElementsToDelete.Get());
-				CommonActions->DeleteElements(ElementsToDelete.Get(), InWorld, SelectionSet, FTypedElementDeletionOptions());
-				ElementsToDelete->Reset();
+				if (ElementsToDelete->Num() > 0)
+				{
+					const FScopedTransaction Transaction(NSLOCTEXT("UnrealEd", "DeleteElements", "Delete Elements"));
+					CommonActions->DeleteElements(ElementsToDelete.Get(), InWorld, SelectionSet, FTypedElementDeletionOptions());
+					ElementsToDelete->Reset();
+				}
 			}
 		}
 		return true;
