Move InteractiveToolsContext calls to mode manager from UEdMode.

diff --git a/Engine/Source/Editor/UnrealEd/Private/EditorModeManager.cpp b/Engine/Source/Editor/UnrealEd/Private/EditorModeManager.cpp
--- a/Engine/Source/Editor/UnrealEd/Private/EditorModeManager.cpp
+++ b/Engine/Source/Editor/UnrealEd/Private/EditorModeManager.cpp
@@ -229,7 +229,7 @@ bool FEditorModeTools::ProcessEditDuplicate()
 
 bool FEditorModeTools::ProcessEditDelete()
 {
-	bool bHandled = false;
+	bool bHandled = InteractiveToolsContext->ProcessEditDelete();
 	for (UEdMode* Mode : ActiveScriptableModes)
 	{
 		bHandled = Mode->ProcessEditDelete();
@@ -1015,10 +1015,9 @@ EAxisList::Type FEditorModeTools::GetWidgetAxisToDraw( UE::Widget::EWidgetMode I
 bool FEditorModeTools::StartTracking(FEditorViewportClient* InViewportClient, FViewport* InViewport)
 {
 	bIsTracking = true;
-	bool bTransactionHandled = false;
-
 	CachedLocation = PivotLocation;	// Cache the pivot location
 
+	bool bTransactionHandled = InteractiveToolsContext->StartTracking(InViewportClient, InViewport);
 	for (UEdMode* Mode : ActiveScriptableModes)
 	{
 		bTransactionHandled |= Mode->StartTracking(InViewportClient, InViewportClient->Viewport);
@@ -1031,7 +1030,7 @@ bool FEditorModeTools::StartTracking(FEditorViewportClient* InViewportClient, FV
 bool FEditorModeTools::EndTracking(FEditorViewportClient* InViewportClient, FViewport* InViewport)
 {
 	bIsTracking = false;
-	bool bTransactionHandled = false;
+	bool bTransactionHandled = InteractiveToolsContext->EndTracking(InViewportClient, InViewport);
 
 	for (UEdMode* Mode : ActiveScriptableModes)
 	{
@@ -1218,6 +1217,7 @@ void FEditorModeTools::Tick( FEditorViewportClient* ViewportClient, float DeltaT
 		ActivateDefaultMode();
 	}
 
+	InteractiveToolsContext->Tick(ViewportClient, DeltaTime);
 	for (UEdMode* Mode : ActiveScriptableModes)
 	{
 		Mode->Tick(ViewportClient, DeltaTime);
@@ -1238,7 +1238,7 @@ bool FEditorModeTools::InputDelta( FEditorViewportClient* InViewportClient,FView
 /** Notifies all active modes of captured mouse movement */	
 bool FEditorModeTools::CapturedMouseMove( FEditorViewportClient* InViewportClient, FViewport* InViewport, int32 InMouseX, int32 InMouseY )
 {
-	bool bHandled = false;
+	bool bHandled = InteractiveToolsContext->CapturedMouseMove(InViewportClient, InViewport, InMouseX, InMouseY);
 	for (UEdMode* Mode : ActiveScriptableModes)
 	{
 		bHandled |= Mode->CapturedMouseMove(InViewportClient, InViewport, InMouseX, InMouseY);
@@ -1260,14 +1260,21 @@ bool FEditorModeTools::ProcessCapturedMouseMoves( FEditorViewportClient* InViewp
 /** Notifies all active modes of keyboard input */
 bool FEditorModeTools::InputKey(FEditorViewportClient* InViewportClient, FViewport* Viewport, FKey Key, EInputEvent Event)
 {
+	bool bToolsContextHandled = InteractiveToolsContext->InputKey(InViewportClient, Viewport, Key, Event);
 	bool bHandled = false;
+	
 	//Copy the modes and iterate of that since a key may remove the edit mode and change CopyActiveScriptableModes
 	TArray<UEdMode*> CopyActiveScriptableModes(ActiveScriptableModes);
 	for (UEdMode* Mode : CopyActiveScriptableModes)
 	{
-		bHandled |= Mode->InputKey(InViewportClient, Viewport, Key, Event);
+		// If the interactive tools context did not handle the InputKey, then pass along to the UEdMode
+		// Always pass through for the FEdModes to handle for backwards compat.
+		if (!bToolsContextHandled || Mode->AsLegacyMode())
+		{
+			bHandled |= Mode->InputKey(InViewportClient, Viewport, Key, Event);
+		}
 	}
-	return bHandled;
+	return bHandled || bToolsContextHandled;
 }
 
 /** Notifies all active modes of axis movement */
@@ -1297,7 +1304,7 @@ bool FEditorModeTools::GetPivotForOrbit( FVector& Pivot ) const
 bool FEditorModeTools::MouseEnter( FEditorViewportClient* InViewportClient, FViewport* Viewport, int32 X, int32 Y )
 {
 	HoveredViewportClient = InViewportClient;
-	bool bHandled = false;
+	bool bHandled = InteractiveToolsContext->MouseEnter(InViewportClient, Viewport, X, Y);;
 	for (UEdMode* Mode : ActiveScriptableModes)
 	{
 		bHandled |= Mode->MouseEnter(InViewportClient, Viewport, X, Y);
@@ -1307,7 +1314,8 @@ bool FEditorModeTools::MouseEnter( FEditorViewportClient* InViewportClient, FVie
 
 bool FEditorModeTools::MouseLeave( FEditorViewportClient* InViewportClient, FViewport* Viewport )
 {
-	bool bHandled = false;
+	bool bHandled = InteractiveToolsContext->MouseLeave(InViewportClient, Viewport);
+
 	for (UEdMode* Mode : ActiveScriptableModes)
 	{
 		bHandled |= Mode->MouseLeave(InViewportClient, Viewport);
@@ -1318,7 +1326,8 @@ bool FEditorModeTools::MouseLeave( FEditorViewportClient* InViewportClient, FVie
 /** Notifies all active modes that the mouse has moved */
 bool FEditorModeTools::MouseMove( FEditorViewportClient* InViewportClient, FViewport* Viewport, int32 X, int32 Y )
 {
-	bool bHandled = false;
+	bool bHandled = InteractiveToolsContext->MouseMove(InViewportClient, Viewport, X, Y);
+
 	for (UEdMode* Mode : ActiveScriptableModes)
 	{
 		bHandled |= Mode->MouseMove(InViewportClient, Viewport, X, Y);
@@ -1362,6 +1371,7 @@ void FEditorModeTools::DrawActiveModes( const FSceneView* InView, FPrimitiveDraw
 /** Renders all active modes */
 void FEditorModeTools::Render( const FSceneView* InView, FViewport* Viewport, FPrimitiveDrawInterface* PDI )
 {
+	InteractiveToolsContext->Render(InView, Viewport, PDI);
 	for (UEdMode* Mode : ActiveScriptableModes)
 	{
 		Mode->Render(InView, Viewport, PDI);
@@ -1371,6 +1381,7 @@ void FEditorModeTools::Render( const FSceneView* InView, FViewport* Viewport, FP
 /** Draws the HUD for all active modes */
 void FEditorModeTools::DrawHUD( FEditorViewportClient* InViewportClient,FViewport* Viewport, const FSceneView* View, FCanvas* Canvas )
 {
+	InteractiveToolsContext->DrawHUD(InViewportClient, Viewport, View, Canvas);
 	for (UEdMode* Mode : ActiveScriptableModes)
 	{
 		Mode->DrawHUD(InViewportClient, Viewport, View, Canvas);
diff --git a/Engine/Source/Editor/UnrealEd/Private/Tools/UEdMode.cpp b/Engine/Source/Editor/UnrealEd/Private/Tools/UEdMode.cpp
--- a/Engine/Source/Editor/UnrealEd/Private/Tools/UEdMode.cpp
+++ b/Engine/Source/Editor/UnrealEd/Private/Tools/UEdMode.cpp
@@ -75,31 +75,16 @@ void UEdMode::Initialize()
 
 bool UEdMode::MouseEnter(FEditorViewportClient* ViewportClient, FViewport* Viewport, int32 x, int32 y)
 {
-	if (ToolsContext->MouseEnter(ViewportClient, Viewport, x, y))
-	{
-		return true;
-	}
-
 	return false;
 }
 
 bool UEdMode::MouseLeave(FEditorViewportClient* ViewportClient, FViewport* Viewport)
 {
-	if (ToolsContext->MouseLeave(ViewportClient, Viewport))
-	{
-		return true;
-	}
-
 	return false;
 }
 
 bool UEdMode::MouseMove(FEditorViewportClient* ViewportClient, FViewport* Viewport, int32 x, int32 y)
 {
-	if (ToolsContext->MouseMove(ViewportClient, Viewport, x, y))
-	{
-		return true;
-	}
-
 	return false;
 }
 
@@ -115,12 +100,6 @@ bool UEdMode::LostFocus(FEditorViewportClient* ViewportClient, FViewport* Viewpo
 
 bool UEdMode::CapturedMouseMove(FEditorViewportClient* InViewportClient, FViewport* InViewport, int32 InMouseX, int32 InMouseY)
 {
-	// if alt is down we will not allow client to see this event
-	if (ToolsContext->CapturedMouseMove(InViewportClient, InViewport, InMouseX, InMouseY))
-	{
-		return true;
-	}
-
 	return false;
 }
 
@@ -131,11 +110,6 @@ bool UEdMode::InputKey(FEditorViewportClient* ViewportClient, FViewport* Viewpor
 		return false;
 	}
 
-	if (ToolsContext->InputKey(ViewportClient, Viewport, Key, Event))
-	{
-		return true;
-	}
-
 	if (Event != IE_Released)
 	{
 		if (ToolsContext->ShouldIgnoreHotkeys() == false)		// allow the context to capture keyboard input if necessary
@@ -181,23 +155,15 @@ bool UEdMode::InputDelta(FEditorViewportClient* InViewportClient, FViewport* InV
 
 void UEdMode::SelectNone()
 {
-	GEditor->SelectNone(true, true);
 }
 
 bool UEdMode::ProcessEditDelete()
 {
-	if (ToolsContext->ProcessEditDelete())
-	{
-		return true;
-	}
-
 	return false;
 }
 
 void UEdMode::Tick(FEditorViewportClient* ViewportClient, float DeltaTime)
 {
-	// Tick the ToolsContext
-	ToolsContext->Tick(ViewportClient, DeltaTime);
 }
 
 
@@ -308,21 +274,10 @@ UTexture2D* UEdMode::GetVertexTexture()
 
 void UEdMode::Render(const FSceneView* View, FViewport* Viewport, FPrimitiveDrawInterface* PDI)
 {
-	// give ToolsContext a chance to render
-	if (ToolsContext.IsValid())
-	{
-		ToolsContext->Render(View, Viewport, PDI);
-	}
 }
 
 void UEdMode::DrawHUD(FEditorViewportClient* ViewportClient, FViewport* Viewport, const FSceneView* View, FCanvas* Canvas)
 {
-	// give ToolsContext a chance to draw in screen space
-	if (ToolsContext != nullptr)
-	{
-		ToolsContext->DrawHUD(ViewportClient, Viewport, View, Canvas);
-	}
-
 	// Render the drag tool.
 	ViewportClient->RenderDragTool(View, Canvas);
 
@@ -458,21 +413,11 @@ class FEditorModeTools* UEdMode::GetModeManager() const
 
 bool UEdMode::StartTracking(FEditorViewportClient* InViewportClient, FViewport* InViewport)
 {
-	if (ToolsContext->StartTracking(InViewportClient, InViewport))
-	{
-		return true;
-	}
-
 	return false; 
 }
 
 bool UEdMode::EndTracking(FEditorViewportClient* InViewportClient, FViewport* InViewport)
 {
-	if (ToolsContext->EndTracking(InViewportClient, InViewport))
-	{
-		return true;
-	}
-
 	return false; 
 }
 
