Fix for CL15273316

diff --git a/Engine/Source/Editor/UnrealEd/Private/EditorModeManager.cpp b/Engine/Source/Editor/UnrealEd/Private/EditorModeManager.cpp
--- a/Engine/Source/Editor/UnrealEd/Private/EditorModeManager.cpp
+++ b/Engine/Source/Editor/UnrealEd/Private/EditorModeManager.cpp
@@ -649,7 +649,6 @@ void FEditorModeTools::RebuildModeToolBar()
 						if (PaletteCount > 1)
 						{
 							PaletteTabBox->AddSlot()
-							.AutoWidth()
 							.Padding(FMargin(0.f, 1.0f, 1.0f, 0.0f))
 							[
 								SNew(SCheckBox)
