Marked some function as final that should no longer be overridden

The older "edact" functions are no longer called directly by typed elements, so any overrides to those would be missed. Client code should override the more specific implementation functions instead.

diff --git a/Engine/Source/Editor/UnrealEd/Classes/Editor/UnrealEdEngine.h b/Engine/Source/Editor/UnrealEd/Classes/Editor/UnrealEdEngine.h
--- a/Engine/Source/Editor/UnrealEd/Classes/Editor/UnrealEdEngine.h
+++ b/Engine/Source/Editor/UnrealEd/Classes/Editor/UnrealEdEngine.h
@@ -208,7 +208,7 @@ public:
 	virtual void RedrawLevelEditingViewports(bool bInvalidateHitProxies=true) override;
 	virtual void TakeHighResScreenShots() override;
 	virtual void GetPackageList( TArray<UPackage*>* InPackages, UClass* InClass ) override;
-	virtual bool ShouldAbortActorDeletion() const override;
+	virtual bool ShouldAbortActorDeletion() const override final; // @note Final - Override ShouldAbortComponentDeletion or ShouldAbortActorDeletion (with parameters) instead.
 	virtual void CloseEditor() override;
 	virtual void OnOpenMatinee() override;
 	virtual bool IsAutosaving() const override;
@@ -314,6 +314,7 @@ public:
 	/**
 	 * Iterate over all levels of the world and create a list of world infos, then
 	 * Iterate over selected actors and assemble a list of actors which can be deleted.
+	 * @see CanDeleteComponent and CanDeleteActor.
 	 *
 	 * @param	InWorld					The world we want to examine
 	 * @param	bStopAtFirst			Whether or not we should stop at the first deletable actor we encounter
@@ -418,7 +419,7 @@ public:
 	 * @param InComponent				Component to check
 	 * @param OutReason					Optional value to fill with the reason the component cannot be deleted, if any
 	 */
-	bool CanDeleteComponent(const UActorComponent* InComponent, FText* OutReason = nullptr) const;
+	virtual bool CanDeleteComponent(const UActorComponent* InComponent, FText* OutReason = nullptr) const;
 
 	/**
 	 * Can the given actor be deleted?
@@ -426,7 +427,7 @@ public:
 	 * @param InActor					Actor to check
 	 * @param OutReason					Optional value to fill with the reason the actor cannot be deleted, if any
 	 */
-	bool CanDeleteActor(const AActor* InActor, FText* OutReason = nullptr) const;
+	virtual bool CanDeleteActor(const AActor* InActor, FText* OutReason = nullptr) const;
 
 	/**
 	 * Should the deletion of the given components be outright aborted?
@@ -434,7 +435,7 @@ public:
 	 * @param InComponentsToDelete		Components to check
 	 * @param OutReason					Optional value to fill with the reason the component deletion was aborted, if any
 	 */
-	bool ShouldAbortComponentDeletion(const TArray<UActorComponent*>& InComponentsToDelete, FText* OutReason = nullptr) const;
+	virtual bool ShouldAbortComponentDeletion(const TArray<UActorComponent*>& InComponentsToDelete, FText* OutReason = nullptr) const;
 
 	/**
 	 * Should the deletion of the given actors be outright aborted?
@@ -442,7 +443,7 @@ public:
 	 * @param InActorsToDelete			Actors to check
 	 * @param OutReason					Optional value to fill with the reason the actor deletion was aborted, if any
 	 */
-	bool ShouldAbortActorDeletion(const TArray<AActor*>& InActorsToDelete, FText* OutReason = nullptr) const;
+	virtual bool ShouldAbortActorDeletion(const TArray<AActor*>& InActorsToDelete, FText* OutReason = nullptr) const;
 
 	/**
 	 * Delete the given components.
@@ -454,7 +455,7 @@ public:
 	 * 
 	 * @return								true unless the delete operation was aborted.
 	 */
-	bool DeleteComponents(const TArray<UActorComponent*>& InComponentsToDelete, UTypedElementSelectionSet* InSelectionSet, const bool bVerifyDeletionCanHappen = true);
+	virtual bool DeleteComponents(const TArray<UActorComponent*>& InComponentsToDelete, UTypedElementSelectionSet* InSelectionSet, const bool bVerifyDeletionCanHappen = true);
 
 	/**
 	 * Deletes the given actors.
@@ -468,10 +469,11 @@ public:
 	 * 
 	 * @return								true unless the delete operation was aborted.
 	 */
-	bool DeleteActors(const TArray<AActor*>& InActorsToDelete, UWorld* InWorld, UTypedElementSelectionSet* InSelectionSet, const bool bVerifyDeletionCanHappen = true, const bool bWarnAboutReferences = true, const bool bWarnAboutSoftReferences = true);
+	virtual bool DeleteActors(const TArray<AActor*>& InActorsToDelete, UWorld* InWorld, UTypedElementSelectionSet* InSelectionSet, const bool bVerifyDeletionCanHappen = true, const bool bWarnAboutReferences = true, const bool bWarnAboutSoftReferences = true);
 
 	/**
 	 * Deletes all selected actors
+	 * @note Final - Override DeleteComponents or DeleteActors instead
 	 *
 	 * @param	InWorld						World context
 	 * @param	bVerifyDeletionCanHappen	[opt] If true (default), verify that deletion can be performed.
@@ -479,7 +481,7 @@ public:
 	 * @param	bWarnAboutSoftReferences	[opt] If true (default), we prompt the user about soft references to actors they are about to delete
 	 * @return								true unless the delete operation was aborted.
 	 */
-	virtual bool edactDeleteSelected( UWorld* InWorld, bool bVerifyDeletionCanHappen=true, bool bWarnAboutReferences = true, bool bWarnAboutSoftReferences = true) override;
+	virtual bool edactDeleteSelected( UWorld* InWorld, bool bVerifyDeletionCanHappen=true, bool bWarnAboutReferences = true, bool bWarnAboutSoftReferences = true) override final;
 
 	/**
 	 * Creates a new group from the current selection removing any existing groups.
@@ -520,11 +522,12 @@ public:
 	
 	/**
 	 * Copy selected actors to the clipboard.  Does not copy PrefabInstance actors or parts of Prefabs.
+	 * @note Final - Override CopyComponents or CopyActors instead.
 	 *
 	 * @param	InWorld					World context
 	 * @param	DestinationData			If != NULL, fill instead of clipboard data
 	 */
-	virtual void edactCopySelected(UWorld* InWorld, FString* DestinationData = nullptr) override;
+	virtual void edactCopySelected(UWorld* InWorld, FString* DestinationData = nullptr) override final;
 
 	/**
 	 * Copy the given components to the clipboard.
@@ -532,7 +535,7 @@ public:
 	 * @param	InComponentsToCopy		Array of components to copy
 	 * @param	DestinationData			If != NULL, fill instead of clipboard data
 	 */
-	void CopyComponents(const TArray<UActorComponent*>& InComponentsToCopy, FString* DestinationData = nullptr) const;
+	virtual void CopyComponents(const TArray<UActorComponent*>& InComponentsToCopy, FString* DestinationData = nullptr) const;
 
 	/**
 	 * Copy the given actors to the clipboard.  Does not copy PrefabInstance actors or parts of Prefabs.
@@ -541,10 +544,11 @@ public:
 	 * @param	InWorld					World context
 	 * @param	DestinationData			If != NULL, fill instead of clipboard data
 	 */
-	void CopyActors(const TArray<AActor*>& InActorsToCopy, UWorld* InWorld, FString* DestinationData = nullptr) const;
+	virtual void CopyActors(const TArray<AActor*>& InActorsToCopy, UWorld* InWorld, FString* DestinationData = nullptr) const;
 
 	/**
 	 * Paste selected actors from the clipboard.
+	 * @note Final - Override PasteComponents or PasteActors instead.
 	 *
 	 * @param	InWorld				World context
 	 * @param	bDuplicate			Is this a duplicate operation (as opposed to a real paste)?
@@ -552,7 +556,7 @@ public:
 	 * @param	bWarnIfHidden		If true displays a warning if the destination level is hidden
 	 * @param	SourceData			If != NULL, use instead of clipboard data
 	 */
-	virtual void edactPasteSelected(UWorld* InWorld, bool bDuplicate, bool bOffsetLocations, bool bWarnIfHidden, const FString* SourceData = nullptr) override;
+	virtual void edactPasteSelected(UWorld* InWorld, bool bDuplicate, bool bOffsetLocations, bool bWarnIfHidden, const FString* SourceData = nullptr) override final;
 
 	/**
 	 * Paste the components from the clipboard.
@@ -562,7 +566,7 @@ public:
 	 * @param	bWarnIfHidden		If true displays a warning if the destination level is hidden
 	 * @param	SourceData			If != NULL, use instead of clipboard data
 	 */
-	void PasteComponents(TArray<UActorComponent*>& OutPastedComponents, AActor* TargetActor, const bool bWarnIfHidden, const FString* SourceData = nullptr);
+	virtual void PasteComponents(TArray<UActorComponent*>& OutPastedComponents, AActor* TargetActor, const bool bWarnIfHidden, const FString* SourceData = nullptr);
 
 	/**
 	 * Paste the actors from the clipboard.
@@ -574,15 +578,16 @@ public:
 	 * @param	bWarnIfHidden		If true displays a warning if the destination level is hidden
 	 * @param	SourceData			If != NULL, use instead of clipboard data
 	 */
-	void PasteActors(TArray<AActor*>& OutPastedActors, UWorld* InWorld, bool bDuplicate, bool bOffsetLocations, bool bWarnIfHidden, const FString* SourceData = nullptr);
+	virtual void PasteActors(TArray<AActor*>& OutPastedActors, UWorld* InWorld, bool bDuplicate, bool bOffsetLocations, bool bWarnIfHidden, const FString* SourceData = nullptr);
 
 	/**
 	 * Duplicates selected actors.  Handles the case where you are trying to duplicate PrefabInstance actors.
+	 * @note Final - Override DuplicateComponents or DuplicateActors instead.
 	 *
 	 * @param	InLevel				Level to place duplicate
 	 * @param	bOffsetLocations	Should the actor locations be offset after they are created?
 	 */
-	virtual void edactDuplicateSelected(ULevel* InLevel, bool bOffsetLocations) override;
+	virtual void edactDuplicateSelected(ULevel* InLevel, bool bOffsetLocations) override final;
 
 	/**
 	 * Duplicate the given components.
@@ -590,7 +595,7 @@ public:
 	 * @param	InComponentsToDuplicate		Array of components to duplicate
 	 * @param	OutNewComponents			List of all the components that were duplicated
 	 */
-	void DuplicateComponents(const TArray<UActorComponent*>& InComponentsToDuplicate, TArray<UActorComponent*>& OutNewComponents);
+	virtual void DuplicateComponents(const TArray<UActorComponent*>& InComponentsToDuplicate, TArray<UActorComponent*>& OutNewComponents);
 
 	/**
 	 * Duplicates the given actors.  Handles the case where you are trying to duplicate PrefabInstance actors.
@@ -600,7 +605,7 @@ public:
 	 * @param	InLevel				Level to place duplicate
 	 * @param	bUseOffset			Should the actor locations be offset after they are created?
 	 */
-	void DuplicateActors(const TArray<AActor*>& InActorsToDuplicate, TArray<AActor*>& OutNewActors, ULevel* InLevel, bool bOffsetLocations);
+	virtual void DuplicateActors(const TArray<AActor*>& InActorsToDuplicate, TArray<AActor*>& OutNewActors, ULevel* InLevel, bool bOffsetLocations);
 
 	/**
 	 * Replace all selected brushes with the default brush.
