When reinstancing a BP external actor, make sure that the old actor is in the same external package of the new one.
This prevents the editor from displaying a package to delete in the save dialogs.

diff --git a/Engine/Source/Editor/UnrealEd/Private/Kismet2/KismetReinstanceUtilities.cpp b/Engine/Source/Editor/UnrealEd/Private/Kismet2/KismetReinstanceUtilities.cpp
--- a/Engine/Source/Editor/UnrealEd/Private/Kismet2/KismetReinstanceUtilities.cpp
+++ b/Engine/Source/Editor/UnrealEd/Private/Kismet2/KismetReinstanceUtilities.cpp
@@ -1952,6 +1952,11 @@ static void ReplaceActorHelper(AActor* OldActor, UClass* OldClass, UObject*& New
 		NewActor = World->SpawnActor(SpawnClass, &Location, &Rotation, SpawnInfo);
 	}
 
+	if (SpawnInfo.OverridePackage)
+	{
+		OldActor->SetExternalPackage(SpawnInfo.OverridePackage);
+	}
+
 	if (OldActor->CurrentTransactionAnnotation.IsValid())
 	{
 		NewActor->CurrentTransactionAnnotation = OldActor->CurrentTransactionAnnotation;
