Make undo transactions work correctly when saving objects over 2GB in size (32->64 bit size)

diff --git a/Engine/Source/Editor/UnrealEd/Classes/Editor/Transactor.h b/Engine/Source/Editor/UnrealEd/Classes/Editor/Transactor.h
--- a/Engine/Source/Editor/UnrealEd/Classes/Editor/Transactor.h
+++ b/Engine/Source/Editor/UnrealEd/Classes/Editor/Transactor.h
@@ -109,7 +109,7 @@ protected:
 				return InPropertyChain.GetPropertyFromRoot(0)->GetFName();
 			}
 
-			void AppendSerializedData(const int32 InOffset, const int32 InSize)
+			void AppendSerializedData(const int64 InOffset, const int64 InSize)
 			{
 				if (DataOffset == INDEX_NONE)
 				{
@@ -124,9 +124,9 @@ protected:
 			}
 
 			/** Offset to the start of this property within the serialized object */
-			int32 DataOffset;
+			int64 DataOffset;
 			/** Size (in bytes) of this property within the serialized object */
-			int32 DataSize;
+			int64 DataSize;
 		};
 
 		struct FSerializedObject
@@ -199,7 +199,7 @@ protected:
 			/** The pending kill state of the object when it was serialized */
 			bool bIsPendingKill;
 			/** The data stream used to serialize/deserialize record */
-			TArray<uint8> Data;
+			TArray64<uint8> Data;
 			/** External objects referenced in the transaction */
 			TArray<FPersistentObjectRef> ReferencedObjects;
 			/** FNames referenced in the object record */
@@ -440,7 +440,7 @@ protected:
 			{
 				if( Num )
 				{
-					int32 DataIndex = ( Offset == SerializedObject.Data.Num() )
+					int64 DataIndex = ( Offset == SerializedObject.Data.Num() )
 						?  SerializedObject.Data.AddUninitialized(Num)
 						:  Offset;
 					
diff --git a/Engine/Source/Editor/UnrealEd/Private/EditorTransaction.cpp b/Engine/Source/Editor/UnrealEd/Private/EditorTransaction.cpp
--- a/Engine/Source/Editor/UnrealEd/Private/EditorTransaction.cpp
+++ b/Engine/Source/Editor/UnrealEd/Private/EditorTransaction.cpp
@@ -433,10 +433,10 @@ void FTransaction::FObjectRecord::Diff( const FTransaction* Owner, const FSerial
 
 	if (OldSerializedObject.SerializedProperties.Num() > 0 || NewSerializedObject.SerializedProperties.Num() > 0)
 	{
-		int32 StartOfOldPropertyBlock = INT_MAX;
-		int32 StartOfNewPropertyBlock = INT_MAX;
-		int32 EndOfOldPropertyBlock = -1;
-		int32 EndOfNewPropertyBlock = -1;
+		int64 StartOfOldPropertyBlock = INT64_MAX;
+		int64 StartOfNewPropertyBlock = INT64_MAX;
+		int64 EndOfOldPropertyBlock = -1;
+		int64 EndOfNewPropertyBlock = -1;
 
 		for (const TPair<FName, FSerializedProperty>& NewNamePropertyPair : NewSerializedObject.SerializedProperties)
 		{
@@ -497,8 +497,8 @@ void FTransaction::FObjectRecord::Diff( const FTransaction* Owner, const FSerial
 			// Compare the data before the property block to see if something else in the object has changed
 			if (!OutDeltaChange.bHasNonPropertyChanges)
 			{
-				const int32 OldHeaderSize = FMath::Min(StartOfOldPropertyBlock, OldSerializedObject.Data.Num());
-				const int32 CurrentHeaderSize = FMath::Min(StartOfNewPropertyBlock, NewSerializedObject.Data.Num());
+				const int64 OldHeaderSize = FMath::Min(StartOfOldPropertyBlock, OldSerializedObject.Data.Num());
+				const int64 CurrentHeaderSize = FMath::Min(StartOfNewPropertyBlock, NewSerializedObject.Data.Num());
 
 				bool bIsHeaderIdentical = OldHeaderSize == CurrentHeaderSize;
 				if (bIsHeaderIdentical && CurrentHeaderSize > 0)
@@ -515,8 +515,8 @@ void FTransaction::FObjectRecord::Diff( const FTransaction* Owner, const FSerial
 			// Compare the data after the property block to see if something else in the object has changed
 			if (!OutDeltaChange.bHasNonPropertyChanges)
 			{
-				const int32 OldFooterSize = OldSerializedObject.Data.Num() - FMath::Max(EndOfOldPropertyBlock, 0);
-				const int32 CurrentFooterSize = NewSerializedObject.Data.Num() - FMath::Max(EndOfNewPropertyBlock, 0);
+				const int64 OldFooterSize = OldSerializedObject.Data.Num() - FMath::Max<int64>(EndOfOldPropertyBlock, 0);
+				const int64 CurrentFooterSize = NewSerializedObject.Data.Num() - FMath::Max<int64>(EndOfNewPropertyBlock, 0);
 
 				bool bIsFooterIdentical = OldFooterSize == CurrentFooterSize;
 				if (bIsFooterIdentical && CurrentFooterSize > 0)
