Setting bCameraCut after level load in editor

diff --git a/Engine/Source/Editor/UnrealEd/Private/LevelEditorViewport.cpp b/Engine/Source/Editor/UnrealEd/Private/LevelEditorViewport.cpp
--- a/Engine/Source/Editor/UnrealEd/Private/LevelEditorViewport.cpp
+++ b/Engine/Source/Editor/UnrealEd/Private/LevelEditorViewport.cpp
@@ -2051,6 +2051,8 @@ void FLevelEditorViewportClient::ResetCamera()
 		((FEdModeInterpEdit*)Mode)->CamMoveNotify(this);
 	}
 
+	SetIsCameraCut();
+
 	// Broadcast 'camera moved' delegate
 	FEditorDelegates::OnEditorCameraMoved.Broadcast(GetViewLocation(), GetViewRotation(), ViewportType, ViewIndex);
 }
