Remove render thread suspension from editor initialization.

diff --git a/Engine/Source/Editor/UnrealEd/Private/UnrealEdGlobals.cpp b/Engine/Source/Editor/UnrealEd/Private/UnrealEdGlobals.cpp
--- a/Engine/Source/Editor/UnrealEd/Private/UnrealEdGlobals.cpp
+++ b/Engine/Source/Editor/UnrealEd/Private/UnrealEdGlobals.cpp
@@ -15,7 +15,6 @@
 #include "Misc/App.h"
 #include "EngineGlobals.h"
 #include "Modules/ModuleInterface.h"
-#include "RenderingThread.h"
 #include "EditorModeManager.h"
 #include "EditorDirectories.h"
 #include "Misc/OutputDeviceConsole.h"
@@ -151,9 +150,6 @@ int32 EditorInit( IEngineLoop& EngineLoop )
 
 	// Do final set up on the editor frame and show it
 	{
-		// Tear down rendering thread once instead of doing it for every window being resized.
-		SCOPED_SUSPEND_RENDERING_THREAD(true);
-
 		// Startup Slate main frame and other editor windows
 		{
 			const bool bStartImmersive = bIsImmersive;
@@ -212,9 +208,6 @@ int32 EditorReinit()
 		const bool bStartImmersive = bIsImmersive;
 		const bool bStartPIE = bIsImmersive;
 
-		// Tear down rendering thread once instead of doing it for every window being resized.
-		SCOPED_SUSPEND_RENDERING_THREAD(true);
-
 		// Startup Slate main frame and other editor windows
 		IMainFrameModule& MainFrameModule = FModuleManager::LoadModuleChecked<IMainFrameModule>(TEXT("MainFrame"));
 		MainFrameModule.RecreateDefaultMainFrame(bStartImmersive, bStartPIE);
