Removing duplicate #include "ShaderCompiler.h"

diff --git a/Engine/Source/Editor/UnrealEd/Private/Commandlets/CookCommandlet.cpp b/Engine/Source/Editor/UnrealEd/Private/Commandlets/CookCommandlet.cpp
--- a/Engine/Source/Editor/UnrealEd/Private/Commandlets/CookCommandlet.cpp
+++ b/Engine/Source/Editor/UnrealEd/Private/Commandlets/CookCommandlet.cpp
@@ -37,7 +37,6 @@
 #include "INetworkFileSystemModule.h"
 #include "GameDelegates.h"
 #include "CookerSettings.h"
-#include "ShaderCompiler.h"
 #include "HAL/MemoryMisc.h"
 #include "ProfilingDebugging/CookStats.h"
 #include "AssetRegistryModule.h"
