Add new type FExpressionExecOutput to represent UMaterialExpression execution outputs
- Previously exec inputs/outputs were stored along with other (data) inputs/outpus
- Problem is, for exec pins, we want each output to only support 1 connection, but multiple connections to each input. This is the opposite of data, where we support multiple output connections, but only 1 input connection
- So we flip the representation around, and store output pointers, rather than input pointers...this also better matches the way we iterate exec connections while compiling
- This does complicate some UI logic, as now input/output graph pins no longer line up exeactly with UMaterialExpression inputs/outputs (we still have a flat list of pins, but separate lists of exec/non-exec inputs/outputs)

diff --git a/Engine/Source/Editor/GraphEditor/Private/MaterialNodes/SGraphNodeMaterialBase.cpp b/Engine/Source/Editor/GraphEditor/Private/MaterialNodes/SGraphNodeMaterialBase.cpp
--- a/Engine/Source/Editor/GraphEditor/Private/MaterialNodes/SGraphNodeMaterialBase.cpp
+++ b/Engine/Source/Editor/GraphEditor/Private/MaterialNodes/SGraphNodeMaterialBase.cpp
@@ -278,8 +278,11 @@ void SGraphNodeMaterialBase::CreatePinWidgets()
 		
 			if ((bIsABreakAttrNode && CurPin->Direction == EGPD_Output) || (bIsAMakeAttrNode && CurPin->Direction == EGPD_Input))
 			{
-				int32 ResultNodeIndex = bIsAMakeAttrNode ? MaterialNode->GetInputIndex(CurPin) : MaterialNode->GetOutputIndex(CurPin);
-				bPinDesiresToBeHidden |= !MaterialGraph->MaterialInputs[ResultNodeIndex].IsVisiblePin(MaterialGraph->Material, true);
+				const FMaterialGraphPinInfo& PinInfo = MaterialNode->GetPinInfo(CurPin);
+				if (!(PinInfo.Type & MCT_Execution))
+				{
+					bPinDesiresToBeHidden |= !MaterialGraph->MaterialInputs[PinInfo.Index].IsVisiblePin(MaterialGraph->Material, true);
+				}
 			}
 		}
 
diff --git a/Engine/Source/Editor/MaterialEditor/Private/MaterialEditor.cpp b/Engine/Source/Editor/MaterialEditor/Private/MaterialEditor.cpp
--- a/Engine/Source/Editor/MaterialEditor/Private/MaterialEditor.cpp
+++ b/Engine/Source/Editor/MaterialEditor/Private/MaterialEditor.cpp
@@ -3360,10 +3360,7 @@ bool FMaterialEditor::IsToggleHideUnrelatedNodesChecked() const
 
 void FMaterialEditor::CollectDownstreamNodes(UMaterialGraphNode* CurrentNode, TArray<UMaterialGraphNode*>& CollectedNodes)
 {
-	TArray<UEdGraphPin*> OutputPins;
-	CurrentNode->GetOutputPins(OutputPins);
-
-	for (auto& OutputPin : OutputPins)
+	for (auto& OutputPin : CurrentNode->OutputPins)
 	{
 		for (auto& Link : OutputPin->LinkedTo)
 		{
@@ -3384,10 +3381,7 @@ void FMaterialEditor::CollectDownstreamNodes(UMaterialGraphNode* CurrentNode, TA
 
 void FMaterialEditor::CollectUpstreamNodes(UMaterialGraphNode* CurrentNode, TArray<UMaterialGraphNode*>& CollectedNodes)
 {
-	TArray<UEdGraphPin*> InputPins;
-	CurrentNode->GetInputPins(InputPins);
-
-	for (auto& InputPin : InputPins)
+	for (auto& InputPin : CurrentNode->InputPins)
 	{
 		for (auto& Link : InputPin->LinkedTo)
 		{
@@ -4132,14 +4126,11 @@ void FMaterialEditor::OnSelectDownstreamNodes()
 	while (NodesToCheck.Num() > 0)
 	{
 		UMaterialGraphNode* CurrentNode = NodesToCheck.Last();
-		TArray<UEdGraphPin*> OutputPins;
-		CurrentNode->GetOutputPins(OutputPins);
-
-		for (int32 Index = 0; Index < OutputPins.Num(); ++Index)
+		for (int32 Index = 0; Index < CurrentNode->OutputPins.Num(); ++Index)
 		{
-			for (int32 LinkIndex = 0; LinkIndex < OutputPins[Index]->LinkedTo.Num(); ++LinkIndex)
+			for (int32 LinkIndex = 0; LinkIndex < CurrentNode->OutputPins[Index]->LinkedTo.Num(); ++LinkIndex)
 			{
-				UMaterialGraphNode* LinkedNode = Cast<UMaterialGraphNode>(OutputPins[Index]->LinkedTo[LinkIndex]->GetOwningNode());
+				UMaterialGraphNode* LinkedNode = Cast<UMaterialGraphNode>(CurrentNode->OutputPins[Index]->LinkedTo[LinkIndex]->GetOwningNode());
 				if (LinkedNode)
 				{
 					int32 FoundIndex = -1;
@@ -4188,14 +4179,11 @@ void FMaterialEditor::OnSelectUpstreamNodes()
 	while (NodesToCheck.Num() > 0)
 	{
 		UMaterialGraphNode* CurrentNode = NodesToCheck.Last();
-		TArray<UEdGraphPin*> InputPins;
-		CurrentNode->GetInputPins(InputPins);
-
-		for (int32 Index = 0; Index < InputPins.Num(); ++Index)
+		for (int32 Index = 0; Index < CurrentNode->InputPins.Num(); ++Index)
 		{
-			for (int32 LinkIndex = 0; LinkIndex < InputPins[Index]->LinkedTo.Num(); ++LinkIndex)
+			for (int32 LinkIndex = 0; LinkIndex < CurrentNode->InputPins[Index]->LinkedTo.Num(); ++LinkIndex)
 			{
-				UMaterialGraphNode* LinkedNode = Cast<UMaterialGraphNode>(InputPins[Index]->LinkedTo[LinkIndex]->GetOwningNode());
+				UMaterialGraphNode* LinkedNode = Cast<UMaterialGraphNode>(CurrentNode->InputPins[Index]->LinkedTo[LinkIndex]->GetOwningNode());
 				if (LinkedNode)
 				{
 					int32 FoundIndex = -1;
diff --git a/Engine/Source/Editor/UnrealEd/Classes/MaterialGraph/MaterialGraphNode.h b/Engine/Source/Editor/UnrealEd/Classes/MaterialGraph/MaterialGraphNode.h
--- a/Engine/Source/Editor/UnrealEd/Classes/MaterialGraph/MaterialGraphNode.h
+++ b/Engine/Source/Editor/UnrealEd/Classes/MaterialGraph/MaterialGraphNode.h
@@ -48,7 +48,7 @@ public:
 	UNREALED_API void RecreateAndLinkNode();
 
 	/** Get the Material value type of an output pin */
-	virtual uint32 GetOutputType(const UEdGraphPin* OutputPin) const override;
+	//virtual uint32 GetOutputType(const UEdGraphPin* OutputPin) const override;
 
 	//~ Begin UObject Interface
 	virtual void PostEditChangeProperty(FPropertyChangedEvent& PropertyChangedEvent) override;
@@ -76,8 +76,6 @@ public:
 	//~ Begin UMaterialGraphNode_Base Interface
 	virtual void CreateInputPins() override;
 	virtual void CreateOutputPins() override;
-	virtual UNREALED_API int32 GetInputIndex(const UEdGraphPin* InputPin) const override;
-	virtual uint32 GetInputType(const UEdGraphPin* InputPin) const override;
 	//~ End UMaterialGraphNode_Base Interface
 
 	/** Will return the shorten pin name to use based on long pin name */
diff --git a/Engine/Source/Editor/UnrealEd/Classes/MaterialGraph/MaterialGraphNode_Base.h b/Engine/Source/Editor/UnrealEd/Classes/MaterialGraph/MaterialGraphNode_Base.h
--- a/Engine/Source/Editor/UnrealEd/Classes/MaterialGraph/MaterialGraphNode_Base.h
+++ b/Engine/Source/Editor/UnrealEd/Classes/MaterialGraph/MaterialGraphNode_Base.h
@@ -10,11 +10,27 @@
 class UEdGraphPin;
 class UEdGraphSchema;
 
+struct FMaterialGraphPinInfo
+{
+	uint32 Type; // EMaterialValueType
+	int32 Index; // index into the expression's list of inputs/outputs (exec inputs/outpus are indexed separately)
+};
+
 UCLASS(MinimalAPI)
 class UMaterialGraphNode_Base : public UEdGraphNode
 {
 	GENERATED_UCLASS_BODY()
 
+	/** Contains additional information about material graph pins, avoid adding material-specific data to base pin type */
+	TMap<UEdGraphPin*, FMaterialGraphPinInfo> PinInfoMap;
+
+	/** Lists of pins that match up with the underlying UMaterialExpression's (non-exec) inputs and outputs */
+	TArray<UEdGraphPin*> InputPins;
+	TArray<UEdGraphPin*> OutputPins;
+
+	/** The exec input pin */
+	UEdGraphPin* ExecInputPin = nullptr;
+
 	/** Create all of the input pins required */
 	virtual void CreateInputPins() {};
 	/** Create all of the output pins required */
@@ -22,26 +38,21 @@ class UMaterialGraphNode_Base : public UEdGraphNode
 	/** Is this the undeletable root node */
 	virtual bool IsRootNode() const {return false;}
 	/** Get a single Input Pin via its index */
-	class UEdGraphPin* GetInputPin(int32 InputIndex) const;
-	/** Get all of the Input Pins */
-	UNREALED_API void GetInputPins(TArray<class UEdGraphPin*>& OutInputPins) const;
+	class UEdGraphPin* GetInputPin(int32 InputIndex) const { return InputPins[InputIndex]; }
 	/** Get a single Output Pin via its index */
-	class UEdGraphPin* GetOutputPin(int32 OutputIndex) const;
-	/** Get all of the Output Pins */
-	UNREALED_API void GetOutputPins(TArray<class UEdGraphPin*>& OutOutputPins) const;
+	class UEdGraphPin* GetOutputPin(int32 OutputIndex) const { return OutputPins[OutputIndex]; }
+	/** Gets the exec input pin */
+	class UEdGraphPin* GetExecInputPin() const { return ExecInputPin; }
 	/** Replace a given node with this one, changing all pin links */
 	UNREALED_API void ReplaceNode(UMaterialGraphNode_Base* OldNode);
 
-	/** Get the Material Expression input index from an input pin */
-	virtual int32 GetInputIndex(const UEdGraphPin* InputPin) const {return -1;}
-	/** Get the Material value type of an input pin */
-	virtual uint32 GetInputType(const UEdGraphPin* InputPin) const;
+	UNREALED_API const FMaterialGraphPinInfo& GetPinInfo(const class UEdGraphPin* Pin) const;
 
-	/** Get the Material Expression output index from an output pin */
-	UNREALED_API int32 GetOutputIndex(const UEdGraphPin* OutputPin) const;
+	/** Get the Material value type of an input pin */
+	uint32 GetInputType(const UEdGraphPin* InputPin) const;
 
 	/** Get the Material value type of an output pin */
-	virtual uint32 GetOutputType(const UEdGraphPin* OutputPin) const { return 0; }
+	uint32 GetOutputType(const UEdGraphPin* OutputPin) const;
 
 	/**
 	 * Handles inserting the node between the FromPin and what the FromPin was original connected to
@@ -63,4 +74,6 @@ class UMaterialGraphNode_Base : public UEdGraphNode
 
 protected:
 	void ModifyAndCopyPersistentPinData(UEdGraphPin& TargetPin, const UEdGraphPin& SourcePin) const;
+
+	void RegisterPin(UEdGraphPin* Pin, int32 Index, uint32 Type);
 };
diff --git a/Engine/Source/Editor/UnrealEd/Classes/MaterialGraph/MaterialGraphNode_Root.h b/Engine/Source/Editor/UnrealEd/Classes/MaterialGraph/MaterialGraphNode_Root.h
--- a/Engine/Source/Editor/UnrealEd/Classes/MaterialGraph/MaterialGraphNode_Root.h
+++ b/Engine/Source/Editor/UnrealEd/Classes/MaterialGraph/MaterialGraphNode_Root.h
@@ -30,7 +30,5 @@ class UMaterialGraphNode_Root : public UMaterialGraphNode_Base
 	//~ Begin UMaterialGraphNode_Base Interface
 	virtual void CreateInputPins() override;
 	virtual bool IsRootNode() const override {return true;}
-	virtual int32 GetInputIndex(const UEdGraphPin* InputPin) const override;
-	virtual uint32 GetInputType(const UEdGraphPin* InputPin) const override;
 	//~ End UMaterialGraphNode_Base Interface
 };
diff --git a/Engine/Source/Editor/UnrealEd/Private/MaterialGraph.cpp b/Engine/Source/Editor/UnrealEd/Private/MaterialGraph.cpp
--- a/Engine/Source/Editor/UnrealEd/Private/MaterialGraph.cpp
+++ b/Engine/Source/Editor/UnrealEd/Private/MaterialGraph.cpp
@@ -310,44 +310,66 @@ void UMaterialGraph::LinkGraphNodesFromMaterial()
 	for (int32 Index = 0; Index < Material->Expressions.Num(); Index++)
 	{
 		UMaterialExpression* Expression = Material->Expressions[Index];
-		
-		if (Expression)
+		if (!Expression)
 		{
-			const TArray<FExpressionInput*> ExpressionInputs = Expression->GetInputs();
-			for (int32 InputIndex = 0; InputIndex < ExpressionInputs.Num(); ++InputIndex)
+			continue;
+		}
+
+		UMaterialGraphNode* MaterialGraphNode = Cast<UMaterialGraphNode>(Expression->GraphNode);
+		if (!MaterialGraphNode)
+		{
+			continue;
+		}
+
+		const TArray<FExpressionInput*> ExpressionInputs = Expression->GetInputs();
+
+		TArray<FExpressionExecOutputEntry> ExecOutputs;
+		Expression->GetExecOutputs(ExecOutputs);
+
+		for (const auto& It : MaterialGraphNode->PinInfoMap)
+		{
+			UEdGraphPin* Pin = It.Key;
+			const FMaterialGraphPinInfo& PinInfo = It.Value;
+			if (Pin->Direction == EGPD_Input && PinInfo.Type != MCT_Execution)
 			{
-				if (UMaterialGraphNode* MaterialGraphNode = Cast<UMaterialGraphNode>(Expression->GraphNode))
+				if (ExpressionInputs[PinInfo.Index]->Expression)
 				{
-					UEdGraphPin* InputPin = MaterialGraphNode->GetInputPin(InputIndex);
-
-					// InputPin can be null during a PostEditChange when there is a circular dependency between nodes, and nodes have pins that are dynamically created
-					if (InputPin != nullptr && ExpressionInputs[InputIndex]->Expression)
+					// Unclear why this is null sometimes outside of composite reroute, but this is safer than crashing
+					if (UMaterialGraphNode* GraphNode = Cast<UMaterialGraphNode>(ExpressionInputs[PinInfo.Index]->Expression->GraphNode))
 					{
-						// Unclear why this is null sometimes outside of composite reroute, but this is safer than crashing
-						if (UMaterialGraphNode* GraphNode = Cast<UMaterialGraphNode>(ExpressionInputs[InputIndex]->Expression->GraphNode))
+						Pin->MakeLinkTo(GraphNode->GetOutputPin(GetValidOutputIndex(ExpressionInputs[PinInfo.Index])));
+					}
+					else if (UMaterialExpressionReroute* CompositeReroute = Cast<UMaterialExpressionReroute>(ExpressionInputs[PinInfo.Index]->Expression))
+					{
+						// This is an unseen composite reroute expression, find the actual expression output to connect to.
+						UMaterialExpressionComposite* OwningComposite = CastChecked<UMaterialExpressionComposite>(CompositeReroute->SubgraphExpression);
+
+						UMaterialGraphNode* OutputGraphNode;
+						int32 OutputPinIndex = OwningComposite->InputExpressions->ReroutePins.FindLastByPredicate(ExpressionMatchesPredicate(CompositeReroute));
+						if (OutputPinIndex != INDEX_NONE)
 						{
-							InputPin->MakeLinkTo(GraphNode->GetOutputPin(GetValidOutputIndex(ExpressionInputs[InputIndex])));
+							OutputGraphNode = CastChecked<UMaterialGraphNode>(OwningComposite->InputExpressions->GraphNode);
 						}
-						else if (UMaterialExpressionReroute* CompositeReroute = Cast<UMaterialExpressionReroute>(ExpressionInputs[InputIndex]->Expression))
+						else
 						{
-							// This is an unseen composite reroute expression, find the actual expression output to connect to.
-							UMaterialExpressionComposite* OwningComposite = CastChecked<UMaterialExpressionComposite>(CompositeReroute->SubgraphExpression);
-
-							UMaterialGraphNode* OutputGraphNode;
-							int32 OutputPinIndex = OwningComposite->InputExpressions->ReroutePins.FindLastByPredicate(ExpressionMatchesPredicate(CompositeReroute));
-							if (OutputPinIndex != INDEX_NONE)
-							{
-								OutputGraphNode = CastChecked<UMaterialGraphNode>(OwningComposite->InputExpressions->GraphNode);
-							}
-							else
-							{
-								// Output pin base in the subgraph cannot have outputs, if this reroute isn't in the inputs connect to composite's outputs
-								OutputPinIndex = OwningComposite->OutputExpressions->ReroutePins.FindLastByPredicate(ExpressionMatchesPredicate(CompositeReroute));
-								OutputGraphNode = CastChecked<UMaterialGraphNode>(OwningComposite->GraphNode);
-							}
-							InputPin->MakeLinkTo(OutputGraphNode->GetOutputPin(OutputPinIndex));
+							// Output pin base in the subgraph cannot have outputs, if this reroute isn't in the inputs connect to composite's outputs
+							OutputPinIndex = OwningComposite->OutputExpressions->ReroutePins.FindLastByPredicate(ExpressionMatchesPredicate(CompositeReroute));
+							OutputGraphNode = CastChecked<UMaterialGraphNode>(OwningComposite->GraphNode);
 						}
+						Pin->MakeLinkTo(OutputGraphNode->GetOutputPin(OutputPinIndex));
+					}
+				}
+			}
+			else if (Pin->Direction == EGPD_Output && PinInfo.Type == MCT_Execution)
+			{
+				FExpressionExecOutput* ExecOutput = ExecOutputs[PinInfo.Index].Output;
+				if (ExecOutput->Expression)
+				{
+					if (UMaterialGraphNode* GraphNode = Cast<UMaterialGraphNode>(ExecOutput->Expression->GraphNode))
+					{
+						Pin->MakeLinkTo(GraphNode->GetExecInputPin());
 					}
+					// TODO - UMaterialExpressionReroute?
 				}
 			}
 		}
@@ -359,32 +381,26 @@ void UMaterialGraph::LinkGraphNodesFromMaterial()
 void UMaterialGraph::LinkMaterialExpressionsFromGraph() const
 {
 	// Use GraphNodes to make Material Expression Connections
-	TArray<UEdGraphPin*> InputPins;
-	TArray<UEdGraphPin*> OutputPins;
-
 	for (int32 NodeIndex = 0; NodeIndex < Nodes.Num(); ++NodeIndex)
 	{
 		if (RootNode && RootNode == Nodes[NodeIndex])
 		{
 			// Setup Material's inputs from root node
 			Material->Modify();
-			InputPins = RootNode->Pins;
 			Material->EditorX = RootNode->NodePosX;
 			Material->EditorY = RootNode->NodePosY;
-			check(InputPins.Num() == MaterialInputs.Num());
-			for (int32 PinIndex = 0; PinIndex < InputPins.Num() && PinIndex < MaterialInputs.Num(); ++PinIndex)
+			check(RootNode->Pins.Num() == MaterialInputs.Num());
+			for (int32 PinIndex = 0; PinIndex < RootNode->Pins.Num() && PinIndex < MaterialInputs.Num(); ++PinIndex)
 			{
 				FExpressionInput& MaterialInput = MaterialInputs[PinIndex].GetExpressionInput(Material);
 
-				if (InputPins[PinIndex]->LinkedTo.Num() > 0)
+				if (RootNode->Pins[PinIndex]->LinkedTo.Num() > 0)
 				{
-					UMaterialGraphNode* ConnectedNode = CastChecked<UMaterialGraphNode>(InputPins[PinIndex]->LinkedTo[0]->GetOwningNode());
-					ConnectedNode->GetOutputPins(OutputPins);
-
+					UMaterialGraphNode* ConnectedNode = CastChecked<UMaterialGraphNode>(RootNode->Pins[PinIndex]->LinkedTo[0]->GetOwningNode());
 					// Work out the index of the connected pin
-					for (int32 OutPinIndex = 0; OutPinIndex < OutputPins.Num(); ++OutPinIndex)
+					for (int32 OutPinIndex = 0; OutPinIndex < ConnectedNode->OutputPins.Num(); ++OutPinIndex)
 					{
-						if (OutputPins[OutPinIndex] == InputPins[PinIndex]->LinkedTo[0])
+						if (ConnectedNode->OutputPins[OutPinIndex] == RootNode->Pins[PinIndex]->LinkedTo[0])
 						{
 							if (!ConnectedNode->MaterialExpression->IsExpressionConnected(&MaterialInput, OutPinIndex))
 							{
@@ -425,45 +441,85 @@ void UMaterialGraph::LinkMaterialExpressionsFromGraph() const
 						Expression->Desc = GraphNode->NodeComment;
 					}
 
-					GraphNode->GetInputPins(InputPins);
 					const TArray<FExpressionInput*> ExpressionInputs = Expression->GetInputs();
-					checkf(InputPins.Num() == ExpressionInputs.Num(), TEXT("Mismatched inputs for '%s'"), *Expression->GetFullName());
-					for (int32 PinIndex = 0; PinIndex < InputPins.Num() && PinIndex < ExpressionInputs.Num(); ++PinIndex)
+
+					TArray<FExpressionExecOutputEntry> ExecOutputs;
+					Expression->GetExecOutputs(ExecOutputs);
+
+					for (const auto& It : GraphNode->PinInfoMap)
 					{
-						FExpressionInput* ExpressionInput = ExpressionInputs[PinIndex];
-						if (InputPins[PinIndex]->LinkedTo.Num() > 0)
+						UEdGraphPin* Pin = It.Key;
+						const FMaterialGraphPinInfo& PinInfo = It.Value;
+						if (Pin->Direction == EGPD_Input && PinInfo.Type != MCT_Execution)
 						{
-							UMaterialGraphNode* ConnectedNode = CastChecked<UMaterialGraphNode>(InputPins[PinIndex]->LinkedTo[0]->GetOwningNode());
-							ConnectedNode->GetOutputPins(OutputPins);
-
-							// Work out the index of the connected pin
-							for (int32 OutPinIndex = 0; OutPinIndex < OutputPins.Num(); ++OutPinIndex)
+							// Wire up non-execution input pins
+							FExpressionInput* ExpressionInput = ExpressionInputs[PinInfo.Index];
+							if (Pin->LinkedTo.Num() > 0)
 							{
-								if (OutputPins[OutPinIndex] == InputPins[PinIndex]->LinkedTo[0])
+								UMaterialGraphNode* ConnectedNode = CastChecked<UMaterialGraphNode>(Pin->LinkedTo[0]->GetOwningNode());
+
+								// Work out the index of the connected pin
+								for (int32 OutPinIndex = 0; OutPinIndex < ConnectedNode->OutputPins.Num(); ++OutPinIndex)
 								{
-									if (ExpressionInput && !ConnectedNode->MaterialExpression->IsExpressionConnected(ExpressionInput, OutPinIndex))
+									if (ConnectedNode->OutputPins[OutPinIndex] == Pin->LinkedTo[0])
 									{
-										if (!bModifiedExpression)
+										if (ExpressionInput && !ConnectedNode->MaterialExpression->IsExpressionConnected(ExpressionInput, OutPinIndex))
 										{
-											bModifiedExpression = true;
-											Expression->Modify();
+											if (!bModifiedExpression)
+											{
+												bModifiedExpression = true;
+												Expression->Modify();
+											}
+
+											ConnectedNode->MaterialExpression->Modify();
+											ExpressionInput->Connect(OutPinIndex, ConnectedNode->MaterialExpression);
 										}
-
-										ConnectedNode->MaterialExpression->Modify();
-										ExpressionInput->Connect(OutPinIndex, ConnectedNode->MaterialExpression);
+										break;
 									}
-									break;
 								}
 							}
+							else if (ExpressionInput && ExpressionInput->Expression)
+							{
+								if (!bModifiedExpression)
+								{
+									bModifiedExpression = true;
+									Expression->Modify();
+								}
+								ExpressionInput->Expression = NULL;
+							}
 						}
-						else if (ExpressionInput && ExpressionInput->Expression)
+						else if (Pin->Direction == EGPD_Output && PinInfo.Type == MCT_Execution)
 						{
-							if (!bModifiedExpression)
+							// Wire up execution output pins
+							FExpressionExecOutput* ExpressionOutput = ExecOutputs[PinInfo.Index].Output;
+
+							if (Pin->LinkedTo.Num() > 0)
 							{
-								bModifiedExpression = true;
-								Expression->Modify();
+								UMaterialGraphNode* ConnectedNode = CastChecked<UMaterialGraphNode>(Pin->LinkedTo[0]->GetOwningNode());
+
+								if (ExpressionOutput &&
+									ExpressionOutput->Expression != ConnectedNode->MaterialExpression &&
+									ConnectedNode->MaterialExpression->HasExecInput())
+								{
+									if (!bModifiedExpression)
+									{
+										bModifiedExpression = true;
+										Expression->Modify();
+									}
+
+									ConnectedNode->MaterialExpression->Modify();
+									ExpressionOutput->Expression = ConnectedNode->MaterialExpression;
+								}
+							}
+							else if (ExpressionOutput && ExpressionOutput->Expression)
+							{
+								if (!bModifiedExpression)
+								{
+									bModifiedExpression = true;
+									Expression->Modify();
+								}
+								ExpressionOutput->Expression = NULL;
 							}
-							ExpressionInput->Expression = NULL;
 						}
 					}
 				}
@@ -525,16 +581,14 @@ void UMaterialGraph::GetUnusedExpressions(TArray<UEdGraphNode*>& UnusedNodes) co
 
 	if (RootNode)
 	{
-		TArray<UEdGraphPin*> InputPins;
-		RootNode->GetInputPins(InputPins);
-		for (int32 Index = 0; Index < InputPins.Num(); ++Index)
+		for (int32 Index = 0; Index < RootNode->InputPins.Num(); ++Index)
 		{
 			check(Index < MaterialInputs.Num());
 			
 			if (MaterialInputs[Index].IsVisiblePin(Material)
-				&& InputPins[Index]->LinkedTo.Num() > 0 && InputPins[Index]->LinkedTo[0])
+				&& RootNode->InputPins[Index]->LinkedTo.Num() > 0 && RootNode->InputPins[Index]->LinkedTo[0])
 			{
-				NodesToCheck.Push(InputPins[Index]->LinkedTo[0]->GetOwningNode());
+				NodesToCheck.Push(RootNode->InputPins[Index]->LinkedTo[0]->GetOwningNode());
 			}
 		}
 
@@ -583,13 +637,11 @@ void UMaterialGraph::GetUnusedExpressions(TArray<UEdGraphNode*>& UnusedNodes) co
 				UsedNodes.Add(GraphNode);
 
 				// Iterate over the expression's inputs and add them to the pending stack.
-				TArray<UEdGraphPin*> InputPins;
-				GraphNode->GetInputPins(InputPins);
-				for (int32 Index = 0; Index < InputPins.Num(); ++Index)
+				for (int32 Index = 0; Index < GraphNode->InputPins.Num(); ++Index)
 				{
-					if (InputPins[Index]->LinkedTo.Num() > 0 && InputPins[Index]->LinkedTo[0])
+					if (GraphNode->InputPins[Index]->LinkedTo.Num() > 0 && GraphNode->InputPins[Index]->LinkedTo[0])
 					{
-						NodesToCheck.Push(InputPins[Index]->LinkedTo[0]->GetOwningNode());
+						NodesToCheck.Push(GraphNode->InputPins[Index]->LinkedTo[0]->GetOwningNode());
 					}
 				}
 
diff --git a/Engine/Source/Editor/UnrealEd/Private/MaterialGraphNode.cpp b/Engine/Source/Editor/UnrealEd/Private/MaterialGraphNode.cpp
--- a/Engine/Source/Editor/UnrealEd/Private/MaterialGraphNode.cpp
+++ b/Engine/Source/Editor/UnrealEd/Private/MaterialGraphNode.cpp
@@ -287,7 +287,7 @@ FLinearColor UMaterialGraphNode::GetNodeTitleColor() const
 	{
 		return Settings->FunctionTerminatorNodeTitleColor;
 	}
-	else if (MaterialExpression->GetExecInput())
+	else if (MaterialExpression->HasExecInput())
 	{
 		return Settings->ExecutionPinTypeColor;
 	}
@@ -622,8 +622,17 @@ FName UMaterialGraphNode::GetShortenPinName(const FName PinName)
 
 void UMaterialGraphNode::CreateInputPins()
 {
-	const TArray<FExpressionInput*> ExpressionInputs = MaterialExpression->GetInputs();
+	if (MaterialExpression->HasExecInput())
+	{
+		UEdGraphPin* NewPin = CreatePin(EGPD_Input, UMaterialGraphSchema::PC_Exec, NAME_None, NAME_None);
+		// Makes sure pin has a name for lookup purposes but user will never see it
+		NewPin->PinName = CreateUniquePinName(TEXT("Input"));
+		NewPin->PinFriendlyName = SpaceText;
+
+		RegisterPin(NewPin, 0, MCT_Execution);
+	}
 
+	const TArray<FExpressionInput*> ExpressionInputs = MaterialExpression->GetInputs();
 	for (int32 Index = 0; Index < ExpressionInputs.Num() ; ++Index)
 	{
 		FExpressionInput* Input = ExpressionInputs[Index];
@@ -655,13 +664,14 @@ void UMaterialGraphNode::CreateInputPins()
 			NewPin->PinName = CreateUniquePinName(TEXT("Input"));
 			NewPin->PinFriendlyName = SpaceText;
 		}
+
+		RegisterPin(NewPin, Index, InputType);
 	}
 }
 
 void UMaterialGraphNode::CreateOutputPins()
 {
 	TArray<FExpressionOutput>& Outputs = MaterialExpression->GetOutputs();
-
 	for(int32 Index = 0; Index < Outputs.Num(); ++Index)
 	{
 		const FExpressionOutput& ExpressionOutput = Outputs[Index];
@@ -715,33 +725,32 @@ void UMaterialGraphNode::CreateOutputPins()
 			NewPin->PinName = CreateUniquePinName(TEXT("Output"));
 			NewPin->PinFriendlyName = SpaceText;
 		}
-	}
-}
 
-uint32 UMaterialGraphNode::GetOutputType(const UEdGraphPin* OutputPin) const
-{
-	return MaterialExpression->GetOutputType(GetOutputIndex(OutputPin));
-}
-
-int32 UMaterialGraphNode::GetInputIndex(const UEdGraphPin* InputPin) const
-{
-	TArray<UEdGraphPin*> InputPins;
-	GetInputPins(InputPins);
+		RegisterPin(NewPin, Index, OutputType);
+	}
 
-	for (int32 Index = 0; Index < InputPins.Num(); ++Index)
+	TArray<FExpressionExecOutputEntry> ExecOutputs;
+	MaterialExpression->GetExecOutputs(ExecOutputs);
+	for (int32 Index = 0; Index < ExecOutputs.Num(); ++Index)
 	{
-		if (InputPin == InputPins[Index])
+		const FExpressionExecOutputEntry& Entry = ExecOutputs[Index];
+
+		FName OutputName;
+		if (MaterialExpression->bShowOutputNameOnPin)
 		{
-			return Index;
+			OutputName = Entry.Name;
 		}
-	}
 
-	return -1;
-}
+		UEdGraphPin* NewPin = CreatePin(EGPD_Output, UMaterialGraphSchema::PC_Exec, NAME_None, OutputName);
+		if (NewPin->PinName.IsNone())
+		{
+			// Makes sure pin has a name for lookup purposes but user will never see it
+			NewPin->PinName = CreateUniquePinName(TEXT("Output"));
+			NewPin->PinFriendlyName = SpaceText;
+		}
 
-uint32 UMaterialGraphNode::GetInputType(const UEdGraphPin* InputPin) const
-{
-	return MaterialExpression->GetInputType(GetInputIndex(InputPin));
+		RegisterPin(NewPin, Index, MCT_Execution);
+	}
 }
 
 void UMaterialGraphNode::ResetMaterialExpressionOwner()
diff --git a/Engine/Source/Editor/UnrealEd/Private/MaterialGraphNode_Base.cpp b/Engine/Source/Editor/UnrealEd/Private/MaterialGraphNode_Base.cpp
--- a/Engine/Source/Editor/UnrealEd/Private/MaterialGraphNode_Base.cpp
+++ b/Engine/Source/Editor/UnrealEd/Private/MaterialGraphNode_Base.cpp
@@ -17,86 +17,21 @@ UMaterialGraphNode_Base::UMaterialGraphNode_Base(const FObjectInitializer& Objec
 {
 }
 
-UEdGraphPin* UMaterialGraphNode_Base::GetInputPin(int32 InputIndex) const
+const FMaterialGraphPinInfo& UMaterialGraphNode_Base::GetPinInfo(const class UEdGraphPin* Pin) const
 {
-	for (int32 PinIndex = 0, FoundInputs = 0; PinIndex < Pins.Num(); PinIndex++)
-	{
-		if (Pins[PinIndex]->Direction == EGPD_Input)
-		{
-			if (InputIndex == FoundInputs)
-			{
-				return Pins[PinIndex];
-			}
-			else
-			{
-				FoundInputs++;
-			}
-		}
-	}
-
-	return NULL;
+	const FMaterialGraphPinInfo* PinInfo = PinInfoMap.Find(Pin);
+	checkf(PinInfo, TEXT("Missing info for pin %s, missing call to RegisterPin()?"), *Pin->GetName());
+	return *PinInfo;
 }
 
-void UMaterialGraphNode_Base::GetInputPins(TArray<class UEdGraphPin*>& OutInputPins) const
+uint32 UMaterialGraphNode_Base::GetOutputType(const UEdGraphPin* OutputPin) const
 {
-	OutInputPins.Empty();
-
-	for (int32 PinIndex = 0; PinIndex < Pins.Num(); PinIndex++)
-	{
-		if (Pins[PinIndex]->Direction == EGPD_Input)
-		{
-			OutInputPins.Add(Pins[PinIndex]);
-		}
-	}
-}
-
-UEdGraphPin* UMaterialGraphNode_Base::GetOutputPin(int32 OutputIndex) const
-{
-	for (int32 PinIndex = 0, FoundOutputs = 0; PinIndex < Pins.Num(); PinIndex++)
-	{
-		if (Pins[PinIndex]->Direction == EGPD_Output)
-		{
-			if (OutputIndex == FoundOutputs)
-			{
-				return Pins[PinIndex];
-			}
-			else
-			{
-				FoundOutputs++;
-			}
-		}
-	}
-
-	return NULL;
-}
-
-void UMaterialGraphNode_Base::GetOutputPins(TArray<class UEdGraphPin*>& OutOutputPins) const
-{
-	OutOutputPins.Empty();
-
-	for (int32 PinIndex = 0; PinIndex < Pins.Num(); PinIndex++)
-	{
-		if (Pins[PinIndex]->Direction == EGPD_Output)
-		{
-			OutOutputPins.Add(Pins[PinIndex]);
-		}
-	}
+	return GetPinInfo(OutputPin).Type;
 }
 
-int32 UMaterialGraphNode_Base::GetOutputIndex(const UEdGraphPin* OutputPin) const
+uint32 UMaterialGraphNode_Base::GetInputType(const UEdGraphPin* InputPin) const
 {
-	TArray<UEdGraphPin*> OutputPins;
-	GetOutputPins(OutputPins);
-
-	for (int32 Index = 0; Index < OutputPins.Num(); ++Index)
-	{
-		if (OutputPin == OutputPins[Index])
-		{
-			return Index;
-		}
-	}
-
-	return -1;
+	return GetPinInfo(InputPin).Type;
 }
 
 void UMaterialGraphNode_Base::ReplaceNode(UMaterialGraphNode_Base* OldNode)
@@ -104,37 +39,25 @@ void UMaterialGraphNode_Base::ReplaceNode(UMaterialGraphNode_Base* OldNode)
 	check(OldNode);
 	check(OldNode != this);
 
-	// Get Pins from node passed in
-	TArray<UEdGraphPin*> OldInputPins;
-	TArray<UEdGraphPin*> OldOutputPins;
-	OldNode->GetInputPins(OldInputPins);
-	OldNode->GetOutputPins(OldOutputPins);
-
-	// Get our Input and Output pins
-	TArray<UEdGraphPin*> NewInputPins;
-	TArray<UEdGraphPin*> NewOutputPins;
-	GetInputPins(NewInputPins);
-	GetOutputPins(NewOutputPins);
-
 	// Copy Inputs from old node
-	for (int32 PinIndex = 0; PinIndex < OldInputPins.Num(); PinIndex++)
+	for (int32 PinIndex = 0; PinIndex < OldNode->InputPins.Num(); PinIndex++)
 	{
-		if (PinIndex < NewInputPins.Num())
+		if (PinIndex < InputPins.Num())
 		{
-			ModifyAndCopyPersistentPinData(*NewInputPins[PinIndex], *OldInputPins[PinIndex]);
+			ModifyAndCopyPersistentPinData(*InputPins[PinIndex], *OldNode->InputPins[PinIndex]);
 		}
 	}
 
 	// Copy Outputs from old node
-	for (int32 PinIndex = 0; PinIndex < OldOutputPins.Num(); PinIndex++)
+	for (int32 PinIndex = 0; PinIndex < OldNode->OutputPins.Num(); PinIndex++)
 	{
 		// Try to find an equivalent output in this node
 		int32 FoundPinIndex = -1;
 		{
 			// First check names
-			for (int32 NewPinIndex = 0; NewPinIndex < NewOutputPins.Num(); NewPinIndex++)
+			for (int32 NewPinIndex = 0; NewPinIndex < OutputPins.Num(); NewPinIndex++)
 			{
-				if (OldOutputPins[PinIndex]->PinName == NewOutputPins[NewPinIndex]->PinName)
+				if (OldNode->OutputPins[PinIndex]->PinName == OutputPins[NewPinIndex]->PinName)
 				{
 					FoundPinIndex = NewPinIndex;
 					break;
@@ -144,9 +67,9 @@ void UMaterialGraphNode_Base::ReplaceNode(UMaterialGraphNode_Base* OldNode)
 		if (FoundPinIndex == -1)
 		{
 			// Now check types
-			for (int32 NewPinIndex = 0; NewPinIndex < NewOutputPins.Num(); NewPinIndex++)
+			for (int32 NewPinIndex = 0; NewPinIndex < OutputPins.Num(); NewPinIndex++)
 			{
-				if (OldOutputPins[PinIndex]->PinType == NewOutputPins[NewPinIndex]->PinType)
+				if (OldNode->OutputPins[PinIndex]->PinType == OutputPins[NewPinIndex]->PinType)
 				{
 					FoundPinIndex = NewPinIndex;
 					break;
@@ -157,9 +80,9 @@ void UMaterialGraphNode_Base::ReplaceNode(UMaterialGraphNode_Base* OldNode)
 		// If we can't find an equivalent output in this node, just use the first
 		// The user will have to fix up any issues from the mismatch
 		FoundPinIndex = FMath::Max(FoundPinIndex, 0);
-		if (FoundPinIndex < NewOutputPins.Num())
+		if (FoundPinIndex < OutputPins.Num())
 		{
-			ModifyAndCopyPersistentPinData(*NewOutputPins[FoundPinIndex], *OldOutputPins[PinIndex]);
+			ModifyAndCopyPersistentPinData(*OutputPins[FoundPinIndex], *OldNode->OutputPins[PinIndex]);
 		}
 	}
 
@@ -208,11 +131,40 @@ void UMaterialGraphNode_Base::InsertNewNode(UEdGraphPin* FromPin, UEdGraphPin* N
 void UMaterialGraphNode_Base::AllocateDefaultPins()
 {
 	check(Pins.Num() == 0);
+	check(InputPins.Num() == 0);
+	check(OutputPins.Num() == 0);
+	check(PinInfoMap.Num() == 0);
 
 	CreateInputPins();
 	CreateOutputPins();
 }
 
+void UMaterialGraphNode_Base::RegisterPin(UEdGraphPin* Pin, int32 Index, uint32 Type)
+{
+	FMaterialGraphPinInfo& PinInfo = PinInfoMap.FindOrAdd(Pin);
+	PinInfo.Type = Type;
+	PinInfo.Index = Index;
+
+	if (Type & MCT_Execution)
+	{
+		if (Pin->Direction == EGPD_Input)
+		{
+			checkf(ExecInputPin == nullptr, TEXT("Only 1 exec input pin allowed"));
+			check(Index == 0);
+			ExecInputPin = Pin;
+		}
+	}
+	else
+	{
+		switch (Pin->Direction)
+		{
+		case EGPD_Input: verify(InputPins.Add(Pin) == Index); break;
+		case EGPD_Output: verify(OutputPins.Add(Pin) == Index); break;
+		default: checkNoEntry(); break;
+		}
+	}
+}
+
 void UMaterialGraphNode_Base::ReconstructNode()
 {
 	Modify();
@@ -233,38 +185,31 @@ void UMaterialGraphNode_Base::ReconstructNode()
 		}
 	}
 
-	// Store the old Input and Output pins
-	TArray<UEdGraphPin*> OldInputPins;
-	TArray<UEdGraphPin*> OldOutputPins;
-	GetInputPins(OldInputPins);
-	GetOutputPins(OldOutputPins);
-
 	// Move the existing pins to a saved array
-	TArray<UEdGraphPin*> OldPins(Pins);
+	TArray<UEdGraphPin*> OldInputPins = MoveTemp(InputPins);
+	TArray<UEdGraphPin*> OldOutputPins = MoveTemp(OutputPins);
+	TArray<UEdGraphPin*> OldPins = MoveTemp(Pins);
 	Pins.Reset();
+	PinInfoMap.Reset();
+	InputPins.Reset();
+	OutputPins.Reset();
 
 	// Recreate the new pins
 	AllocateDefaultPins();
 
-	// Get new Input and Output pins
-	TArray<UEdGraphPin*> NewInputPins;
-	TArray<UEdGraphPin*> NewOutputPins;
-	GetInputPins(NewInputPins);
-	GetOutputPins(NewOutputPins);
-
 	for (int32 PinIndex = 0; PinIndex < OldInputPins.Num(); PinIndex++)
 	{
-		if (PinIndex < NewInputPins.Num())
+		if (PinIndex < InputPins.Num())
 		{
-			NewInputPins[PinIndex]->MovePersistentDataFromOldPin(*OldInputPins[PinIndex]);
+			InputPins[PinIndex]->MovePersistentDataFromOldPin(*OldInputPins[PinIndex]);
 		}
 	}
 
 	for (int32 PinIndex = 0; PinIndex < OldOutputPins.Num(); PinIndex++)
 	{
-		if (PinIndex < NewOutputPins.Num())
+		if (PinIndex < OutputPins.Num())
 		{
-			NewOutputPins[PinIndex]->MovePersistentDataFromOldPin(*OldOutputPins[PinIndex]);
+			OutputPins[PinIndex]->MovePersistentDataFromOldPin(*OldOutputPins[PinIndex]);
 		}
 	}
 
@@ -351,7 +296,3 @@ FString UMaterialGraphNode_Base::GetDocumentationLink() const
 	return TEXT("Shared/GraphNodes/Material");
 }
 
-uint32 UMaterialGraphNode_Base::GetInputType(const UEdGraphPin* InputPin) const
-{
-	return MCT_Unknown;
-}
diff --git a/Engine/Source/Editor/UnrealEd/Private/MaterialGraphNode_Knot.cpp b/Engine/Source/Editor/UnrealEd/Private/MaterialGraphNode_Knot.cpp
--- a/Engine/Source/Editor/UnrealEd/Private/MaterialGraphNode_Knot.cpp
+++ b/Engine/Source/Editor/UnrealEd/Private/MaterialGraphNode_Knot.cpp
@@ -27,8 +27,10 @@ void UMaterialGraphNode_Knot::AllocateDefaultPins()
 
 	UEdGraphPin* MyInputPin = CreatePin(EGPD_Input, PC_Wildcard, InputPinName);
 	MyInputPin->bDefaultValueIsIgnored = true;
+	RegisterPin(MyInputPin, 0, MCT_Unknown);
 
 	UEdGraphPin* MyOutputPin = CreatePin(EGPD_Output, PC_Wildcard, OutputPinName);
+	RegisterPin(MyOutputPin, 0, MCT_Unknown);
 }
 
 FText UMaterialGraphNode_Knot::GetTooltipText() const
diff --git a/Engine/Source/Editor/UnrealEd/Private/MaterialGraphNode_Root.cpp b/Engine/Source/Editor/UnrealEd/Private/MaterialGraphNode_Root.cpp
--- a/Engine/Source/Editor/UnrealEd/Private/MaterialGraphNode_Root.cpp
+++ b/Engine/Source/Editor/UnrealEd/Private/MaterialGraphNode_Root.cpp
@@ -49,42 +49,29 @@ void UMaterialGraphNode_Root::CreateInputPins()
 {
 	UMaterialGraph* MaterialGraph = CastChecked<UMaterialGraph>(GetGraph());
 
-	for (const FMaterialInputInfo& MaterialInput : MaterialGraph->MaterialInputs)
+	for (int32 Index = 0; Index < MaterialGraph->MaterialInputs.Num(); ++Index)
 	{
-		UEdGraphPin* InputPin = CreatePin(EGPD_Input, UMaterialGraphSchema::PC_MaterialInput, *FString::Printf(TEXT("%d"), (int32)MaterialInput.GetProperty()), *MaterialInput.GetName().ToString());
-	}
-
-}
-
-int32 UMaterialGraphNode_Root::GetInputIndex(const UEdGraphPin* InputPin) const
-{
-	for (int32 Index = 0; Index < Pins.Num(); ++Index)
-	{
-		if (InputPin == Pins[Index])
+		const FMaterialInputInfo& MaterialInput = MaterialGraph->MaterialInputs[Index];
+		EMaterialProperty Property = MaterialInput.GetProperty();
+	
+		uint32 MaterialType = 0u;
+		if (Property == MP_MaterialAttributes)
 		{
-			return Index;
+			MaterialType = MCT_MaterialAttributes;
+		}
+		else if (Property == MP_FrontMaterial)
+		{
+			MaterialType = MCT_Strata;
+		}
+		else
+		{
+			MaterialType = FMaterialAttributeDefinitionMap::GetValueType(Property);
 		}
-	}
-
-	return -1;
-}
 
-uint32 UMaterialGraphNode_Root::GetInputType(const UEdGraphPin* InputPin) const
-{
-	UMaterialGraph* MaterialGraph = CastChecked<UMaterialGraph>(GetGraph());
-	EMaterialProperty Property = MaterialGraph->MaterialInputs[GetInputIndex(InputPin)].GetProperty();
-	if (Property == MP_MaterialAttributes)
-	{
-		return MCT_MaterialAttributes;
-	}
-	else if (Property == MP_FrontMaterial)
-	{
-		return MCT_Strata;
-	}
-	else
-	{
-		return FMaterialAttributeDefinitionMap::GetValueType(Property);
+		UEdGraphPin* InputPin = CreatePin(EGPD_Input, UMaterialGraphSchema::PC_MaterialInput, *FString::Printf(TEXT("%d"), (int32)Property), *MaterialInput.GetName().ToString());
+		RegisterPin(InputPin, Index, MaterialType);
 	}
+
 }
 
 #undef LOCTEXT_NAMESPACE
diff --git a/Engine/Source/Editor/UnrealEd/Private/MaterialGraphSchema.cpp b/Engine/Source/Editor/UnrealEd/Private/MaterialGraphSchema.cpp
--- a/Engine/Source/Editor/UnrealEd/Private/MaterialGraphSchema.cpp
+++ b/Engine/Source/Editor/UnrealEd/Private/MaterialGraphSchema.cpp
@@ -560,12 +560,11 @@ const FPinConnectionResponse UMaterialGraphSchema::CanCreateConnection(const UEd
 		return FPinConnectionResponse(CONNECT_RESPONSE_DISALLOW, ResponseMessage);
 	}
 
-	// Break existing connections on inputs only - multiple output connections are acceptable
+	// For non-exec pins, break existing connections on inputs only - multiple output connections are acceptable
 	if (InputPin->LinkedTo.Num() > 0)
 	{
 		const uint32 InputType = GetMaterialValueType(InputPin);
-		// TODO - Allow multiple exec inputs, will generate a compile error where not supported
-		//if (!(InputType & MCT_Execution))
+		if (!(InputType & MCT_Execution))
 		{
 			ECanCreateConnectionResponse ReplyBreakOutputs;
 			if (InputPin == A)
@@ -584,6 +583,29 @@ const FPinConnectionResponse UMaterialGraphSchema::CanCreateConnection(const UEd
 		}
 	}
 
+	// For exec pins, reverse is true - multiple input connections are acceptable
+	if (OutputPin->LinkedTo.Num() > 0)
+	{
+		const uint32 OutputType = GetMaterialValueType(InputPin);
+		if (OutputType & MCT_Execution)
+		{
+			ECanCreateConnectionResponse ReplyBreakInputs;
+			if (OutputPin == A)
+			{
+				ReplyBreakInputs = CONNECT_RESPONSE_BREAK_OTHERS_A;
+			}
+			else
+			{
+				ReplyBreakInputs = CONNECT_RESPONSE_BREAK_OTHERS_B;
+			}
+			if (ResponseMessage.IsEmpty())
+			{
+				ResponseMessage = LOCTEXT("ConnectionReplace", "Replace existing connections");
+			}
+			return FPinConnectionResponse(ReplyBreakInputs, ResponseMessage);
+		}
+	}
+
 	return FPinConnectionResponse(CONNECT_RESPONSE_MAKE, ResponseMessage);
 }
 
diff --git a/Engine/Source/Editor/UnrealEd/Public/MaterialGraphNode_Knot.h b/Engine/Source/Editor/UnrealEd/Public/MaterialGraphNode_Knot.h
--- a/Engine/Source/Editor/UnrealEd/Public/MaterialGraphNode_Knot.h
+++ b/Engine/Source/Editor/UnrealEd/Public/MaterialGraphNode_Knot.h
@@ -26,11 +26,11 @@ public:
 	
 	UEdGraphPin* GetInputPin() const
 	{
-		return Pins[0];
+		return InputPins[0];
 	}
 
 	UEdGraphPin* GetOutputPin() const
 	{
-		return Pins[1];
+		return OutputPins[0];
 	}
 };
diff --git a/Engine/Source/Runtime/Engine/Classes/Materials/MaterialExpression.h b/Engine/Source/Runtime/Engine/Classes/Materials/MaterialExpression.h
--- a/Engine/Source/Runtime/Engine/Classes/Materials/MaterialExpression.h
+++ b/Engine/Source/Runtime/Engine/Classes/Materials/MaterialExpression.h
@@ -99,6 +99,25 @@ struct FExpressionOutput
 };
 #endif
 
+USTRUCT()
+struct FExpressionExecOutput
+{
+	GENERATED_BODY()
+
+#if WITH_EDITOR
+	ENGINE_API int32 Compile(class FMaterialCompiler* Compiler);
+#endif // WITH_EDITOR
+
+	UPROPERTY()
+	class UMaterialExpression* Expression = nullptr;
+};
+
+struct FExpressionExecOutputEntry
+{
+	FName Name;
+	FExpressionExecOutput* Output = nullptr;
+};
+
 UCLASS(abstract, BlueprintType, hidecategories=Object)
 class ENGINE_API UMaterialExpression : public UObject
 {
@@ -269,10 +288,9 @@ class ENGINE_API UMaterialExpression : public UObject
 	virtual uint32 GetInputType(int32 InputIndex);
 	virtual uint32 GetOutputType(int32 OutputIndex);
 
-	/** Gets this expression's execution input, if any */
-	virtual FExpressionInput* GetExecInput();
+	virtual void GetExecOutputs(TArray<FExpressionExecOutputEntry>& Outputs);
 
-	bool IsExecInputConnection(UMaterialExpression* InExpression, int32 InOutputIndex);
+	virtual bool HasExecInput();
 
 	virtual FText GetCreationDescription() const;
 	virtual FText GetCreationName() const;
diff --git a/Engine/Source/Runtime/Engine/Classes/Materials/MaterialExpressionExecBegin.h b/Engine/Source/Runtime/Engine/Classes/Materials/MaterialExpressionExecBegin.h
--- a/Engine/Source/Runtime/Engine/Classes/Materials/MaterialExpressionExecBegin.h
+++ b/Engine/Source/Runtime/Engine/Classes/Materials/MaterialExpressionExecBegin.h
@@ -14,11 +14,13 @@ class UMaterialExpressionExecBegin : public UMaterialExpression
 {
 	GENERATED_UCLASS_BODY()
 
+	UPROPERTY()
+	FExpressionExecOutput Exec;
+
 	//~ Begin UMaterialExpression Interface
 #if WITH_EDITOR
 	virtual int32 Compile(class FMaterialCompiler* Compiler, int32 OutputIndex) override;
 	virtual void GetCaption(TArray<FString>& OutCaptions) const override;
-	virtual uint32 GetOutputType(int32 OutputIndex);
 	virtual bool CanUserDeleteExpression() const override { return false; }
 #endif
 	//~ End UMaterialExpression Interface
diff --git a/Engine/Source/Runtime/Engine/Classes/Materials/MaterialExpressionForLoop.h b/Engine/Source/Runtime/Engine/Classes/Materials/MaterialExpressionForLoop.h
--- a/Engine/Source/Runtime/Engine/Classes/Materials/MaterialExpressionForLoop.h
+++ b/Engine/Source/Runtime/Engine/Classes/Materials/MaterialExpressionForLoop.h
@@ -15,7 +15,10 @@ class UMaterialExpressionForLoop : public UMaterialExpression
 	GENERATED_UCLASS_BODY()
 
 	UPROPERTY()
-	FExpressionInput Exec;
+	FExpressionExecOutput LoopBody;
+
+	UPROPERTY()
+	FExpressionExecOutput Completed;
 
 	UPROPERTY()
 	FExpressionInput StartIndex;
@@ -32,7 +35,7 @@ class UMaterialExpressionForLoop : public UMaterialExpression
 	virtual void GetCaption(TArray<FString>& OutCaptions) const override;
 	virtual uint32 GetInputType(int32 InputIndex) override;
 	virtual uint32 GetOutputType(int32 InputIndex) override;
-	virtual FExpressionInput* GetExecInput() override { return &Exec; }
+	virtual bool HasExecInput() override { return true; }
 #endif
 	//~ End UMaterialExpression Interface
 };
diff --git a/Engine/Source/Runtime/Engine/Classes/Materials/MaterialExpressionIfThenElse.h b/Engine/Source/Runtime/Engine/Classes/Materials/MaterialExpressionIfThenElse.h
--- a/Engine/Source/Runtime/Engine/Classes/Materials/MaterialExpressionIfThenElse.h
+++ b/Engine/Source/Runtime/Engine/Classes/Materials/MaterialExpressionIfThenElse.h
@@ -15,7 +15,10 @@ class UMaterialExpressionIfThenElse : public UMaterialExpression
 	GENERATED_UCLASS_BODY()
 
 	UPROPERTY()
-	FExpressionInput Exec;
+	FExpressionExecOutput Then;
+
+	UPROPERTY()
+	FExpressionExecOutput Else;
 
 	UPROPERTY()
 	FExpressionInput Condition;
@@ -25,8 +28,7 @@ class UMaterialExpressionIfThenElse : public UMaterialExpression
 	virtual int32 Compile(class FMaterialCompiler* Compiler, int32 OutputIndex) override;
 	virtual void GetCaption(TArray<FString>& OutCaptions) const override;
 	virtual uint32 GetInputType(int32 InputIndex) override;
-	virtual uint32 GetOutputType(int32 InputIndex) override;
-	virtual FExpressionInput* GetExecInput() override { return &Exec; }
+	virtual bool HasExecInput() override { return true; }
 #endif
 	//~ End UMaterialExpression Interface
 };
diff --git a/Engine/Source/Runtime/Engine/Classes/Materials/MaterialExpressionReturnMaterialAttributes.h b/Engine/Source/Runtime/Engine/Classes/Materials/MaterialExpressionReturnMaterialAttributes.h
--- a/Engine/Source/Runtime/Engine/Classes/Materials/MaterialExpressionReturnMaterialAttributes.h
+++ b/Engine/Source/Runtime/Engine/Classes/Materials/MaterialExpressionReturnMaterialAttributes.h
@@ -14,9 +14,6 @@ class UMaterialExpressionReturnMaterialAttributes : public UMaterialExpression
 {
 	GENERATED_UCLASS_BODY()
 
-	UPROPERTY()
-	FExpressionInput Exec;
-
 	UPROPERTY()
 	FMaterialAttributesInput MaterialAttributes;
 
@@ -27,7 +24,7 @@ class UMaterialExpressionReturnMaterialAttributes : public UMaterialExpression
 	virtual const TArray<FExpressionInput*> GetInputs() override;
 	virtual FExpressionInput* GetInput(int32 InputIndex) override;
 	virtual uint32 GetInputType(int32 InputIndex) override;
-	virtual FExpressionInput* GetExecInput() override { return &Exec; }
+	virtual bool HasExecInput() override { return true; }
 #endif
 	//~ End UMaterialExpression Interface
 };
diff --git a/Engine/Source/Runtime/Engine/Classes/Materials/MaterialExpressionSetLocal.h b/Engine/Source/Runtime/Engine/Classes/Materials/MaterialExpressionSetLocal.h
--- a/Engine/Source/Runtime/Engine/Classes/Materials/MaterialExpressionSetLocal.h
+++ b/Engine/Source/Runtime/Engine/Classes/Materials/MaterialExpressionSetLocal.h
@@ -15,7 +15,7 @@ class UMaterialExpressionSetLocal : public UMaterialExpression
 	GENERATED_UCLASS_BODY()
 
 	UPROPERTY()
-	FExpressionInput Exec;
+	FExpressionExecOutput Exec;
 
 	UPROPERTY()
 	FExpressionInput Value;
@@ -28,8 +28,7 @@ class UMaterialExpressionSetLocal : public UMaterialExpression
 	virtual int32 Compile(class FMaterialCompiler* Compiler, int32 OutputIndex) override;
 	virtual void GetCaption(TArray<FString>& OutCaptions) const override;
 	virtual uint32 GetInputType(int32 InputIndex) override;
-	virtual uint32 GetOutputType(int32 InputIndex) override;
-	virtual FExpressionInput* GetExecInput() override { return &Exec; }
+	virtual bool HasExecInput() override { return true; }
 #endif
 	//~ End UMaterialExpression Interface
 };
diff --git a/Engine/Source/Runtime/Engine/Private/Materials/HLSLMaterialTranslator.cpp b/Engine/Source/Runtime/Engine/Private/Materials/HLSLMaterialTranslator.cpp
--- a/Engine/Source/Runtime/Engine/Private/Materials/HLSLMaterialTranslator.cpp
+++ b/Engine/Source/Runtime/Engine/Private/Materials/HLSLMaterialTranslator.cpp
@@ -3951,25 +3951,73 @@ int32 FHLSLMaterialTranslator::CallExpression(FMaterialExpressionKey ExpressionK
 	return Result;
 }
 
+int32 FHLSLMaterialTranslator::CallExpressionExec(UMaterialExpression* Expression)
+{
+	check(ShaderFrequency < SF_NumFrequencies);
+	auto& CurrentFunctionStack = FunctionStacks[ShaderFrequency];
+	FMaterialFunctionCompileState* CurrentFunctionState = CurrentFunctionStack.Last();
+
+	int32 Result = INDEX_NONE;
+	/*int32* ExistingCodeIndex = CurrentFunctionState->ExecExpressionCodeMap.Find(Expression);
+	if (ExistingCodeIndex)
+	{
+		Result = *ExistingCodeIndex;
+		FShaderCodeChunk& Chunk = (*CurrentScopeChunks)[Result];
+		const int32 CurrentScopeIndex = ScopeStack.Last();
+		const FShaderCodeChunk& CurrentScope = (*CurrentScopeChunks)[CurrentScopeIndex];
+		const FShaderCodeChunk& PrevScope = (*CurrentScopeChunks)[Chunk.UsedScopeIndex];
+		if (CurrentScope.UsedScopeIndex == PrevScope.UsedScopeIndex)
+		{
+			AddCodeChunkToScope(Result, CurrentScope.UsedScopeIndex);
+		}
+		else
+		{
+			return Errorf(TEXT("Invalid scopes"));
+		}
+	}
+	else*/
+	{
+		ReferencedCodeChunks.Reset();
+		Result = Expression->Compile(this, UMaterialExpression::CompileExecutionOutputIndex);
+		CurrentFunctionState->ExecExpressionCodeMap.Add(Expression, Result);
+
+		if (Result != INDEX_NONE)
+		{
+			FShaderCodeChunk& ResultChunk = (*CurrentScopeChunks)[Result];
+			ResultChunk.ReferencedCodeChunks = MoveTemp(ReferencedCodeChunks);
+		}
+		ReferencedCodeChunks.Reset();
+	}
+	return Result;
+}
+
 void FHLSLMaterialTranslator::AddCodeChunkToCurrentScope(int32 ChunkIndex)
 {
 	if (ChunkIndex != INDEX_NONE && ScopeStack.Num() > 0)
 	{
 		const int32 CurrentScopeIndex = ScopeStack.Last();
-		const FShaderCodeChunk& CurrentScope = (*CurrentScopeChunks)[CurrentScopeIndex];
+		AddCodeChunkToScope(ChunkIndex, CurrentScopeIndex);
+	}
+}
+
+void FHLSLMaterialTranslator::AddCodeChunkToScope(int32 ChunkIndex, int32 ScopeIndex)
+{
+	if (ChunkIndex != INDEX_NONE)
+	{
+		FShaderCodeChunk& CurrentScope = (*CurrentScopeChunks)[ScopeIndex];
 
 		FShaderCodeChunk& Chunk = (*CurrentScopeChunks)[ChunkIndex];
 		if (Chunk.DeclaredScopeIndex == INDEX_NONE)
 		{
 			check(Chunk.UsedScopeIndex == INDEX_NONE);
-			Chunk.DeclaredScopeIndex = CurrentScopeIndex;
-			Chunk.UsedScopeIndex = CurrentScopeIndex;
+			Chunk.DeclaredScopeIndex = ScopeIndex;
+			Chunk.UsedScopeIndex = ScopeIndex;
 			Chunk.ScopeLevel = CurrentScope.ScopeLevel + 1;
 		}
-		else if(Chunk.UsedScopeIndex != CurrentScopeIndex)
+		else if (Chunk.UsedScopeIndex != ScopeIndex)
 		{
 			// Find the most derived scope that's shared by the current scope, and the scope this code was previously referenced from
-			int32 ScopeIndex0 = CurrentScopeIndex;
+			int32 ScopeIndex0 = ScopeIndex;
 			int32 ScopeIndex1 = Chunk.UsedScopeIndex;
 			while (ScopeIndex0 != ScopeIndex1)
 			{
diff --git a/Engine/Source/Runtime/Engine/Private/Materials/HLSLMaterialTranslator.h b/Engine/Source/Runtime/Engine/Private/Materials/HLSLMaterialTranslator.h
--- a/Engine/Source/Runtime/Engine/Private/Materials/HLSLMaterialTranslator.h
+++ b/Engine/Source/Runtime/Engine/Private/Materials/HLSLMaterialTranslator.h
@@ -578,6 +578,7 @@ protected:
 	FString CreateSymbolName(const TCHAR* SymbolNameHint);
 
 	void AddCodeChunkToCurrentScope(int32 ChunkIndex);
+	void AddCodeChunkToScope(int32 ChunkIndex, int32 ScopeIndex);
 
 	/** Adds an already formatted inline or referenced code chunk */
 	int32 AddCodeChunkInner(uint64 Hash, const TCHAR* FormattedCode, EMaterialValueType Type, bool bInlined);
@@ -663,6 +664,8 @@ protected:
 
 	virtual int32 CallExpression(FMaterialExpressionKey ExpressionKey, FMaterialCompiler* Compiler) override;
 
+	virtual int32 CallExpressionExec(UMaterialExpression* Expression) override;
+
 	virtual EMaterialValueType GetType(int32 Code) override;
 	virtual EMaterialQualityLevel::Type GetQualityLevel() override;
 	virtual ERHIFeatureLevel::Type GetFeatureLevel() override;
diff --git a/Engine/Source/Runtime/Engine/Private/Materials/MaterialExpressions.cpp b/Engine/Source/Runtime/Engine/Private/Materials/MaterialExpressions.cpp
--- a/Engine/Source/Runtime/Engine/Private/Materials/MaterialExpressions.cpp
+++ b/Engine/Source/Runtime/Engine/Private/Materials/MaterialExpressions.cpp
@@ -1016,6 +1016,21 @@ FName UMaterialExpression::GetInputName(int32 InputIndex) const
 	return NAME_None;
 }
 
+void UMaterialExpression::GetExecOutputs(TArray<FExpressionExecOutputEntry>& OutExpressionOutpus)
+{
+	for (TFieldIterator<FStructProperty> It(GetClass(), EFieldIteratorFlags::IncludeSuper, EFieldIteratorFlags::ExcludeDeprecated); It; ++It)
+	{
+		const FStructProperty* StructProp = *It;
+		if (StructProp->Struct == FExpressionExecOutput::StaticStruct())
+		{
+			check(StructProp->ArrayDim == 1);
+			FExpressionExecOutputEntry& Entry = OutExpressionOutpus.AddDefaulted_GetRef();
+			Entry.Output = StructProp->ContainerPtrToValuePtr<FExpressionExecOutput>(this);
+			Entry.Name = StructProp->GetFName();
+		}
+	}
+}
+
 FText UMaterialExpression::GetCreationDescription() const
 {
 	return FText::GetEmpty();
@@ -1128,19 +1143,8 @@ uint32 UMaterialExpression::GetOutputType(int32 OutputIndex)
 	}
 }
 
-FExpressionInput* UMaterialExpression::GetExecInput()
-{
-	return nullptr;
-}
-
-bool UMaterialExpression::IsExecInputConnection(UMaterialExpression* InExpression, int32 InOutputIndex)
+bool UMaterialExpression::HasExecInput()
 {
-	FExpressionInput* Input = GetExecInput();
-	if (Input)
-	{
-		FExpressionInput TracedInput = Input->GetTracedInput();
-		return TracedInput.Expression == InExpression && TracedInput.OutputIndex == InOutputIndex;
-	}
 	return false;
 }
 
@@ -1149,8 +1153,6 @@ int32 UMaterialExpression::GetWidth() const
 	return ME_STD_THUMBNAIL_SZ+(ME_STD_BORDER*2);
 }
 
-
-
 int32 UMaterialExpression::GetHeight() const
 {
 	return FMath::Max(ME_CAPTION_HEIGHT + (Outputs.Num() * ME_STD_TAB_HEIGHT),ME_CAPTION_HEIGHT+ME_STD_THUMBNAIL_SZ+(ME_STD_BORDER*2));
@@ -1421,9 +1423,7 @@ bool UMaterialExpression::HasConnectedOutputs() const
 		UMaterialGraphNode* MatGraphNode = Cast<UMaterialGraphNode>(GraphNode);
 		if (MatGraphNode)
 		{
-			TArray<UEdGraphPin*> OutputPins;
-			MatGraphNode->GetOutputPins(OutputPins);
-			for (UEdGraphPin* Pin : OutputPins)
+			for (UEdGraphPin* Pin : MatGraphNode->OutputPins)
 			{
 				if (Pin->LinkedTo.Num() > 0)
 				{
@@ -20708,7 +20708,6 @@ UMaterialExpressionExecBegin::UMaterialExpressionExecBegin(const FObjectInitiali
 {
 #if WITH_EDITORONLY_DATA
 	Outputs.Reset();
-	Outputs.Add(FExpressionOutput(TEXT("exec")));
 	bHidePreviewWindow = true;
 #endif
 }
@@ -20719,37 +20718,13 @@ void UMaterialExpressionExecBegin::GetCaption(TArray<FString>& OutCaptions) cons
 	OutCaptions.Add(TEXT("Material Execution Begin"));
 }
 
-uint32 UMaterialExpressionExecBegin::GetOutputType(int32 OutputIndex)
-{
-	check(OutputIndex == 0);
-	return MCT_Execution;
-}
-
-static UMaterialExpression* GetExecConnection(UMaterialExpression* InExpression, int32 InOutputIndex)
-{
-	const uint32 OutputType = InExpression->GetOutputType(InOutputIndex);
-	check(OutputType == MCT_Execution);
-
-	const auto& Expressions = InExpression->Material->Expressions;
-	for (UMaterialExpression* Expression : Expressions)
-	{
-		if (Expression->IsExecInputConnection(InExpression, InOutputIndex))
-		{
-			return Expression;
-		}
-	}
-
-	return nullptr;
-}
-
 int32 UMaterialExpressionExecBegin::Compile(class FMaterialCompiler* Compiler, int32 OutputIndex)
 {
 	check(OutputIndex == CompileExecutionOutputIndex);
-	UMaterialExpression* NextExpression = GetExecConnection(this, 0);
-	if (NextExpression)
+	if (Exec.Expression)
 	{
 		const int32 RootScope = Compiler->BeginScope();
-		NextExpression->Compile(Compiler, CompileExecutionOutputIndex);
+		Exec.Compile(Compiler);
 		Compiler->EndScope();
 		return RootScope;
 	}
@@ -20789,7 +20764,6 @@ int32 UMaterialExpressionReturnMaterialAttributes::Compile(class FMaterialCompil
 const TArray<FExpressionInput*> UMaterialExpressionReturnMaterialAttributes::GetInputs()
 {
 	TArray<FExpressionInput*> Result;
-	Result.Add(&Exec);
 	Result.Add(&MaterialAttributes);
 	return Result;
 }
@@ -20798,8 +20772,7 @@ FExpressionInput* UMaterialExpressionReturnMaterialAttributes::GetInput(int32 In
 {
 	switch (InputIndex)
 	{
-	case 0: return &Exec;
-	case 1: return &MaterialAttributes;
+	case 0: return &MaterialAttributes;
 	default: checkNoEntry(); return nullptr;
 	}
 }
@@ -20808,8 +20781,7 @@ uint32 UMaterialExpressionReturnMaterialAttributes::GetInputType(int32 InputInde
 {
 	switch (InputIndex)
 	{
-	case 0: return MCT_Execution;
-	case 1: return MCT_MaterialAttributes;
+	case 0: return MCT_MaterialAttributes;
 	default: checkNoEntry(); return 0u;
 	}
 }
@@ -20827,8 +20799,6 @@ UMaterialExpressionIfThenElse::UMaterialExpressionIfThenElse(const FObjectInitia
 #if WITH_EDITORONLY_DATA
 	MenuCategories.Add(ConstructorStatics.NAME_Execution);
 	Outputs.Reset();
-	Outputs.Add(FExpressionOutput(TEXT("Then")));
-	Outputs.Add(FExpressionOutput(TEXT("Else")));
 	bShowOutputNameOnPin = true;
 	bHidePreviewWindow = true;
 #endif
@@ -20838,8 +20808,7 @@ UMaterialExpressionIfThenElse::UMaterialExpressionIfThenElse(const FObjectInitia
 int32 UMaterialExpressionIfThenElse::Compile(class FMaterialCompiler* Compiler, int32 OutputIndex)
 {
 	check(OutputIndex == CompileExecutionOutputIndex);
-	UMaterialExpression* ThenExpression = GetExecConnection(this, 0);
-	if (!ThenExpression)
+	if (!Then.Expression)
 	{
 		return Compiler->Error(TEXT("Then pin must be connected"));
 	}
@@ -20851,15 +20820,14 @@ int32 UMaterialExpressionIfThenElse::Compile(class FMaterialCompiler* Compiler,
 	}
 
 	const int32 IfScopeIndex = Compiler->BeginScope_If(ConditionResult);
-	ThenExpression->Compile(Compiler, CompileExecutionOutputIndex);
+	Then.Compile(Compiler);
 	Compiler->EndScope();
 
-	UMaterialExpression* ElseExpression = GetExecConnection(this, 1);
 	int32 ElseScope = INDEX_NONE;
-	if (ElseExpression)
+	if (Else.Expression)
 	{
 		ElseScope = Compiler->BeginScope_Else();
-		ElseExpression->Compile(Compiler, CompileExecutionOutputIndex);
+		Else.Compile(Compiler);
 		Compiler->EndScope();
 	}
 
@@ -20875,16 +20843,10 @@ uint32 UMaterialExpressionIfThenElse::GetInputType(int32 InputIndex)
 {
 	switch (InputIndex)
 	{
-	case 0: return MCT_Execution;
-	case 1: return MCT_Float1;
+	case 0: return MCT_Float1;
 	default: checkNoEntry(); return 0u;
 	}
 }
-
-uint32 UMaterialExpressionIfThenElse::GetOutputType(int32 InputIndex)
-{
-	return MCT_Execution;
-}
 #endif
 
 UMaterialExpressionForLoop::UMaterialExpressionForLoop(const FObjectInitializer& ObjectInitializer)
@@ -20899,9 +20861,7 @@ UMaterialExpressionForLoop::UMaterialExpressionForLoop(const FObjectInitializer&
 #if WITH_EDITORONLY_DATA
 	MenuCategories.Add(ConstructorStatics.NAME_Execution);
 	Outputs.Reset();
-	Outputs.Add(FExpressionOutput(TEXT("Loop Body")));
 	Outputs.Add(FExpressionOutput(TEXT("Index")));
-	Outputs.Add(FExpressionOutput(TEXT("Completed")));
 	bShowOutputNameOnPin = true;
 	bHidePreviewWindow = true;
 #endif
@@ -20912,14 +20872,11 @@ int32 UMaterialExpressionForLoop::Compile(class FMaterialCompiler* Compiler, int
 {
 	if (OutputIndex == CompileExecutionOutputIndex)
 	{
-		UMaterialExpression* LoopExpression = GetExecConnection(this, 0);
-		if (!LoopExpression)
+		if (!LoopBody.Expression)
 		{
 			return Compiler->Error(TEXT("Loop pin must be connected"));
 		}
 
-		UMaterialExpression* CompletedExpression = GetExecConnection(this, 2);
-
 		const int32 StartResult = StartIndex.GetTracedInput().Compile(Compiler);
 		if (StartResult == INDEX_NONE)
 		{
@@ -20933,19 +20890,16 @@ int32 UMaterialExpressionForLoop::Compile(class FMaterialCompiler* Compiler, int
 		}
 
 		const int32 LoopScopeIndex = Compiler->BeginScope_For(this, StartResult, EndResult, Compiler->Constant(1.0f));
-		LoopExpression->Compile(Compiler, CompileExecutionOutputIndex);
+		LoopBody.Compile(Compiler);
 		Compiler->EndScope();
 
-		if (CompletedExpression)
-		{
-			CompletedExpression->Compile(Compiler, CompileExecutionOutputIndex);
-		}
+		Completed.Compile(Compiler);
 
 		return INDEX_NONE;
 	}
 
 	// Loop index output
-	check(OutputIndex == 1);
+	check(OutputIndex == 0);
 	return Compiler->ForLoopIndex(this);
 }
 
@@ -20958,10 +20912,9 @@ uint32 UMaterialExpressionForLoop::GetInputType(int32 InputIndex)
 {
 	switch (InputIndex)
 	{
-	case 0: return MCT_Execution;
+	case 0: return MCT_Float1;
 	case 1: return MCT_Float1;
 	case 2: return MCT_Float1;
-	case 3: return MCT_Float1;
 	default: checkNoEntry(); return 0u;
 	}
 }
@@ -20970,9 +20923,7 @@ uint32 UMaterialExpressionForLoop::GetOutputType(int32 OutputIndex)
 {
 	switch (OutputIndex)
 	{
-	case 0: return MCT_Execution;
-	case 1: return MCT_Float1;
-	case 2: return MCT_Execution;
+	case 0: return MCT_Float1;
 	default: checkNoEntry(); return 0u;
 	}
 }
@@ -21023,7 +20974,6 @@ UMaterialExpressionSetLocal::UMaterialExpressionSetLocal(const FObjectInitialize
 #if WITH_EDITORONLY_DATA
 	MenuCategories.Add(ConstructorStatics.NAME_Execution);
 	Outputs.Reset();
-	Outputs.Add(FExpressionOutput(TEXT("Exec")));
 	bHidePreviewWindow = true;
 #endif
 }
@@ -21033,8 +20983,7 @@ uint32 UMaterialExpressionSetLocal::GetInputType(int32 InputIndex)
 {
 	switch (InputIndex)
 	{
-	case 0: return MCT_Execution;
-	case 1: return MCT_Float1;
+	case 0: return MCT_Float1;
 	default: checkNoEntry(); return 0u;
 	}
 }
@@ -21051,11 +21000,7 @@ int32 UMaterialExpressionSetLocal::Compile(class FMaterialCompiler* Compiler, in
 
 	Compiler->SetLocal(LocalName, ValueIndex);
 
-	UMaterialExpression* NextExpression = GetExecConnection(this, 0);
-	if (NextExpression)
-	{
-		NextExpression->Compile(Compiler, CompileExecutionOutputIndex);
-	}
+	Exec.Compile(Compiler);
 
 	return INDEX_NONE;
 }
@@ -21064,11 +21009,6 @@ void UMaterialExpressionSetLocal::GetCaption(TArray<FString>& OutCaptions) const
 {
 	OutCaptions.Add(FString::Printf(TEXT("Set %s"), *LocalName.ToString()));
 }
-
-uint32 UMaterialExpressionSetLocal::GetOutputType(int32 InputIndex)
-{
-	return MCT_Execution;
-}
 #endif
 
 #undef LOCTEXT_NAMESPACE
diff --git a/Engine/Source/Runtime/Engine/Private/Materials/MaterialShared.cpp b/Engine/Source/Runtime/Engine/Private/Materials/MaterialShared.cpp
--- a/Engine/Source/Runtime/Engine/Private/Materials/MaterialShared.cpp
+++ b/Engine/Source/Runtime/Engine/Private/Materials/MaterialShared.cpp
@@ -276,6 +276,17 @@ FExpressionInput FExpressionInput::GetTracedInput() const
 	return *this;
 }
 
+#if WITH_EDITOR
+int32 FExpressionExecOutput::Compile(class FMaterialCompiler* Compiler)
+{
+	if (Expression)
+	{
+		return Compiler->CallExpressionExec(Expression);
+	}
+	return INDEX_NONE;
+}
+#endif // WITH_EDITOR
+
 /** Native serialize for FMaterialExpression struct */
 static bool SerializeExpressionInput(FArchive& Ar, FExpressionInput& Input)
 {
diff --git a/Engine/Source/Runtime/Engine/Public/MaterialCompiler.h b/Engine/Source/Runtime/Engine/Public/MaterialCompiler.h
--- a/Engine/Source/Runtime/Engine/Public/MaterialCompiler.h
+++ b/Engine/Source/Runtime/Engine/Public/MaterialCompiler.h
@@ -84,6 +84,8 @@ public:
 
 	virtual int32 CallExpression(FMaterialExpressionKey ExpressionKey,FMaterialCompiler* InCompiler) = 0;
 
+	virtual int32 CallExpressionExec(UMaterialExpression* Expression) = 0;
+
 	virtual EMaterialCompilerType GetCompilerType() const { return EMaterialCompilerType::Standard; }
 	inline bool IsMaterialProxyCompiler() const { return GetCompilerType() == EMaterialCompilerType::MaterialProxy; }
 	inline bool IsLightmassCompiler() const { return GetCompilerType() == EMaterialCompilerType::Lightmass; }
@@ -475,6 +477,8 @@ public:
 
 	virtual int32 CallExpression(FMaterialExpressionKey ExpressionKey, FMaterialCompiler* InCompiler) override { return Compiler->CallExpression(ExpressionKey, InCompiler); }
 
+	virtual int32 CallExpressionExec(UMaterialExpression* Expression) { return Compiler->CallExpressionExec(Expression); }
+
 	virtual void PushFunction(FMaterialFunctionCompileState* FunctionState) override { Compiler->PushFunction(FunctionState); }
 	virtual FMaterialFunctionCompileState* PopFunction() override { return Compiler->PopFunction(); }
 	virtual int32 GetCurrentFunctionStackDepth() override { return Compiler->GetCurrentFunctionStackDepth(); }
diff --git a/Engine/Source/Runtime/Engine/Public/MaterialShared.h b/Engine/Source/Runtime/Engine/Public/MaterialShared.h
--- a/Engine/Source/Runtime/Engine/Public/MaterialShared.h
+++ b/Engine/Source/Runtime/Engine/Public/MaterialShared.h
@@ -1472,6 +1472,8 @@ public:
 	/** A map from material expression to the index into CodeChunks of the code for the material expression. */
 	TMap<FMaterialExpressionKey,int32> ExpressionCodeMap;
 
+	TMap<UMaterialExpression*, int32> ExecExpressionCodeMap;
+
 private:
 	/** Cache of MaterialFunctionOutput CodeChunks.  Allows for further reuse than just the ExpressionCodeMap */
 	TMap<FMaterialExpressionKey, FMaterialFunctionCompileState*> SharedFunctionStates;
