Preserve group actors folders when copy/pasting them.

diff --git a/Engine/Source/Editor/UnrealEd/Private/EditorExporters.cpp b/Engine/Source/Editor/UnrealEd/Private/EditorExporters.cpp
--- a/Engine/Source/Editor/UnrealEd/Private/EditorExporters.cpp
+++ b/Engine/Source/Editor/UnrealEd/Private/EditorExporters.cpp
@@ -510,9 +510,10 @@ bool ULevelExporterT3D::ExportText( const FExportObjectInnerContext* Context, UO
 				FString ParentActorString = ( ParentActor ? FString::Printf(TEXT(" ParentActor=%s"), *ParentActor->GetName() ) : TEXT(""));
 				FString SocketNameString = ( (ParentActor && SocketName != NAME_None) ? FString::Printf(TEXT(" SocketName=%s"), *SocketName.ToString() ) : TEXT(""));
 				FString GroupActor = (Actor->GroupActor? FString::Printf(TEXT(" GroupActor=%s"), *Actor->GroupActor->GetName() ) : TEXT(""));
-				Ar.Logf( TEXT("%sBegin Actor Class=%s Name=%s Archetype=%s'%s'%s%s%s") LINE_TERMINATOR, 
+				FString GroupFolder = (Actor->GroupActor? FString::Printf(TEXT(" GroupFolder=%s"), *Actor->GroupActor->GetFolderPath().ToString() ) : TEXT(""));
+				Ar.Logf( TEXT("%sBegin Actor Class=%s Name=%s Archetype=%s'%s'%s%s%s%s") LINE_TERMINATOR, 
 					FCString::Spc(TextIndent), *Actor->GetClass()->GetPathName(), *Actor->GetName(),
-					*Actor->GetArchetype()->GetClass()->GetPathName(), *Actor->GetArchetype()->GetPathName(), *ParentActorString, *SocketNameString, *GroupActor );
+					*Actor->GetArchetype()->GetClass()->GetPathName(), *Actor->GetArchetype()->GetPathName(), *ParentActorString, *SocketNameString, *GroupActor, *GroupFolder );
 
 				ExportRootScope = Actor;
 				ExportObjectInner( Context, Actor, Ar, PortFlags | PPF_ExportsNotFullyQualified );
diff --git a/Engine/Source/Editor/UnrealEd/Private/Factories/EditorFactories.cpp b/Engine/Source/Editor/UnrealEd/Private/Factories/EditorFactories.cpp
--- a/Engine/Source/Editor/UnrealEd/Private/Factories/EditorFactories.cpp
+++ b/Engine/Source/Editor/UnrealEd/Private/Factories/EditorFactories.cpp
@@ -976,10 +976,12 @@ UObject* ULevelFactory::FactoryCreateText
 								AGroupActor** tmpNewGroup = nullptr;
 								// We need to add all the objects we selected into groups with new objects that were in their group before.
 								FString GroupName;
+								FString GroupFolder;
 								if(FParse::Value(Str, TEXT("GroupActor="), GroupName))
 								{
 									tmpNewGroup = NewGroups.Find(GroupName);
 									bGrouped = true;
+									FParse::Value(Str, TEXT("GroupFolder="), GroupFolder);
 								}
 
 								// Does the group exist?
@@ -994,6 +996,7 @@ UObject* ULevelFactory::FactoryCreateText
 								{
 									// Create a new group and add the actor.
 									AGroupActor* SpawnedGroupActor = NewActor->GetWorld()->SpawnActor<AGroupActor>();
+									SpawnedGroupActor->SetFolderPath(*GroupFolder);
 									SpawnedGroupActor->Add(*NewActor);
 
 									// Place the group in the map so we can find it later.
