redo cl 15323093 in UE5:

Do not select actor during the Editor Level Libray SpawnActor function.
This first significantly spead up the operation and reduce memory usage spike by not allocating all the details panel widget layout for each actor for no reason.

When used in a loop in BP, it could dramatically increase memory usage without being able to return to the system at the end.

diff --git a/Engine/Source/Editor/UnrealEd/Private/Subsystems/EditorActorSubsystem.cpp b/Engine/Source/Editor/UnrealEd/Private/Subsystems/EditorActorSubsystem.cpp
--- a/Engine/Source/Editor/UnrealEd/Private/Subsystems/EditorActorSubsystem.cpp
+++ b/Engine/Source/Editor/UnrealEd/Private/Subsystems/EditorActorSubsystem.cpp
@@ -110,7 +110,7 @@ namespace InternalActorUtilitiesSubsystemLibrary
 		}
 
 		UActorFactory* FactoryToUse = nullptr;
-		bool bSelectActors = true;
+		bool bSelectActors = false;
 		TArray<AActor*> Actors = FLevelEditorViewportClient::TryPlacingActorFromObject(DesiredLevel, ObjToUse, bSelectActors, NewObjectFlags, FactoryToUse);
 
 		if (Actors.Num() == 0 || Actors[0] == nullptr)
