Skip unsaved actor packages that contains pending kill actors from the list of packages to save.

diff --git a/Engine/Source/Editor/UnrealEd/Private/FileHelpers.cpp b/Engine/Source/Editor/UnrealEd/Private/FileHelpers.cpp
--- a/Engine/Source/Editor/UnrealEd/Private/FileHelpers.cpp
+++ b/Engine/Source/Editor/UnrealEd/Private/FileHelpers.cpp
@@ -4331,7 +4331,30 @@ void FEditorFileUtils::GetDirtyWorldPackages(TArray<UPackage*>& OutDirtyPackages
 				{
 					if (ExternalPackage->IsDirty())
 					{
-						OutDirtyPackages.Add(ExternalPackage);
+						bool bActorPackageNeedsToSave = true;
+
+						// Skip unsaved packages containing only pending kill actors
+						if (ExternalPackage->HasAnyPackageFlags(PKG_NewlyCreated))
+						{
+							bActorPackageNeedsToSave = false;
+							ForEachObjectWithPackage(ExternalPackage, [&bActorPackageNeedsToSave](UObject* Object)
+							{
+								if (Cast<AActor>(Object))
+								{
+									if (!Object->IsPendingKill())
+									{
+										bActorPackageNeedsToSave = true;
+										return false;
+									}
+								}
+								return true;
+							}, false);
+						}
+
+						if (bActorPackageNeedsToSave)
+						{
+							OutDirtyPackages.Add(ExternalPackage);
+						}
 					}
 				}
 			}
