Fix replicating material graph node data after a copy/paste

diff --git a/Engine/Source/Editor/UnrealEd/Private/MaterialGraphNode_Base.cpp b/Engine/Source/Editor/UnrealEd/Private/MaterialGraphNode_Base.cpp
--- a/Engine/Source/Editor/UnrealEd/Private/MaterialGraphNode_Base.cpp
+++ b/Engine/Source/Editor/UnrealEd/Private/MaterialGraphNode_Base.cpp
@@ -186,9 +186,8 @@ void UMaterialGraphNode_Base::ReconstructNode()
 	}
 
 	// Move the existing pins to a saved array
-	TArray<UEdGraphPin*> OldInputPins = MoveTemp(InputPins);
-	TArray<UEdGraphPin*> OldOutputPins = MoveTemp(OutputPins);
 	TArray<UEdGraphPin*> OldPins = MoveTemp(Pins);
+
 	Pins.Reset();
 	PinInfoMap.Reset();
 	InputPins.Reset();
@@ -197,25 +196,16 @@ void UMaterialGraphNode_Base::ReconstructNode()
 	// Recreate the new pins
 	AllocateDefaultPins();
 
-	for (int32 PinIndex = 0; PinIndex < OldInputPins.Num(); PinIndex++)
-	{
-		if (PinIndex < InputPins.Num())
-		{
-			InputPins[PinIndex]->MovePersistentDataFromOldPin(*OldInputPins[PinIndex]);
-		}
-	}
-
-	for (int32 PinIndex = 0; PinIndex < OldOutputPins.Num(); PinIndex++)
+	for (UEdGraphPin* OldPin : OldPins)
 	{
-		if (PinIndex < OutputPins.Num())
+		// Transfer data to new pin
+		UEdGraphPin* NewPin = FindPin(OldPin->PinName, OldPin->Direction);
+		if (NewPin)
 		{
-			OutputPins[PinIndex]->MovePersistentDataFromOldPin(*OldOutputPins[PinIndex]);
+			NewPin->MovePersistentDataFromOldPin(*OldPin);
 		}
-	}
 
-	// Throw away the original pins
-	for (UEdGraphPin* OldPin : OldPins)
-	{
+		// Throw away the original pins
 		OldPin->Modify();
 		UEdGraphNode::DestroyPin(OldPin);
 	}
