Updated EditorHeaderButton widget to properly bind the OnClicked event

diff --git a/Engine/Source/Editor/UnrealEd/Private/SEditorHeaderButton.cpp b/Engine/Source/Editor/UnrealEd/Private/SEditorHeaderButton.cpp
--- a/Engine/Source/Editor/UnrealEd/Private/SEditorHeaderButton.cpp
+++ b/Engine/Source/Editor/UnrealEd/Private/SEditorHeaderButton.cpp
@@ -44,6 +44,7 @@ void SEditorHeaderButton::Construct(const FArguments& InArgs)
 			.HAlign(HAlign_Center)
 			.VAlign(VAlign_Center)
 			.ContentPadding(0)
+			.OnClicked(InArgs._OnClicked)
 			[
 				ButtonContent
 			]
