Slight modification to factory.cpp code to check the BulkDataPointer instead of the RawData to see if PVS-Studio likes that code better

#rnx

diff --git a/Engine/Source/Editor/UnrealEd/Private/Factories/Factory.cpp b/Engine/Source/Editor/UnrealEd/Private/Factories/Factory.cpp
--- a/Engine/Source/Editor/UnrealEd/Private/Factories/Factory.cpp
+++ b/Engine/Source/Editor/UnrealEd/Private/Factories/Factory.cpp
@@ -541,8 +541,8 @@ bool UFactory::ImportUntypedBulkDataFromText(const TCHAR*& Buffer, FUntypedBulkD
 				}
 				else if (FParse::Value(Str, TEXT("BEGIN "), ParsedText) && (ParsedText.ToUpper() == TEXT("BINARY")))
 				{
-					check(RawData);
 					uint8* BulkDataPointer = RawData;
+					check(BulkDataPointer);
 					while(FParse::Line(&Buffer,StrLine))
 					{
 						Str = *StrLine;
@@ -568,9 +568,7 @@ bool UFactory::ImportUntypedBulkDataFromText(const TCHAR*& Buffer, FUntypedBulkD
 								ParseStr +=2;
 							}
 							Value = FParse::HexDigit(ParseStr[0]) * 16 + FParse::HexDigit(ParseStr[1]);
-#ifndef PVS_STUDIO // Build machine refuses to disable this warning
-							*BulkDataPointer = (uint8)Value; //-V522
-#endif
+							*BulkDataPointer = (uint8)Value;
 							BulkDataPointer++;
 							ParseStr += 2;
 							ParseStr++;
