Support focus camera on actors if the selected actor doesn't have a primitive component.

diff --git a/Engine/Source/Editor/UnrealEd/Private/EditorServer.cpp b/Engine/Source/Editor/UnrealEd/Private/EditorServer.cpp
--- a/Engine/Source/Editor/UnrealEd/Private/EditorServer.cpp
+++ b/Engine/Source/Editor/UnrealEd/Private/EditorServer.cpp
@@ -4836,6 +4836,7 @@ void UEditorEngine::MoveViewportCamerasToActor(const TArray<AActor*> &Actors, co
 						RootComponent->GetChildrenComponents(true, SceneComponents);
 						SceneComponents.Add(RootComponent);
 
+						bool bHasAtLeastOnePrimitiveComponent = false;
 						for (USceneComponent* SceneComponent : SceneComponents)
 						{
 							UPrimitiveComponent* PrimitiveComponent = Cast<UPrimitiveComponent>(SceneComponent);
@@ -4856,9 +4857,17 @@ void UEditorEngine::MoveViewportCamerasToActor(const TArray<AActor*> &Actors, co
 									{
 										BoundingBox += PrimitiveComponent->Bounds.GetBox();
 									}
+
+									bHasAtLeastOnePrimitiveComponent = true;
 								}
 							}
 						}
+
+						if (!bHasAtLeastOnePrimitiveComponent)
+						{
+							BoundingBox += RootComponent->GetComponentLocation();
+						}
+
 					}
 				}
 			}
