Fix state reset on material graph nodes

diff --git a/Engine/Source/Editor/UnrealEd/Classes/MaterialGraph/MaterialGraphNode_Base.h b/Engine/Source/Editor/UnrealEd/Classes/MaterialGraph/MaterialGraphNode_Base.h
--- a/Engine/Source/Editor/UnrealEd/Classes/MaterialGraph/MaterialGraphNode_Base.h
+++ b/Engine/Source/Editor/UnrealEd/Classes/MaterialGraph/MaterialGraphNode_Base.h
@@ -76,4 +76,6 @@ protected:
 	void ModifyAndCopyPersistentPinData(UEdGraphPin& TargetPin, const UEdGraphPin& SourcePin) const;
 
 	void RegisterPin(UEdGraphPin* Pin, int32 Index, uint32 Type);
+
+	void EmptyPins();
 };
diff --git a/Engine/Source/Editor/UnrealEd/Private/MaterialGraphNode.cpp b/Engine/Source/Editor/UnrealEd/Private/MaterialGraphNode.cpp
--- a/Engine/Source/Editor/UnrealEd/Private/MaterialGraphNode.cpp
+++ b/Engine/Source/Editor/UnrealEd/Private/MaterialGraphNode.cpp
@@ -87,7 +87,7 @@ void UMaterialGraphNode::RecreateAndLinkNode()
 
 		UEdGraphNode::DestroyPin(Pin);
 	}
-	Pins.Empty();
+	EmptyPins();
 
 	AllocateDefaultPins();
 
diff --git a/Engine/Source/Editor/UnrealEd/Private/MaterialGraphNode_Base.cpp b/Engine/Source/Editor/UnrealEd/Private/MaterialGraphNode_Base.cpp
--- a/Engine/Source/Editor/UnrealEd/Private/MaterialGraphNode_Base.cpp
+++ b/Engine/Source/Editor/UnrealEd/Private/MaterialGraphNode_Base.cpp
@@ -165,6 +165,14 @@ void UMaterialGraphNode_Base::RegisterPin(UEdGraphPin* Pin, int32 Index, uint32
 	}
 }
 
+void UMaterialGraphNode_Base::EmptyPins()
+{
+	Pins.Reset();
+	PinInfoMap.Reset();
+	InputPins.Reset();
+	OutputPins.Reset();
+}
+
 void UMaterialGraphNode_Base::ReconstructNode()
 {
 	Modify();
@@ -188,10 +196,7 @@ void UMaterialGraphNode_Base::ReconstructNode()
 	// Move the existing pins to a saved array
 	TArray<UEdGraphPin*> OldPins = MoveTemp(Pins);
 
-	Pins.Reset();
-	PinInfoMap.Reset();
-	InputPins.Reset();
-	OutputPins.Reset();
+	EmptyPins();
 
 	// Recreate the new pins
 	AllocateDefaultPins();
