Don't cache the results of GetInput/OutputType() for material pins, as this type can change for certain expressions

diff --git a/Engine/Source/Editor/GraphEditor/Private/MaterialNodes/SGraphNodeMaterialBase.cpp b/Engine/Source/Editor/GraphEditor/Private/MaterialNodes/SGraphNodeMaterialBase.cpp
--- a/Engine/Source/Editor/GraphEditor/Private/MaterialNodes/SGraphNodeMaterialBase.cpp
+++ b/Engine/Source/Editor/GraphEditor/Private/MaterialNodes/SGraphNodeMaterialBase.cpp
@@ -279,7 +279,7 @@ void SGraphNodeMaterialBase::CreatePinWidgets()
 			if ((bIsABreakAttrNode && CurPin->Direction == EGPD_Output) || (bIsAMakeAttrNode && CurPin->Direction == EGPD_Input))
 			{
 				const FMaterialGraphPinInfo& PinInfo = MaterialNode->GetPinInfo(CurPin);
-				if (!(PinInfo.Type & MCT_Execution))
+				if (PinInfo.PinType == EMaterialGraphPinType::Data)
 				{
 					bPinDesiresToBeHidden |= !MaterialGraph->MaterialInputs[PinInfo.Index].IsVisiblePin(MaterialGraph->Material, true);
 				}
diff --git a/Engine/Source/Editor/UnrealEd/Classes/MaterialGraph/MaterialGraphNode.h b/Engine/Source/Editor/UnrealEd/Classes/MaterialGraph/MaterialGraphNode.h
--- a/Engine/Source/Editor/UnrealEd/Classes/MaterialGraph/MaterialGraphNode.h
+++ b/Engine/Source/Editor/UnrealEd/Classes/MaterialGraph/MaterialGraphNode.h
@@ -70,7 +70,8 @@ public:
 	virtual void OnCommentBubbleToggled( bool bInCommentBubbleVisible ) override;
 	virtual void GetPinHoverText(const UEdGraphPin& Pin, FString& HoverTextOut) const override;
 	virtual FString GetDocumentationExcerptName() const override;
-	virtual bool CanUserDeleteNode() const override; 
+	virtual bool CanUserDeleteNode() const override;
+	virtual uint32 GetPinMaterialType(const UEdGraphPin* Pin, const FMaterialGraphPinInfo& PinInfo) const override;
 	//~ End UEdGraphNode Interface.
 
 	//~ Begin UMaterialGraphNode_Base Interface
diff --git a/Engine/Source/Editor/UnrealEd/Classes/MaterialGraph/MaterialGraphNode_Base.h b/Engine/Source/Editor/UnrealEd/Classes/MaterialGraph/MaterialGraphNode_Base.h
--- a/Engine/Source/Editor/UnrealEd/Classes/MaterialGraph/MaterialGraphNode_Base.h
+++ b/Engine/Source/Editor/UnrealEd/Classes/MaterialGraph/MaterialGraphNode_Base.h
@@ -10,9 +10,15 @@
 class UEdGraphPin;
 class UEdGraphSchema;
 
+enum class EMaterialGraphPinType
+{
+	Data,
+	Exec,
+};
+
 struct FMaterialGraphPinInfo
 {
-	uint32 Type; // EMaterialValueType
+	EMaterialGraphPinType PinType;
 	int32 Index; // index into the expression's list of inputs/outputs (exec inputs/outpus are indexed separately)
 };
 
@@ -75,7 +81,9 @@ class UMaterialGraphNode_Base : public UEdGraphNode
 protected:
 	void ModifyAndCopyPersistentPinData(UEdGraphPin& TargetPin, const UEdGraphPin& SourcePin) const;
 
-	void RegisterPin(UEdGraphPin* Pin, int32 Index, uint32 Type);
+	void RegisterPin(UEdGraphPin* Pin, EMaterialGraphPinType Type, int32 Index);
+
+	virtual uint32 GetPinMaterialType(const UEdGraphPin* Pin, const FMaterialGraphPinInfo& PinInfo) const;
 
 	void EmptyPins();
 };
diff --git a/Engine/Source/Editor/UnrealEd/Classes/MaterialGraph/MaterialGraphNode_Root.h b/Engine/Source/Editor/UnrealEd/Classes/MaterialGraph/MaterialGraphNode_Root.h
--- a/Engine/Source/Editor/UnrealEd/Classes/MaterialGraph/MaterialGraphNode_Root.h
+++ b/Engine/Source/Editor/UnrealEd/Classes/MaterialGraph/MaterialGraphNode_Root.h
@@ -28,6 +28,7 @@ class UMaterialGraphNode_Root : public UMaterialGraphNode_Base
 	//~ End UEdGraphNode Interface.
 
 	//~ Begin UMaterialGraphNode_Base Interface
+	virtual uint32 GetPinMaterialType(const UEdGraphPin* Pin, const FMaterialGraphPinInfo& PinInfo) const override;
 	virtual void CreateInputPins() override;
 	virtual bool IsRootNode() const override {return true;}
 	//~ End UMaterialGraphNode_Base Interface
diff --git a/Engine/Source/Editor/UnrealEd/Private/MaterialGraph.cpp b/Engine/Source/Editor/UnrealEd/Private/MaterialGraph.cpp
--- a/Engine/Source/Editor/UnrealEd/Private/MaterialGraph.cpp
+++ b/Engine/Source/Editor/UnrealEd/Private/MaterialGraph.cpp
@@ -330,7 +330,7 @@ void UMaterialGraph::LinkGraphNodesFromMaterial()
 		{
 			UEdGraphPin* Pin = It.Key;
 			const FMaterialGraphPinInfo& PinInfo = It.Value;
-			if (Pin->Direction == EGPD_Input && PinInfo.Type != MCT_Execution)
+			if (Pin->Direction == EGPD_Input && PinInfo.PinType == EMaterialGraphPinType::Data)
 			{
 				if (ExpressionInputs[PinInfo.Index]->Expression)
 				{
@@ -360,7 +360,7 @@ void UMaterialGraph::LinkGraphNodesFromMaterial()
 					}
 				}
 			}
-			else if (Pin->Direction == EGPD_Output && PinInfo.Type == MCT_Execution)
+			else if (Pin->Direction == EGPD_Output && PinInfo.PinType == EMaterialGraphPinType::Exec)
 			{
 				FExpressionExecOutput* ExecOutput = ExecOutputs[PinInfo.Index].Output;
 				if (ExecOutput->Expression)
@@ -450,7 +450,7 @@ void UMaterialGraph::LinkMaterialExpressionsFromGraph() const
 					{
 						UEdGraphPin* Pin = It.Key;
 						const FMaterialGraphPinInfo& PinInfo = It.Value;
-						if (Pin->Direction == EGPD_Input && PinInfo.Type != MCT_Execution)
+						if (Pin->Direction == EGPD_Input && PinInfo.PinType == EMaterialGraphPinType::Data)
 						{
 							// Wire up non-execution input pins
 							FExpressionInput* ExpressionInput = ExpressionInputs[PinInfo.Index];
@@ -488,7 +488,7 @@ void UMaterialGraph::LinkMaterialExpressionsFromGraph() const
 								ExpressionInput->Expression = NULL;
 							}
 						}
-						else if (Pin->Direction == EGPD_Output && PinInfo.Type == MCT_Execution)
+						else if (Pin->Direction == EGPD_Output && PinInfo.PinType == EMaterialGraphPinType::Exec)
 						{
 							// Wire up execution output pins
 							FExpressionExecOutput* ExpressionOutput = ExecOutputs[PinInfo.Index].Output;
diff --git a/Engine/Source/Editor/UnrealEd/Private/MaterialGraphNode.cpp b/Engine/Source/Editor/UnrealEd/Private/MaterialGraphNode.cpp
--- a/Engine/Source/Editor/UnrealEd/Private/MaterialGraphNode.cpp
+++ b/Engine/Source/Editor/UnrealEd/Private/MaterialGraphNode.cpp
@@ -620,6 +620,27 @@ FName UMaterialGraphNode::GetShortenPinName(const FName PinName)
 	return InputName;
 }
 
+uint32 UMaterialGraphNode::GetPinMaterialType(const UEdGraphPin* Pin, const FMaterialGraphPinInfo& PinInfo) const
+{
+	switch (PinInfo.PinType)
+	{
+	case EMaterialGraphPinType::Exec:
+		return MCT_Execution;
+	case EMaterialGraphPinType::Data:
+		if (Pin->Direction == EGPD_Input)
+		{
+			return MaterialExpression->GetInputType(PinInfo.Index);
+		}
+		else
+		{
+			return MaterialExpression->GetOutputType(PinInfo.Index);
+		}
+	default:
+		checkNoEntry();
+		return 0u;
+	}
+}
+
 void UMaterialGraphNode::CreateInputPins()
 {
 	if (MaterialExpression->HasExecInput())
@@ -629,32 +650,24 @@ void UMaterialGraphNode::CreateInputPins()
 		NewPin->PinName = CreateUniquePinName(TEXT("Input"));
 		NewPin->PinFriendlyName = SpaceText;
 
-		RegisterPin(NewPin, 0, MCT_Execution);
+		RegisterPin(NewPin, EMaterialGraphPinType::Exec, 0);
 	}
 
 	const TArray<FExpressionInput*> ExpressionInputs = MaterialExpression->GetInputs();
-	for (int32 Index = 0; Index < ExpressionInputs.Num() ; ++Index)
+	for (int32 Index = 0; Index < ExpressionInputs.Num(); ++Index)
 	{
 		FExpressionInput* Input = ExpressionInputs[Index];
-		const uint32 InputType = MaterialExpression->GetInputType(Index);
 		FName PinCategory;
-		FName InputName;
-		if (InputType == MCT_Execution)
+
+		FName InputName = MaterialExpression->GetInputName(Index);
+		InputName = GetShortenPinName(InputName);
+		if (MaterialExpression->IsInputConnectionRequired(Index))
 		{
-			PinCategory = UMaterialGraphSchema::PC_Exec;
+			PinCategory = UMaterialGraphSchema::PC_Required;
 		}
 		else
 		{
-			InputName = MaterialExpression->GetInputName(Index);
-			InputName = GetShortenPinName(InputName);
-			if (MaterialExpression->IsInputConnectionRequired(Index))
-			{
-				PinCategory = UMaterialGraphSchema::PC_Required;
-			}
-			else
-			{
-				PinCategory = UMaterialGraphSchema::PC_Optional;
-			}
+			PinCategory = UMaterialGraphSchema::PC_Optional;
 		}
 
 		UEdGraphPin* NewPin = CreatePin(EGPD_Input, PinCategory, InputName);
@@ -665,7 +678,7 @@ void UMaterialGraphNode::CreateInputPins()
 			NewPin->PinFriendlyName = SpaceText;
 		}
 
-		RegisterPin(NewPin, Index, InputType);
+		RegisterPin(NewPin, EMaterialGraphPinType::Data, Index);
 	}
 }
 
@@ -675,7 +688,6 @@ void UMaterialGraphNode::CreateOutputPins()
 	for(int32 Index = 0; Index < Outputs.Num(); ++Index)
 	{
 		const FExpressionOutput& ExpressionOutput = Outputs[Index];
-		const uint32 OutputType = MaterialExpression->GetOutputType(Index);
 
 		FName PinCategory;
 		FName PinSubCategory;
@@ -685,36 +697,29 @@ void UMaterialGraphNode::CreateOutputPins()
 			OutputName = ExpressionOutput.OutputName;
 		}
 
-		if (OutputType == MCT_Execution)
+		if (MaterialExpression->bShowMaskColorsOnPin && ExpressionOutput.Mask)
 		{
-			PinCategory = UMaterialGraphSchema::PC_Exec;
-		}
-		else
-		{
-			if (MaterialExpression->bShowMaskColorsOnPin && ExpressionOutput.Mask)
-			{
-				PinCategory = UMaterialGraphSchema::PC_Mask;
+			PinCategory = UMaterialGraphSchema::PC_Mask;
 
-				if (ExpressionOutput.MaskR && !ExpressionOutput.MaskG && !ExpressionOutput.MaskB && !ExpressionOutput.MaskA)
-				{
-					PinSubCategory = UMaterialGraphSchema::PSC_Red;
-				}
-				else if (!ExpressionOutput.MaskR &&  ExpressionOutput.MaskG && !ExpressionOutput.MaskB && !ExpressionOutput.MaskA)
-				{
-					PinSubCategory = UMaterialGraphSchema::PSC_Green;
-				}
-				else if (!ExpressionOutput.MaskR && !ExpressionOutput.MaskG &&  ExpressionOutput.MaskB && !ExpressionOutput.MaskA)
-				{
-					PinSubCategory = UMaterialGraphSchema::PSC_Blue;
-				}
-				else if (!ExpressionOutput.MaskR && !ExpressionOutput.MaskG && !ExpressionOutput.MaskB &&  ExpressionOutput.MaskA)
-				{
-					PinSubCategory = UMaterialGraphSchema::PSC_Alpha;
-				}
-				else if (ExpressionOutput.MaskR && ExpressionOutput.MaskG && ExpressionOutput.MaskB &&  ExpressionOutput.MaskA)
-				{
-					PinSubCategory = UMaterialGraphSchema::PSC_RGBA;
-				}
+			if (ExpressionOutput.MaskR && !ExpressionOutput.MaskG && !ExpressionOutput.MaskB && !ExpressionOutput.MaskA)
+			{
+				PinSubCategory = UMaterialGraphSchema::PSC_Red;
+			}
+			else if (!ExpressionOutput.MaskR && ExpressionOutput.MaskG && !ExpressionOutput.MaskB && !ExpressionOutput.MaskA)
+			{
+				PinSubCategory = UMaterialGraphSchema::PSC_Green;
+			}
+			else if (!ExpressionOutput.MaskR && !ExpressionOutput.MaskG && ExpressionOutput.MaskB && !ExpressionOutput.MaskA)
+			{
+				PinSubCategory = UMaterialGraphSchema::PSC_Blue;
+			}
+			else if (!ExpressionOutput.MaskR && !ExpressionOutput.MaskG && !ExpressionOutput.MaskB && ExpressionOutput.MaskA)
+			{
+				PinSubCategory = UMaterialGraphSchema::PSC_Alpha;
+			}
+			else if (ExpressionOutput.MaskR && ExpressionOutput.MaskG && ExpressionOutput.MaskB && ExpressionOutput.MaskA)
+			{
+				PinSubCategory = UMaterialGraphSchema::PSC_RGBA;
 			}
 		}
 
@@ -726,7 +731,7 @@ void UMaterialGraphNode::CreateOutputPins()
 			NewPin->PinFriendlyName = SpaceText;
 		}
 
-		RegisterPin(NewPin, Index, OutputType);
+		RegisterPin(NewPin, EMaterialGraphPinType::Data, Index);
 	}
 
 	TArray<FExpressionExecOutputEntry> ExecOutputs;
@@ -749,7 +754,7 @@ void UMaterialGraphNode::CreateOutputPins()
 			NewPin->PinFriendlyName = SpaceText;
 		}
 
-		RegisterPin(NewPin, Index, MCT_Execution);
+		RegisterPin(NewPin, EMaterialGraphPinType::Exec, Index);
 	}
 }
 
diff --git a/Engine/Source/Editor/UnrealEd/Private/MaterialGraphNode_Base.cpp b/Engine/Source/Editor/UnrealEd/Private/MaterialGraphNode_Base.cpp
--- a/Engine/Source/Editor/UnrealEd/Private/MaterialGraphNode_Base.cpp
+++ b/Engine/Source/Editor/UnrealEd/Private/MaterialGraphNode_Base.cpp
@@ -26,12 +26,26 @@ const FMaterialGraphPinInfo& UMaterialGraphNode_Base::GetPinInfo(const class UEd
 
 uint32 UMaterialGraphNode_Base::GetOutputType(const UEdGraphPin* OutputPin) const
 {
-	return GetPinInfo(OutputPin).Type;
+	return GetPinMaterialType(OutputPin, GetPinInfo(OutputPin));
 }
 
 uint32 UMaterialGraphNode_Base::GetInputType(const UEdGraphPin* InputPin) const
 {
-	return GetPinInfo(InputPin).Type;
+	return GetPinMaterialType(InputPin, GetPinInfo(InputPin));
+}
+
+uint32 UMaterialGraphNode_Base::GetPinMaterialType(const UEdGraphPin* Pin, const FMaterialGraphPinInfo& PinInfo) const
+{
+	switch (PinInfo.PinType)
+	{
+	case EMaterialGraphPinType::Data:
+		return MCT_Unknown;
+	case EMaterialGraphPinType::Exec:
+		return MCT_Execution;
+	default:
+		checkNoEntry();
+		return 0u;
+	}
 }
 
 void UMaterialGraphNode_Base::ReplaceNode(UMaterialGraphNode_Base* OldNode)
@@ -139,13 +153,13 @@ void UMaterialGraphNode_Base::AllocateDefaultPins()
 	CreateOutputPins();
 }
 
-void UMaterialGraphNode_Base::RegisterPin(UEdGraphPin* Pin, int32 Index, uint32 Type)
+void UMaterialGraphNode_Base::RegisterPin(UEdGraphPin* Pin, EMaterialGraphPinType Type, int32 Index)
 {
 	FMaterialGraphPinInfo& PinInfo = PinInfoMap.FindOrAdd(Pin);
-	PinInfo.Type = Type;
+	PinInfo.PinType = Type;
 	PinInfo.Index = Index;
 
-	if (Type & MCT_Execution)
+	if (Type == EMaterialGraphPinType::Exec)
 	{
 		if (Pin->Direction == EGPD_Input)
 		{
diff --git a/Engine/Source/Editor/UnrealEd/Private/MaterialGraphNode_Knot.cpp b/Engine/Source/Editor/UnrealEd/Private/MaterialGraphNode_Knot.cpp
--- a/Engine/Source/Editor/UnrealEd/Private/MaterialGraphNode_Knot.cpp
+++ b/Engine/Source/Editor/UnrealEd/Private/MaterialGraphNode_Knot.cpp
@@ -27,10 +27,10 @@ void UMaterialGraphNode_Knot::AllocateDefaultPins()
 
 	UEdGraphPin* MyInputPin = CreatePin(EGPD_Input, PC_Wildcard, InputPinName);
 	MyInputPin->bDefaultValueIsIgnored = true;
-	RegisterPin(MyInputPin, 0, MCT_Unknown);
+	RegisterPin(MyInputPin, EMaterialGraphPinType::Data, 0);
 
 	UEdGraphPin* MyOutputPin = CreatePin(EGPD_Output, PC_Wildcard, OutputPinName);
-	RegisterPin(MyOutputPin, 0, MCT_Unknown);
+	RegisterPin(MyOutputPin, EMaterialGraphPinType::Data, 0);
 }
 
 FText UMaterialGraphNode_Knot::GetTooltipText() const
diff --git a/Engine/Source/Editor/UnrealEd/Private/MaterialGraphNode_Root.cpp b/Engine/Source/Editor/UnrealEd/Private/MaterialGraphNode_Root.cpp
--- a/Engine/Source/Editor/UnrealEd/Private/MaterialGraphNode_Root.cpp
+++ b/Engine/Source/Editor/UnrealEd/Private/MaterialGraphNode_Root.cpp
@@ -45,6 +45,29 @@ void UMaterialGraphNode_Root::PostPlacedNewNode()
 	}
 }
 
+uint32 UMaterialGraphNode_Root::GetPinMaterialType(const UEdGraphPin* Pin, const FMaterialGraphPinInfo& PinInfo) const
+{
+	check(PinInfo.PinType == EMaterialGraphPinType::Data);
+	const UMaterialGraph* MaterialGraph = CastChecked<UMaterialGraph>(GetGraph());
+	const FMaterialInputInfo& MaterialInput = MaterialGraph->MaterialInputs[PinInfo.Index];
+	EMaterialProperty Property = MaterialInput.GetProperty();
+
+	uint32 MaterialType = 0u;
+	if (Property == MP_MaterialAttributes)
+	{
+		MaterialType = MCT_MaterialAttributes;
+	}
+	else if (Property == MP_FrontMaterial)
+	{
+		MaterialType = MCT_Strata;
+	}
+	else
+	{
+		MaterialType = FMaterialAttributeDefinitionMap::GetValueType(Property);
+	}
+	return MaterialType;
+}
+
 void UMaterialGraphNode_Root::CreateInputPins()
 {
 	UMaterialGraph* MaterialGraph = CastChecked<UMaterialGraph>(GetGraph());
@@ -54,22 +77,8 @@ void UMaterialGraphNode_Root::CreateInputPins()
 		const FMaterialInputInfo& MaterialInput = MaterialGraph->MaterialInputs[Index];
 		EMaterialProperty Property = MaterialInput.GetProperty();
 	
-		uint32 MaterialType = 0u;
-		if (Property == MP_MaterialAttributes)
-		{
-			MaterialType = MCT_MaterialAttributes;
-		}
-		else if (Property == MP_FrontMaterial)
-		{
-			MaterialType = MCT_Strata;
-		}
-		else
-		{
-			MaterialType = FMaterialAttributeDefinitionMap::GetValueType(Property);
-		}
-
 		UEdGraphPin* InputPin = CreatePin(EGPD_Input, UMaterialGraphSchema::PC_MaterialInput, *FString::Printf(TEXT("%d"), (int32)Property), *MaterialInput.GetName().ToString());
-		RegisterPin(InputPin, Index, MaterialType);
+		RegisterPin(InputPin, EMaterialGraphPinType::Data, Index);
 	}
 
 }
