Editor Domain Prototype
#rn Minor Cooking

diff --git a/Engine/Config/BaseEditor.ini b/Engine/Config/BaseEditor.ini
--- a/Engine/Config/BaseEditor.ini
+++ b/Engine/Config/BaseEditor.ini
@@ -369,6 +369,8 @@ BatchCount=100
 MaxMemoryAllowance=16384
 
 [CookSettings]
+; EditorDomain is in prototype and currently is a significant deoptimization
+EditorDomainEnabled=false
 PackagesPerGC=0
 IdleTimeToGC=20
 MemoryMaxUsedPhysical=16384
diff --git a/Engine/Source/Editor/UnrealEd/Classes/Commandlets/EditorDomainSaveCommandlet.h b/Engine/Source/Editor/UnrealEd/Classes/Commandlets/EditorDomainSaveCommandlet.h
new file mode 100644
--- /dev/null
+++ b/Engine/Source/Editor/UnrealEd/Classes/Commandlets/EditorDomainSaveCommandlet.h
@@ -0,0 +1,22 @@
+// Copyright Epic Games, Inc. All Rights Reserved.
+
+#pragma once
+
+#include "Commandlets/Commandlet.h"
+
+#include "EditorDomainSaveCommandlet.generated.h"
+
+/** Runs the EditorDomainSave server for client Editor and CookCommandlet processes. */
+UCLASS(config = Editor)
+class UEditorDomainSaveCommandlet
+	: public UCommandlet
+{
+	GENERATED_UCLASS_BODY()
+
+public:
+
+	//~ Begin UCommandlet Interface
+	virtual int32 Main(const FString& CmdLineParams) override;
+	//~ End UCommandlet Interface
+};
+
diff --git a/Engine/Source/Editor/UnrealEd/Private/CookOnTheFlyServer.cpp b/Engine/Source/Editor/UnrealEd/Private/CookOnTheFlyServer.cpp
--- a/Engine/Source/Editor/UnrealEd/Private/CookOnTheFlyServer.cpp
+++ b/Engine/Source/Editor/UnrealEd/Private/CookOnTheFlyServer.cpp
@@ -16,6 +16,7 @@
 #include "Cooker/PackageTracker.h"
 #include "CookPackageSplitter.h"
 #include "Containers/RingBuffer.h"
+#include "EditorDomain/EditorDomain.h"
 #include "HAL/PlatformFileManager.h"
 #include "HAL/FileManager.h"
 #include "Misc/CommandLine.h"
@@ -4201,7 +4202,9 @@ void UCookOnTheFlyServer::Initialize( ECookMode::Type DesiredCookMode, ECookInit
 		MemoryMaxUsedVirtual / 1024 / 1024, MemoryMaxUsedPhysical / 1024 / 1024, MemoryMinFreeVirtual / 1024 / 1024, MemoryMinFreePhysical / 1024 / 1024);
 
 	if (IsCookByTheBookMode() && !IsCookingInEditor() &&
-		FPlatformMisc::SupportsMultithreadedFileHandles()// Preloading moves file handles between threads
+		FPlatformMisc::SupportsMultithreadedFileHandles() && // Preloading moves file handles between threads
+		!FEditorDomain::Get() && // EDITOR_DOMAIN_TODO EditorDomain does not yet support opening packages from async threads, because AssetRegistry is not asynchronous
+		!GAllowCookedDataInEditorBuilds // // Use of preloaded files is not yet implemented when GAllowCookedDataInEditorBuilds is on, see FLinkerLoad::CreateLoader
 		)
 	{
 		bPreloadingEnabled = true;
diff --git a/Engine/Source/Editor/UnrealEd/Private/Cooker/CookPackageData.cpp b/Engine/Source/Editor/UnrealEd/Private/Cooker/CookPackageData.cpp
--- a/Engine/Source/Editor/UnrealEd/Private/Cooker/CookPackageData.cpp
+++ b/Engine/Source/Editor/UnrealEd/Private/Cooker/CookPackageData.cpp
@@ -678,12 +678,6 @@ namespace Cook
 			SetIsPreloadAttempted(true);
 			return true;
 		}
-		if (GAllowCookedDataInEditorBuilds)
-		{
-			// Use of preloaded files is not yet implemented when GAllowCookedDataInEditorBuilds is on, see FLinkerLoad::CreateLoader
-			SetIsPreloadAttempted(true);
-			return true;
-		}
 		if (!PreloadableFile)
 		{
 			TStringBuilder<NAME_SIZE> FileNameString;
diff --git a/Engine/Source/Editor/UnrealEd/Private/EditorDomain/EditorDomain.cpp b/Engine/Source/Editor/UnrealEd/Private/EditorDomain/EditorDomain.cpp
new file mode 100644
--- /dev/null
+++ b/Engine/Source/Editor/UnrealEd/Private/EditorDomain/EditorDomain.cpp
@@ -0,0 +1,337 @@
+// Copyright Epic Games, Inc. All Rights Reserved.
+
+#include "EditorDomain/EditorDomain.h"
+
+#include "AssetRegistry/AssetRegistryModule.h"
+#include "AssetRegistry/IAssetRegistry.h"
+#include "Async/AsyncFileHandleNull.h"
+#include "Containers/UnrealString.h"
+#include "DerivedDataCache.h"
+#include "DerivedDataCacheInterface.h"
+#include "EditorDomain/EditorDomainArchive.h"
+#include "EditorDomain/EditorDomainSave.h"
+#include "EditorDomain/EditorDomainUtils.h"
+#include "HAL/CriticalSection.h"
+#include "Misc/CommandLine.h"
+#include "Misc/ConfigCacheIni.h"
+#include "Misc/PackagePath.h"
+#include "Misc/PackageSegment.h"
+#include "Misc/ScopeLock.h"
+#include "Serialization/Archive.h"
+#include "Templates/RefCounting.h"
+#include "Templates/UniquePtr.h"
+#include "UObject/PackageResourceManagerFile.h"
+
+DEFINE_LOG_CATEGORY(LogEditorDomain);
+
+/** Add a hook to the PackageResourceManager's startup delegate to use the EditorDomain as the IPackageResourceManager */
+class FEditorDomainRegisterAsPackageResourceManager
+{
+public:
+	FEditorDomainRegisterAsPackageResourceManager()
+	{
+		IPackageResourceManager::GetSetPackageResourceManagerDelegate().BindStatic(SetPackageResourceManager);
+	}
+
+	static IPackageResourceManager* SetPackageResourceManager()
+	{
+		bool bEditorDomainEnabled = false;
+		if (GIsEditor && (!IsRunningCommandlet() || IsRunningCookCommandlet()))
+		{
+			GConfig->GetBool(TEXT("CookSettings"), TEXT("EditorDomainEnabled"), bEditorDomainEnabled, GEditorIni);
+			if (bEditorDomainEnabled)
+			{
+				check(FEditorDomain::RegisteredEditorDomain == nullptr);
+				FEditorDomain::RegisteredEditorDomain = new FEditorDomain();
+				return FEditorDomain::RegisteredEditorDomain;
+			}
+		}
+		return nullptr;
+	}
+} GRegisterAsPackageResourceManager;
+
+FEditorDomain* FEditorDomain::RegisteredEditorDomain = nullptr;
+
+FEditorDomain::FEditorDomain()
+{
+	Locks = TRefCountPtr<FLocks>(new FLocks(*this));
+	Workspace.Reset(MakePackageResourceManagerFile());
+	SaveClient.Reset(new FEditorDomainSaveClient());
+	AssetRegistry = IAssetRegistry::Get();
+	// We require calling SearchAllAssets, because we rely on being able to call WaitOnAsset
+	// without needing to call ScanPathsSynchronous
+	AssetRegistry->SearchAllAssets(false /* bSynchronousSearch */);
+
+	bEditorDomainReadEnabled = !FParse::Param(FCommandLine::Get(), TEXT("noeditordomainread"));
+}
+
+FEditorDomain::~FEditorDomain()
+{
+	FScopeLock ScopeLock(&Locks->Lock);
+	Locks->Owner = nullptr;
+	AssetRegistry = nullptr;
+	Workspace.Reset();
+
+	if (RegisteredEditorDomain == this)
+	{
+		RegisteredEditorDomain = nullptr;
+	}
+}
+
+FEditorDomain* FEditorDomain::Get()
+{
+	return RegisteredEditorDomain;
+}
+
+bool FEditorDomain::SupportsLocalOnlyPaths()
+{
+	// Local Only paths are supported by falling back to the WorkspaceDomain
+	return true;
+}
+
+bool FEditorDomain::SupportsPackageOnlyPaths()
+{
+	return true;
+}
+
+bool FEditorDomain::DoesPackageExist(const FPackagePath& PackagePath, EPackageSegment PackageSegment,
+	FPackagePath* OutUpdatedPath)
+{
+	return Workspace->DoesPackageExist(PackagePath, PackageSegment, OutUpdatedPath);
+}
+
+FEditorDomain::FLocks::FLocks(FEditorDomain& InOwner)
+	:Owner(&InOwner)
+{
+}
+
+bool FEditorDomain::TryGetPackageSource(const FPackagePath& PackagePath, TRefCountPtr<FPackageSource>& OutSource)
+{
+	// Called within Locks.Lock
+	using namespace UE::EditorDomain;
+
+	// EDITOR_DOMAIN_TODO: Need to delete entries from PackageSources when the assetregistry reports the package is
+	// resaved on disk.
+	FName PackageName = PackagePath.GetPackageFName();
+	if (PackageName.IsNone())
+	{
+		OutSource.SafeRelease();
+		return false;
+	}
+
+	TRefCountPtr<FPackageSource>& PackageSource = PackageSources.FindOrAdd(PackageName);
+	if (PackageSource)
+	{
+		OutSource = PackageSource;
+		return true;
+	}
+
+	FPackageDigest PackageDigest;
+	EPackageDigestResult Result = GetPackageDigest(*AssetRegistry, PackageName, PackageDigest);
+	switch (Result)
+	{
+	case EPackageDigestResult::Success:
+		PackageSource = new FPackageSource();
+		PackageSource->Digest = PackageDigest;
+		OutSource = PackageSource;
+		if (!bEditorDomainReadEnabled)
+		{
+			PackageSource->Source = EPackageSource::Workspace;
+		}
+		return true;
+	case EPackageDigestResult::FileDoesNotExist:
+		OutSource.SafeRelease();
+		// Remove the entry in PackageSources that we added; we added it to avoid a double lookup for new packages,
+		// but for non-existent packages we want it not to be there to avoid wasting memory on it
+		PackageSources.Remove(PackageName);
+		return false;
+	case EPackageDigestResult::WrongThread:
+		PackageSource = new FPackageSource();
+		PackageSource->Source = EPackageSource::Workspace;
+		OutSource = PackageSource;
+		return true;
+	default:
+		check(false);
+		return false;
+	}
+}
+
+int64 FEditorDomain::FileSize(const FPackagePath& PackagePath, EPackageSegment PackageSegment,
+	FPackagePath* OutUpdatedPath)
+{
+	using namespace UE::EditorDomain;
+
+	if (PackageSegment != EPackageSegment::Header)
+	{
+		return Workspace->FileSize(PackagePath, PackageSegment, OutUpdatedPath);
+	}
+
+	UE::DerivedData::FCacheRequest Request;
+	int64 FileSize = -1;
+	{
+		FScopeLock ScopeLock(&Locks->Lock);
+		TRefCountPtr<FPackageSource> PackageSource;
+		if (!TryGetPackageSource(PackagePath, PackageSource) || PackageSource->Source == EPackageSource::Workspace)
+		{
+			return Workspace->FileSize(PackagePath, PackageSegment, OutUpdatedPath);
+		}
+
+		UE::DerivedData::ICache& Cache = GetDerivedDataCacheRef().GetCache();
+		auto MetaDataGetComplete =
+			[&FileSize, &PackageSource, &PackagePath, PackageSegment, Locks=this->Locks, OutUpdatedPath]
+			(UE::DerivedData::FCacheGetCompleteParams&& Params)
+		{
+			FScopeLock ScopeLock(&Locks->Lock);
+			if ((PackageSource->Source == FEditorDomain::EPackageSource::Undecided || PackageSource->Source == FEditorDomain::EPackageSource::Editor) &&
+				Params.Status == UE::DerivedData::ECacheStatus::Cached)
+			{
+				const FCbObjectRef& MetaData = Params.Record.GetMeta();
+				FCbField FileSizeField = MetaData["FileSize"];
+				FileSize = FileSizeField.AsInt64();
+				PackageSource->Source = EPackageSource::Editor;
+			}
+			else
+			{
+				checkf(PackageSource->Source == EPackageSource::Undecided || PackageSource->Source == EPackageSource::Workspace,
+					TEXT("%s was previously loaded from the EditorDomain but now is unavailable."),
+					*PackagePath.GetDebugName());
+				if (Locks->Owner)
+				{
+					FEditorDomain& EditorDomain = *Locks->Owner;
+					PackageSource->Source = EPackageSource::Workspace;
+					EditorDomain.SaveClient->RequestSave(PackagePath);
+					FileSize = EditorDomain.Workspace->FileSize(PackagePath, PackageSegment, OutUpdatedPath);
+				}
+				else
+				{
+					UE_LOG(LogEditorDomain, Warning, TEXT("%s size read after EditorDomain shutdown. Returning -1."),
+						*PackagePath.GetDebugName());
+					FileSize = -1;
+				}
+			}
+		};
+		Request = RequestEditorDomainPackage(PackagePath, PackageSource->Digest,
+			UE::DerivedData::ECachePriority::Highest, MoveTemp(MetaDataGetComplete));
+	}
+	Request.Wait();
+	return FileSize;
+}
+
+FOpenPackageResult FEditorDomain::OpenReadPackage(const FPackagePath& PackagePath, EPackageSegment PackageSegment,
+	FPackagePath* OutUpdatedPath)
+{
+	using namespace UE::EditorDomain;
+
+	FScopeLock ScopeLock(&Locks->Lock);
+	if (PackageSegment != EPackageSegment::Header)
+	{
+		return Workspace->OpenReadPackage(PackagePath, PackageSegment, OutUpdatedPath);
+	}
+	TRefCountPtr<FPackageSource> PackageSource;
+	if (!TryGetPackageSource(PackagePath, PackageSource) ||	(PackageSource->Source == EPackageSource::Workspace))
+	{
+		return Workspace->OpenReadPackage(PackagePath, PackageSegment, OutUpdatedPath);
+	}
+
+	FEditorDomainReadArchive* Result = new FEditorDomainReadArchive(Locks, PackagePath, PackageSource);
+	UE::DerivedData::FCacheRequest Request = RequestEditorDomainPackage(PackagePath, PackageSource->Digest,
+		UE::DerivedData::ECachePriority::Normal,
+		[Result](UE::DerivedData::FCacheGetCompleteParams&& Params)
+		{
+			// Note that ~FEditorDomainReadArchive waits for this callback to be called, so Result cannot dangle
+			Result->OnCacheRequestComplete(MoveTemp(Params));
+		});
+	Result->SetRequest(Request);
+
+	EPackageFormat Format;
+	if (PackageSource->Source == EPackageSource::Editor)
+	{
+		Format = EPackageFormat::Binary;
+	}
+	else
+	{
+		// EDITOR_DOMAIN_TODO: Reading GetPackageFormat forces us to wait for the cache response
+		// We should read just the metadata for the package so we don't have to block here on the transfer of the bytes.
+		Format = Result->GetPackageFormat();
+	}
+	if (OutUpdatedPath)
+	{
+		*OutUpdatedPath = PackagePath;
+	}
+	return FOpenPackageResult{ TUniquePtr<FArchive>(Result), Format };
+}
+
+IAsyncReadFileHandle* FEditorDomain::OpenAsyncReadPackage(const FPackagePath& PackagePath, EPackageSegment PackageSegment)
+{
+	using namespace UE::EditorDomain;
+
+	FScopeLock ScopeLock(&Locks->Lock);
+	if (PackageSegment != EPackageSegment::Header)
+	{
+		return Workspace->OpenAsyncReadPackage(PackagePath, PackageSegment);
+	}
+
+	TRefCountPtr<FPackageSource> PackageSource;
+	if (!TryGetPackageSource(PackagePath, PackageSource) ||
+		(PackageSource->Source == EPackageSource::Workspace))
+	{
+		return Workspace->OpenAsyncReadPackage(PackagePath, PackageSegment);
+	}
+
+	FEditorDomainAsyncReadFileHandle* Result = new FEditorDomainAsyncReadFileHandle(Locks, PackagePath, PackageSource);
+	UE::DerivedData::FCacheRequest Request = RequestEditorDomainPackage(PackagePath, PackageSource->Digest,
+		UE::DerivedData::ECachePriority::Normal,
+		[Result](UE::DerivedData::FCacheGetCompleteParams&& Params)
+		{
+			// Note that ~FEditorDomainAsyncReadFileHandle waits for this callback to be called, so Result cannot dangle
+			Result->OnCacheRequestComplete(MoveTemp(Params));
+		});
+	Result->SetRequest(Request);
+
+	return Result;
+}
+
+IMappedFileHandle* FEditorDomain::OpenMappedHandleToPackage(const FPackagePath& PackagePath,
+	EPackageSegment PackageSegment, FPackagePath* OutUpdatedPath)
+{
+	// No need to implement this runtime feature in the editor domain.
+	return nullptr;
+}
+
+bool FEditorDomain::TryMatchCaseOnDisk(const FPackagePath& PackagePath, FPackagePath* OutNormalizedPath)
+{
+	return Workspace->TryMatchCaseOnDisk(PackagePath, OutNormalizedPath);
+}
+
+void FEditorDomain::FindPackagesRecursive(TArray<TPair<FPackagePath, EPackageSegment>>& OutPackages,
+	FStringView PackageMount, FStringView FileMount, FStringView RootRelPath, FStringView BasenameWildcard)
+{
+	return Workspace->FindPackagesRecursive(OutPackages, PackageMount, FileMount, RootRelPath, BasenameWildcard);
+}
+
+void FEditorDomain::IteratePackagesInPath(FStringView PackageMount, FStringView FileMount, FStringView RootRelPath,
+	FPackageSegmentVisitor Callback)
+{
+	Workspace->IteratePackagesInPath(PackageMount, FileMount, RootRelPath, Callback);
+
+}
+void FEditorDomain::IteratePackagesInLocalOnlyDirectory(FStringView RootDir, FPackageSegmentVisitor Callback)
+{
+	Workspace->IteratePackagesInLocalOnlyDirectory(RootDir, Callback);
+}
+
+void FEditorDomain::IteratePackagesStatInPath(FStringView PackageMount, FStringView FileMount,
+	FStringView RootRelPath, FPackageSegmentStatVisitor Callback)
+{
+	Workspace->IteratePackagesStatInPath(PackageMount, FileMount, RootRelPath, Callback);
+}
+
+void FEditorDomain::IteratePackagesStatInLocalOnlyDirectory(FStringView RootDir, FPackageSegmentStatVisitor Callback)
+{
+	Workspace->IteratePackagesStatInLocalOnlyDirectory(RootDir, Callback);
+}
+
+void FEditorDomain::Tick(float DeltaTime)
+{
+	SaveClient->Tick(DeltaTime);
+}
diff --git a/Engine/Source/Editor/UnrealEd/Private/EditorDomain/EditorDomain.h b/Engine/Source/Editor/UnrealEd/Private/EditorDomain/EditorDomain.h
new file mode 100644
--- /dev/null
+++ b/Engine/Source/Editor/UnrealEd/Private/EditorDomain/EditorDomain.h
@@ -0,0 +1,134 @@
+// Copyright Epic Games, Inc. All Rights Reserved.
+
+#pragma once
+
+#include "AssetRegistry/AssetData.h"
+#include "DerivedDataCache.h"
+#include "HAL/CriticalSection.h"
+#include "HAL/Platform.h"
+#include "IO/IoHash.h"
+#include "Logging/LogMacros.h"
+#include "Templates/RefCounting.h"
+#include "TickableEditorObject.h"
+#include "UObject/PackageResourceManager.h"
+
+class FAssetPackageData;
+class FEditorDomainSaveClient;
+class IAssetRegistry;
+namespace UE
+{
+namespace EditorDomain
+{
+	typedef FIoHash FPackageDigest;
+}
+}
+
+DECLARE_LOG_CATEGORY_EXTERN(LogEditorDomain, Log, All);
+
+/**
+ * The EditorDomain is container for optimized but still editor-usable versions of WorkspaceDomain packages.
+ * The WorkspaceDomain is the source data for Unreal packages; packages created by the editor or compatible importers
+ * that can be read by any future build of the project's editor. This source data is converted to an optimized format
+ * for the current binary and saved into the EditorDomain, for faster loads when requested again by a later invocation
+ * of the editor. The optimizations include running upgrades in UObjects' PostLoad and Serialize, and saving the
+ * package in Unversioned format.
+ *
+ * FEditorDomain is a subclass of IPackageResourceManager that handles PackagePath requests by looking up the
+ * package in the EditorDomain cache, stored in the DerivedDataCache. If a version of the package matching the current
+ * WorkspaceDomain package and the current binary does not exist, then the EditorDomain falls back to loading from
+ * the WorkspaceDomain (through ordinary IFileManager operations on the Root/Game/Content folders) and creates
+ * the EditorDomain version for next time.
+ */
+class FEditorDomain : public IPackageResourceManager, public FTickableEditorObject
+{
+public:
+	FEditorDomain();
+	virtual ~FEditorDomain();
+
+	/** Return the EditorDomain that is registered as the global PackageResourceManager, if there is one. */
+	static FEditorDomain* Get();
+
+	// IPackageResourceManager interface
+	virtual bool SupportsLocalOnlyPaths() override;
+	virtual bool SupportsPackageOnlyPaths() override;
+	virtual bool DoesPackageExist(const FPackagePath& PackagePath, EPackageSegment PackageSegment,
+		FPackagePath* OutUpdatedPath = nullptr) override;
+	virtual int64 FileSize(const FPackagePath& PackagePath, EPackageSegment PackageSegment,
+		FPackagePath* OutUpdatedPath = nullptr) override;
+	virtual FOpenPackageResult OpenReadPackage(const FPackagePath& PackagePath, EPackageSegment PackageSegment,
+		FPackagePath* OutUpdatedPath = nullptr) override;
+	virtual IAsyncReadFileHandle* OpenAsyncReadPackage(const FPackagePath& PackagePath,
+		EPackageSegment PackageSegment) override;
+	virtual IMappedFileHandle* OpenMappedHandleToPackage(const FPackagePath& PackagePath,
+		EPackageSegment PackageSegment, FPackagePath* OutUpdatedPath = nullptr) override;
+	virtual bool TryMatchCaseOnDisk(const FPackagePath& PackagePath, FPackagePath* OutNormalizedPath = nullptr) override;
+	virtual void FindPackagesRecursive(TArray<TPair<FPackagePath, EPackageSegment>>& OutPackages, FStringView PackageMount,
+		FStringView FileMount, FStringView RootRelPath, FStringView BasenameWildcard) override;
+	virtual void IteratePackagesInPath(FStringView PackageMount, FStringView FileMount, FStringView RootRelPath,
+		FPackageSegmentVisitor Callback) override;
+	virtual void IteratePackagesInLocalOnlyDirectory(FStringView RootDir, FPackageSegmentVisitor Callback) override;
+	virtual void IteratePackagesStatInPath(FStringView PackageMount, FStringView FileMount, FStringView RootRelPath,
+		FPackageSegmentStatVisitor Callback) override;
+	virtual void IteratePackagesStatInLocalOnlyDirectory(FStringView RootDir, 
+		FPackageSegmentStatVisitor Callback) override;
+
+	// FTickableEditorObject interface
+	virtual void Tick(float DeltaTime) override;
+	virtual ETickableTickType GetTickableTickType() const override { return ETickableTickType::Always; }
+	virtual TStatId GetStatId() const override { return TStatId(); }
+
+private:
+	/**
+	 * Reference-counted struct holding the locks used for multithreaded synchronization.
+	 * Shared with Archives and other helpers that might outlive *this.
+	 */
+	class FLocks : public FThreadSafeRefCountedObject
+	{
+	public:
+		FLocks(FEditorDomain& InOwner);
+		FEditorDomain* Owner;
+		FCriticalSection Lock;
+	};
+	/** Different options for which domain a package comes from. */
+	enum class EPackageSource
+	{
+		Undecided,
+		Workspace,
+		Editor
+	};
+	/**
+	 * Data about which domain a package comes from. Multiple queries of the same
+	 * package have to align (for e.g. BulkData offsets) so we have to keep track of this.
+	 */
+	struct FPackageSource : public FThreadSafeRefCountedObject
+	{
+		UE::EditorDomain::FPackageDigest Digest;
+		EPackageSource Source = EPackageSource::Undecided;
+	};
+
+	/** Disallow copy constructors */
+	FEditorDomain(const FEditorDomain& Other) = delete;
+	FEditorDomain(FEditorDomain&& Other) = delete;
+
+	/** Read the PackageSource data (domain&digest) from the cache, or from the asset registry if not in cache. */
+	bool TryGetPackageSource(const FPackagePath& PackagePath, TRefCountPtr<FPackageSource>& OutSource);
+
+	/** Subsystem used to request the save of missing packages into the EditorDomain from a separate process. */
+	TUniquePtr<FEditorDomainSaveClient> SaveClient;
+	/** PackageResourceManagerFile to fall back to WorkspaceDomain when packages are missing from EditorDomain. */
+	TUniquePtr<IPackageResourceManager> Workspace;
+	/** Cached pointer to the global AssetRegistry. */
+	IAssetRegistry* AssetRegistry = nullptr;
+	/** Locks used by *this and its helper objects. */
+	TRefCountPtr<FLocks> Locks;
+	/** Digests previously found for a package. Used for optimization, but also to record loaded-from-domain. */
+	TMap<FName, TRefCountPtr<FPackageSource>> PackageSources;
+	/** True by default, set to false when reading is disabled for testing. */
+	bool bEditorDomainReadEnabled = true;
+
+	static FEditorDomain* RegisteredEditorDomain;
+
+	friend class FEditorDomainRegisterAsPackageResourceManager;
+	friend class FEditorDomainReadArchive;
+	friend class FEditorDomainAsyncReadFileHandle;
+};
\ No newline at end of file
diff --git a/Engine/Source/Editor/UnrealEd/Private/EditorDomain/EditorDomainArchive.cpp b/Engine/Source/Editor/UnrealEd/Private/EditorDomain/EditorDomainArchive.cpp
new file mode 100644
--- /dev/null
+++ b/Engine/Source/Editor/UnrealEd/Private/EditorDomain/EditorDomainArchive.cpp
@@ -0,0 +1,471 @@
+// Copyright Epic Games, Inc. All Rights Reserved.
+
+#include "EditorDomain/EditorDomainArchive.h"
+
+#include "AssetRegistry/AssetData.h"
+#include "Async/AsyncFileHandleNull.h"
+#include "EditorDomain/EditorDomainSave.h"
+#include "HAL/UnrealMemory.h"
+#include "Memory/SharedBuffer.h"
+#include "Misc/AssertionMacros.h"
+#include "Misc/PackageSegment.h"
+#include "Misc/ScopeLock.h"
+
+FEditorDomainReadArchive::FEditorDomainReadArchive(const TRefCountPtr<FEditorDomain::FLocks>& InLocks,
+	const FPackagePath& InPackagePath, const TRefCountPtr<FEditorDomain::FPackageSource>& InPackageSource)
+	: EditorDomainLocks(InLocks)
+	, PackagePath(InPackagePath)
+	, PackageSource(InPackageSource)
+{
+	this->SetIsLoading(true);
+	this->SetIsPersistent(true);
+}
+
+FEditorDomainReadArchive::~FEditorDomainReadArchive()
+{
+	Close();
+}
+
+void FEditorDomainReadArchive::SetRequest(const UE::DerivedData::FCacheRequest& InRequest)
+{
+	Request = InRequest;
+}
+
+void FEditorDomainReadArchive::Seek(int64 InPos)
+{
+	switch (Source)
+	{
+	case ESource::Archive:
+		InnerArchive->Seek(InPos);
+		break;
+	default:
+		Pos = InPos;
+		break;
+	}
+}
+
+int64 FEditorDomainReadArchive::Tell()
+{
+	switch (Source)
+	{
+	case ESource::Archive:
+		return InnerArchive->Tell();
+	default:
+		return Pos;
+	}
+}
+
+int64 FEditorDomainReadArchive::TotalSize()
+{
+	WaitForReady();
+	return Size;
+}
+
+bool FEditorDomainReadArchive::Close()
+{
+	{
+		FScopeLock ScopeLock(&Lock);
+		if (AsyncSource == ESource::Uninitialized)
+		{
+			AsyncSource = ESource::Closed;
+		}
+	}
+	Request.Cancel();
+	InnerArchive.Reset();
+	Bytes.Reset();
+	Source = ESource::Closed;
+	return true;
+}
+
+void FEditorDomainReadArchive::Serialize(void* V, int64 Length)
+{
+	switch (Source)
+	{
+	case ESource::Uninitialized:
+		WaitForReady();
+		check(Source != ESource::Uninitialized);
+		Serialize(V, Length);
+		break;
+	case ESource::Bytes:
+		if (Pos + Length > Size)
+		{
+			SetError();
+			UE_LOG(LogEditorDomain, Error, TEXT("Requested read of %d bytes when %d bytes remain (file=%s, size=%d)"),
+				Length, Size - Pos, *PackagePath.GetDebugName(), Size);
+			return;
+		}
+		FMemory::Memcpy(V, static_cast<const uint8*>(Bytes.GetData()) + Pos, Length);
+		Pos += Length;
+		break;
+	case ESource::Archive:
+		InnerArchive->Serialize(V, Length);
+		break;
+	case ESource::Closed:
+		UE_LOG(LogEditorDomain, Error, TEXT("Requested read after close (file=%s)"), *PackagePath.GetDebugName());
+		break;
+	default:
+		checkNoEntry();
+		break;
+	}
+}
+
+FString FEditorDomainReadArchive::GetArchiveName() const
+{
+	return PackagePath.GetDebugName();
+}
+
+void FEditorDomainReadArchive::Flush()
+{
+	switch (Source)
+	{
+	case ESource::Uninitialized:
+		WaitForReady();
+		check(Source != ESource::Uninitialized);
+		Flush();
+		break;
+	case ESource::Archive:
+		InnerArchive->Flush();
+		break;
+	default:
+		break;
+	}
+}
+
+void FEditorDomainReadArchive::FlushCache()
+{
+	switch (Source)
+	{
+	case ESource::Uninitialized:
+		WaitForReady();
+		check(Source != ESource::Uninitialized);
+		FlushCache();
+		break;
+	case ESource::Archive:
+		InnerArchive->FlushCache();
+		break;
+	default:
+		break;
+	}
+}
+
+bool FEditorDomainReadArchive::Precache(int64 PrecacheOffset, int64 PrecacheSize)
+{
+	switch (Source)
+	{
+	case ESource::Uninitialized:
+		WaitForReady();
+		check(Source != ESource::Uninitialized);
+		return Precache(PrecacheOffset, PrecacheSize);
+	case ESource::Archive:
+		return InnerArchive->Precache(PrecacheOffset, PrecacheSize);
+	default:
+		return true;
+	}
+}
+
+EPackageFormat FEditorDomainReadArchive::GetPackageFormat() const
+{
+	WaitForReady();
+	return PackageFormat;
+}
+
+
+void FEditorDomainReadArchive::OnCacheRequestComplete(UE::DerivedData::FCacheGetCompleteParams&& Params)
+{
+	FScopeLock ScopeLock(&Lock);
+	if (AsyncSource == ESource::Closed)
+	{
+		return;
+	}
+
+	check(AsyncSource == ESource::Uninitialized);
+	FScopeLock DomainScopeLock(&EditorDomainLocks->Lock);
+	if ((PackageSource->Source == FEditorDomain::EPackageSource::Undecided || PackageSource->Source == FEditorDomain::EPackageSource::Editor) &&
+		Params.Status == UE::DerivedData::ECacheStatus::Cached)
+	{
+		const FCbObjectRef& MetaData = Params.Record.GetMeta();
+		int64 FileSize = MetaData["FileSize"].AsInt64(-1);
+		const FSharedBuffer& LocalBytes = Params.Record.AsBinary();
+		if (static_cast<int64>(LocalBytes.GetSize()) != FileSize)
+		{
+			UE_LOG(LogEditorDomain, Warning, TEXT("Package %s received invalid record from EditorDomainPackage table ")
+				TEXT("with blob size %" INT64_FMT " not equal to FileSize in metadata %" INT64_FMT)
+				TEXT(". Reading from workspace domain instead."),
+				*PackagePath.GetDebugName(), static_cast<int64>(LocalBytes.GetSize()), FileSize);
+		}
+		else
+		{
+			AsyncSource = ESource::Bytes;
+			Size = FileSize;
+			Bytes = LocalBytes;
+			Bytes.MakeOwned();
+			PackageFormat = EPackageFormat::Binary;
+			PackageSource->Source = FEditorDomain::EPackageSource::Editor;
+		}
+	}
+	if (AsyncSource == ESource::Uninitialized)
+	{
+		checkf(PackageSource->Source == FEditorDomain::EPackageSource::Undecided || PackageSource->Source == FEditorDomain::EPackageSource::Workspace,
+			TEXT("%s was previously loaded from the EditorDomain but now is unavailable."), *PackagePath.GetDebugName());
+		if (EditorDomainLocks->Owner)
+		{
+			FEditorDomain& EditorDomain(*EditorDomainLocks->Owner);
+			IPackageResourceManager& Workspace = *EditorDomain.Workspace;
+			EditorDomain.SaveClient->RequestSave(PackagePath);
+			FOpenPackageResult Result = Workspace.OpenReadPackage(PackagePath, EPackageSegment::Header);
+			if (Result.Archive)
+			{
+				PackageSource->Source = FEditorDomain::EPackageSource::Workspace;
+				InnerArchive = MoveTemp(Result.Archive);
+				AsyncSource = ESource::Archive;
+				Size = InnerArchive->TotalSize();
+				PackageFormat = Result.Format;
+			}
+			else
+			{
+				UE_LOG(LogEditorDomain, Warning, TEXT("%s could not be read from WorkspaceDomain. Archive Set to Error."),
+					*PackagePath.GetDebugName());
+				AsyncSource = ESource::Bytes;
+				Size = 0;
+				PackageFormat = EPackageFormat::Binary;
+				SetError();
+			}
+		}
+		else
+		{
+			UE_LOG(LogEditorDomain, Warning, TEXT("%s read after EditorDomain shutdown. Archive Set to Error."),
+				*PackagePath.GetDebugName());
+			AsyncSource = ESource::Bytes;
+			Size = 0;
+			PackageFormat = EPackageFormat::Binary;
+			SetError();
+		}
+	}
+}
+
+void FEditorDomainReadArchive::WaitForReady() const
+{
+	if (Source != ESource::Uninitialized)
+	{
+		return;
+	}
+	Request.Wait();
+	{
+		// Even though we know that the asynchronous task has left the critical section,
+		// we still need to synchronize the memory order.
+		// Entering the critical section will activate the equivalent of std::memory_order::acquire that we need.
+		FScopeLock ReceiverScopeLock(&Lock);
+		Source = AsyncSource;
+	}
+
+	switch (Source)
+	{
+	case ESource::Archive:
+		// Copy local variables over to InnerArchive
+		if (Pos != 0)
+		{
+			InnerArchive->Seek(Pos);
+		}
+		break;
+	case ESource::Bytes:
+		break;
+	default:
+		check(false);
+		break;
+	}
+}
+
+/** An IAsyncReadRequest SizeRequest that returns a value known at construction time. */
+class FAsyncSizeRequestConstant : public IAsyncReadRequest
+{
+public:
+	FAsyncSizeRequestConstant(int64 InSize, FAsyncFileCallBack* InCallback)
+		: IAsyncReadRequest(InCallback, true /* bInSizeRequest */, nullptr /* UserSuppliedMemory */)
+	{
+		Size = InSize;
+		SetComplete();
+	}
+
+protected:
+	virtual void WaitCompletionImpl(float TimeLimitSeconds) override
+	{
+	}
+
+	virtual void CancelImpl() override
+	{
+	}
+};
+
+/** An IAsyncReadRequest that reads from a SharedBuffer that was already populated at construction time. */
+class FAsyncReadRequestConstant : public IAsyncReadRequest
+{
+public:
+	FAsyncReadRequestConstant(const FSharedBuffer& Bytes, FAsyncFileCallBack* InCallback, int64 Offset,
+		int64 BytesToRead, uint8* UserSuppliedMemory, const FPackagePath& PackagePath)
+		: IAsyncReadRequest(InCallback, false /* bInSizeRequest */, UserSuppliedMemory)
+	{
+		if (Offset < 0 || BytesToRead < 0 || Offset + BytesToRead > static_cast<int64>(Bytes.GetSize()))
+		{
+			UE_LOG(LogEditorDomain, Fatal, TEXT("FAsyncReadRequestConstant bogus request Offset = %" INT64_FMT)
+				TEXT(" BytesToRead = %" INT64_FMT " Bytes.GetSize() == %" UINT64_FMT " File = %s"),
+				Offset, BytesToRead, Bytes.GetSize(), *PackagePath.GetDebugName());
+		}
+		if (!UserSuppliedMemory)
+		{
+			Memory = (uint8*)FMemory::Malloc(BytesToRead);
+		}
+		check(Memory);
+		FMemory::Memcpy(Memory, reinterpret_cast<const uint8*>(Bytes.GetData()) + Offset, BytesToRead);
+		SetComplete();
+	}
+
+protected:
+	virtual void WaitCompletionImpl(float TimeLimitSeconds) override
+	{
+	}
+
+	virtual void CancelImpl() override
+	{
+	}
+};
+
+
+FEditorDomainAsyncReadFileHandle::FEditorDomainAsyncReadFileHandle(const TRefCountPtr<FEditorDomain::FLocks>& InLocks,
+	const FPackagePath& InPackagePath, const TRefCountPtr<FEditorDomain::FPackageSource>& InPackageSource)
+	: EditorDomainLocks(InLocks)
+	, PackagePath(InPackagePath)
+	, PackageSource(InPackageSource)
+{
+}
+
+FEditorDomainAsyncReadFileHandle::~FEditorDomainAsyncReadFileHandle()
+{
+	{
+		FScopeLock ScopeLock(&Lock);
+		if (AsyncSource == ESource::Uninitialized)
+		{
+			AsyncSource = ESource::Closed;
+		}
+	}
+	Request.Cancel();
+	Source = ESource::Closed;
+}
+
+void FEditorDomainAsyncReadFileHandle::SetRequest(const UE::DerivedData::FCacheRequest& InRequest)
+{
+	Request = InRequest;
+}
+
+IAsyncReadRequest* FEditorDomainAsyncReadFileHandle::SizeRequest(FAsyncFileCallBack* CompleteCallback)
+{
+	switch (Source)
+	{
+	case ESource::Uninitialized:
+		WaitForReady();
+		check(Source != ESource::Uninitialized);
+		return SizeRequest(CompleteCallback);
+	case ESource::Bytes:
+		return new FAsyncSizeRequestConstant(Bytes.GetSize(), CompleteCallback);
+	case ESource::Archive:
+		return InnerArchive->SizeRequest(CompleteCallback);
+	default:
+		checkNoEntry();
+		return nullptr;
+	}
+}
+
+IAsyncReadRequest* FEditorDomainAsyncReadFileHandle::ReadRequest(int64 Offset, int64 BytesToRead,
+	EAsyncIOPriorityAndFlags PriorityAndFlags, FAsyncFileCallBack* CompleteCallback,
+	uint8* UserSuppliedMemory)
+{
+	switch (Source)
+	{
+	case ESource::Uninitialized:
+		WaitForReady();
+		check(Source != ESource::Uninitialized);
+		return ReadRequest(Offset, BytesToRead, PriorityAndFlags, CompleteCallback, UserSuppliedMemory);
+	case ESource::Bytes:
+		return new FAsyncReadRequestConstant(Bytes, CompleteCallback, Offset, BytesToRead, UserSuppliedMemory, PackagePath);
+	case ESource::Archive:
+		return InnerArchive->ReadRequest(Offset, BytesToRead, PriorityAndFlags, CompleteCallback, UserSuppliedMemory);
+	default:
+		checkNoEntry();
+		return nullptr;
+	}
+}
+
+bool FEditorDomainAsyncReadFileHandle::UsesCache()
+{
+	return false;
+}
+
+void FEditorDomainAsyncReadFileHandle::OnCacheRequestComplete(UE::DerivedData::FCacheGetCompleteParams&& Params)
+{
+	FScopeLock ScopeLock(&Lock);
+	if (AsyncSource == ESource::Closed)
+	{
+		return;
+	}
+
+	check(AsyncSource == ESource::Uninitialized);
+	FScopeLock DomainScopeLock(&EditorDomainLocks->Lock);
+	if ((PackageSource->Source == FEditorDomain::EPackageSource::Undecided || PackageSource->Source == FEditorDomain::EPackageSource::Editor) &&
+		Params.Status == UE::DerivedData::ECacheStatus::Cached)
+	{
+		const FCbObjectRef& MetaData = Params.Record.GetMeta();
+		int64 FileSize = MetaData["FileSize"].AsInt64(-1);
+		const FSharedBuffer& LocalBytes = Params.Record.AsBinary();
+		if (static_cast<int64>(LocalBytes.GetSize()) != FileSize)
+		{
+			UE_LOG(LogEditorDomain, Warning, TEXT("Package %s received invalid record from EditorDomainPackage ")
+				TEXT("table with blob size %" INT64_FMT "not equal to FileSize in metadata %" INT64_FMT)
+				TEXT(". Reading from workspace domain instead."),
+				*PackagePath.GetDebugName(), static_cast<int64>(LocalBytes.GetSize()), FileSize);
+		}
+		else
+		{
+			AsyncSource = ESource::Bytes;
+			Bytes = LocalBytes;
+			Bytes.MakeOwned();
+			PackageSource->Source = FEditorDomain::EPackageSource::Editor;
+		}
+	}
+	if (AsyncSource == ESource::Uninitialized)
+	{
+		checkf(PackageSource->Source == FEditorDomain::EPackageSource::Undecided || PackageSource->Source == FEditorDomain::EPackageSource::Workspace,
+			TEXT("%s was previously loaded from the EditorDomain but now is unavailable."), *PackagePath.GetDebugName());
+		if (EditorDomainLocks->Owner)
+		{
+			FEditorDomain& EditorDomain(*EditorDomainLocks->Owner);
+			IPackageResourceManager& Workspace = *EditorDomain.Workspace;
+			EditorDomain.SaveClient->RequestSave(PackagePath);
+			IAsyncReadFileHandle* Result = Workspace.OpenAsyncReadPackage(PackagePath, EPackageSegment::Header);
+			check(Result);
+			AsyncSource = ESource::Archive;
+			InnerArchive.Reset(Result);
+			PackageSource->Source = FEditorDomain::EPackageSource::Workspace;
+		}
+		else
+		{
+			UE_LOG(LogEditorDomain, Warning, TEXT("%s read after EditorDomain shutdown. Returning null archive"),
+				*PackagePath.GetDebugName());
+			AsyncSource = ESource::Archive;
+			InnerArchive.Reset(new FAsyncReadFileHandleNull());
+		}
+	}
+}
+
+void FEditorDomainAsyncReadFileHandle::WaitForReady() const
+{
+	if (Source != ESource::Uninitialized)
+	{
+		return;
+	}
+	Request.Wait();
+	{
+		// Even though we know that the asynchronous task has left the critical section, we still need to synchronize the memory order
+		// Entering the critical section will activate the equivalent of std::memory_order::acquire that we need
+		FScopeLock ReceiverScopeLock(&Lock);
+		Source = AsyncSource;
+	}
+}
diff --git a/Engine/Source/Editor/UnrealEd/Private/EditorDomain/EditorDomainArchive.h b/Engine/Source/Editor/UnrealEd/Private/EditorDomain/EditorDomainArchive.h
new file mode 100644
--- /dev/null
+++ b/Engine/Source/Editor/UnrealEd/Private/EditorDomain/EditorDomainArchive.h
@@ -0,0 +1,132 @@
+// Copyright Epic Games, Inc. All Rights Reserved.
+
+#pragma once
+
+#include "Async/AsyncFileHandle.h"
+#include "DerivedDataCache.h"
+#include "EditorDomain/EditorDomain.h"
+#include "HAL/CriticalSection.h"
+#include "HAL/Platform.h"
+#include "Misc/PackagePath.h"
+#include "Serialization/Archive.h"
+#include "Templates/RefCounting.h"
+#include "Templates/UniquePtr.h"
+#include "UObject/PackageResourceManager.h"
+
+class FAssetPackageData;
+
+/**
+ * An Archive that asynchronously waits for the cache request to complete, and reads either from the returned cache bytes
+ * or from the fallback WorkspaceDomain Archive for the given PackagePath.
+ *
+ * This class is a serialization archive rather than a full archive; it overrides the serialization functions used by
+ * LinkerLoad and BulkData but does not override all of the functions used by general Archive use as FArchiveProxy does.
+ */
+class FEditorDomainReadArchive : public FArchive
+{
+public:
+	FEditorDomainReadArchive(const TRefCountPtr<FEditorDomain::FLocks>& InLocks,
+		const FPackagePath& InPackagePath, const TRefCountPtr<FEditorDomain::FPackageSource>& InPackageSource);
+	~FEditorDomainReadArchive();
+
+	/** Set the CacheRequest handle that will feed this archive. */
+	void SetRequest(const UE::DerivedData::FCacheRequest& InRequest);
+	/** Callback from the CacheRequest; set whether we're reading from EditorDomain bytes or WorkspaceDomain archive. */
+	void OnCacheRequestComplete(UE::DerivedData::FCacheGetCompleteParams&& Params);
+	/** Get the PackageFormat, which depends on the domain the data is read from. */
+	EPackageFormat GetPackageFormat() const;
+
+	// FArchive interface
+	virtual void Seek(int64 InPos) override;
+	virtual int64 Tell() override;
+	virtual int64 TotalSize() override;
+	virtual bool Close() override;
+	virtual void Serialize(void* V, int64 Length) override;
+	virtual FString GetArchiveName() const override;
+	virtual void Flush() override;
+	virtual void FlushCache() override;
+	virtual bool Precache(int64 PrecacheOffset, int64 PrecacheSize) override;
+
+private:
+	enum class ESource : uint8
+	{
+		Uninitialized,
+		Bytes,
+		Archive,
+		Closed
+	};
+
+	/** Wait for the handle to call OnCacheRequestComplete and make the size and bytes available. */
+	void WaitForReady() const;
+
+	/** Lock to synchronize the CacheCompletion and the public interface thread. */
+	mutable FCriticalSection Lock;
+
+	// Data in this section is either read-only, or is read and written only on the public interface thread.
+	TRefCountPtr<FEditorDomain::FLocks> EditorDomainLocks; // Pointer is read-only, *pointer has internal lock
+	UE::DerivedData::FCacheRequest Request; // Read-only
+	FPackagePath PackagePath; // Read-only
+	int64 Pos = 0; // Interface-thread-only
+	mutable ESource Source = ESource::Uninitialized; // Interface-thread-only
+	TRefCountPtr<FEditorDomain::FPackageSource> PackageSource; // Read-only, *pointer requires EditorDomainLocks.Lock
+
+	// Data in this section is written on the callback thread.
+	// It can not be read until after WaitForReady or Cancel.
+	TUniquePtr<FArchive> InnerArchive;
+	FSharedBuffer Bytes;
+	int64 Size = 0;
+	ESource AsyncSource = ESource::Uninitialized;
+	EPackageFormat PackageFormat = EPackageFormat::Binary;
+};
+
+/**
+ * An IAsyncReadFileHandle that asynchronously waits for the cache request to complete, and reads either from the
+ * returned cache bytes or from the fallback WorkspaceDomain Archive for the given PackagePath.
+ */
+class FEditorDomainAsyncReadFileHandle : public IAsyncReadFileHandle
+{
+public:
+	FEditorDomainAsyncReadFileHandle(const TRefCountPtr<FEditorDomain::FLocks>& InLocks,
+		const FPackagePath& InPackagePath, const TRefCountPtr<FEditorDomain::FPackageSource>& InPackageSource);
+	~FEditorDomainAsyncReadFileHandle();
+
+	/** Set the CacheRequest handle that will feed this archive. */
+	void SetRequest(const UE::DerivedData::FCacheRequest& InRequest);
+	/** Callback from the CacheRequest; set whether we're reading from EditorDomain bytes or WorkspaceDomain archive. */
+	void OnCacheRequestComplete(UE::DerivedData::FCacheGetCompleteParams&& Params);
+
+	// IAsyncReadFileHandle interface
+	virtual IAsyncReadRequest* SizeRequest(FAsyncFileCallBack* CompleteCallback = nullptr) override;
+	virtual IAsyncReadRequest* ReadRequest(int64 Offset, int64 BytesToRead,
+		EAsyncIOPriorityAndFlags PriorityAndFlags = AIOP_Normal, FAsyncFileCallBack* CompleteCallback = nullptr,
+		uint8* UserSuppliedMemory = nullptr) override;
+	virtual bool UsesCache();
+
+private:
+	enum class ESource : uint8
+	{
+		Uninitialized,
+		Bytes,
+		Archive,
+		Closed
+	};
+
+	/** Wait for the handle to call OnCacheRequestComplete and make the size and bytes available. */
+	void WaitForReady() const;
+
+	/** Lock to synchronize the CacheCompletion and the public interface thread. */
+	mutable FCriticalSection Lock;
+
+	// Data in this section is either read-only, or is read and written only on the public interface thread.
+	TRefCountPtr<FEditorDomain::FLocks> EditorDomainLocks; // Pointer is read-only, *pointer has internal lock
+	UE::DerivedData::FCacheRequest Request; // Read-only
+	FPackagePath PackagePath; // Read-only
+	mutable ESource Source = ESource::Uninitialized; // Interface-thread-only
+	TRefCountPtr<FEditorDomain::FPackageSource> PackageSource; // Read-only, *pointer requires EditorDomainLocks.Lock
+
+	// Data in this section is written on the callback thread.
+	// It can not be read until after WaitForReady or Cancel.
+	TUniquePtr<IAsyncReadFileHandle> InnerArchive;
+	FSharedBuffer Bytes;
+	ESource AsyncSource = ESource::Uninitialized;
+};
\ No newline at end of file
diff --git a/Engine/Source/Editor/UnrealEd/Private/EditorDomain/EditorDomainSave.cpp b/Engine/Source/Editor/UnrealEd/Private/EditorDomain/EditorDomainSave.cpp
new file mode 100644
--- /dev/null
+++ b/Engine/Source/Editor/UnrealEd/Private/EditorDomain/EditorDomainSave.cpp
@@ -0,0 +1,1541 @@
+// Copyright Epic Games, Inc. All Rights Reserved.
+
+#include "EditorDomainSave.h"
+
+#include "AssetRegistry/IAssetRegistry.h"
+#include "Async/Async.h"
+#include "Commandlets/EditorDomainSaveCommandlet.h"
+#include "DerivedDataCache.h"
+#include "DerivedDataCacheInterface.h"
+#include "Dom/JsonObject.h"
+#include "EditorDomain/EditorDomainUtils.h"
+#include "HAL/FileManager.h"
+#include "HAL/PlatformTime.h"
+#include "HAL/RunnableThread.h"
+#include "IPAddress.h"
+#include "Misc/DateTime.h"
+#include "Misc/FileHelper.h"
+#include "Misc/PackagePath.h"
+#include "Misc/Paths.h"
+#include "Serialization/Archive.h"
+#include "Serialization/CompactBinaryWriter.h"
+#include "Serialization/JsonReader.h"
+#include "Serialization/JsonSerializer.h"
+#include "Serialization/JsonWriter.h"
+#include "Serialization/MemoryReader.h"
+#include "Serialization/MemoryWriter.h"
+#include "Sockets.h"
+#include "SocketSubsystem.h"
+#include "UnrealEdMisc.h"
+#include "UObject/Package.h"
+#include "UObject/PackageResourceManager.h"
+#include "UObject/UObjectGlobals.h"
+
+DEFINE_LOG_CATEGORY(LogEditorDomainSave);
+
+/** A header for network messages with magic number identifier, message enum, and payload size. */
+template<typename MessageEnum, uint32 ExpectedMagic, uint32 MessageStart, uint32 MessageEnd, uint32 MaxSize>
+struct TMessageHeader
+{
+	uint32 Magic;
+	uint32 Size;
+	uint32 MessageType;
+
+	void Initialize(MessageEnum InMessageType)
+	{
+		Magic = ExpectedMagic;
+		SetMessageType(InMessageType);
+	}
+	void SetMessageType(MessageEnum InMessageType)
+	{
+		MessageType = static_cast<uint32>(InMessageType);
+	}
+	MessageEnum GetMessageType() const
+	{
+		return static_cast<MessageEnum>(MessageType);
+	}
+	FString ToString() const
+	{
+		return FString::Printf(TEXT("(Magic=%u, MessageType=%u, Size=%u)"), Magic, MessageType, Size);
+	}
+	bool IsValid() const
+	{
+		return Magic == ExpectedMagic &&
+			MessageStart <= MessageType && MessageType < MessageEnd &&
+			Size <= MaxSize;
+	}
+	friend FArchive& operator<<(FArchive& Ar, TMessageHeader& Header)
+	{
+		return Ar << Header.Magic << Header.Size << Header.MessageType;
+	}
+};
+
+/** A ScopeLock for FPlatformProcess::NewInterprocessSynchObject types. */
+struct FProcessSemaphoreScopeLock
+{
+	FProcessSemaphoreScopeLock(FGenericPlatformProcess::FSemaphore* InSemaphore, float TimeOutSeconds)
+		: Semaphore(InSemaphore)
+	{
+		if (TimeOutSeconds < 0)
+		{
+			Semaphore->Lock();
+			bLocked = true;
+		}
+		else
+		{
+			float NanosecondsPerSecond = 1e9f;
+			bLocked = Semaphore->TryLock(static_cast<uint64>(TimeOutSeconds * NanosecondsPerSecond));
+		}
+	}
+
+	bool IsLocked() const
+	{
+		return bLocked;
+	}
+
+	void Unlock()
+	{
+		if (bLocked)
+		{
+			Semaphore->Unlock();
+		}
+	}
+
+	~FProcessSemaphoreScopeLock()
+	{
+		Unlock();
+	}
+
+private:
+	FGenericPlatformProcess::FSemaphore* Semaphore;
+	bool bLocked = false;
+};
+
+namespace UE
+{
+namespace EditorDomainSave
+{
+
+/** MessageTypes that can be sent between EditorDomainSaveServer to EditorDomainSaveClient. */
+enum class EMessageType : uint32
+{
+	MessageTypeMagicNumber = 0xED1104,
+	MessageTypeStart = MessageTypeMagicNumber,
+
+	Save = MessageTypeStart,
+	CloseConnection,
+
+	MessageTypeEnd,
+};
+constexpr uint32 MessageMagicNumber = 0x854EBA92;
+constexpr int32 MaxMessageSize = 1000000 * 10;
+
+typedef TMessageHeader<EMessageType, MessageMagicNumber,
+	(uint32)EMessageType::MessageTypeStart, (uint32)EMessageType::MessageTypeEnd,
+	MaxMessageSize> FMessageHeader;
+
+/** Ids used in the JSON settings file. */
+namespace SettingNames
+{
+	const TCHAR* const ListenPortFieldName = TEXT("ListenPort");
+	const TCHAR* const ProcessIdFieldName = TEXT("ProcessId");
+}
+
+/** The string identifier for the OS to Find or FindAndCreate the interprocess synchronization object. */
+FString GetLockName();
+/**
+ * Create the inprocess handle to the interprocess synchronization object.
+ * Caller must call DeleteInterprocessSynchObject when done.
+ */
+FPlatformProcess::FSemaphore* CreateProcessLock();
+/** Read the ServerSettings JSON file from disk into a JsonObject. */
+bool TryReadServerSettings(TSharedPtr<FJsonObject>& OutRootObject, FDateTime& OutServerSettingsTimestamp);
+/** Write the ServerSettings JSON file to disk from a JsonObject */
+bool TryWriteServerSettings(TSharedPtr<FJsonObject>& RootObject);
+/** The full path to the intermediate JSON file used to store settings for the active EditorDomainSaveServer. */
+FString GetServerSettingsFilename();
+/** Get the timestamp of the ServerSettings JSON file, used to respond to updates. */
+void GetServerSettingsTimeStamp(bool& bOutExists, FDateTime& OutModificationTime);
+/** Delete the server settings file; indicates that a new server will need to be created. */
+bool TryDeleteSettingsFile();
+
+namespace Constants
+{
+	/** Timeout in seconds to wait for the interprocess lock before giving up trying again on the next tick. */
+	constexpr double ProcessLockTimeoutSecondsWhenPolling = .005;
+	/* Timeout in seconds to wait for the interprocess lock before aborting the server. */
+	constexpr double ProcessLockTimeoutSecondsWhenBlocking = 60.;
+
+	/** Timeout in seconds to wait for the server to respond to its duties, (write port, read socket). */
+	constexpr double ServerUnresponsiveTimeoutSeconds = 10 * 60.;
+	/** Cooldown in seconds before CollectGarbage should be called on the server after it becomes idle. */
+	constexpr double CollectGarbageIdleCooldownSeconds = 10.;
+	/** Cooldown in seconds before CollectGarbage should be called again on the server when it is active. */
+	constexpr double CollectGarbageActiveCooldownSeconds = 60. * 10;
+	/** Duration in seconds the server should sleep before checking for messages when it is idle. */
+	constexpr double ServerIdleSleepPeriodSeconds = 1.;
+	/** Max duration in seconds of the server's single tick of package saving. */
+	constexpr double ServerPackageLoadBudgetSeconds = .5;
+	/** Cooldown in seconds before the Server will shutdown when it has no clients. */
+	constexpr double ServerAbdicationCooldownSeconds = 10.;
+
+	/** Timeout in seconds to wait for the socket to connect before rechecking whether the process exists. */
+	constexpr double SocketUnresponsiveTimeoutSeconds = 5.;
+	/** Max duration in seconds of the clients wait for queued messages to send when it is shutting down. */
+	constexpr double ClosingFlushBudgetSeconds = 10.;
+	/** Cooldown in seconds between attempts to create the Server process. */
+	constexpr double ProcessCreationCooldownSeconds = 10.;
+	/** Cooldown in seconds before logging another connection warning. */
+	constexpr double ConnectionWarningCooldownSeconds = 60.;
+	/** Duration in seconds the client should sleep in its async run before trying another Tick. */
+	constexpr double RunCommunicationSleepPeriodSeconds = .001;
+}
+
+}
+}
+
+
+FEditorDomainSaveServer::FEditorDomainSaveServer()
+{
+	AssetRegistry = IAssetRegistry::Get();
+}
+
+FEditorDomainSaveServer::~FEditorDomainSaveServer()
+{
+	Shutdown();
+}
+
+int32 FEditorDomainSaveServer::Run()
+{
+	if (!TryInitialize())
+	{
+		Shutdown();
+		return 1;
+	}
+
+	while (PollShouldRun())
+	{
+		bool bIsIdle = true;
+		TickPendingPackages(bIsIdle);
+		PollIncomingConnections(bIsIdle);
+		PollConnections(bIsIdle);
+		TickMaintenance(bIsIdle);
+	}
+
+	Shutdown();
+	return 0;
+}
+
+
+bool FEditorDomainSaveServer::TryInitialize()
+{
+	using namespace UE::EditorDomainSave;
+
+	const TCHAR* CommandLine = FCommandLine::Get();
+	if (!FParse::Value(CommandLine, TEXT("-creatorprocessid="), CreatorProcessId))
+	{
+		CreatorProcessId = 0;
+	}
+	bIsResidentServer = FParse::Param(CommandLine, TEXT("editordomainserverresident"));
+
+	ProcessLock = CreateProcessLock();
+	if (!ProcessLock)
+	{
+		UE_LOG(LogEditorDomainSave, Error, TEXT("Could not create InterprocessSynchObject %s, ")
+			TEXT("EditorDomain saving is not enabled."), *GetLockName());
+		return false;
+	}
+
+	FProcessSemaphoreScopeLock ProcessScopeLock(ProcessLock, Constants::ProcessLockTimeoutSecondsWhenBlocking);
+	if (!ProcessScopeLock.IsLocked())
+	{
+		UE_LOG(LogEditorDomainSave, Error, TEXT("Could not lock InterprocessSynchObject %s for %lf seconds, ")
+			TEXT("aborting creation of EditorDomainSaveServer."),
+			*GetLockName(), Constants::ProcessLockTimeoutSecondsWhenBlocking);
+		return false;
+	}
+
+	TSharedPtr<FJsonObject> SettingsObject;
+	if (!TryRefreshAuthority(nullptr /* bOutOwnsSettingFile */, &SettingsObject, &ListenPort))
+	{
+		// Log statement sent by TryRefreshAuthority
+		return false;
+	}
+
+	ISocketSubsystem* SocketSubsystem = ISocketSubsystem::Get();
+	if (!SocketSubsystem)
+	{
+		UE_LOG(LogEditorDomainSave, Error, TEXT("No SocketSubsystem, aborting creation of EditorDomainServer."));
+		return false;
+	}
+	TSharedRef<FInternetAddr> ListenAddr = SocketSubsystem->GetLocalBindAddr(*GLog);
+	ListenAddr->SetPort(ListenPort);
+	ListenSocket = SocketSubsystem->CreateSocket(NAME_Stream, TEXT("FEditorDomainSaveServer tcp-listen"),
+		ListenAddr->GetProtocolType());
+	if (!ListenSocket)
+	{
+		UE_LOG(LogEditorDomainSave, Error, TEXT("CreateSocket failed, aborting creation of EditorDomainServer."));
+		return false;
+	}
+
+	ListenSocket->SetReuseAddr();
+	ListenSocket->SetNonBlocking(true);
+	if (!ListenSocket->Bind(*ListenAddr))
+	{
+		UE_LOG(LogEditorDomainSave, Error,
+			TEXT("Failed to bind listen socket %s, aborting creation of EditorDomainServer."),
+			*ListenAddr->ToString(true));
+		return false;
+	}
+	if (!ListenSocket->Listen(16))
+	{
+		UE_LOG(LogEditorDomainSave, Warning,
+			TEXT("Failed to listen to socket %s, aborting creation of EditorDomainServer."),
+			*ListenAddr->ToString(true));
+		return false;
+	}
+	ListenPort = ListenSocket->GetPortNo();
+	ListenAddr->SetPort(ListenPort);
+	UE_LOG(LogEditorDomainSave, Display, TEXT("EditorDomainSaveServer is listening for client connections on %s."),
+		*ListenAddr->ToString(true));
+
+	SettingsObject->SetNumberField(SettingNames::ListenPortFieldName, ListenPort);
+	if (!TryWriteServerSettings(SettingsObject))
+	{
+		UE_LOG(LogEditorDomainSave, Warning,
+			TEXT("Failed to write ServerSettings file %s, aborting creation of EditorDomainServer."),
+			*GetServerSettingsFilename());
+		return false;
+	}
+
+	IAssetRegistry::Get()->SearchAllAssets(false /* bSynchronousSearch */);
+	LastGarbageTime = FPlatformTime::Seconds();
+	return true;
+}
+
+bool FEditorDomainSaveServer::TryRefreshAuthority(bool* bOutOwnsSettingsFile, TSharedPtr<FJsonObject>* OutRootObject, 
+	int32* OutListenPort)
+{
+	using namespace UE::EditorDomainSave;
+
+	if (bOutOwnsSettingsFile)
+	{
+		*bOutOwnsSettingsFile = false;
+	}
+
+	bool bIsInitialAcquisition = ListenPort == 0;
+	const TCHAR* ErrorActionDesc = bIsInitialAcquisition
+		? TEXT("aborting creation")
+		: TEXT("relinquishing ownership and shutting down");
+
+	int32 ServerProcessId = 0;
+	int32 LocalListenPort = 0;
+	TSharedPtr<FJsonObject> RootObject;
+	if (bIsInitialAcquisition && bIsResidentServer)
+	{
+		ServerProcessId = FPlatformProcess::GetCurrentProcessId();
+		RootObject = MakeShared<FJsonObject>();
+		RootObject->SetNumberField(SettingNames::ProcessIdFieldName, ServerProcessId);
+	}
+	else
+	{
+		TryReadServerSettings(RootObject, SettingsTimestamp);
+		if (!RootObject)
+		{
+			UE_LOG(LogEditorDomainSave, Error, TEXT("Could not read settings file %s, %s."),
+				*GetServerSettingsFilename(), ErrorActionDesc);
+			return false;
+		}
+		if (!RootObject->TryGetNumberField(SettingNames::ProcessIdFieldName, ServerProcessId))
+		{
+			UE_LOG(LogEditorDomainSave, Error, TEXT("Could not read %s out of settings file %s, %s."),
+				SettingNames::ProcessIdFieldName, *GetServerSettingsFilename(), ErrorActionDesc);
+				return false;
+		}
+		int32 Value;
+		if (RootObject->TryGetNumberField(SettingNames::ListenPortFieldName, Value))
+		{
+			LocalListenPort = Value;
+		}
+	}
+
+	if (ServerProcessId != FPlatformProcess::GetCurrentProcessId())
+	{
+		UE_LOG(LogEditorDomainSave, Error, TEXT("EditorDomainSave:Server Another process %d is registered as ")
+			TEXT("the singleton EditorDomainSaveServer in settings file %s, this process %d is %s."),
+			ServerProcessId, *GetServerSettingsFilename(), FPlatformProcess::GetCurrentProcessId(), ErrorActionDesc);
+		return false;
+	}
+	if (bOutOwnsSettingsFile)
+	{
+		*bOutOwnsSettingsFile = true;
+	}
+
+	if (!bIsInitialAcquisition && LocalListenPort != ListenPort)
+	{
+		UE_LOG(LogEditorDomainSave, Error,
+			TEXT("Listen port has changed from %d to %d. Change in listen port is unsupported; %s."),
+			ListenPort, LocalListenPort, ErrorActionDesc);
+		return false;
+	}
+	if (OutListenPort)
+	{
+		*OutListenPort = LocalListenPort;
+	}
+	if (OutRootObject)
+	{
+		*OutRootObject = RootObject;
+	}
+	return true;
+}
+
+void FEditorDomainSaveServer::Shutdown()
+{
+	ShutdownSocket();
+	if (ProcessLock)
+	{
+		FPlatformProcess::DeleteInterprocessSynchObject(ProcessLock);
+		ProcessLock = nullptr;
+	}
+	GetDerivedDataCacheRef().WaitForQuiescence(true);
+}
+
+void FEditorDomainSaveServer::ShutdownSocket()
+{
+	ClientConnections.Empty();
+	if (ListenSocket)
+	{
+		ISocketSubsystem* SocketSubsystem = ISocketSubsystem::Get();
+		check(SocketSubsystem);
+		ListenSocket->Close();
+		SocketSubsystem->DestroySocket(ListenSocket);
+		ListenSocket = nullptr;
+	}
+}
+
+bool FEditorDomainSaveServer::TryAbdicate()
+{
+	using namespace UE::EditorDomainSave;
+
+	if (bIsResidentServer)
+	{
+		return false;
+	}
+	FProcessSemaphoreScopeLock ProcessScopeLock(ProcessLock, Constants::ProcessLockTimeoutSecondsWhenPolling);
+	if (!ProcessScopeLock.IsLocked())
+	{
+		return false;
+	}
+	if (HasExpectedConnections())
+	{
+		return false;
+	}
+
+	UE_LOG(LogEditorDomainSave, Display,
+		TEXT("EditorDomainSaveServer has no remaining connections and is shutting down."));
+	bAbdicated = true;
+	ShutdownSocket();
+	bool bOwnsSettingsFile;
+	TryRefreshAuthority(&bOwnsSettingsFile);
+	if (bOwnsSettingsFile && !TryDeleteSettingsFile())
+	{
+		UE_LOG(LogEditorDomainSave, Warning,
+			TEXT("Unable to delete server settings file %s. Future clients will have to delete the file."),
+			*GetServerSettingsFilename());
+	}
+	return true;
+}
+
+bool FEditorDomainSaveServer::IsAbdicated() const
+{
+	return bAbdicated;
+}
+
+bool FEditorDomainSaveServer::HasExpectedConnections() const
+{
+	if (!bHasEverConnected && CreatorProcessId)
+	{
+		FProcHandle Handle = FPlatformProcess::OpenProcess(CreatorProcessId);
+		ON_SCOPE_EXIT{ FPlatformProcess::CloseProc(Handle); };
+		if (Handle.IsValid())
+		{
+			// We are still waiting for the creator process to connect
+			return true;
+		}
+	}
+	return ClientConnections.Num() > 0;
+}
+
+bool FEditorDomainSaveServer::PollShouldRun()
+{
+	using namespace UE::EditorDomainSave;
+
+	if (IsAbdicated() || ListenSocket == nullptr)
+	{
+		return false;
+	}
+	FDateTime LocalSettingsTimestamp;
+	bool bExists;
+	GetServerSettingsTimeStamp(bExists, LocalSettingsTimestamp);
+	if (!bExists || SettingsTimestamp != LocalSettingsTimestamp)
+	{
+		FProcessSemaphoreScopeLock ProcessScopeLock(ProcessLock, Constants::ProcessLockTimeoutSecondsWhenPolling);
+		if (ProcessScopeLock.IsLocked())
+		{
+			bool bOwnsSettingsFile;
+			if (!TryRefreshAuthority(&bOwnsSettingsFile))
+			{
+				ShutdownSocket();
+				if (bOwnsSettingsFile && !TryDeleteSettingsFile())
+				{
+					UE_LOG(LogEditorDomainSave, Warning,
+						TEXT("Unable to delete server settings file %s. Future clients will have to delete the file."),
+						*GetServerSettingsFilename());
+				}
+				return false;
+			}
+		}
+	}
+	return true;
+}
+
+void FEditorDomainSaveServer::PollIncomingConnections(bool& bInOutIsIdle)
+{
+	bool bReadReady;
+	while (ListenSocket->WaitForPendingConnection(bReadReady, FTimespan::FromSeconds(0)) && bReadReady)
+	{
+		FSocket* ClientSocket = ListenSocket->Accept(TEXT("Client Connection"));
+		if (!ClientSocket)
+		{
+			UE_LOG(LogEditorDomainSave, Warning,
+				TEXT("Listen socket received a pending connection event but ListenSocket->Accept ")
+				TEXT("failed to create a ClientSocket."));
+		}
+		else
+		{
+			ClientSocket->SetNonBlocking(true);
+			bHasEverConnected = true;
+			bInOutIsIdle = false;
+			ClientConnections.Add(TRefCountPtr<FClientConnection>(new FClientConnection(ClientSocket)));
+		}
+	}
+}
+
+void FEditorDomainSaveServer::PollConnections(bool& bInOutIsIdle)
+{
+	for (int Index = 0; Index < ClientConnections.Num(); ++Index)
+	{
+		TRefCountPtr<FClientConnection>& ClientConnection = ClientConnections[Index];
+		bool bStillAlive;
+		bool bClientIsIdle;
+		ClientConnection->Poll(*this, bClientIsIdle, bStillAlive);
+		if (!bStillAlive)
+		{
+			bInOutIsIdle = false;
+			ClientConnections.RemoveAtSwap(Index);
+		}
+		else if (!bClientIsIdle)
+		{
+			bInOutIsIdle = false;
+		}
+	}
+}
+
+void FEditorDomainSaveServer::TickMaintenance(bool bIsIdle)
+{
+	using namespace UE::EditorDomainSave;
+
+	SetIdle(bIsIdle);
+
+	double CurrentTime = FPlatformTime::Seconds();
+	if (bIsIdle)
+	{
+		if (!HasExpectedConnections() && CurrentTime - IdleStartTime > Constants::ServerAbdicationCooldownSeconds)
+		{
+			if (TryAbdicate())
+			{
+				return;
+			}
+		}
+	}
+
+	double CollectGarbageCooldownSeconds = bIsIdle ?
+		Constants::CollectGarbageIdleCooldownSeconds :
+		Constants::CollectGarbageActiveCooldownSeconds;
+	bool bCollectedGarbageAfterIdle = bIsIdle && LastGarbageTime >= IdleStartTime;
+	if (!bCollectedGarbageAfterIdle && CurrentTime - LastGarbageTime > CollectGarbageCooldownSeconds)
+	{
+		CollectGarbage(RF_NoFlags);
+		LastGarbageTime = FPlatformTime::Seconds();
+	}
+
+	if (bIsIdle)
+	{
+		FPlatformProcess::Sleep(Constants::ServerIdleSleepPeriodSeconds);
+	}
+}
+
+void FEditorDomainSaveServer::SetIdle(bool bIsIdle)
+{
+	bool bWasIdle = IdleStartTime != 0.;
+	if (bIsIdle == bWasIdle)
+	{
+		return;
+	}
+	if (bIsIdle)
+	{
+		IdleStartTime = FPlatformTime::Seconds();
+	}
+	else
+	{
+		IdleStartTime = 0.;
+	}
+}
+
+
+void FEditorDomainSaveServer::ProcessMessage(FClientConnection& ClientConnection,
+	UE::EditorDomainSave::EMessageType MessageType, const TArray<uint8>& MessageBuffer, bool& bOutStillAlive)
+{
+	using namespace UE::EditorDomainSave;
+
+	bOutStillAlive = true;
+
+	switch (MessageType)
+	{
+	case EMessageType::Save:
+	{
+		FMemoryReader Reader(MessageBuffer);
+		int32 NumPackages;
+		const int32 MinSizePerPackage = sizeof(int32) + 2;
+		Reader << NumPackages;
+		if (Reader.IsError() || NumPackages < 0 || NumPackages > MessageBuffer.Num() / MinSizePerPackage)
+		{
+			UE_LOG(LogEditorDomainSave, Warning, TEXT("EditorDomainServer received corrupt save message with invalid ")
+				TEXT("NumPackages %d from client, closing the connection"), NumPackages);
+			bOutStillAlive = false;
+		}
+		else
+		{
+			for (int n = 0; n < NumPackages; ++n)
+			{
+				FString PackageName;
+				Reader << PackageName;
+				if (Reader.IsError())
+				{
+					UE_LOG(LogEditorDomainSave, Warning, TEXT("EditorDomainServer received corrupt save message with ")
+						TEXT("invalid PackageName instance %d from client, closing the connection"), n);
+					bOutStillAlive = false;
+					break;
+				}
+				PendingPackageNames.Add(MoveTemp(PackageName));
+			}
+		}
+		break;
+	}
+	case EMessageType::CloseConnection:
+		UE_LOG(LogEditorDomainSave, Display,
+			TEXT("EditorDomainServer received CloseConnection message from client, closing the connection"));
+		bOutStillAlive = false;
+		break;
+	default:
+		UE_LOG(LogEditorDomainSave, Warning,
+			TEXT("EditorDomainServer received invalid messagetype %u from client, closing the connection"),
+			static_cast<uint32>(MessageType));
+		bOutStillAlive = false;
+		break;
+	}
+}
+
+void FEditorDomainSaveServer::TickPendingPackages(bool& bInOutIsIdle)
+{
+	using namespace UE::EditorDomainSave;
+
+	if (PendingPackageNames.Num() == 0)
+	{
+		return;
+	}
+	bInOutIsIdle = false;
+	double StartTime = FPlatformTime::Seconds();
+
+	if (!bHasInitializedAssetRegistry)
+	{
+		IAssetRegistry::Get()->SearchAllAssets(true /* bSynchronousSearch */);
+		bHasInitializedAssetRegistry = true;
+	}
+
+	while (PendingPackageNames.Num() > 0)
+	{
+		double CurrentTime = FPlatformTime::Seconds();
+		if (CurrentTime - StartTime > Constants::ServerPackageLoadBudgetSeconds)
+		{
+			break;
+		}
+		FString PackageName = PendingPackageNames.PopFrontValue();
+		FString ErrorMessage;
+		if (!TrySavePackage(PackageName, ErrorMessage))
+		{
+			UE_LOG(LogEditorDomainSave, Warning, TEXT("%s"), *ErrorMessage);
+		}
+		else
+		{
+			UE_LOG(LogEditorDomainSave, Display, TEXT("Saved package %s into the EditorDomain"), *PackageName);
+		}
+	}
+}
+
+bool FEditorDomainSaveServer::TrySavePackage(FStringView PackageName, FString& OutErrorMessage)
+{
+	FPackagePath PackagePath;
+	if (!FPackagePath::TryFromPackageName(PackageName, PackagePath) || !PackagePath.IsMountedPath())
+	{
+		OutErrorMessage = FString::Printf(TEXT("Package %.*s is not mounted. Skipping its save."),
+			PackageName.Len(), PackageName.GetData());
+		return false;
+	}
+	return TrySavePackage(PackagePath, OutErrorMessage);
+}
+
+bool FEditorDomainSaveServer::TrySavePackage(const FPackagePath& PackagePath, FString& OutErrorMessage)
+{
+	using namespace UE::EditorDomain;
+
+	FName PackageName = PackagePath.GetPackageFName();
+	if (PackageName.IsNone())
+	{
+		OutErrorMessage = FString::Printf(TEXT("Package %s is not mounted. Skipping its save."), *PackagePath.GetDebugName());
+		return false;
+	}
+
+	FPackageDigest PackageDigest;
+	EPackageDigestResult FindHashResult = GetPackageDigest(*AssetRegistry, PackageName, PackageDigest);
+	switch (FindHashResult)
+	{
+	case EPackageDigestResult::Success:
+		break;
+	case EPackageDigestResult::FileDoesNotExist:
+		OutErrorMessage = FString::Printf(TEXT("Package %s no longer exists in the Workspace Domain. Skipping its save."),
+			*PackagePath.GetDebugName());
+		return false;
+	case EPackageDigestResult::WrongThread:
+		OutErrorMessage = FString::Printf(TEXT("Package %s failed because GetPackageDigest returned WrongThread; ")
+			TEXT("this should not be possible in the EditorDomainSaveServer since it only works from main thread. ")
+			TEXT("Skipping its save."), *PackagePath.GetDebugName());
+		return false;
+	default:
+		check(false);
+		OutErrorMessage = TEXT("");
+		return false;
+	}
+
+	// EDITOR_DOMAIN_TODO: Check whether the package already exists in the cache before spending time on it
+	UPackage* Package = LoadPackage(nullptr, PackagePath, 0);
+	if (!Package)
+	{
+		OutErrorMessage = FString::Printf(TEXT("Package %s could not be loaded. Skipping its save."),
+			*PackagePath.GetDebugName());
+		return false;
+	}
+
+	// EDITOR_DOMAIN_TODO: Need to extend SavePackage to allow saving to an archive
+	FString TempFilename = FPaths::Combine(FPaths::ProjectIntermediateDir(), FGuid::NewGuid().ToString());
+	ON_SCOPE_EXIT{ IFileManager::Get().Delete(*TempFilename); };
+
+	FSavePackageResultStruct Result = GEditor->Save(Package, nullptr, RF_Standalone, *TempFilename);
+	if (Result.Result != ESavePackageResult::Success)
+	{
+		OutErrorMessage = FString::Printf(TEXT("Package %s could not be saved to temporary file %s: ResultCode = %d"),
+			*PackagePath.GetDebugName(), *TempFilename, Result.Result);
+		return false;
+	}
+
+	TArray64<uint8> TempBytes;
+	bool bBytesRead = FFileHelper::LoadFileToArray(TempBytes, *TempFilename);
+	IFileManager::Get().Delete(*TempFilename);
+	if (!bBytesRead)
+	{
+		OutErrorMessage = FString::Printf(TEXT("Package %s could not be read from temporary file %s."),
+			*PackagePath.GetDebugName(), *TempFilename, Result.Result);
+		return false;
+	}
+
+	UE::DerivedData::FCacheRecord Record;
+	Record.SetKey(GetEditorDomainPackageKey(PackageDigest));
+	TCbWriter<256> MetaData;
+	MetaData.BeginObject();
+	MetaData << "FileSize" << TempBytes.Num();
+	MetaData.EndObject();
+	Record.SetMeta(MetaData.Save().AsObjectRef());
+	Record.SetBinary(FSharedBuffer::MakeView(TempBytes.GetData(), TempBytes.Num()));
+	GetDerivedDataCacheRef().GetCache().Put({ Record }, PackagePath.GetDebugName());
+	return true;
+}
+
+FEditorDomainSaveServer::FClientConnection::FClientConnection(FSocket* InClientSocket)
+	: ClientSocket(InClientSocket)
+	, MessageType(UE::EditorDomainSave::EMessageType::Save)
+{
+}
+
+FEditorDomainSaveServer::FClientConnection::~FClientConnection()
+{
+	Shutdown();
+}
+
+void FEditorDomainSaveServer::FClientConnection::Shutdown()
+{
+	if (ClientSocket)
+	{
+		ISocketSubsystem* SocketSubsystem = ISocketSubsystem::Get();
+		check(SocketSubsystem);
+		ClientSocket->Close();
+		SocketSubsystem->DestroySocket(ClientSocket);
+		ClientSocket = nullptr;
+	}
+}
+
+void FEditorDomainSaveServer::FClientConnection::Poll(FEditorDomainSaveServer& Server,
+	bool& bOutIsIdle, bool& bOutStillAlive)
+{
+	using namespace UE::EditorDomainSave;
+
+	if (!ClientSocket)
+	{
+		bOutStillAlive = false;
+		bOutIsIdle = true;
+		return;
+	}
+
+	bOutStillAlive = true;
+	bOutIsIdle = true;
+	if (MessageBuffer.Num() > 0)
+	{
+		bOutIsIdle = false;
+		int32 RemainingSize = MessageBuffer.Num() - BufferOffset;
+		int32 BytesRead;
+		bool bConnectionAlive = ClientSocket->Recv(MessageBuffer.GetData() + BufferOffset, RemainingSize, BytesRead);
+		if (BytesRead > 0)
+		{
+			check(BytesRead <= RemainingSize);
+			BufferOffset += BytesRead;
+			if (BufferOffset == MessageBuffer.Num())
+			{
+				Server.ProcessMessage(*this, MessageType, MessageBuffer, bOutStillAlive);
+				MessageBuffer.Empty();
+			}
+		}
+		else if (!bConnectionAlive)
+		{
+			bOutStillAlive = false;
+		}
+	}
+	else
+	{
+		FMessageHeader MessageHeader;
+		int32 BytesRead;
+		bool bConnectionAlive = ClientSocket->Recv(reinterpret_cast<uint8*>(&MessageHeader), sizeof(MessageHeader),
+			BytesRead, ESocketReceiveFlags::Peek);
+		check(BytesRead <= sizeof(FMessageHeader));
+		if (BytesRead > 0)
+		{
+			bOutIsIdle = false;
+			if (BytesRead == sizeof(FMessageHeader))
+			{
+				ClientSocket->Recv(reinterpret_cast<uint8*>(&MessageHeader), sizeof(MessageHeader),
+					BytesRead, ESocketReceiveFlags::None);
+				check(BytesRead == sizeof(FMessageHeader));
+				if (!MessageHeader.IsValid())
+				{
+					UE_LOG(LogEditorDomainSave, Warning,
+						TEXT("EditorDomainServer received invalid message header %s from client, closing the connection."),
+						*MessageHeader.ToString());
+					bOutStillAlive = false;
+					return;
+				}
+				MessageBuffer.SetNum(MessageHeader.Size - sizeof(FMessageHeader));
+				BufferOffset = 0;
+				MessageType = MessageHeader.GetMessageType();
+
+				bool bUnusedIsIdle;
+				Poll(Server, bUnusedIsIdle, bOutStillAlive);
+			}
+		}
+		else if (!bConnectionAlive)
+		{
+			bOutStillAlive = false;
+		}
+	}
+}
+
+
+FEditorDomainSaveClient::~FEditorDomainSaveClient()
+{
+	Shutdown();
+}
+
+void FEditorDomainSaveClient::RequestSave(const FPackagePath& PackagePath)
+{
+	FScopeLock ScopeLock(&Lock);
+	if (!bEnabled)
+	{
+		return;
+	}
+	Requests.Add(PackagePath);
+	if (!bAsyncActive)
+	{
+		KickCommunication();
+	}
+}
+
+void FEditorDomainSaveClient::Tick(float DeltaTime)
+{
+	FScopeLock ScopeLock(&Lock);
+	if (Requests.Num() == 0)
+	{
+		return;
+	}
+	if (!bAsyncActive)
+	{
+		KickCommunication();
+		PollServerMessages();
+	}
+}
+
+void FEditorDomainSaveClient::KickCommunication()
+{
+	// Called within ThreadingLock
+	check(!bAsyncActive); // Caller should have checked
+	if (!TrySendBatchRequest())
+	{
+		bAsyncActive = true;
+		if (!InitializeCommunication())
+		{
+			Requests.Empty();
+		}
+		else if (!bAsyncSupported)
+		{
+			TickCommunication();
+		}
+		else
+		{
+			Async(EAsyncExecution::Thread, [this]() { RunCommunication(); });
+		}
+	}
+}
+
+bool FEditorDomainSaveClient::TrySendBatchRequest()
+{
+	// Called within ThreadingLock,
+	// Called only on public interface thread (if !bAsyncActive) or on async thread (if bAsyncActive)
+	if (Requests.Num() == 0)
+	{
+		return true;
+	}
+
+	if (bServerSocketReady)
+	{
+		TArray<FPackagePath> LocalRequests;
+		const int32 MaxRequestSize = 100;
+		int32 PeelCount = FMath::Min<int32>(MaxRequestSize, Requests.Num());
+		LocalRequests.Reserve(PeelCount);
+		for (; PeelCount > 0; --PeelCount)
+		{
+			LocalRequests.Add(Requests.PopFrontValue());
+		}
+
+		if (TrySendRequests(LocalRequests))
+		{
+			return true;
+		}
+		else
+		{
+			// Push back on front in reverse order to restore original order
+			for (int Index = LocalRequests.Num() - 1; Index >= 0; --Index)
+			{
+				Requests.AddFront(LocalRequests[Index]);
+			}
+			return false;
+		}
+	}
+	else
+	{
+		return false;
+	}
+}
+
+bool FEditorDomainSaveClient::InitializeCommunication()
+{
+	// Called within ThreadingLock
+	using namespace UE::EditorDomainSave;
+
+	if (!ISocketSubsystem::Get())
+	{
+		UE_LOG(LogEditorDomainSave, Error,
+			TEXT("Platform does not support network sockets, EditorDomain saving is not enabled."));
+		Shutdown();
+		bEnabled = false;
+		return false;
+	}
+	FString LockName = FString(FApp::GetProjectName()) + TEXT("EditorDomainServer");
+	ProcessLock = CreateProcessLock();
+	if (!ProcessLock)
+	{
+		UE_LOG(LogEditorDomainSave, Error,
+			TEXT("Could not create InterprocessSynchObject %s, EditorDomain saving is not enabled."), *LockName);
+		Shutdown();
+		bEnabled = false;
+		return false;
+	}
+	bAsyncSupported = FPlatformProcess::SupportsMultithreading();
+
+	return true;
+}
+
+void FEditorDomainSaveClient::RunCommunication()
+{
+	using namespace UE::EditorDomainSave;
+
+	bool bIdle = false;
+	double FlushStartTime = 0.;
+	for (;;)
+	{
+		check(bAsyncActive);
+		if (TryConnect())
+		{
+			FScopeLock ScopeLock(&Lock);
+			if (TrySendBatchRequest())
+			{
+				if (Requests.Num() == 0)
+				{
+					bAsyncActive = false;
+					break;
+				}
+			}
+		}
+		PollServerMessages();
+
+		if (IsStopped)
+		{
+			if (FlushStartTime == 0.)
+			{
+				FlushStartTime = FPlatformTime::Seconds();
+			}
+			if (FPlatformTime::Seconds() - FlushStartTime > Constants::ClosingFlushBudgetSeconds)
+			{
+				UE_LOG(LogEditorDomainSave, Warning,
+					TEXT("Could not flush all results before shutting down. Failing to save %d remaining requests."),
+					Requests.Num());
+				FScopeLock ScopeLock(&Lock);
+				bAsyncActive = false;
+				break;
+			}
+		}
+
+		FPlatformProcess::Sleep(Constants::RunCommunicationSleepPeriodSeconds);
+	}
+}
+
+void FEditorDomainSaveClient::TickCommunication()
+{
+	check(!bAsyncActive);
+	if (TryConnect())
+	{
+		FScopeLock ScopeLock(&Lock);
+		TrySendBatchRequest();
+	}
+}
+
+void FEditorDomainSaveClient::Shutdown()
+{
+	using namespace UE::EditorDomainSave;
+
+	IsStopped = true;
+	for (;;)
+	{
+		{
+			FScopeLock ScopeLock(&Lock);
+			if (!bAsyncActive)
+			{
+				break;
+			}
+		}
+		FPlatformProcess::Sleep(Constants::RunCommunicationSleepPeriodSeconds);
+	}
+	FScopeLock ScopeLock(&Lock);
+	Disconnect();
+	KillDanglingProcess();
+	if (ProcessLock)
+	{
+		FPlatformProcess::DeleteInterprocessSynchObject(ProcessLock);
+		ProcessLock = nullptr;
+	}
+}
+
+bool FEditorDomainSaveClient::TryConnect()
+{
+	// Called within either the ThreadingLock or the bAsyncActive flag
+	if (!bServerSocketReady)
+	{
+		TickConnect();
+	}
+	return bServerSocketReady;
+}
+
+
+void FEditorDomainSaveClient::TickConnect()
+{
+	// Called within either the ThreadingLock or the bAsyncActive flag
+	if (ConnectionAttemptStartTime == 0.)
+	{
+		ConnectionAttemptStartTime = FPlatformTime::Seconds();
+	}
+	if (!TryConnectProcess())
+	{
+		return;
+	}
+	TryConnectSocket();
+}
+
+bool FEditorDomainSaveClient::TryConnectProcess()
+{
+	// Called within either the ThreadingLock or the bAsyncActive flag
+	using namespace UE::EditorDomainSave;
+
+	if (ServerProcessId)
+	{
+		return true;
+	}
+
+	bool bSuccessful = false;
+	FProcessSemaphoreScopeLock ScopeLock(ProcessLock, Constants::ProcessLockTimeoutSecondsWhenPolling);
+	if (!ScopeLock.IsLocked())
+	{
+		return false;
+	}
+
+	uint32 LocalServerListenPort = 0;
+	uint32 LocalServerProcessId = 0;
+	FDateTime ServerSettingsTimestamp;
+	if (TryReadConnectionData(LocalServerProcessId, LocalServerListenPort, ServerSettingsTimestamp))
+	{
+		return true;
+	}
+
+	if (LocalServerProcessId)
+	{
+		FTimespan UnresponsiveTime = FDateTime::Now() - ServerSettingsTimestamp;
+		if (UnresponsiveTime.GetTotalSeconds() < Constants::ServerUnresponsiveTimeoutSeconds)
+		{
+			// Continue waiting for the server to post the listenport
+			return false;
+		}
+	}
+
+	double CurrentTime = FPlatformTime::Seconds();
+	if (CurrentTime - LastProcessCreationAttemptTimeSeconds < Constants::ProcessCreationCooldownSeconds)
+	{
+		// Don't create a new process (wasting resources and potentially triggering virus scanners)
+		// until our cooldown period has elapsed
+		return false;
+	}
+	LastProcessCreationAttemptTimeSeconds = CurrentTime;
+
+	int32 CurrentProcessId = FPlatformProcess::GetCurrentProcessId();
+	FString CommandletExecutable = FUnrealEdMisc::Get().GetProjectEditorBinaryPath();
+	FString Commandline = FString::Printf(TEXT("%s -run=EditorDomainSave -creatorprocessid=%d"),
+		*FPaths::GetProjectFilePath(), CurrentProcessId);
+	FProcHandle Handle = FPlatformProcess::CreateProc(*CommandletExecutable, *Commandline,
+		true /* bLaunchDetached */, true /* bLaunchHidden */, true /* bLaunchReallyHidden */,
+		&LocalServerProcessId, 0 /* PriorityModifier */, *FPaths::GetPath(CommandletExecutable),
+		nullptr /* PipeWriteChild */);
+	if (!Handle.IsValid())
+	{
+		ConnectionWarning(FString::Printf(TEXT("Could not create EditorDomainSave process %s %s"),
+			*CommandletExecutable, *Commandline));
+		return false;
+	}
+
+	TSharedPtr<FJsonObject> RootObject = MakeShareable(new FJsonObject());
+	RootObject->SetNumberField(TEXT("ProcessId"), (double)LocalServerProcessId);
+	if (!TryWriteServerSettings(RootObject))
+	{
+		FPlatformProcess::TerminateProc(Handle);
+		TryDeleteSettingsFile();
+		ConnectionWarning(FString::Printf(TEXT("EditorDomainSaveServer is %s, and we are unable to create ")
+			TEXT("a new settings file %s so we can not create a new one"),
+			LocalServerProcessId != 0 ? TEXT("unresponsive") : TEXT("not yet started"), *GetServerSettingsFilename()));
+		return false;
+
+	}
+
+	// We are not successful until we the server has written the ListenPort for us to read; retry the read again on next tick
+	return false;
+}
+
+bool FEditorDomainSaveClient::TryReadConnectionData(uint32& OutLocalServerProcessId, uint32& OutLocalServerListenPort,
+	FDateTime& OutServerSettingsTimestamp)
+{
+	using namespace UE::EditorDomainSave;
+
+	// Called within either the ThreadingLock or the bAsyncActive flag
+	// Called within ProcessLock
+	OutLocalServerListenPort = 0;
+	OutLocalServerProcessId = 0;
+
+	TSharedPtr<FJsonObject> RootObject;
+	if (TryReadServerSettings(RootObject, OutServerSettingsTimestamp))
+	{
+		uint32 PotentialProcessId;
+		if (RootObject->TryGetNumberField(SettingNames::ProcessIdFieldName, PotentialProcessId))
+		{
+			FProcHandle ServerHandle = FPlatformProcess::OpenProcess(PotentialProcessId);
+			ON_SCOPE_EXIT{ FPlatformProcess::CloseProc(ServerHandle); };
+			if (ServerHandle.IsValid())
+			{
+				OutLocalServerProcessId = PotentialProcessId;
+				uint32 PotentialListenPort;
+				if (RootObject->TryGetNumberField(SettingNames::ListenPortFieldName, PotentialListenPort))
+				{
+					OutLocalServerListenPort = PotentialListenPort;
+				}
+			}
+		}
+	}
+
+	if (OutLocalServerProcessId && OutLocalServerListenPort)
+	{
+		ServerProcessId = OutLocalServerProcessId;
+		ServerListenPort = OutLocalServerListenPort;
+		return true;
+	}
+	return false;
+}
+
+bool FEditorDomainSaveClient::TryConnectSocket()
+{
+	// Called within either the ThreadingLock or the bAsyncActive flag
+	if (bServerSocketReady)
+	{
+		return true;
+	}
+
+	auto PollServerSocketNotReady = [this]()
+	{
+		using namespace UE::EditorDomainSave;
+
+		if (PollKillServerProcess())
+		{
+			// PollKillServerProcess has taken action, don't check any further timeouts
+		}
+		else if (FPlatformTime::Seconds() - ConnectionAttemptSocketStartTime > Constants::SocketUnresponsiveTimeoutSeconds)
+		{
+			UE_LOG(LogEditorDomainSave, Verbose,
+				TEXT("EditorDomainSaveClient: Could not connect to %s which should be opened by process %d"),
+				*ServerAddr->ToString(true), ServerProcessId);
+			DisconnectProcess(); // Reverify that the process is running
+			ConnectionAttemptSocketStartTime = 0;
+		}
+	};
+	if (!ServerSocket)
+	{
+		if (ConnectionAttemptSocketStartTime == 0.)
+		{
+			ConnectionAttemptSocketStartTime = FPlatformTime::Seconds();
+		}
+
+		ISocketSubsystem* SocketSubsystem = ISocketSubsystem::Get();
+		check(SocketSubsystem); // This should have been checked already in Initialize
+		bool bCanBindAllUnused;
+		ServerAddr = SocketSubsystem->GetLocalHostAddr(*GLog, bCanBindAllUnused);
+		ServerAddr->SetPort(ServerListenPort);
+		FSocket* LocalServerSocket = SocketSubsystem->CreateSocket(NAME_Stream, TEXT("FEditorDomainSaveClient tcp"),
+			ServerAddr->GetProtocolType());
+		if (LocalServerSocket == nullptr)
+		{
+			ConnectionWarning(TEXT("Could not create socket"));
+			DisconnectSocket();
+			return false;
+		}
+
+		LocalServerSocket->SetNonBlocking(true);
+		if (!LocalServerSocket->Connect(*ServerAddr))
+		{
+			LocalServerSocket->Close();
+			SocketSubsystem->DestroySocket(LocalServerSocket);
+			PollServerSocketNotReady();
+			return false;
+		}
+		ServerSocket = LocalServerSocket;
+	}
+	if (!bServerSocketReady)
+	{
+		bServerSocketReady = ServerSocket->Wait(ESocketWaitConditions::WaitForWrite, FTimespan::FromSeconds(0));
+		if (!bServerSocketReady)
+		{
+			PollServerSocketNotReady();
+			return false;
+		}
+	}
+
+	UE_LOG(LogEditorDomainSave, Display,
+		TEXT("EditorDomainSaveClient: Connected to EditorDomainSaveServer with pid %d at %s"),
+		ServerProcessId, *ServerAddr->ToString(true));
+	return true;
+}
+
+bool FEditorDomainSaveClient::PollKillServerProcess()
+{
+	using namespace UE::EditorDomainSave;
+
+	if (FPlatformTime::Seconds() - ConnectionAttemptStartTime <= Constants::ServerUnresponsiveTimeoutSeconds)
+	{
+		return false;
+	}
+
+	ConnectionWarning(FString::Printf(
+		TEXT("EditorDomainSaveClient: Could not connect to %s which should be opened by process %d, for %lf seconds"),
+		*ServerAddr->ToString(true), ServerProcessId, Constants::ServerUnresponsiveTimeoutSeconds),
+		TEXT("Revoking the process's listen authority and starting a new one"),
+		true /* bIgnoreCooldown */);
+	{
+		FProcessSemaphoreScopeLock ProcessScopeLock(ProcessLock, Constants::ProcessLockTimeoutSecondsWhenPolling);
+		if (!ProcessScopeLock.IsLocked())
+		{
+			return false;
+		}
+		TryDeleteSettingsFile();
+
+		// We are about to assign a new DanglingProcessId, so time is up for any previous dangling process
+		KillDanglingProcess();
+
+		// We will give the server some time to notice that its authority has been revoked and shut itself down,
+		// But after a timeout we kill the process
+		DanglingProcessId = ServerProcessId;
+	}
+	DisconnectProcess();
+	ConnectionAttemptStartTime = 0.;
+	return true;
+}
+
+void FEditorDomainSaveClient::KillDanglingProcess()
+{
+	if (!DanglingProcessId)
+	{
+		return;
+	}
+	int32 ProcessId = DanglingProcessId;
+	DanglingProcessId = 0;
+
+	if (!FPathViews::GetBaseFilenameWithPath(FPlatformProcess::GetApplicationName(ProcessId)).Equals(
+		FPathViews::GetBaseFilenameWithPath(FUnrealEdMisc::Get().GetProjectEditorBinaryPath()), ESearchCase::IgnoreCase))
+	{
+		// The process must have exited on its own and some other process has taken over the pid it had
+		return;
+	}
+
+	FProcHandle Handle = FPlatformProcess::OpenProcess(ProcessId);
+	ON_SCOPE_EXIT{ FPlatformProcess::CloseProc(Handle); };
+	if (Handle.IsValid())
+	{
+		FPlatformProcess::TerminateProc(Handle, true /* KillTree */);
+	}
+}
+
+bool FEditorDomainSaveClient::TrySendRequests(TArray<FPackagePath>& LocalRequests)
+{
+	// Called within ThreadingLock
+	// Called only on public interface thread (if !bAsyncActive) or on async thread (if bAsyncActive)
+	using namespace UE::EditorDomainSave;
+
+	TArray<uint8> MessageBytes;
+	FMessageHeader MessageHeader;
+	{
+		FMemoryWriter Message(MessageBytes);
+
+		MessageHeader.Initialize(EMessageType::Save);
+		Message << MessageHeader;
+
+		int32 NumRequests = LocalRequests.Num();
+		Message << NumRequests;
+		for (const FPackagePath& PackagePath : LocalRequests)
+		{
+			FString PackageName = PackagePath.GetPackageName();
+			Message << PackageName;
+		}
+		MessageHeader.Size = Message.Tell();
+		Message.Seek(0);
+		Message << MessageHeader;
+		Message.Seek(MessageHeader.Size);
+	}
+	check(MessageBytes.Num() == MessageHeader.Size);
+
+	int32 BytesSent;
+	ServerSocket->Send(MessageBytes.GetData(), MessageBytes.Num(), BytesSent);
+	if (BytesSent != MessageBytes.Num())
+	{
+		ConnectionWarning(TEXT("Socket::Send failed to send message to server"), TEXT("Will reconnect"));
+		if (bPreviousSendFailed)
+		{
+			PollKillServerProcess();
+		}
+		else
+		{
+			bPreviousSendFailed = true;
+			ConnectionAttemptStartTime = 0; // Start the timeout from now, the first point of failure
+		}
+		DisconnectProcess();
+		return false;
+	}
+	bPreviousSendFailed = false;
+
+	LocalRequests.Reset();
+	return true;
+}
+
+void FEditorDomainSaveClient::Disconnect()
+{
+	DisconnectProcess();
+	ConnectionAttemptStartTime = 0.;
+	bPreviousSendFailed = false;
+}
+
+void FEditorDomainSaveClient::DisconnectProcess()
+{
+	DisconnectSocket();
+	ServerProcessId = 0;
+	ServerListenPort = 0;
+	ConnectionAttemptSocketStartTime = 0.;
+}
+
+void FEditorDomainSaveClient::DisconnectSocket()
+{
+	bServerSocketReady = false;
+	ServerAddr.Reset();
+	if (ServerSocket)
+	{
+		ISocketSubsystem* SocketSubsystem = ISocketSubsystem::Get();
+		check(SocketSubsystem);
+		ServerSocket->Close();
+		SocketSubsystem->DestroySocket(ServerSocket);
+		ServerSocket = nullptr;
+	}
+}
+
+void FEditorDomainSaveClient::ConnectionWarning(FStringView Message, FStringView Consequence, bool bIgnoreCooldown)
+{
+	using namespace UE::EditorDomainSave;
+
+	double CurrentTime = FPlatformTime::Seconds();
+	if (bIgnoreCooldown || CurrentTime - LastConnectionWarningTimeSeconds >= Constants::ConnectionWarningCooldownSeconds)
+	{
+		if (!bIgnoreCooldown)
+		{
+			LastConnectionWarningTimeSeconds = CurrentTime;
+		}
+		FStringView DefaultConsequence = TEXT("EditorDomainSave is currently unavailable"_SV);
+		if (Consequence.Len() == 0)
+		{
+			Consequence = DefaultConsequence;
+		}
+		UE_LOG(LogEditorDomainSave, Warning, TEXT("%.*s. %.*s."), Message.Len(), Message.GetData(), Consequence.Len(),
+			Consequence.GetData());
+	}
+}
+
+void FEditorDomainSaveClient::PollServerMessages()
+{
+	// Called within either the ThreadingLock or the bAsyncActive flag
+	if (!ServerSocket)
+	{
+		return;
+	}
+
+	// EDITOR_DOMAIN_TODO: Read messages from the server
+
+}
+
+namespace UE
+{
+namespace EditorDomainSave
+{
+
+FString GetServerSettingsFilename()
+{
+	return FPaths::ProjectIntermediateDir() / TEXT("Cook") / TEXT("EditorDomainSave.json");
+}
+
+FString GetLockName()
+{
+	return FString(FApp::GetProjectName()) + TEXT("EditorDomainServer");
+
+}
+FPlatformProcess::FSemaphore* CreateProcessLock()
+{
+	// EDITOR_DOMAIN_TODO: Need to add a HasInterprocessSyncLock function to prevent the
+	// use of editor saving rather than throwing a fatal
+	return FPlatformProcess::NewInterprocessSynchObject(GetLockName(), true /* bCreate */);
+}
+
+bool TryReadServerSettings(TSharedPtr<FJsonObject>& OutRootObject, FDateTime& OutServerSettingsTimestamp)
+{
+	OutRootObject = TSharedPtr<FJsonObject>();
+	OutServerSettingsTimestamp = FDateTime();
+
+	FString ServerSettingsFilename = GetServerSettingsFilename();
+	IFileManager& FileManager = IFileManager::Get();
+	FFileStatData StatData = FileManager.GetStatData(*ServerSettingsFilename);
+	if (!StatData.bIsValid)
+	{
+		return false;
+	}
+	TUniquePtr<FArchive> FileAr(FileManager.CreateFileReader(*ServerSettingsFilename));
+	if (!FileAr)
+	{
+		return false;
+	}
+
+	OutRootObject = MakeShareable(new FJsonObject);
+	TSharedRef<TJsonReader<TCHAR>> Reader = TJsonReaderFactory<TCHAR>::Create(FileAr.Get());
+	if (!FJsonSerializer::Deserialize(Reader, OutRootObject))
+	{
+		OutRootObject.Reset();
+		return false;
+	}
+
+	OutServerSettingsTimestamp = StatData.ModificationTime;
+	return true;
+}
+
+bool TryWriteServerSettings(TSharedPtr<FJsonObject>& RootObject)
+{
+	if (!RootObject)
+	{
+		return false;
+	}
+	IFileManager& FileManager = IFileManager::Get();
+	FString ServerSettingsFilename = GetServerSettingsFilename();
+	FileManager.MakeDirectory(*FPaths::GetPath(ServerSettingsFilename), true /* Tree */);
+	TUniquePtr<FArchive> FileAr(IFileManager::Get().CreateFileWriter(*ServerSettingsFilename));
+	if (!FileAr)
+	{
+		return false;
+	}
+	TSharedRef<TJsonWriter<TCHAR>> Writer = TJsonWriterFactory<TCHAR>::Create(FileAr.Get());
+	return FJsonSerializer::Serialize(RootObject.ToSharedRef(), Writer, false/* bCloseWriter */);
+}
+
+bool TryDeleteSettingsFile()
+{
+	FString ServerSettingsFilename = GetServerSettingsFilename();
+	IFileManager& FileManager = IFileManager::Get();
+	return !FileManager.FileExists(*ServerSettingsFilename) || FileManager.Delete(*ServerSettingsFilename);
+}
+
+void GetServerSettingsTimeStamp(bool& bOutExists, FDateTime& OutModificationTime)
+{
+	OutModificationTime = IFileManager::Get().GetTimeStamp(*GetServerSettingsFilename());
+	bOutExists = OutModificationTime != FDateTime::MinValue();
+}
+
+
+}
+}
diff --git a/Engine/Source/Editor/UnrealEd/Private/EditorDomain/EditorDomainSave.h b/Engine/Source/Editor/UnrealEd/Private/EditorDomain/EditorDomainSave.h
new file mode 100644
--- /dev/null
+++ b/Engine/Source/Editor/UnrealEd/Private/EditorDomain/EditorDomainSave.h
@@ -0,0 +1,279 @@
+// Copyright Epic Games, Inc. All Rights Reserved.
+
+#pragma once
+
+#include <atomic>
+#include "Containers/Array.h"
+#include "Containers/RingBuffer.h"
+#include "Containers/StringView.h"
+#include "Containers/UnrealString.h"
+#include "HAL/CriticalSection.h"
+#include "HAL/Platform.h"
+#include "HAL/PlatformProcess.h"
+#include "HAL/Runnable.h"
+#include "Misc/DateTime.h"
+#include "Templates/RefCounting.h"
+#include "Templates/SharedPointer.h"
+#include "Templates/UniquePtr.h"
+
+class FJsonObject;
+class FInternetAddr;
+class FPackagePath;
+class FRunnableThread;
+class FSocket;
+class IAssetRegistry;
+class IPackageResourceManager;
+namespace UE
+{
+namespace EditorDomainSave
+{
+	enum class EMessageType : uint32;
+}
+}
+
+/**
+ * System to handle save requests sent from editor and commandlet processes to save packages
+ * into the Editor Domain.
+ */
+class FEditorDomainSaveServer
+{
+public:
+	FEditorDomainSaveServer();
+	~FEditorDomainSaveServer();
+
+	int32 Run();
+
+private:
+	/** A connection to a client held by EditorDomainSaveServer. */
+	struct FClientConnection : public FRefCountedObject
+	{
+		FClientConnection(FSocket* InClientSocket);
+		~FClientConnection();
+		/** Release resources and indicate this connection should no longer be used. */
+		void Shutdown();
+		/** Read message bytes from the socket. */
+		void Poll(FEditorDomainSaveServer& Server, bool& bOutIsIdle, bool& bOutStillAlive);
+
+		/** The socket to the client process, returned from accept on the server's ListenSocket. */
+		FSocket* ClientSocket;
+		/**
+		 * Bytes being read from the client process, held until all the expected size from the message header is read.
+		 * Array is sized to the size from the message.
+		 */
+		TArray<uint8> MessageBuffer;
+		/** Number of bytes we have so far read into the buffer. */
+		int32 BufferOffset = 0;
+		/** The MessageType for the message being read, as reported in the message header. */
+		UE::EditorDomainSave::EMessageType MessageType;
+	};
+
+	/**
+	 * Initialize the listen socket, failing if this process is not the designated handler process
+	 * or if any required features are missing in the current development environment.
+	 */
+	bool TryInitialize();
+	/**
+	 * Called on startup or when the ServerSettings file on disk changes.
+	 * Read the settings file, return false if this process is not (or is no longer or can no longer be) the
+	 * designated handler, and consume any new or modified settings from the file
+	 * 
+	 * @param bOutOwnsSettingsFile [Optional] Set to true if this process still owns the settings file
+	 * and should delete it on exit. This can be true even if the process should shut down due to e.g. an error.
+	 * @return Whether this FEditorDomainSaveServer should continue to handle save messages
+	 */
+	bool TryRefreshAuthority(bool* bOutOwnsSettingsFile = nullptr, TSharedPtr<FJsonObject>* OutRootObject = nullptr,
+		int32* OutListenPort = nullptr);
+	/** Free used resources; destructor code moved into separate function to make flexible when it operates. */
+	void Shutdown();
+	/** Shut down the Listen socket and all client sockets created from it. */
+	void ShutdownSocket();
+	/**
+	 * Check whether this server should shutdown because no clients requesting a handler currently exist.
+	 * Note this check in the ensuing shutdown have to be done inside the interprocess lock to prevent races
+	 * with clients that want to create a new handler process if one is not already running
+	 */
+	bool TryAbdicate();
+	/** Return whether TryAbdicate has returned false and this server should therefore shut down. */
+	bool IsAbdicated() const;
+	/**
+	 * Return true if there are any active client connections, or if the expected connection from the 
+	 * initial client that caused the creation of this server is still pending.
+	 */
+	bool HasExpectedConnections() const;
+	/**
+	 * Called periodically to check whether this server is still the designated handler and still has clients.
+	 *
+	 * @return whether the server should continue running.
+	 */
+	bool PollShouldRun();
+	/** Check the listen socket for new client connections. */
+	void PollIncomingConnections(bool& bInOutIsIdle);
+	/** Check the client sockets for new messages. */
+	void PollConnections(bool& bInOutIsIdle);
+	/** Poll periodic maintenance such as garbage collection. */
+	void TickMaintenance(bool bIsIdle);
+	/**
+	 * Called each tick loop when whether the server is idle is known (e.g. no pending messages, no pending loads).
+	 * Perform actions that should be taken when IsIdle changes
+	 */
+	void SetIdle(bool bIsIdle);
+	/**  Switch on messagetype and handle a message sent from a client (Save, Close, others). */
+	void ProcessMessage(FClientConnection& ClientConnection, UE::EditorDomainSave::EMessageType MessageType,
+		const TArray<uint8>& MessageBuffer, bool& bOutStillAlive);
+	/** Time-sliced tick function to load and save packages clients have requested. */
+	void TickPendingPackages(bool& bInOutIsIdle);
+	/** Load/Save a given PackageName/PackagePath. */
+	bool TrySavePackage(FStringView PackageName, FString& OutErrorMessage);
+	bool TrySavePackage(const FPackagePath& PackagePath, FString& OutErrorMessage);
+
+	/** Network socket, pending message process, and and other information about each connected client. */
+	TArray<TRefCountPtr<FClientConnection>> ClientConnections;
+	/** The list of packages that have been requested saved but not yet saved. */
+	TRingBuffer<FString> PendingPackageNames;
+	/** Timestamp of the settings file on disk the last time we read it, to detect settings updates. */
+	FDateTime SettingsTimestamp;
+	/** The socket listening for client connections. */
+	FSocket* ListenSocket = nullptr;
+	/** Cached pointer to the global asset registry. */
+	IAssetRegistry* AssetRegistry;
+	/** In-process handle to the interprocess synchronization object. */
+	FPlatformProcess::FSemaphore* ProcessLock = nullptr;
+	/** The time at which the current idleness started. Used to decide abdication. */
+	double IdleStartTime = 0;
+	/** The time at which garbage was last collected. Used for scheduling the next collection. */
+	double LastGarbageTime = 0;
+	/** The process id that created this process (and is expected to connect). Used to decide abdication. */
+	int32 CreatorProcessId = 0;
+	/** The port to listen on. */
+	int32 ListenPort = 0;
+	/** True if the server was idle (no pending messages, no pending packages) in its previous tick. */
+	bool bIdle = false;
+	/** True if any client has so-far connected. Used to decide abdication. */
+	bool bHasEverConnected = false;
+	/** True if the server has abdicated and should shutdown. */
+	bool bAbdicated = false;
+	/** Tracks whether initialize-on-demand steps on the assetregistry have been done. */
+	bool bHasInitializedAssetRegistry = false;
+	/** Whether this server was launched as a resident server; it takes over authority and never abdicates. */
+	bool bIsResidentServer = false;
+};
+
+/** Client-side system to send SaveRequests to the EditorDomainSaveServer in another process. */
+class FEditorDomainSaveClient
+{
+public:
+	~FEditorDomainSaveClient();
+
+	/** Send a request to the EditorDomainSaveServer to save the given PackagePath into the Editor Domain. */
+	void RequestSave(const FPackagePath& PackagePath);
+	/** Tick communication with the server. */
+	void Tick(float DeltaTime);
+
+private:
+	/**
+	 * If already connected, send message synchronously.
+	 * If not yet connected or send failed, start an async task to connect and send message.
+	 * Early-exit the entire function if async task is already running.
+	 * If async is not supported, runs one tick of the connection attempt per call.
+	 */
+	void KickCommunication();
+	/** If connected, send a batch of names to the server. Return true on success or empty. */
+	bool TrySendBatchRequest();
+	/** Construct on-demand resources and flags. Returns true iff saving is enabled. */
+	bool InitializeCommunication();
+	/** Called on an async thread. Try to connect until successful, send messages until empty, return. */
+	void RunCommunication();
+	/** Called from the public interface thread. Pump one connection step, send one message, return. */
+	void TickCommunication();
+	/** Stop the async task if it is running, wait for the async task to complete, shutdown all resources. */
+	void Shutdown();
+	/** Tick the connection code if it is not already ready and return true if it is ready. */
+	bool TryConnect();
+	/** Pump one tick of the connection process - settings file read, process creation, opening socket. */
+	void TickConnect();
+	/**
+	 * Read the server settings file to find the process id and port of the save server process. 
+	 * Create the process and server settings file if it is not running.
+	 * Return whether we now have the process id and port number.
+	 */
+	bool TryConnectProcess();
+	/** Read the server settings file to find the process id and port of the save server process. */
+	bool TryReadConnectionData(uint32& OutLocalServerProcessId, uint32& OutLocalServerListenPort,
+		FDateTime& OutServerSettingsTimestamp);
+	/** Connect to the server's listen port. Set bServerSocketReady if completed. Retry the connect if necessary. */
+	bool TryConnectSocket();
+	/**
+	 * Check whether the server connection has failed for too long,
+	 * and if so revoke its authority to listen by rewriting the settings file.
+	 */
+	bool PollKillServerProcess();
+	/** Terminate the process of the process from which we earlier revoked authority */
+	void KillDanglingProcess();
+	/** Send the batch of requests to the server connection, and remove sent requests from the LocalRequests list. */
+	bool TrySendRequests(TArray<FPackagePath>& LocalRequests);
+	/** Close the current connection attempt with the server, including the connection starttime. */
+	void Disconnect();
+	/** Remove information about the server process; it will be reconstructed when TickConnect is called. */
+	void DisconnectProcess();
+	/** Close the socket communicating with the server; it will be reconstructed when TickConnect is called. */
+	void DisconnectSocket();
+	/** Issue a connection warning with reason and consequence; apply a cooldown to prevent spamming warnings. */
+	void ConnectionWarning(FStringView Message, FStringView Consequence = FStringView(), bool bIgnoreCooldown=false);
+	/** If the socket is ready, check for messages from the server and process them. */
+	void PollServerMessages();
+
+	// Shared variables - threadsafe
+	/** False by default, set to true if the public interface wants to shutdown the async task. */
+	std::atomic<bool> IsStopped{ false };
+	/** Lock around data accessed by the public interface. */
+	FCriticalSection Lock;
+
+	// PublicInterface RestrictedWrite variables
+	// These variables can only be accessed on the public interface when inside the lock.
+	// They may not be written when bAsyncActive is true.
+	/** True by default, set to false if the communication with the server is unavailable. */
+	bool bEnabled = true;
+	/** True iff async calls are supported. If false, Tick will be used instead of Run. */
+	bool bAsyncSupported = false;
+
+	// SharedVariables - read/write only inside the lock.
+	/** Requests through the public interface that have not yet been sent. */
+	TRingBuffer<FPackagePath> Requests;
+	/**
+	 * True if an async task has started and not yet finished.
+	 * If true (inside the lock), the async task will send all requests before exiting.
+	 */
+	bool bAsyncActive = false;
+
+	// ProcessCommunication variables. These variables can only be read or written
+	// 1) In public interface when inside the Lock and !bAsyncActive
+	// OR
+	// 2) From the async task and bAsyncActive
+	/** Socket to the server process; requests are sent on this socket. */
+	FSocket* ServerSocket = nullptr;
+	/** In-process handle to the interprocess synchronization object. */
+	FGenericPlatformProcess::FSemaphore* ProcessLock = nullptr;
+	/** The url:port the server is listening on. */
+	TSharedPtr<FInternetAddr> ServerAddr;
+	/** When we started the current connection attempt; used to decide when the server process is not responding. */
+	double ConnectionAttemptStartTime = 0.;
+	/** When we started the current socket connection attempt; used to check to see if the server process is alive. */
+	double ConnectionAttemptSocketStartTime = 0.;
+	/** The last time we attempted to create a process. */
+	double LastProcessCreationAttemptTimeSeconds = 0.;
+	/* The last time we sent a warning to the log. */
+	double LastConnectionWarningTimeSeconds = 0.;
+	/** The pid of the process running the server, created by us or another client. */
+	int32 ServerProcessId = 0;
+	/** The port the server is listening on. */
+	int32 ServerListenPort = 0;
+	/** A server that we have told to shutdown; kill the process if it doesn't shutdown. */
+	int32 DanglingProcessId = 0;
+	/** Indicator that the socket asynchronous connect has completed and we can call Send on it. */
+	bool bServerSocketReady = false;
+	/**
+	 * Tracks that our previous attempt to send failed, even though the serversocket was ready.
+	 * Used to detect connecting socket but failing send.
+	 */
+	bool bPreviousSendFailed = false;
+};
diff --git a/Engine/Source/Editor/UnrealEd/Private/EditorDomain/EditorDomainSaveCommandlet.cpp b/Engine/Source/Editor/UnrealEd/Private/EditorDomain/EditorDomainSaveCommandlet.cpp
new file mode 100644
--- /dev/null
+++ b/Engine/Source/Editor/UnrealEd/Private/EditorDomain/EditorDomainSaveCommandlet.cpp
@@ -0,0 +1,16 @@
+// Copyright Epic Games, Inc. All Rights Reserved.
+
+#include "Commandlets/EditorDomainSaveCommandlet.h"
+
+#include "EditorDomainSave.h"
+
+UEditorDomainSaveCommandlet::UEditorDomainSaveCommandlet(const FObjectInitializer& ObjectInitializer)
+	: Super(ObjectInitializer)
+{
+}
+
+int32 UEditorDomainSaveCommandlet::Main(const FString& CmdLineParams)
+{
+	FEditorDomainSaveServer Server;
+	return Server.Run();
+}
diff --git a/Engine/Source/Editor/UnrealEd/Private/EditorDomain/EditorDomainUtils.cpp b/Engine/Source/Editor/UnrealEd/Private/EditorDomain/EditorDomainUtils.cpp
new file mode 100644
--- /dev/null
+++ b/Engine/Source/Editor/UnrealEd/Private/EditorDomain/EditorDomainUtils.cpp
@@ -0,0 +1,88 @@
+// Copyright Epic Games, Inc. All Rights Reserved.
+
+#include "EditorDomain/EditorDomainUtils.h"
+
+#include "Algo/IsSorted.h"
+#include "AssetRegistry/AssetData.h"
+#include "AssetRegistry/IAssetRegistry.h"
+#include "Containers/Array.h"
+#include "DerivedDataCacheInterface.h"
+#include "DerivedDataCacheKey.h"
+#include "Misc/PackagePath.h"
+#include "Misc/StringBuilder.h"
+#include "Serialization/CompactBinaryWriter.h"
+#include "UObject/ObjectVersion.h"
+
+namespace UE
+{
+namespace EditorDomain
+{
+
+// Change to a new guid when EditorDomain needs to be invalidated
+const TCHAR* EditorDomainVersion = TEXT("D3317196-098A-449B-B740-631D94CD7C93");
+// Identifier of the CacheBucket for EditorDomainPackages
+const TCHAR* EditorDomainPackageBucketName = TEXT("EditorDomainPackage");
+
+FPackageDigest GetPackageDigest(const FAssetPackageData& PackageData, FName PackageName)
+{
+	FCbWriter Writer;
+	int32 CurrentFileVersionUE4 = GPackageFileUE4Version;
+	int32 CurrentFileVersionLicenseeUE4 = GPackageFileLicenseeUE4Version;
+	Writer << EditorDomainVersion;
+	Writer << const_cast<FGuid&>(PackageData.PackageGuid);
+	Writer << CurrentFileVersionUE4;
+	Writer << CurrentFileVersionLicenseeUE4;
+	check(Algo::IsSorted(PackageData.GetCustomVersions()));
+	for (const UE::AssetRegistry::FPackageCustomVersion& PackageVersion : PackageData.GetCustomVersions())
+	{
+		Writer << const_cast<FGuid&>(PackageVersion.Key);
+		TOptional<FCustomVersion> CurrentVersion = FCurrentCustomVersions::Get(PackageVersion.Key);
+		if (CurrentVersion.IsSet())
+		{
+			Writer << CurrentVersion->Version;
+		}
+		else
+		{
+			UE_LOG(LogEditorDomain, Error, TEXT("CustomVersion guid %s is used by package %s but is not ")
+				TEXT("available in FCurrentCustomVersions. PackageDigest will set its version to 0."),
+				*PackageVersion.Key.ToString(), *PackageName.ToString());
+		}
+	}
+	return Writer.Save().GetRangeHash();
+}
+
+EPackageDigestResult GetPackageDigest(IAssetRegistry& AssetRegistry, FName PackageName,
+	FPackageDigest& OutPackageDigest)
+{
+	if (!IsInGameThread())
+	{
+		return EPackageDigestResult::WrongThread;
+	}
+	AssetRegistry.WaitForPackage(PackageName.ToString());
+	const FAssetPackageData* PackageData = AssetRegistry.GetAssetPackageData(PackageName);
+	if (!PackageData)
+	{
+		return EPackageDigestResult::FileDoesNotExist;
+	}
+	OutPackageDigest = GetPackageDigest(*PackageData, PackageName);
+	return EPackageDigestResult::Success;
+}
+
+UE::DerivedData::FCacheKey GetEditorDomainPackageKey(const FPackageDigest& PackageDigest)
+{
+	static UE::DerivedData::FCacheBucket EditorDomainPackageCacheBucket(EditorDomainPackageBucketName);
+	return UE::DerivedData::FCacheKey(EditorDomainPackageCacheBucket, PackageDigest);
+}
+
+UE::DerivedData::FCacheRequest RequestEditorDomainPackage(const FPackagePath& PackagePath,
+	const FPackageDigest& PackageDigest, UE::DerivedData::ECachePriority CachePriority,
+	UE::DerivedData::FOnCacheGetComplete&& Callback)
+{
+	UE::DerivedData::ICache& Cache = GetDerivedDataCacheRef().GetCache();
+	return Cache.Get({ GetEditorDomainPackageKey(PackageDigest) },
+		PackagePath.GetDebugName(),
+		UE::DerivedData::ECachePolicy::QueryLocal, CachePriority, MoveTemp(Callback));
+}
+
+}
+}
diff --git a/Engine/Source/Editor/UnrealEd/Private/EditorDomain/EditorDomainUtils.h b/Engine/Source/Editor/UnrealEd/Private/EditorDomain/EditorDomainUtils.h
new file mode 100644
--- /dev/null
+++ b/Engine/Source/Editor/UnrealEd/Private/EditorDomain/EditorDomainUtils.h
@@ -0,0 +1,49 @@
+// Copyright Epic Games, Inc. All Rights Reserved.
+
+#pragma once
+
+#include "Containers/UnrealString.h"
+#include "EditorDomain/EditorDomain.h"
+
+DECLARE_LOG_CATEGORY_EXTERN(LogEditorDomainSave, Log, All);
+
+class FAssetPackageData;
+class FPackagePath;
+class IAssetRegistry;
+namespace UE
+{
+namespace DerivedData
+{
+	class FCacheKey;
+}
+}
+
+namespace UE
+{
+namespace EditorDomain
+{
+
+enum class EPackageDigestResult
+{
+	Success,
+	FileDoesNotExist,
+	WrongThread,
+};
+
+/**
+ * Calculate the PackageDigest for the given packagePath.
+ * Reads information from the AssetRegistry to compute the digest.
+ */
+EPackageDigestResult GetPackageDigest(IAssetRegistry& AssetRegistry, FName PackageName,
+	FPackageDigest& OutPackageDigest);
+
+/** Convert the given PackageDigest into CacheKey format. */
+UE::DerivedData::FCacheKey GetEditorDomainPackageKey(const FPackageDigest& PackageDigest);
+
+/** Get the CacheRequest for the given package from the EditorDomain cache bucket. */
+UE::DerivedData::FCacheRequest RequestEditorDomainPackage(const FPackagePath& PackagePath,
+	const FPackageDigest& PackageDigest, UE::DerivedData::ECachePriority CachePriority,
+	UE::DerivedData::FOnCacheGetComplete&& Callback);
+
+}
+}
diff --git a/Engine/Source/Editor/UnrealEd/Private/UnrealEdMisc.cpp b/Engine/Source/Editor/UnrealEd/Private/UnrealEdMisc.cpp
--- a/Engine/Source/Editor/UnrealEd/Private/UnrealEdMisc.cpp
+++ b/Engine/Source/Editor/UnrealEd/Private/UnrealEdMisc.cpp
@@ -910,6 +910,11 @@ FString CreateProjectPath()
 #endif
 }
 
+FString FUnrealEdMisc::GetProjectEditorBinaryPath()
+{
+	return CreateProjectPath();
+}
+
 void FUnrealEdMisc::OnExit()
 {
 	if ( !bInitialized )
diff --git a/Engine/Source/Editor/UnrealEd/Public/UnrealEdMisc.h b/Engine/Source/Editor/UnrealEd/Public/UnrealEdMisc.h
--- a/Engine/Source/Editor/UnrealEd/Public/UnrealEdMisc.h
+++ b/Engine/Source/Editor/UnrealEd/Public/UnrealEdMisc.h
@@ -308,6 +308,9 @@ public:
 
 	bool EnableWorldPartition(UWorld* InWorld, bool bEnable);
 
+	/** Get the path to the executable that runs the editor */
+	FString GetProjectEditorBinaryPath();
+
 private:
 
 	/** The current state of the autosave */
diff --git a/Engine/Source/Runtime/Core/Public/Async/AsyncFileHandleNull.h b/Engine/Source/Runtime/Core/Public/Async/AsyncFileHandleNull.h
new file mode 100644
--- /dev/null
+++ b/Engine/Source/Runtime/Core/Public/Async/AsyncFileHandleNull.h
@@ -0,0 +1,57 @@
+// Copyright Epic Games, Inc. All Rights Reserved.
+
+#pragma once
+
+#include "Async/AsyncFileHandle.h"
+
+
+/**
+ * IAsyncReadRequest returned from FAsyncReadFileHandleNull;
+ * guaranteed to be a cancelled readrequest with no size or bytes when the Callback is called.
+ */
+class FAsyncReadRequestNull : public IAsyncReadRequest
+{
+public:
+	FAsyncReadRequestNull(FAsyncFileCallBack* InCallback, bool bInSizeRequest)
+		: IAsyncReadRequest(InCallback, bInSizeRequest, nullptr /* UserSuppliedMemory */)
+	{
+		bCanceled = true;
+		SetComplete();
+	}
+
+protected:
+	virtual void WaitCompletionImpl(float TimeLimitSeconds) override
+	{
+	}
+
+	virtual void CancelImpl() override
+	{
+	}
+};
+
+/**
+ * An IAsyncReadFileHandle that returns only failed results,
+ * used when a function has failed but needs to return a non-null IAsyncReadFileHandle.
+ */
+class FAsyncReadFileHandleNull : public IAsyncReadFileHandle
+{
+public:
+	using IAsyncReadFileHandle::IAsyncReadFileHandle;
+
+	virtual IAsyncReadRequest* SizeRequest(FAsyncFileCallBack* CompleteCallback = nullptr) override
+	{
+		return new FAsyncReadRequestNull(CompleteCallback, true /* bInSizeRequest */);
+	}
+
+	virtual IAsyncReadRequest* ReadRequest(int64 Offset, int64 BytesToRead,
+		EAsyncIOPriorityAndFlags PriorityAndFlags = AIOP_Normal, FAsyncFileCallBack* CompleteCallback = nullptr,
+		uint8* UserSuppliedMemory = nullptr) override
+	{
+		return new FAsyncReadRequestNull(CompleteCallback, false /* bInSizeRequest */);
+	}
+
+	virtual bool UsesCache()
+	{
+		return false;
+	}
+};
\ No newline at end of file
diff --git a/Engine/Source/Runtime/CoreUObject/Private/UObject/PackageResourceManagerFile.cpp b/Engine/Source/Runtime/CoreUObject/Private/UObject/PackageResourceManagerFile.cpp
--- a/Engine/Source/Runtime/CoreUObject/Private/UObject/PackageResourceManagerFile.cpp
+++ b/Engine/Source/Runtime/CoreUObject/Private/UObject/PackageResourceManagerFile.cpp
@@ -4,7 +4,7 @@
 
 #include "CoreMinimal.h"
 
-#include "Async/AsyncFileHandle.h"
+#include "Async/AsyncFileHandleNull.h"
 #include "HAL/FileManager.h"
 #include "HAL/PlatformFile.h"
 #include "HAL/PlatformFileManager.h"
@@ -104,57 +104,6 @@ protected:
 	};
 };
 
-/**
- * IAsyncReadRequest returned from FAsyncReadFileHandleNull;
- * guaranteed to be a cancelled readrequest with no size or bytes when the Callback is called.
- */
-class FAsyncReadRequestNull : public IAsyncReadRequest
-{
-public:
-	FAsyncReadRequestNull(FAsyncFileCallBack* InCallback, bool bInSizeRequest)
-		: IAsyncReadRequest(InCallback, bInSizeRequest, nullptr /* UserSuppliedMemory */)
-	{
-		bCanceled = true;
-		SetComplete();
-	}
-
-protected:
-	virtual void WaitCompletionImpl(float TimeLimitSeconds) override
-	{
-	}
-
-	virtual void CancelImpl() override
-	{
-	}
-};
-
-/**
- * An IAsyncReadFileHandle that returns only failed results,
- * used when a function has failed but needs to return a non-null IAsyncReadFileHandle.
- */
-class FAsyncReadFileHandleNull : public IAsyncReadFileHandle
-{
-public:
-	using IAsyncReadFileHandle::IAsyncReadFileHandle;
-
-	virtual IAsyncReadRequest* SizeRequest(FAsyncFileCallBack* CompleteCallback = nullptr) override
-	{
-		return new FAsyncReadRequestNull(CompleteCallback, true /* bInSizeRequest */);
-	}
-
-	virtual IAsyncReadRequest* ReadRequest(int64 Offset, int64 BytesToRead,
-		EAsyncIOPriorityAndFlags PriorityAndFlags = AIOP_Normal, FAsyncFileCallBack* CompleteCallback = nullptr,
-		uint8* UserSuppliedMemory = nullptr) override
-	{
-		return new FAsyncReadRequestNull(CompleteCallback, false /* bInSizeRequest */);
-	}
-
-	virtual bool UsesCache()
-	{
-		return false;
-	}
-};
-
 IPackageResourceManager* MakePackageResourceManagerFile()
 {
 	return new FPackageResourceManagerFile();
