UEdMode uses element selection set and handles to do additional actor processing.

diff --git a/Engine/Source/Editor/UnrealEd/Private/Tools/UEdMode.cpp b/Engine/Source/Editor/UnrealEd/Private/Tools/UEdMode.cpp
--- a/Engine/Source/Editor/UnrealEd/Private/Tools/UEdMode.cpp
+++ b/Engine/Source/Editor/UnrealEd/Private/Tools/UEdMode.cpp
@@ -26,6 +26,9 @@
 #include "Editor/EditorEngine.h"
 #include "GameFramework/Actor.h"
 
+#include "Elements/Framework/TypedElementSelectionSet.h"
+#include "Elements/Interfaces/TypedElementObjectInterface.h"
+
 /** Hit proxy used for editable properties */
 struct HPropertyWidgetProxyTools : public HHitProxy
 {
@@ -131,14 +134,11 @@ bool UEdMode::InputKey(FEditorViewportClient* ViewportClient, FViewport* Viewpor
 	}
 
 	// Finally, pass input up to selected actors if not in a tool mode
-	TArray<AActor*> SelectedActors;
-	Owner->GetSelectedActors()->GetSelectedObjects<AActor>(SelectedActors);
-
-	for (TArray<AActor*>::TIterator It(SelectedActors); It; ++It)
-	{
-		// Tell the object we've had a key press
-		(*It)->EditorKeyPressed(Key, Event);
-	}
+	Owner->GetEditorSelectionSet()->ForEachSelectedObject<AActor>([Key, Event](AActor* ActorPtr)
+		{
+			ActorPtr->EditorKeyPressed(Key, Event);
+			return true;
+		});
 
 	return false;
 }
@@ -176,11 +176,11 @@ void UEdMode::Enter()
 {
 	// Update components for selected actors, in case the mode we just exited
 	// was hijacking selection events selection and not updating components.
-	for (FSelectionIterator It(*Owner->GetSelectedActors()); It; ++It)
-	{
-		AActor* SelectedActor = CastChecked<AActor>(*It);
-		SelectedActor->MarkComponentsRenderStateDirty();
-	}
+	Owner->GetEditorSelectionSet()->ForEachSelectedObject<AActor>([](AActor* ActorPtr)
+		{
+			ActorPtr->MarkComponentsRenderStateDirty();
+			return true;
+		});
 
 	bPendingDeletion = false;
 
@@ -306,33 +306,26 @@ void UEdMode::DrawHUD(FEditorViewportClient* ViewportClient, FViewport* Viewport
 	}
 
 	const bool bLargeVertices = View->Family->EngineShowFlags.LargeVertices;
-	const bool bShowBrushes = View->Family->EngineShowFlags.Brushes;
-	const bool bShowBSP = View->Family->EngineShowFlags.BSP;
-	const bool bShowBuilderBrush = View->Family->EngineShowFlags.BuilderBrush != 0;
+	if (!bLargeVertices)
+	{
+		return;
+	}
 
+	// Temporaries.
 	UTexture2D* VertexTexture = GetVertexTexture();
 	const float TextureSizeX = VertexTexture->GetSizeX() * (bLargeVertices ? 1.0f : 0.5f);
 	const float TextureSizeY = VertexTexture->GetSizeY() * (bLargeVertices ? 1.0f : 0.5f);
 
-	// Temporaries.
-	TArray<FVector> Vertices;
-
-	for (FSelectionIterator It(*Owner->GetSelectedActors()); It; ++It)
-	{
-		AActor* SelectedActor = static_cast<AActor*>(*It);
-		checkSlow(SelectedActor->IsA(AActor::StaticClass()));
-
-		if (bLargeVertices)
+	// Static mesh vertices
+	Owner->GetEditorSelectionSet()->ForEachSelectedObject<AStaticMeshActor>([View, Canvas, VertexTexture, TextureSizeX, TextureSizeY, bIsHitTesting](AStaticMeshActor* Actor)
 		{
+			TArray<FVector> Vertices;
 			FCanvasItemTestbed::bTestState = !FCanvasItemTestbed::bTestState;
 
-			// Static mesh vertices
-			AStaticMeshActor* Actor = Cast<AStaticMeshActor>(SelectedActor);
-			if (Actor && Actor->GetStaticMeshComponent() && Actor->GetStaticMeshComponent()->GetStaticMesh()
+			if (Actor->GetStaticMeshComponent() && Actor->GetStaticMeshComponent()->GetStaticMesh()
 				&& Actor->GetStaticMeshComponent()->GetStaticMesh()->GetRenderData())
 			{
 				FTransform ActorToWorld = Actor->ActorToWorld();
-				Vertices.Empty();
 				const FPositionVertexBuffer& VertexBuffer = Actor->GetStaticMeshComponent()->GetStaticMesh()->GetRenderData()->LODResources[0].VertexBuffers.PositionVertexBuffer;
 				for (uint32 i = 0; i < VertexBuffer.GetNumVertices(); i++)
 				{
@@ -352,8 +345,8 @@ void UEdMode::DrawHUD(FEditorViewportClient* ViewportClient, FViewport* Viewport
 						PixelLocation *= InvDpiScale;
 
 						const bool bOutside =
-							PixelLocation.X < 0.0f || PixelLocation.X > View->UnscaledViewRect.Width()*InvDpiScale ||
-							PixelLocation.Y < 0.0f || PixelLocation.Y > View->UnscaledViewRect.Height()*InvDpiScale;
+							PixelLocation.X < 0.0f || PixelLocation.X > View->UnscaledViewRect.Width() * InvDpiScale ||
+							PixelLocation.Y < 0.0f || PixelLocation.Y > View->UnscaledViewRect.Height() * InvDpiScale;
 						if (!bOutside)
 						{
 							const float X = PixelLocation.X - (TextureSizeX / 2);
@@ -374,17 +367,13 @@ void UEdMode::DrawHUD(FEditorViewportClient* ViewportClient, FViewport* Viewport
 					}
 				}
 			}
-		}
-	}
+			return true;
+		});
 }
 
 void UEdMode::DrawBrackets(FEditorViewportClient* ViewportClient, FViewport* Viewport, const FSceneView* View, FCanvas* Canvas)
 {
-	USelection& SelectedActors = *Owner->GetSelectedActors();
-	for (int32 CurSelectedActorIndex = 0; CurSelectedActorIndex < SelectedActors.Num(); ++CurSelectedActorIndex)
-	{
-		AActor* SelectedActor = Cast<AActor>(SelectedActors.GetSelectedObject(CurSelectedActorIndex));
-		if (SelectedActor != NULL)
+	Owner->GetEditorSelectionSet()->ForEachSelectedObject<AActor>([ViewportClient, Canvas, View, Viewport](AActor* SelectedActor)
 		{
 			// Draw a bracket for selected "paintable" static mesh actors
 			const bool bIsValidActor = (Cast< AStaticMeshActor >(SelectedActor) != NULL);
@@ -392,8 +381,9 @@ void UEdMode::DrawBrackets(FEditorViewportClient* ViewportClient, FViewport* Vie
 			const FLinearColor SelectedActorBoxColor(0.6f, 0.6f, 1.0f);
 			const bool bDrawBracket = bIsValidActor;
 			ViewportClient->DrawActorScreenSpaceBoundingBox(Canvas, View, Viewport, SelectedActor, SelectedActorBoxColor, bDrawBracket);
-		}
-	}
+
+			return true;
+		});
 }
 
 bool UEdMode::UsesToolkits() const
@@ -423,7 +413,7 @@ bool UEdMode::EndTracking(FEditorViewportClient* InViewportClient, FViewport* In
 
 AActor* UEdMode::GetFirstSelectedActorInstance() const
 {
-	return Owner->GetSelectedActors()->GetTop<AActor>();
+	return Owner->GetEditorSelectionSet()->GetTopSelectedObject<AActor>();
 }
 
 UInteractiveToolManager* UEdMode::GetToolManager() const
