Fix padding around "Launch Game on Device" message

diff --git a/Engine/Source/Editor/UnrealEd/Private/Kismet2/DebuggerCommands.cpp b/Engine/Source/Editor/UnrealEd/Private/Kismet2/DebuggerCommands.cpp
--- a/Engine/Source/Editor/UnrealEd/Private/Kismet2/DebuggerCommands.cpp
+++ b/Engine/Source/Editor/UnrealEd/Private/Kismet2/DebuggerCommands.cpp
@@ -876,10 +876,14 @@ TSharedRef< SWidget > FPlayWorldCommands::GeneratePlayMenuContent(TSharedRef<FUI
 			Section.AddSeparator(NAME_None);
 			Section.AddEntry(FToolMenuEntry::InitWidget(
 				"PlayIn",
-				SNew(STextBlock)
-				.ColorAndOpacity(FSlateColor::UseSubduedForeground())
-				.Text(LOCTEXT("PlayInTip", "Launching a game (preview or on device) with a different mode will change your default 'Play' mode in the toolbar"))
-				.WrapTextAt(250),
+				SNew(SBox)
+				.Padding(FMargin(10.0f, 0.0f))
+				[
+					SNew(STextBlock)
+					.ColorAndOpacity(FSlateColor::UseSubduedForeground())
+					.Text(LOCTEXT("PlayInTip", "Launching a game (preview or on device) with a different mode will change your default 'Play' mode in the toolbar"))
+					.WrapTextAt(250)
+				],
 				FText::GetEmpty()));
 		}
 
