Check if the given asset editor is a blueprint editor before attempting to cast it, allowing for asset editors that are not IBlueprintEditors to call this util function.

PR #7750: Prevent crash with the custom editor operating on the project-specific UBlueprint class (Contributed by [REDACTED_PERSON])

diff --git a/Engine/Source/Editor/UnrealEd/Private/Kismet2/BlueprintEditorUtils.cpp b/Engine/Source/Editor/UnrealEd/Private/Kismet2/BlueprintEditorUtils.cpp
--- a/Engine/Source/Editor/UnrealEd/Private/Kismet2/BlueprintEditorUtils.cpp
+++ b/Engine/Source/Editor/UnrealEd/Private/Kismet2/BlueprintEditorUtils.cpp
@@ -2132,8 +2132,14 @@ void FBlueprintEditorUtils::MarkBlueprintAsModified(UBlueprint* Blueprint, FProp
 		IAssetEditorInstance* AssetEditor = GEditor->GetEditorSubsystem<UAssetEditorSubsystem>()->FindEditorForAsset(Blueprint, false);
 		if (AssetEditor)
 		{
-			FBlueprintEditor* BlueprintEditor = static_cast<FBlueprintEditor*>(AssetEditor);
-			BlueprintEditor->UpdateNodesUnrelatedStatesAfterGraphChange();
+			// Prevent crash with the custom editor operating on the project-specific UBlueprint class
+			// Such custom editor might not inherit after FBlueprintEditor, but could still utilize FBlueprintEditorUtils
+			FAssetEditorToolkit* AssetEditorToolkit = static_cast<FAssetEditorToolkit*>(AssetEditor);
+			if (AssetEditorToolkit->IsBlueprintEditor())
+			{
+				FBlueprintEditor* BlueprintEditor = static_cast<FBlueprintEditor*>(AssetEditor);
+				BlueprintEditor->UpdateNodesUnrelatedStatesAfterGraphChange();
+			}
 		}
 	}
 }
