Another attempt at eliminating the bogus warning.  PVS-Studio is ignoring the //-V522

#rnx

diff --git a/Engine/Source/Editor/UnrealEd/Private/Factories/Factory.cpp b/Engine/Source/Editor/UnrealEd/Private/Factories/Factory.cpp
--- a/Engine/Source/Editor/UnrealEd/Private/Factories/Factory.cpp
+++ b/Engine/Source/Editor/UnrealEd/Private/Factories/Factory.cpp
@@ -541,8 +541,10 @@ bool UFactory::ImportUntypedBulkDataFromText(const TCHAR*& Buffer, FUntypedBulkD
 				}
 				else if (FParse::Value(Str, TEXT("BEGIN "), ParsedText) && (ParsedText.ToUpper() == TEXT("BINARY")))
 				{
+					check(RawData);
+#ifndef PVS_STUDIO // Build machine refuses to disable to warning below
 					uint8* BulkDataPointer = RawData;
-					check(BulkDataPointer);
+#endif
 					while(FParse::Line(&Buffer,StrLine))
 					{
 						Str = *StrLine;
@@ -568,8 +570,10 @@ bool UFactory::ImportUntypedBulkDataFromText(const TCHAR*& Buffer, FUntypedBulkD
 								ParseStr +=2;
 							}
 							Value = FParse::HexDigit(ParseStr[0]) * 16 + FParse::HexDigit(ParseStr[1]);
-							*BulkDataPointer = (uint8)Value;
+#ifndef PVS_STUDIO // Build machine refuses to disable to warning below
+							*BulkDataPointer = (uint8)Value; //-V522
 							BulkDataPointer++;
+#endif
 							ParseStr += 2;
 							ParseStr++;
 						}
