Fix uninitialized variable warning.
#rnx

diff --git a/Engine/Source/Editor/UnrealEd/Private/EditorDomain/EditorDomainSave.cpp b/Engine/Source/Editor/UnrealEd/Private/EditorDomain/EditorDomainSave.cpp
--- a/Engine/Source/Editor/UnrealEd/Private/EditorDomain/EditorDomainSave.cpp
+++ b/Engine/Source/Editor/UnrealEd/Private/EditorDomain/EditorDomainSave.cpp
@@ -605,7 +605,7 @@ void FEditorDomainSaveServer::ProcessMessage(FClientConnection& ClientConnection
 	case EMessageType::Save:
 	{
 		FMemoryReader Reader(MessageBuffer);
-		int32 NumPackages;
+		int32 NumPackages = -1;
 		const int32 MinSizePerPackage = sizeof(int32) + 2;
 		Reader << NumPackages;
 		if (Reader.IsError() || NumPackages < 0 || NumPackages > MessageBuffer.Num() / MinSizePerPackage)
@@ -740,7 +740,10 @@ bool FEditorDomainSaveServer::TrySavePackage(const FPackagePath& PackagePath, FS
 	FString TempFilename = FPaths::Combine(FPaths::ProjectIntermediateDir(), FGuid::NewGuid().ToString());
 	ON_SCOPE_EXIT{ IFileManager::Get().Delete(*TempFilename); };
 
-	FSavePackageResultStruct Result = GEditor->Save(Package, nullptr, RF_Standalone, *TempFilename);
+	uint32 SaveFlags = SAVE_SaveBulkDataByReference;
+	FSavePackageResultStruct Result = GEditor->Save(Package, nullptr /* InBase */, RF_Standalone, *TempFilename,
+		GError, nullptr /* Conform */, false /* bForceByteSwapping */, false /* bWarnOfLongFilename */, SaveFlags , nullptr /* TargetPlatform */,
+		FDateTime::MinValue(), true /* bSlowTask */, nullptr /* InOutDiffMap */, nullptr /* SavePackageContext */);
 	if (Result.Result != ESavePackageResult::Success)
 	{
 		OutErrorMessage = FString::Printf(TEXT("Package %s could not be saved to temporary file %s: ResultCode = %d"),
