Revert accidentally submitted change.
#b None, trivial
#rnx

diff --git a/Engine/Source/Editor/UnrealEd/Private/EditorDomain/EditorDomainSave.cpp b/Engine/Source/Editor/UnrealEd/Private/EditorDomain/EditorDomainSave.cpp
--- a/Engine/Source/Editor/UnrealEd/Private/EditorDomain/EditorDomainSave.cpp
+++ b/Engine/Source/Editor/UnrealEd/Private/EditorDomain/EditorDomainSave.cpp
@@ -740,10 +740,7 @@ bool FEditorDomainSaveServer::TrySavePackage(const FPackagePath& PackagePath, FS
 	FString TempFilename = FPaths::Combine(FPaths::ProjectIntermediateDir(), FGuid::NewGuid().ToString());
 	ON_SCOPE_EXIT{ IFileManager::Get().Delete(*TempFilename); };
 
-	uint32 SaveFlags = SAVE_SaveBulkDataByReference;
-	FSavePackageResultStruct Result = GEditor->Save(Package, nullptr /* InBase */, RF_Standalone, *TempFilename,
-		GError, nullptr /* Conform */, false /* bForceByteSwapping */, false /* bWarnOfLongFilename */, SaveFlags , nullptr /* TargetPlatform */,
-		FDateTime::MinValue(), true /* bSlowTask */, nullptr /* InOutDiffMap */, nullptr /* SavePackageContext */);
+	FSavePackageResultStruct Result = GEditor->Save(Package, nullptr, RF_Standalone, *TempFilename);
 	if (Result.Result != ESavePackageResult::Success)
 	{
 		OutErrorMessage = FString::Printf(TEXT("Package %s could not be saved to temporary file %s: ResultCode = %d"),
