TL;DR

Improved behavior of PIE automation to deal with extremly long hitches and better represet the time a user is going to wait before the editor is usable.
Improved usability of BenchmarkBuild script

Changes

+ Added new automation command for waiting until average fps reaches a specified value. Using this with a conservative value is a better way of checking for interactivity than looping at shader compiles etc.
+ When Errors are added to a test as a result of logging, this itself is now logged as an error. This makes it clearer why a test may have failed despite its own checks having passed.
+ BenchmarkBuild can now just specify empty arguments for PIE/Cook tests (e.g. -PIE1_Args="") to allow comparisons between default/other.
+ Project.Maps.PIE can now take a "-maps" argument on the command line for the list of maps to test
+ Turned on AutomationController.SuppresLogErrors for Frosty to prevent errors causing errors to fail
+ Cleaned up naming for mulltiple-variations of PIE/Cook tests in BenchmarkBuild. Now use -PIE1Args, -PIE2Args etc

Fixes -

* Fixed issue in AutomationControllerManager where extreme hitching resulted in an editor test timing out due to outstanding pings being given a "TimeSince" that would exceed a timeout
* Fixed issue where bTreatLogErrorsAsTestErrors was (strangely) causing the editor to exit with a zero-error code when tests failed.
* Made FWaitForShadersToFinishCompiling automation task non-blocking to prevent timeouts. (It also checks texture compilation as the blocking version used to)
* Fixed BenchmarkBuild using the old System.Maps.PIE test that isn't enabled by default in projects
* Removed flush of shaders/textures in Project.Maps.PIE since this doesn't represent when the level is playable. Measuring this will be added as a separate test.
* Fixed issue where a user specifying a map with an extension would cause Project.Maps.PIE not to recognize the level had loaded

diff --git a/Engine/Source/Developer/AutomationController/Private/AutomationCommandline.cpp b/Engine/Source/Developer/AutomationController/Private/AutomationCommandline.cpp
--- a/Engine/Source/Developer/AutomationController/Private/AutomationCommandline.cpp
+++ b/Engine/Source/Developer/AutomationController/Private/AutomationCommandline.cpp
@@ -513,8 +513,7 @@ public:
 				{
 					if (!GIsCriticalError)
 					{
-						const UAutomationControllerSettings* Settings = GetDefault<UAutomationControllerSettings>();
-						if (AutomationController->ReportsHaveErrors() && Settings->bTreatLogErrorsAsTestErrors)
+						if (AutomationController->ReportsHaveErrors())
 						{
 							UE_LOG(LogAutomationCommandLine, Display, TEXT("Setting GIsCriticalError due to test failures (will cause non-zero exit code)."));
 							GIsCriticalError = AutomationController->ReportsHaveErrors();
diff --git a/Engine/Source/Developer/AutomationController/Private/AutomationControllerManger.cpp b/Engine/Source/Developer/AutomationController/Private/AutomationControllerManger.cpp
--- a/Engine/Source/Developer/AutomationController/Private/AutomationControllerManger.cpp
+++ b/Engine/Source/Developer/AutomationController/Private/AutomationControllerManger.cpp
@@ -941,81 +941,94 @@ void FAutomationControllerManager::AddPingResult(const FMessageAddress& Responde
 
 void FAutomationControllerManager::UpdateTests()
 {
-	CheckTestTimer += FPlatformTime::Seconds() - LastTimeUpdateTicked;
+	const double kIgnoreAsFrameHitchDelta = 2.0f;
+	const double TickDelta = FPlatformTime::Seconds() - LastTimeUpdateTicked;
+
 	LastTimeUpdateTicked = FPlatformTime::Seconds();
-	if (CheckTestTimer > CheckTestIntervalSeconds)
+
+	// If this tick was very long then don't check the health of tests. If we've been blocked for a while on a load and ping responses haven't yet been processed
+	// then otherwise might add a delta to the ping time that causes the test to look like it's timed out.
+	if (TickDelta < kIgnoreAsFrameHitchDelta)
 	{
-		for ( int32 Index = 0; Index < TestRunningArray.Num(); Index++ )
+		CheckTestTimer += TickDelta;
+		if (CheckTestTimer > CheckTestIntervalSeconds)
 		{
-			TestRunningArray[Index].LastPingTime += CheckTestTimer;
-
-			if (TestRunningArray[Index].LastPingTime > GameInstanceLostTimerSeconds)
+			for ( int32 Index = 0; Index < TestRunningArray.Num(); Index++ )
 			{
-				// Find the game session instance info
-				int32 ClusterIndex;
-				int32 DeviceIndex;
-				verify(DeviceClusterManager.FindDevice(TestRunningArray[Index].OwnerMessageAddress, ClusterIndex, DeviceIndex));
-				//verify this device thought it was busy
-				TSharedPtr <IAutomationReport> Report = DeviceClusterManager.GetTest(ClusterIndex, DeviceIndex);
-				check(Report.IsValid());
-				// A dummy array used to report the result
-
-				TArray<FString> EmptyStringArray;
-				TArray<FString> ErrorStringArray;
-				ErrorStringArray.Add(FString(TEXT("Failed")));
-				bHasErrors = true;
-				UE_LOG(LogAutomationController, Display, TEXT("Timeout hit. Nooooooo."));
-
-				FAutomationTestResults TestResults;
-				TestResults.State = EAutomationState::Fail;
-				TestResults.GameInstance = DeviceClusterManager.GetClusterDeviceName(ClusterIndex, DeviceIndex);
-				TestResults.AddEvent(FAutomationEvent(EAutomationEventType::Error, FString::Printf(TEXT("Timeout waiting for device %s"), *TestResults.GameInstance)));
-
-				// Set the results
-				Report->SetResults(ClusterIndex, CurrentTestPass, TestResults);
-				bTestResultsAvailable = true;
-
-				const FAutomationTestResults& FinalResults = Report->GetResults(ClusterIndex, CurrentTestPass);
+				TestRunningArray[Index].LastPingTime += CheckTestTimer;
 
-				// Gather all of the data relevant to this test for our json reporting.
-				CollectTestResults(Report, FinalResults);
-
-				// Disable the device in the cluster so it is not used again
-				DeviceClusterManager.DisableDevice(ClusterIndex, DeviceIndex);
-
-				// Remove the running test
-				TestRunningArray.RemoveAt(Index--);
-
-				// If there are no more devices, set the module state to disabled
-				if ( DeviceClusterManager.HasActiveDevice() == false )
-				{
-					// Process results first to write out the report
-					ProcessResults();
-
-					UE_LOG(LogAutomationController, Display, TEXT("Module disabled"));
-					SetControllerStatus(EAutomationControllerModuleState::Disabled);
-					ClusterDistributionMask = 0;
-				}
-				else
+				if (TestRunningArray[Index].LastPingTime > GameInstanceLostTimerSeconds)
 				{
-					UE_LOG(LogAutomationController, Display, TEXT("Module not disabled. Keep looking."));
-					// Remove the cluster from the mask if there are no active devices left
-					if ( DeviceClusterManager.GetNumActiveDevicesInCluster(ClusterIndex) == 0 )
+					// Find the game session instance info
+					int32 ClusterIndex;
+					int32 DeviceIndex;
+					verify(DeviceClusterManager.FindDevice(TestRunningArray[Index].OwnerMessageAddress, ClusterIndex, DeviceIndex));
+					//verify this device thought it was busy
+					TSharedPtr <IAutomationReport> Report = DeviceClusterManager.GetTest(ClusterIndex, DeviceIndex);
+					check(Report.IsValid());
+					// A dummy array used to report the result
+
+					TArray<FString> EmptyStringArray;
+					TArray<FString> ErrorStringArray;
+					ErrorStringArray.Add(FString(TEXT("Failed")));
+					bHasErrors = true;
+					UE_LOG(LogAutomationController, Display, TEXT("Timeout hit. Nooooooo."));
+
+					FAutomationTestResults TestResults;
+					TestResults.State = EAutomationState::Fail;
+					TestResults.GameInstance = DeviceClusterManager.GetClusterDeviceName(ClusterIndex, DeviceIndex);
+					TestResults.AddEvent(FAutomationEvent(EAutomationEventType::Error, FString::Printf(TEXT("Timeout waiting for device %s"), *TestResults.GameInstance)));
+
+					// Set the results
+					Report->SetResults(ClusterIndex, CurrentTestPass, TestResults);
+					bTestResultsAvailable = true;
+
+					const FAutomationTestResults& FinalResults = Report->GetResults(ClusterIndex, CurrentTestPass);
+
+					// Gather all of the data relevant to this test for our json reporting.
+					CollectTestResults(Report, FinalResults);
+
+					// Disable the device in the cluster so it is not used again
+					DeviceClusterManager.DisableDevice(ClusterIndex, DeviceIndex);
+
+					// Remove the running test
+					TestRunningArray.RemoveAt(Index--);
+
+					// If there are no more devices, set the module state to disabled
+					if ( DeviceClusterManager.HasActiveDevice() == false )
 					{
-						ClusterDistributionMask &= ~( 1 << ClusterIndex );
+						// Process results first to write out the report
+						ProcessResults();
+
+						UE_LOG(LogAutomationController, Display, TEXT("Module disabled"));
+						SetControllerStatus(EAutomationControllerModuleState::Disabled);
+						ClusterDistributionMask = 0;
 					}
-					if ( TestRunningArray.Num() == 0 )
+					else
 					{
-						SetControllerStatus(EAutomationControllerModuleState::Ready);
+						UE_LOG(LogAutomationController, Display, TEXT("Module not disabled. Keep looking."));
+						// Remove the cluster from the mask if there are no active devices left
+						if ( DeviceClusterManager.GetNumActiveDevicesInCluster(ClusterIndex) == 0 )
+						{
+							ClusterDistributionMask &= ~( 1 << ClusterIndex );
+						}
+						if ( TestRunningArray.Num() == 0 )
+						{
+							SetControllerStatus(EAutomationControllerModuleState::Ready);
+						}
 					}
 				}
+				else
+				{
+					MessageEndpoint->Send(new FAutomationWorkerPing(), TestRunningArray[Index].OwnerMessageAddress);
+				}
 			}
-			else
-			{
-				MessageEndpoint->Send(new FAutomationWorkerPing(), TestRunningArray[Index].OwnerMessageAddress);
-			}
+			CheckTestTimer = 0.f;
 		}
-		CheckTestTimer = 0.f;
+	}
+	else
+	{
+		UE_LOG(LogAutomationController, Log, TEXT("Ignoring very large delta of %.02f seconds in calls to FAutomationControllerManager::Tick() and not penalizing unresponsive tests"), TickDelta);
 	}
 }
 
diff --git a/Engine/Source/Editor/UnrealEd/Private/Tests/AutomationEditorCommon.cpp b/Engine/Source/Editor/UnrealEd/Private/Tests/AutomationEditorCommon.cpp
--- a/Engine/Source/Editor/UnrealEd/Private/Tests/AutomationEditorCommon.cpp
+++ b/Engine/Source/Editor/UnrealEd/Private/Tests/AutomationEditorCommon.cpp
@@ -39,6 +39,7 @@
 #include "Bookmarks/IBookmarkTypeTools.h"
 #include "GameMapsSettings.h"
 #include "Editor/EditorPerformanceSettings.h"
+#include "TextureCompiler.h"
 
 
 #define COOK_TIMEOUT 3600
@@ -833,7 +834,49 @@ bool FEditorLoadMap::Update()
 */
 bool FWaitForShadersToFinishCompiling::Update()
 {
-	UE_LOG(LogEditorAutomationTests, Log, TEXT("Waiting for %i shaders to finish."), GShaderCompilingManager->GetNumRemainingJobs());
+	static double TimeShadersFinishedCompiling = 0;
+	static double LastReportTime = FPlatformTime::Seconds();
+	const double TimeToWaitForJobs = 2.0;
+
+	bool ShadersCompiling = GShaderCompilingManager && GShaderCompilingManager->IsCompiling();
+	bool TexturesCompiling = FTextureCompilingManager::Get().GetNumRemainingTextures() > 0;
+
+	double TimeNow = FPlatformTime::Seconds();
+
+	if (ShadersCompiling || TexturesCompiling)
+	{
+		if (TimeNow - LastReportTime > 5.0)
+		{
+			LastReportTime = TimeNow;
+
+			if (ShadersCompiling)
+			{
+				UE_LOG(LogEditorAutomationTests, Log, TEXT("Waiting for %i shaders to finish."), GShaderCompilingManager->GetNumRemainingJobs() + GShaderCompilingManager->GetNumPendingJobs());
+			}
+
+			if (TexturesCompiling)
+			{
+				UE_LOG(LogEditorAutomationTests, Log, TEXT("Waiting for %i texures to finish."), FTextureCompilingManager::Get().GetNumRemainingTextures());
+			}
+		}
+
+		TimeShadersFinishedCompiling = 0;
+
+		return false;
+	}
+
+	// Current jobs are done, but things may still come in on subsequent frames..
+	if (TimeShadersFinishedCompiling == 0)
+	{
+		TimeShadersFinishedCompiling = FPlatformTime::Seconds();
+	}
+
+	if (FPlatformTime::Seconds() - TimeShadersFinishedCompiling < TimeToWaitForJobs)
+	{
+		return false;
+	}
+
+	// may not be necessary, but just double-check everything is finished and ready
 	GShaderCompilingManager->FinishAllCompilation();
 	UE_LOG(LogEditorAutomationTests, Log, TEXT("Done waiting for shaders to finish."));
 	return true;
@@ -977,22 +1020,22 @@ bool FWaitToFinishCookByTheBookCommand::Update()
 
 bool FDeleteDirCommand::Update()
 {
-	FString FullFolderPath = FPaths::ConvertRelativePathToFull(*InFolderLocation);
-	if ( IFileManager::Get().DirectoryExists(*FullFolderPath) )
-	{
-		IFileManager::Get().DeleteDirectory(*FullFolderPath, false, true);
-	}
-	return true;
+FString FullFolderPath = FPaths::ConvertRelativePathToFull(*InFolderLocation);
+if (IFileManager::Get().DirectoryExists(*FullFolderPath))
+{
+	IFileManager::Get().DeleteDirectory(*FullFolderPath, false, true);
+}
+return true;
 }
 
 bool FWaitToFinishBuildDeployCommand::Update()
 {
-	if ( GEditor->LauncherWorker->GetStatus() == ELauncherWorkerStatus::Completed )
+	if (GEditor->LauncherWorker->GetStatus() == ELauncherWorkerStatus::Completed)
 	{
 		UE_LOG(LogEditorAutomationTests, Log, TEXT("The build game and deploy operation has finished."));
 		return true;
 	}
-	else if ( GEditor->LauncherWorker->GetStatus() == ELauncherWorkerStatus::Canceled || GEditor->LauncherWorker->GetStatus() == ELauncherWorkerStatus::Canceling )
+	else if (GEditor->LauncherWorker->GetStatus() == ELauncherWorkerStatus::Canceled || GEditor->LauncherWorker->GetStatus() == ELauncherWorkerStatus::Canceling)
 	{
 		UE_LOG(LogEditorAutomationTests, Warning, TEXT("The build was canceled."));
 		return true;
@@ -1035,7 +1078,9 @@ bool FWaitForSpecifiedPIEMapToEndCommand::Update()
 		return true;
 	}
 
+	// remove any paths or extensions to match the name of the world
 	FString ShortMapName = FPackageName::GetShortName(MapName);
+	ShortMapName = FPaths::GetBaseFilename(ShortMapName);
 
 	// Handle both ways the user may have specified this
 	if (TestWorld->GetName() != ShortMapName)
@@ -1048,9 +1093,9 @@ bool FWaitForSpecifiedPIEMapToEndCommand::Update()
 
 // agrant-todo: Move this into BasicTests.cpp for 4.27
 /**
- * Generic Pie Test for projects. 
+ * Generic Pie Test for projects.
  * By default this test will PIE the lit of MapsToPIETest from automation settings. if that is empty it will PIE the default editor and game (if they're different)
- * maps. 
+ * maps.
  *
  * If the editor session was started with a map on the command line then that's the only map that will be PIE'd. This allows project to set up tests that PIE
  * a list of maps from an external source.
@@ -1063,7 +1108,7 @@ IMPLEMENT_SIMPLE_AUTOMATION_TEST(FProjectMapsPIETest, "Project.Maps.PIE", EAutom
  * @param Parameters - Unused for this test
  * @return	TRUE if the test was successful, FALSE otherwise
  */
-bool FProjectMapsPIETest::RunTest(const FString& Parameters)
+	bool FProjectMapsPIETest::RunTest(const FString& Parameters)
 {
 	UAutomationTestSettings const* AutomationTestSettings = GetDefault<UAutomationTestSettings>();
 	check(AutomationTestSettings);
@@ -1072,13 +1117,19 @@ bool FProjectMapsPIETest::RunTest(const FString& Parameters)
 
 	// If the user has specified a map on the command line then that is what we'll PIE
 
-	// Taken from FUnrealEdMisc::OnInit which determines if there's a map on the command line
 	const TCHAR* ParsedCmdLine = FCommandLine::Get();
-
 	FString ParsedMapName;
-	// Check the first arg if it's not a parameter
-	if (FParse::Token(ParsedCmdLine, ParsedMapName, false) && ParsedMapName.StartsWith(TEXT("-")) == false)
+
+	// If there is an explicit list of maps on the command line via -map or -maps the use those.
+	if (FParse::Value(FCommandLine::Get(), TEXT("-maps="), ParsedMapName) || FParse::Value(FCommandLine::Get(), TEXT("-map="), ParsedMapName))
+	{
+		ParsedMapName.ParseIntoArray(PIEMaps, TEXT("+"), true);
+
+		UE_LOG(LogEditorAutomationTests, Display, TEXT("Found Maps %s on command line. PIE Test will use these maps"), *ParsedMapName);
+	}
+	else if (FParse::Token(ParsedCmdLine, ParsedMapName, false) && ParsedMapName.StartsWith(TEXT("-")) == false)
 	{
+		// If the user specified a map as the first param after the project, we'll PIE that
 		FString InitialMapName;
 
 		// If the specified package exists
@@ -1087,11 +1138,11 @@ bool FProjectMapsPIETest::RunTest(const FString& Parameters)
 			FPaths::GetExtension(InitialMapName, /*bIncludeDot=*/true) == FPackageName::GetMapPackageExtension())
 		{
 			PIEMaps.Add(InitialMapName);
-			UE_LOG(LogEditorAutomationTests, Display, TEXT("Found Map %s on command line. PIE Test will be restructed to this map"), *InitialMapName);
+			UE_LOG(LogEditorAutomationTests, Display, TEXT("Found Map %s on command line. PIE Test will be restricted to this map"), *InitialMapName);
 		}
 	}
 
-	// If there was no command line map then default to the project settings
+	// Ok, at this point there were no command line maps so default to the project settings. We PIE the editor startup map and the game startup map
 	if (PIEMaps.Num() == 0)
 	{
 		// If the project has maps configured for PIE then use those
@@ -1124,7 +1175,7 @@ bool FProjectMapsPIETest::RunTest(const FString& Parameters)
 	// Uh-oh
 	if (PIEMaps.Num() == 0)
 	{
-		UE_LOG(LogEditorAutomationTests, Error, TEXT("No automation or default maps are configured for PIE!"));
+		UE_LOG(LogEditorAutomationTests, Fatal, TEXT("No automation or default maps are configured for PIE!"));
 	}
 
 	// Don't want these settings affecting metrics
@@ -1147,12 +1198,13 @@ bool FProjectMapsPIETest::RunTest(const FString& Parameters)
 		}		
 		
 		AddCommand(new FEditorAutomationLogCommand(FString::Printf(TEXT("LoadMap-Begin: %s"), *MapPackageName)));
-		AddCommand(new FEditorLoadMap(Map));
+		AddCommand(new FEditorLoadMap(MapPackageName));
+		AddCommand(new FWaitLatentCommand(1.0f));
 		AddCommand(new FEditorAutomationLogCommand(FString::Printf(TEXT("LoadMap-End: %s"), *MapPackageName)));
 		AddCommand(new FEditorAutomationLogCommand(FString::Printf(TEXT("PIE-Begin: %s"), *MapPackageName)));
 		AddCommand(new FStartPIECommand(false));
-		AddCommand(new FWaitForShadersToFinishCompiling());
 		AddCommand(new FWaitForSpecifiedMapToLoadCommand(MapPackageName));  // need at least some frames before starting & ending PIE
+		AddCommand(new FWaitForAverageFrameRate(5.0f));	// wait until the editor reaches something vaguely usable
 		AddCommand(new FWaitLatentCommand(AutomationTestSettings->PIETestDuration));
 		AddCommand(new FEndPlayMapCommand());
 		AddCommand(new FWaitForSpecifiedPIEMapToEndCommand(MapPackageName));  // need at least some frames before starting & ending PIE
diff --git a/Engine/Source/Programs/AutomationTool/Scripts/BenchmarkBuild/BenchmarkBuild.Automation.cs b/Engine/Source/Programs/AutomationTool/Scripts/BenchmarkBuild/BenchmarkBuild.Automation.cs
--- a/Engine/Source/Programs/AutomationTool/Scripts/BenchmarkBuild/BenchmarkBuild.Automation.cs
+++ b/Engine/Source/Programs/AutomationTool/Scripts/BenchmarkBuild/BenchmarkBuild.Automation.cs
@@ -32,7 +32,7 @@ namespace AutomationTool.Benchmark
 	[Help("cores=X+Y+Z", "Do noxge builds with these processor counts (default is Environment.ProcessorCount)")]
 	[Help("cook", "Do a cook for the specified platform")]
 	[Help("pie", "Launch the editor (only valid when -project is specified")]
-	[Help("map", "Map to PIE with (only valid when using a single project")]
+	[Help("maps", "Map to PIE with (only valid when using a single project")]
 	[Help("warmddc", "Cook / PIE with a warm DDC")]
 	[Help("hotddc", "Cook / PIE with a hot local DDC (an untimed pre-run is performed)")]
 	[Help("coldddc", "Cook / PIE with a cold local DDC (a temporary folder is used)")]
@@ -41,8 +41,8 @@ namespace AutomationTool.Benchmark
 	[Help("wait=<n>", "How many seconds to wait between each test)")]
 	[Help("filename", "Name/path of file to write CSV results to. If empty the local machine name will be used")]
 	[Help("noclean", "Don't build from clean. (Mostly just to speed things up when testing)")]
-	[Help("CookArgs=", "Extra args to use when cooking. -CookArgs1=\"-foo\" -CookArgs2=\"-bar\" will run two cooks with each argument set")]
-	[Help("PIEArgs=", "Extra args to use when running the editor. -PIEArgs1=\"-foo\" -PIEArgs2=\"-bar\" will run two PIE tests with each argument set")]
+	[Help("Cook[N]Args=", "Extra args to use when cooking. -Cook1Args=\"-foo\" -Cook2Args=\"-bar\" will run two cooks with each argument set. Use -CookArgs for a single cook")]
+	[Help("PIE[N]Args=", "Extra args to use when running the editor. -PIE1Args=\"-foo\" -PIE2Args=\"-bar\" will run three PIE tests with each argument set. Use -PIEArgs for a single PIE")]
 	class BenchmarkBuild : BuildCommand
 	{
 		class BenchmarkOptions : BuildCommand
@@ -69,7 +69,7 @@ namespace AutomationTool.Benchmark
 			// editor startup tests
 			public bool DoPIETests = false;
 
-			public IEnumerable<string> MapList = Enumerable.Empty<string>();
+			public IEnumerable<string> PIEMapList = Enumerable.Empty<string>();
 
 			// misc
 			public int Iterations = 1;
@@ -120,26 +120,25 @@ namespace AutomationTool.Benchmark
 				Iterations = ParseParamInt("Iterations", Iterations);
 				TimeBetweenTasks = ParseParamInt("Wait", TimeBetweenTasks);
 
-				// allow up to 10 cook & PIE variations. -CookArgs=etc -CookArgs1=etc2 etc
+				// allow up to 10 cook & PIE variations. -Cook1_Args=etc -Cook2_Args=etc2 etc
 				for (int i = 0; i < 10; i++)
 				{
 					string PostFix = i == 0 ? "" : i.ToString();
-					string CookParam = ParseParamValue("CookArgs" + PostFix, "");
+					string CookParam = ParseParamValue("Cook" + PostFix + "Args", null);
 
-					if (!string.IsNullOrEmpty(CookParam))
+					if (CookParam != null)
 					{
 						CookArgs.Add(CookParam);
 					}
 
-					string PIEParam = ParseParamValue("PIEArgs" + PostFix, "");
+					string PIEParam = ParseParamValue("PIE" + PostFix + "Args", null);
 
-					if (!string.IsNullOrEmpty(PIEParam))
+					if (PIEParam != null)
 					{
 						PIEArgs.Add(PIEParam);
 					}
 				}
-
-			
+							
 				FileName = ParseParamValue("filename", FileName);	
 
 				// Parse the project arg
@@ -198,11 +197,12 @@ namespace AutomationTool.Benchmark
 
 				// parse map args
 				{
-					string Arg = ParseParamValue("map", "");
+					string Arg = ParseParamValue("maps", "");
+					Arg = ParseParamValue("map", Arg);
 
 					if (!string.IsNullOrEmpty(Arg))
 					{
-						MapList = Arg.Split(new[] { '+', ',' }, StringSplitOptions.RemoveEmptyEntries);
+						PIEMapList = Arg.Split(new[] { '+', ',' }, StringSplitOptions.RemoveEmptyEntries);
 					}
 				}
 			}
@@ -441,49 +441,41 @@ namespace AutomationTool.Benchmark
 
 			List<BenchmarkTaskBase> NewTasks = new List<BenchmarkTaskBase>();
 
-			List<string> MapsToTest = InOptions.MapList.ToList();
+			string PIEArgs = "";
 
-			if (!MapsToTest.Any())
+			if (InOptions.PIEMapList.Any())
 			{
-				MapsToTest.Add("");
+				PIEArgs += string.Format("-map=\"{0}\"", string.Join("+", InOptions.PIEMapList));
 			}
-			
+						
 			IEnumerable<string> PIEVariations = InOptions.PIEArgs.Any() ? InOptions.PIEArgs : new List<string> { "" };
 
-			foreach (string PIEArgs in PIEVariations)
+			foreach (string VariationArgs in PIEVariations)
 			{
-				foreach (var Map in MapsToTest)
-				{
-					string FinalArgs = PIEArgs;
-
-					if (!string.IsNullOrEmpty(Map))
-					{
-						FinalArgs += " -map=" + Map;
-					}
+				string FinalArgs = PIEArgs + " " + VariationArgs;
 
-					// if no options assume warm
-					if (InOptions.DDCOptions == DDCTaskOptions.None || InOptions.DDCOptions.HasFlag(DDCTaskOptions.WarmDDC))
-					{
-						NewTasks.Add(new BenchmarkRunEditorTask(InProjectFile, DDCTaskOptions.WarmDDC, FinalArgs));
-					}
+				// if no options assume warm
+				if (InOptions.DDCOptions == DDCTaskOptions.None || InOptions.DDCOptions.HasFlag(DDCTaskOptions.WarmDDC))
+				{
+					NewTasks.Add(new BenchmarkRunEditorTask(InProjectFile, DDCTaskOptions.WarmDDC, FinalArgs));
+				}
 
-					// hot ddc
-					if (InOptions.DDCOptions.HasFlag(DDCTaskOptions.HotDDC))
-					{
-						NewTasks.Add(new BenchmarkRunEditorTask(InProjectFile, DDCTaskOptions.HotDDC, FinalArgs));
-					}
+				// hot ddc
+				if (InOptions.DDCOptions.HasFlag(DDCTaskOptions.HotDDC))
+				{
+					NewTasks.Add(new BenchmarkRunEditorTask(InProjectFile, DDCTaskOptions.HotDDC, FinalArgs));
+				}
 
-					// cold ddc
-					if (InOptions.DDCOptions.HasFlag(DDCTaskOptions.ColdDDC))
-					{
-						NewTasks.Add(new BenchmarkRunEditorTask(InProjectFile, DDCTaskOptions.ColdDDC, FinalArgs));
-					}
+				// cold ddc
+				if (InOptions.DDCOptions.HasFlag(DDCTaskOptions.ColdDDC))
+				{
+					NewTasks.Add(new BenchmarkRunEditorTask(InProjectFile, DDCTaskOptions.ColdDDC, FinalArgs));
+				}
 
-					// no shaders in the ddc
-					if (InOptions.DDCOptions.HasFlag(DDCTaskOptions.NoShaderDDC))
-					{
-						NewTasks.Add(new BenchmarkRunEditorTask(InProjectFile, DDCTaskOptions.NoShaderDDC, FinalArgs));
-					}					
+				// no shaders in the ddc
+				if (InOptions.DDCOptions.HasFlag(DDCTaskOptions.NoShaderDDC))
+				{
+					NewTasks.Add(new BenchmarkRunEditorTask(InProjectFile, DDCTaskOptions.NoShaderDDC, FinalArgs));
 				}
 			}
 		
diff --git a/Engine/Source/Programs/AutomationTool/Scripts/BenchmarkBuild/BenchmarkRunEditorTask.cs b/Engine/Source/Programs/AutomationTool/Scripts/BenchmarkBuild/BenchmarkRunEditorTask.cs
--- a/Engine/Source/Programs/AutomationTool/Scripts/BenchmarkBuild/BenchmarkRunEditorTask.cs
+++ b/Engine/Source/Programs/AutomationTool/Scripts/BenchmarkBuild/BenchmarkRunEditorTask.cs
@@ -85,7 +85,7 @@ namespace AutomationTool.Benchmark
 			string ProjectArg = ProjectFile != null ? ProjectFile.ToString() : "";
 			string EditorPath = HostPlatform.Current.GetUE4ExePath("UnrealEditor.exe");
 			string LogArg = string.Format("-log={0}.log", MakeValidFileName(GetFullTaskName()).Replace(" ", "_"));
-			string Arguments = string.Format("{0} {1} -execcmds=\"automation runtest System.Maps.PIE;Quit\" -stdout -FullStdOutLogOutput -unattended {2}", ProjectArg, EditorArgs, LogArg);
+			string Arguments = string.Format("{0} -execcmds=\"automation runtest Project.Maps.PIE;Quit\" {1} -stdout -FullStdOutLogOutput -unattended {2}", ProjectArg, EditorArgs, LogArg);
 
 			if (!bIsWarming && TaskOptions.HasFlag(DDCTaskOptions.NoShaderDDC))
 			{
@@ -105,9 +105,7 @@ namespace AutomationTool.Benchmark
 			LastStdOutTime = DateTime.Now;
 			CurrentProcess = CommandUtils.Run(EditorPath, Arguments, Options: RunOptions, SpewFilterCallback: SpewDelegate);
 
-			DateTime StartTime = DateTime.Now;
-
-			int MaxWaitMins = 90;
+			int TimeoutMins = 15;
 
 			while (!CurrentProcess.HasExited)
 			{
@@ -115,9 +113,9 @@ namespace AutomationTool.Benchmark
 
 				lock (CurrentProcess)
 				{
-					if ((LastStdOutTime - StartTime).TotalMinutes >= MaxWaitMins)
+					if ((DateTime.Now - LastStdOutTime).TotalMinutes >= TimeoutMins)
 					{
-						Log.TraceError("Gave up waiting for task after {0} minutes of no output", MaxWaitMins);
+						Log.TraceError("Gave up waiting for task after {0} minutes of no output", TimeoutMins);
 						CurrentProcess.ProcessObject.Kill();
 					}
 				}
diff --git a/Engine/Source/Runtime/Core/Private/Misc/AutomationTest.cpp b/Engine/Source/Runtime/Core/Private/Misc/AutomationTest.cpp
--- a/Engine/Source/Runtime/Core/Private/Misc/AutomationTest.cpp
+++ b/Engine/Source/Runtime/Core/Private/Misc/AutomationTest.cpp
@@ -101,6 +101,11 @@ void FAutomationTestFramework::FAutomationTestOutputDevice::Serialize( const TCH
 			// Errors
 			if (EffectiveVerbosity == ELogVerbosity::Error)
 			{
+				if (!LoggedFailureCause.Contains(CurTest))
+				{
+					CurTest->AddError(FString::Printf(TEXT("%s will be marked as failing due to errors being logged"), *CurTest->GetBeautifiedTestName()), STACK_OFFSET);
+					LoggedFailureCause.Add(CurTest);
+				}
 				CurTest->AddError(FString(V), STACK_OFFSET);
 			}
 			// Warnings
diff --git a/Engine/Source/Runtime/Core/Public/Misc/AutomationTest.h b/Engine/Source/Runtime/Core/Public/Misc/AutomationTest.h
--- a/Engine/Source/Runtime/Core/Public/Misc/AutomationTest.h
+++ b/Engine/Source/Runtime/Core/Public/Misc/AutomationTest.h
@@ -915,6 +915,9 @@ private:
 	private:
 		/** Associated automation test; all warnings, errors, etc. are routed to the automation test to track */
 		class FAutomationTestBase* CurTest;
+
+		/** Tests that we've logged the failure cause when an error is involved */
+		TSet<FAutomationTestBase*> LoggedFailureCause;
 	};
 
 	 /** Special feedback context used during automated testing to filter messages that happen during tests */
@@ -3356,3 +3359,4 @@ private:
 	TFunction<bool()> TimeoutCallback;
 	float Timeout;
 };
+
diff --git a/Engine/Source/Runtime/Engine/Private/Tests/AutomationCommon.cpp b/Engine/Source/Runtime/Engine/Private/Tests/AutomationCommon.cpp
--- a/Engine/Source/Runtime/Engine/Private/Tests/AutomationCommon.cpp
+++ b/Engine/Source/Runtime/Engine/Private/Tests/AutomationCommon.cpp
@@ -21,6 +21,7 @@
 #include "Matinee/MatineeActor.h"
 #include "StereoRendering.h"
 #include "Misc/PackageName.h"
+#include "TextureCompiler.h"
 
 #if WITH_AUTOMATION_TESTS
 
@@ -343,7 +344,10 @@ bool FWaitForSpecifiedMapToLoadCommand::Update()
 		AGameStateBase* GameState = TestWorld->GetGameState();
 		if (GameState && GameState->HasMatchStarted())
 		{
+			// remove any paths or extensions to match the name of the world
 			FString ShortMapName = FPackageName::GetShortName(MapName);
+			ShortMapName = FPaths::GetBaseFilename(ShortMapName);
+
 			// Handle both ways the user may have specified this
 			if (TestWorld->GetName() == ShortMapName)
 			{
@@ -355,6 +359,38 @@ bool FWaitForSpecifiedMapToLoadCommand::Update()
 	return false;
 }
 
+bool FWaitForAverageFrameRate::Update()
+{
+	if (StartTime == 0)
+	{
+		StartTime = FPlatformTime::Seconds();
+	}
+	else
+	{
+		const double ElapsedTime = FPlatformTime::Seconds() - StartTime;
+
+		if (ElapsedTime > Delay)
+		{
+			extern ENGINE_API float GAverageFPS;
+			if (GAverageFPS >= DesiredFrameRate)
+			{
+				return true;
+			}
+
+			if (ElapsedTime >= MaxWaitTime)
+			{
+				UE_LOG(LogEngineAutomationLatentCommand, Error, TEXT("FWaitForAverageFrameRate: Game did not reach %.02f FPS within %.02f seconds. Giving up."), DesiredFrameRate, MaxWaitTime);
+				
+				return true;
+			}
+		}
+	}
+
+	return false;
+}
+
+
+
 ///////////////////////////////////////////////////////////////////////
 // Common Latent commands which are used across test type. I.e. Engine, Network, etc...
 
@@ -501,12 +537,56 @@ bool FExecWorldStringLatentCommand::Update()
 */
 bool FWaitForShadersToFinishCompilingInGame::Update()
 {
-	if (GShaderCompilingManager)
+#if WITH_EDITOR
+	static double TimeShadersFinishedCompiling = 0;
+	static double LastReportTime = FPlatformTime::Seconds();
+	const double TimeToWaitForJobs = 2.0;
+
+	bool ShadersCompiling = GShaderCompilingManager && GShaderCompilingManager->IsCompiling();
+	bool TexturesCompiling = FTextureCompilingManager::Get().GetNumRemainingTextures() > 0;
+
+	
+	double TimeNow = FPlatformTime::Seconds();
+
+	if (ShadersCompiling || TexturesCompiling)
+	{
+		if (TimeNow - LastReportTime > 5.0)
+		{
+			LastReportTime = TimeNow;
+
+			if (ShadersCompiling)
+			{
+				UE_LOG(LogEditorAutomationTests, Log, TEXT("Waiting for %i shaders to finish."), GShaderCompilingManager->GetNumRemainingJobs() + GShaderCompilingManager->GetNumPendingJobs());
+			}
+
+			if (TexturesCompiling)
+			{
+				UE_LOG(LogEditorAutomationTests, Log, TEXT("Waiting for %i texures to finish."), FTextureCompilingManager::Get().GetNumRemainingTextures());
+			}
+		}
+
+		TimeShadersFinishedCompiling = 0;
+
+		return false;
+	}
+
+	// Current jobs are done, but things may still come in on subsequent frames..
+	if (TimeShadersFinishedCompiling == 0)
+	{
+		TimeShadersFinishedCompiling = FPlatformTime::Seconds();
+	}
+
+	if (FPlatformTime::Seconds() - TimeShadersFinishedCompiling < TimeToWaitForJobs)
 	{
-		UE_LOG(LogEditorAutomationTests, Log, TEXT("Waiting for %i shaders to finish."), GShaderCompilingManager->GetNumRemainingJobs());
-		GShaderCompilingManager->FinishAllCompilation();
-		UE_LOG(LogEditorAutomationTests, Log, TEXT("Done waiting for shaders to finish."));
+		return false;
 	}
+
+	// may not be necessary, but just double-check everything is finished and ready
+	GShaderCompilingManager->FinishAllCompilation();
+	UE_LOG(LogEditorAutomationTests, Log, TEXT("Done waiting for shaders to finish."));
+#endif
+
 	return true;
 }
+
 #endif //WITH_DEV_AUTOMATION_TESTS
diff --git a/Engine/Source/Runtime/Engine/Public/Tests/AutomationCommon.h b/Engine/Source/Runtime/Engine/Public/Tests/AutomationCommon.h
--- a/Engine/Source/Runtime/Engine/Public/Tests/AutomationCommon.h
+++ b/Engine/Source/Runtime/Engine/Public/Tests/AutomationCommon.h
@@ -117,6 +117,7 @@ DEFINE_ENGINE_LATENT_AUTOMATION_COMMAND(FWaitForMapToLoadCommand);
 */
 DEFINE_ENGINE_LATENT_AUTOMATION_COMMAND_ONE_PARAMETER(FWaitForSpecifiedMapToLoadCommand, FString, MapName);
 
+
 /**
 * Force a matinee to not loop and request that it play
 */
@@ -167,4 +168,39 @@ DEFINE_ENGINE_LATENT_AUTOMATION_COMMAND_ONE_PARAMETER(FExecWorldStringLatentComm
 */
 DEFINE_ENGINE_LATENT_AUTOMATION_COMMAND(FWaitForShadersToFinishCompilingInGame);
 
+
+/**
+ * Waits until the average framerate meets to exceeds the specified value. Mostly intended as a way to ensure that a level load etc has completed
+ * and an interactive framerate is present.
+ *
+ * InDelay is how long to wait before checking, InMaxWaitTIme is how long to wait before throwing an error
+ */
+class ENGINE_API FWaitForAverageFrameRate : public IAutomationLatentCommand
+{
+public:
+
+	FWaitForAverageFrameRate(float InDesiredFrameRate, float InDelay = 1, float InMaxWaitTime = 60)
+		: StartTime(0)
+		, DesiredFrameRate(InDesiredFrameRate)
+		, Delay(InDelay)
+		, MaxWaitTime(InMaxWaitTime)
+	{}
+
+	bool Update() override;
+
+private:
+
+	// Time we began executing
+	double StartTime;
+
+	// Framerate we want to see
+	float DesiredFrameRate;
+
+	// How long before we start loiking at FPS
+	float Delay;
+
+	// Max time so spend waiting
+	float MaxWaitTime;
+};
+
 #endif
