Fix regression: short package names should be handled correctly in MapsToCook.
#rnx

diff --git a/Engine/Source/Editor/UnrealEd/Private/CookOnTheFlyServer.cpp b/Engine/Source/Editor/UnrealEd/Private/CookOnTheFlyServer.cpp
--- a/Engine/Source/Editor/UnrealEd/Private/CookOnTheFlyServer.cpp
+++ b/Engine/Source/Editor/UnrealEd/Private/CookOnTheFlyServer.cpp
@@ -5690,7 +5690,13 @@ void UCookOnTheFlyServer::GenerateLongPackageNames(TArray<FName>& FilesInPath)
 		{
 			FString LongPackageName;
 			FPackageName::EErrorCode FailureReason;
-			if (FPackageName::TryConvertToMountedPath(FileInPath, nullptr /* LocalPath */, &LongPackageName, nullptr /* ObjectName */, nullptr /* SubObjectName */, nullptr /* Extension */, nullptr /* FlexNameType */, &FailureReason))
+			if (FPackageName::TryConvertToMountedPath(FileInPath, nullptr /* LocalPath */, &LongPackageName,
+				nullptr /* ObjectName */, nullptr /* SubObjectName */, nullptr /* Extension */,
+				nullptr /* FlexNameType */, &FailureReason)
+				||
+				(FPackageName::IsShortPackageName(FileInPath) &&
+					FPackageName::SearchForPackageOnDisk(FileInPath, &LongPackageName, nullptr))
+				)
 			{
 				const FName LongPackageFName(*LongPackageName);
 				bool bIsAlreadyAdded;
@@ -5702,7 +5708,8 @@ void UCookOnTheFlyServer::GenerateLongPackageNames(TArray<FName>& FilesInPath)
 			}
 			else
 			{
-				LogCookerMessage(FString::Printf(TEXT("Unable to generate long package name, %s"), *FileInPath, *FPackageName::FormatErrorAsString(FileInPath, FailureReason)), EMessageSeverity::Warning);
+				LogCookerMessage(FString::Printf(TEXT("Unable to generate long package name, %s"), *FileInPath,
+					*FPackageName::FormatErrorAsString(FileInPath, FailureReason)), EMessageSeverity::Warning);
 			}
 		}
 	}
