WorldPartitionConvertCommandlet: Added an option to exclude sublevels from being converted by the commandlet

diff --git a/Engine/Source/Editor/UnrealEd/Classes/Commandlets/WorldPartitionConvertCommandlet.h b/Engine/Source/Editor/UnrealEd/Classes/Commandlets/WorldPartitionConvertCommandlet.h
--- a/Engine/Source/Editor/UnrealEd/Classes/Commandlets/WorldPartitionConvertCommandlet.h
+++ b/Engine/Source/Editor/UnrealEd/Classes/Commandlets/WorldPartitionConvertCommandlet.h
@@ -29,7 +29,7 @@ struct FHLODLayerActorMapping
 	FString HLODLayer;
 };
 
-UCLASS(Config=Engine)
+UCLASS(Config = Engine)
 class UNREALED_API UWorldPartitionConvertCommandlet : public UCommandlet
 {
 	GENERATED_UCLASS_BODY()
@@ -46,6 +46,7 @@ private:
 protected:
 	virtual bool GetAdditionalLevelsToConvert(ULevel* Level, TArray<ULevel*>& SubLevels);
 	virtual bool PrepareStreamingLevelForConversion(const UWorldPartition* WorldPartition, ULevelStreaming* StreamingLevel);
+	virtual bool ShouldConvertStreamingLevel(ULevelStreaming* StreamingLevel);
 	virtual bool ShouldDeleteActor(AActor* Actor, bool bMainLevel) const;
 	virtual void PerformAdditionalWorldCleanup(UWorld* World) const;
 	virtual void OutputConversionReport() const;
@@ -100,6 +101,10 @@ protected:
 	UPROPERTY(Config)
 	TMap<FName, EActorGridPlacement> LevelsGridPlacement;
 
+	// Levels excluded from conversion.
+	UPROPERTY(Config)
+	TArray<FString> ExcludedLevels;
+
 	UPROPERTY(Config)
 	bool bConvertActorsNotReferencedByLevelScript;
 
diff --git a/Engine/Source/Editor/UnrealEd/Private/Commandlets/WorldPartitionConvertCommandlet.cpp b/Engine/Source/Editor/UnrealEd/Private/Commandlets/WorldPartitionConvertCommandlet.cpp
--- a/Engine/Source/Editor/UnrealEd/Private/Commandlets/WorldPartitionConvertCommandlet.cpp
+++ b/Engine/Source/Editor/UnrealEd/Private/Commandlets/WorldPartitionConvertCommandlet.cpp
@@ -248,16 +248,25 @@ void UWorldPartitionConvertCommandlet::GatherAndPrepareSubLevelsToConvert(const
 	UWorld* World = Level->GetTypedOuter<UWorld>();	
 
 	// Set all streaming levels to be loaded/visible for next Flush
+	TArray<ULevelStreaming*> StreamingLevels;
 	for (ULevelStreaming* StreamingLevel : World->GetStreamingLevels())
 	{
-		StreamingLevel->SetShouldBeLoaded(true);
-		StreamingLevel->SetShouldBeVisible(true);
-		StreamingLevel->SetShouldBeVisibleInEditor(true);
+		if (ShouldConvertStreamingLevel(StreamingLevel))
+		{
+			StreamingLevels.Add(StreamingLevel);
+			StreamingLevel->SetShouldBeLoaded(true);
+			StreamingLevel->SetShouldBeVisible(true);
+			StreamingLevel->SetShouldBeVisibleInEditor(true);
+		}
+		else
+		{
+			UE_LOG(LogWorldPartitionConvertCommandlet, Log, TEXT("Skipping conversion of streaming Level %s"), *StreamingLevel->GetWorldAssetPackageName());
+		}
 	}
 
 	World->FlushLevelStreaming(EFlushLevelStreamingType::Full);
 	
-	for(ULevelStreaming* StreamingLevel: World->GetStreamingLevels())
+	for(ULevelStreaming* StreamingLevel: StreamingLevels)
 	{
 		if (PrepareStreamingLevelForConversion(WorldPartition, StreamingLevel))
 		{
@@ -1480,3 +1489,9 @@ int32 UWorldPartitionConvertCommandlet::Main(const FString& Params)
 
 	return 0;
 }
+
+bool UWorldPartitionConvertCommandlet::ShouldConvertStreamingLevel(ULevelStreaming* StreamingLevel)
+{
+	return StreamingLevel && !ExcludedLevels.Contains(StreamingLevel->GetWorldAssetPackageName());
+}
+
