wrap pvs warning in `#ifndef PVS_STUDIO` since the disable warning comment isn't properly working

diff --git a/Engine/Source/Editor/UnrealEd/Private/ObjectTools.cpp b/Engine/Source/Editor/UnrealEd/Private/ObjectTools.cpp
--- a/Engine/Source/Editor/UnrealEd/Private/ObjectTools.cpp
+++ b/Engine/Source/Editor/UnrealEd/Private/ObjectTools.cpp
@@ -3855,8 +3855,11 @@ namespace ObjectTools
 
 	bool RenameObjects( const TArray< UObject* >& SelectedObjects, bool bIncludeLocInstances, const FString& SourcePath, const FString& DestinationPath, bool bOpenDialog )
 	{
+	// seems like bug in pvs makes disabling the warning not work as expected
+	#ifndef PVS_STUDIO
 		// @todo asset: Find a proper location for localized files
 		bIncludeLocInstances = false; //-V763
+	#endif
 		if( !bIncludeLocInstances )
 		{
 			return RenameObjectsInternal( SelectedObjects, bIncludeLocInstances, NULL, SourcePath, DestinationPath, bOpenDialog );
