Added export of UWorldPartitionBuilder

diff --git a/Engine/Source/Editor/UnrealEd/Public/WorldPartition/WorldPartitionBuilder.h b/Engine/Source/Editor/UnrealEd/Public/WorldPartition/WorldPartitionBuilder.h
--- a/Engine/Source/Editor/UnrealEd/Public/WorldPartition/WorldPartitionBuilder.h
+++ b/Engine/Source/Editor/UnrealEd/Public/WorldPartition/WorldPartitionBuilder.h
@@ -6,7 +6,7 @@
 #include "WorldPartitionBuilder.generated.h"
 
 UCLASS(Abstract, Config=Engine)
-class UWorldPartitionBuilder : public UObject
+class UNREALED_API UWorldPartitionBuilder : public UObject
 {
 	GENERATED_UCLASS_BODY()
 
