Expose UDefaultEdMode so that it can be inherited from in the case we want default level editor gizmos, instead of relying on the legacy widget ed mode.

diff --git a/Engine/Source/Editor/UnrealEd/Private/Tools/DefaultEdMode.h b/Engine/Source/Editor/UnrealEd/Private/Tools/DefaultEdMode.h
deleted file mode 100644
diff --git a/Engine/Source/Editor/UnrealEd/Public/Tools/DefaultEdMode.h b/Engine/Source/Editor/UnrealEd/Public/Tools/DefaultEdMode.h
new file mode 100644
--- /dev/null
+++ b/Engine/Source/Editor/UnrealEd/Public/Tools/DefaultEdMode.h
@@ -0,0 +1,22 @@
+// Copyright Epic Games, Inc. All Rights Reserved.
+
+#pragma once
+
+#include "Tools/LegacyEdModeWidgetHelpers.h"
+
+#include "DefaultEdMode.generated.h"
+
+UCLASS(Transient)
+class UNREALED_API UEdModeDefault : public UBaseLegacyWidgetEdMode
+{
+	GENERATED_BODY()
+
+public:
+	UEdModeDefault();
+
+	virtual bool UsesPropertyWidgets() const override;
+	virtual bool UsesToolkits() const override;
+
+protected:
+	virtual TSharedRef<FLegacyEdModeWidgetHelper> CreateWidgetHelper() override;
+};
