Missing ShouldIgnorePackageFunction call on External Packages was causing out of context packages to be saved

diff --git a/Engine/Source/Editor/UnrealEd/Private/FileHelpers.cpp b/Engine/Source/Editor/UnrealEd/Private/FileHelpers.cpp
--- a/Engine/Source/Editor/UnrealEd/Private/FileHelpers.cpp
+++ b/Engine/Source/Editor/UnrealEd/Private/FileHelpers.cpp
@@ -4329,7 +4329,7 @@ void FEditorFileUtils::GetDirtyWorldPackages(TArray<UPackage*>& OutDirtyPackages
 			{
 				for (UPackage* ExternalPackage : WorldIt->PersistentLevel->GetLoadedExternalActorPackages())
 				{
-					if (ExternalPackage->IsDirty())
+					if (ExternalPackage->IsDirty() && !ShouldIgnorePackageFunction(ExternalPackage))
 					{
 						bool bActorPackageNeedsToSave = true;
 
