Fix a lingering strong ref count in the asset editor viewport layouts.

diff --git a/Engine/Source/Editor/UnrealEd/Private/AssetEditorViewportLayout.cpp b/Engine/Source/Editor/UnrealEd/Private/AssetEditorViewportLayout.cpp
--- a/Engine/Source/Editor/UnrealEd/Private/AssetEditorViewportLayout.cpp
+++ b/Engine/Source/Editor/UnrealEd/Private/AssetEditorViewportLayout.cpp
@@ -89,7 +89,7 @@ const FVector2D& SAssetEditorViewportsOverlay::GetCachedSize() const
 	return CachedSize;
 }
 
-TSharedPtr<FViewportTabContent> SAssetEditorViewportsOverlay::GetViewportTab() const
+TWeakPtr<FViewportTabContent> SAssetEditorViewportsOverlay::GetViewportTab() const
 {
 	return ViewportTab;
 }
diff --git a/Engine/Source/Editor/UnrealEd/Public/AssetEditorViewportLayout.h b/Engine/Source/Editor/UnrealEd/Public/AssetEditorViewportLayout.h
--- a/Engine/Source/Editor/UnrealEd/Public/AssetEditorViewportLayout.h
+++ b/Engine/Source/Editor/UnrealEd/Public/AssetEditorViewportLayout.h
@@ -72,7 +72,7 @@ public:
 
 	SLATE_BEGIN_ARGS(SAssetEditorViewportsOverlay) {}
 	SLATE_DEFAULT_SLOT(FArguments, Content)
-		SLATE_ARGUMENT(TSharedPtr<FViewportTabContent>, ViewportTab)
+		SLATE_ARGUMENT(TWeakPtr<FViewportTabContent>, ViewportTab)
 		SLATE_END_ARGS()
 
 		void Construct(const FArguments& InArgs);
@@ -99,12 +99,12 @@ public:
 	const FVector2D& GetCachedSize() const;
 
 	/** Gets the  Viewport Tab that created this overlay */
-	TSharedPtr<FViewportTabContent> GetViewportTab() const;
+	TWeakPtr<FViewportTabContent> GetViewportTab() const;
 
 private:
 
 	/** Reference to the owning  viewport tab */
-	TSharedPtr<FViewportTabContent> ViewportTab;
+	TWeakPtr<FViewportTabContent> ViewportTab;
 
 	/** The overlay widget we're containing */
 	TSharedPtr< SOverlay > OverlayWidget;
