Fix double drop of static mesh on drag + drop from content browser.

diff --git a/Engine/Source/Editor/UnrealEd/Private/AssetSelection.cpp b/Engine/Source/Editor/UnrealEd/Private/AssetSelection.cpp
--- a/Engine/Source/Editor/UnrealEd/Private/AssetSelection.cpp
+++ b/Engine/Source/Editor/UnrealEd/Private/AssetSelection.cpp
@@ -618,7 +618,11 @@ static AActor* PrivateAddActor( UObject* Asset, UActorFactory* Factory, bool Sel
 			PlacementInfo.FinalizedTransform = ActorTransform;
 			PlacementInfo.FactoryOverride = Factory;
 
-			TArray<FTypedElementHandle> PlacedElements = PlacementSubsystem->PlaceAsset(PlacementInfo, FPlacementOptions());
+			FPlacementOptions PlacementOptions;
+			PlacementOptions.bIsCreatingPreviewElements = FLevelEditorViewportClient::IsDroppingPreviewActor();
+			PlacementOptions.bPreferInstancedPlacement = false;
+
+			TArray<FTypedElementHandle> PlacedElements = PlacementSubsystem->PlaceAsset(PlacementInfo, PlacementOptions);
 			if (PlacedElements.Num())
 			{
 				if (TTypedElement<UTypedElementObjectInterface> ObjectInterface = UTypedElementRegistry::GetInstance()->GetElement<UTypedElementObjectInterface>(PlacedElements[0]))
