Cached the elements to manipulate until the selection set actually changes

This avoids the calculation being made repeatedly when nothing has changed (eg, it used to be cached 3-4 times during a selection change event), as well as not needing to calculate it each frame (eg, for the gizmo visibility, or during a drag operation)

diff --git a/Engine/Source/Editor/UnrealEd/Private/EditorSelectUtils.cpp b/Engine/Source/Editor/UnrealEd/Private/EditorSelectUtils.cpp
--- a/Engine/Source/Editor/UnrealEd/Private/EditorSelectUtils.cpp
+++ b/Engine/Source/Editor/UnrealEd/Private/EditorSelectUtils.cpp
@@ -29,7 +29,6 @@
 #include "ActorGroupingUtils.h"
 #include "Subsystems/BrushEditingSubsystem.h"
 #include "Elements/Framework/TypedElementList.h"
-#include "Elements/Framework/TypedElementRegistry.h"
 #include "Elements/Framework/TypedElementSelectionSet.h"
 #include "Elements/Framework/EngineElementsLibrary.h"
 #include "Elements/Interfaces/TypedElementWorldInterface.h"
@@ -204,46 +203,41 @@ void UUnrealEdEngine::SetPivot( FVector NewPivot, bool bSnapPivotToGrid, bool bI
 	FVector TranslateRotateWidgetWorldXAxis = FVector::ZeroVector;
 	FVector Widget2DWorldXAxis = FVector::ZeroVector;
 
-	UTypedElementRegistry* Registry = UTypedElementRegistry::GetInstance();
-
 	// Pick a new common pivot, or not.
 	TTypedElement<UTypedElementWorldInterface> SingleWorldElement;
+
+	const UTypedElementList* ElementsToManipulate = GCurrentLevelEditingViewportClient->GetElementsToManipulate();
+	ElementsToManipulate->ForEachElement<UTypedElementWorldInterface>([&NumElements, &TranslateRotateWidgetWorldXAxis, &Widget2DWorldXAxis, &SingleWorldElement](const TTypedElement<UTypedElementWorldInterface>& InWorldElement)
 	{
-		const TArray<FTypedElementHandle> ElementsToManipulate = GCurrentLevelEditingViewportClient->GetElementsToManipulate();
-		for (const FTypedElementHandle& Element : ElementsToManipulate)
+		if (NumElements == 0)
 		{
-			if (TTypedElement<UTypedElementWorldInterface> WorldElement = Registry->GetElement<UTypedElementWorldInterface>(Element))
-			{
-				if (NumElements == 0)
-				{
-					FTransform ElementWorldTransform;
-					WorldElement.GetWorldTransform(ElementWorldTransform);
-
-					TranslateRotateWidgetWorldXAxis = ElementWorldTransform.TransformVector(FVector(1.0f, 0.0f, 0.0f));
-					//get the xy plane project of this vector
-					TranslateRotateWidgetWorldXAxis.Z = 0.0f;
-					if (!TranslateRotateWidgetWorldXAxis.Normalize())
-					{
-						TranslateRotateWidgetWorldXAxis = FVector(1.0f, 0.0f, 0.0f);
-					}
+			FTransform ElementWorldTransform;
+			InWorldElement.GetWorldTransform(ElementWorldTransform);
 
-					Widget2DWorldXAxis = ElementWorldTransform.TransformVector(FVector(1, 0, 0));
-					Widget2DWorldXAxis.Y = 0;
-					if (!Widget2DWorldXAxis.Normalize())
-					{
-						Widget2DWorldXAxis = FVector(1, 0, 0);
-					}
-				}
+			TranslateRotateWidgetWorldXAxis = ElementWorldTransform.TransformVector(FVector(1.0f, 0.0f, 0.0f));
+			//get the xy plane project of this vector
+			TranslateRotateWidgetWorldXAxis.Z = 0.0f;
+			if (!TranslateRotateWidgetWorldXAxis.Normalize())
+			{
+				TranslateRotateWidgetWorldXAxis = FVector(1.0f, 0.0f, 0.0f);
+			}
 
-				SingleWorldElement = MoveTemp(WorldElement);
-				++NumElements;
+			Widget2DWorldXAxis = ElementWorldTransform.TransformVector(FVector(1, 0, 0));
+			Widget2DWorldXAxis.Y = 0;
+			if (!Widget2DWorldXAxis.Normalize())
+			{
+				Widget2DWorldXAxis = FVector(1, 0, 0);
 			}
 		}
-	}
+
+		SingleWorldElement = InWorldElement;
+		++NumElements;
+		return true;
+	});
 	
 	if (bAssignPivot && UActorGroupingUtils::IsGroupingActive())
 	{
-		if (TTypedElement<UTypedElementObjectInterface> ObjectElement = Registry->GetElement<UTypedElementObjectInterface>(SingleWorldElement))
+		if (TTypedElement<UTypedElementObjectInterface> ObjectElement = ElementsToManipulate->GetElement<UTypedElementObjectInterface>(SingleWorldElement))
 		{
 			if (AActor* SingleActor = Cast<AActor>(ObjectElement.GetObject()))
 			{
@@ -331,29 +325,24 @@ void UUnrealEdEngine::UpdatePivotLocationForSelection( bool bOnChange )
 		return;
 	}
 
-	UTypedElementRegistry* Registry = UTypedElementRegistry::GetInstance();
-
 	// Pick a new common pivot, or not.
 	TTypedElement<UTypedElementWorldInterface> SingleWorldElement;
+	
+	const UTypedElementList* ElementsToManipulate = GCurrentLevelEditingViewportClient->GetElementsToManipulate();
+	ElementsToManipulate->ForEachElement<UTypedElementWorldInterface>([&SingleWorldElement](const TTypedElement<UTypedElementWorldInterface>& InWorldElement)
 	{
-		const TArray<FTypedElementHandle> ElementsToManipulate = GCurrentLevelEditingViewportClient->GetElementsToManipulate();
-		for (const FTypedElementHandle& Element : ElementsToManipulate)
-		{
-			if (TTypedElement<UTypedElementWorldInterface> WorldElement = Registry->GetElement<UTypedElementWorldInterface>(Element))
-			{
 #if DO_CHECK
-				{
-					ULevel* OwnerLevel = WorldElement.GetOwnerLevel();
-					const bool bIsTemplate = WorldElement.IsTemplateElement();
-					const bool bLevelLocked = FLevelUtils::IsLevelLocked(OwnerLevel);
-					check(bIsTemplate || !bLevelLocked);
-				}
+		{
+			ULevel* OwnerLevel = InWorldElement.GetOwnerLevel();
+			const bool bIsTemplate = InWorldElement.IsTemplateElement();
+			const bool bLevelLocked = FLevelUtils::IsLevelLocked(OwnerLevel);
+			check(bIsTemplate || !bLevelLocked);
+		}
 #endif	// DO_CHECK
 
-				SingleWorldElement = MoveTemp(WorldElement);
-			}
-		}
-	}
+		SingleWorldElement = InWorldElement;
+		return true;
+	});
 	
 	if (SingleWorldElement)
 	{
@@ -376,7 +365,7 @@ void UUnrealEdEngine::UpdatePivotLocationForSelection( bool bOnChange )
 			// If grouping is active, see if this element is an actor that's part of a locked group and use that pivot instead
 			if (UActorGroupingUtils::IsGroupingActive())
 			{
-				if (TTypedElement<UTypedElementObjectInterface> ObjectElement = Registry->GetElement<UTypedElementObjectInterface>(SingleWorldElement))
+				if (TTypedElement<UTypedElementObjectInterface> ObjectElement = ElementsToManipulate->GetElement<UTypedElementObjectInterface>(SingleWorldElement))
 				{
 					if (AActor* SingleActor = Cast<AActor>(ObjectElement.GetObject()))
 					{
diff --git a/Engine/Source/Editor/UnrealEd/Private/LevelEditorViewport.cpp b/Engine/Source/Editor/UnrealEd/Private/LevelEditorViewport.cpp
--- a/Engine/Source/Editor/UnrealEd/Private/LevelEditorViewport.cpp
+++ b/Engine/Source/Editor/UnrealEd/Private/LevelEditorViewport.cpp
@@ -1709,7 +1709,7 @@ FLevelEditorViewportClient::FLevelEditorViewportClient(const TSharedPtr<SLevelVi
 	, bNeedToRestoreComponentBeingMovedFlag(false)
 	, bReceivedFocusRecently(false)
 	, bAlwaysShowModeWidgetAfterSelectionChanges(true)
-	, ElementsToManipulate(UTypedElementRegistry::GetInstance()->CreateElementList())
+	, CachedElementsToManipulate(UTypedElementRegistry::GetInstance()->CreateElementList())
 	, SpriteCategoryVisibility()
 	, World(nullptr)
 	, TrackingTransaction()
@@ -1759,6 +1759,12 @@ FLevelEditorViewportClient::FLevelEditorViewportClient(const TSharedPtr<SLevelVi
 
 FLevelEditorViewportClient::~FLevelEditorViewportClient()
 {
+	if (UTypedElementSelectionSet* SelectionSet = GetMutableSelectionSet())
+	{
+		SelectionSet->OnPreChange().RemoveAll(this);
+	}
+	ResetElementsToManipulate();
+
 	// Unregister for all global callbacks to this object
 	FEditorSupportDelegates::CleanseEditor.RemoveAll(this);
 	FEditorDelegates::PreBeginPIE.RemoveAll(this);
@@ -1814,6 +1820,11 @@ void FLevelEditorViewportClient::InitializeVisibilityFlags()
 
 void FLevelEditorViewportClient::InitializeViewportInteraction()
 {
+	{
+		UTypedElementSelectionSet* SelectionSet = GetMutableSelectionSet();
+		SelectionSet->OnPreChange().AddRaw(this, &FLevelEditorViewportClient::ResetElementsToManipulateFromSelectionChange);
+	}
+
 	{
 		TUniquePtr<FActorElementLevelEditorViewportInteractionCustomization> ActorCustomization = MakeUnique<FActorElementLevelEditorViewportInteractionCustomization>();
 		ActorCustomization->SetLevelEditorViewportClient(this);
@@ -2581,9 +2592,8 @@ bool FLevelEditorViewportClient::InputWidgetDelta(FViewport* InViewport, EAxisLi
 							UTypedElementCommonActions* CommonActions = LevelEditor ? LevelEditor->GetCommonActions() : nullptr;
 							if (CommonActions)
 							{
-								CacheElementsToManipulate();
+								const UTypedElementList* ElementsToManipulate = GetElementsToManipulate();
 								TArray<FTypedElementHandle> DuplicatedElements = CommonActions->DuplicateElements(ElementsToManipulate, GetWorld(), /*bOffsetLocations*/false);
-								ElementsToManipulate->Reset();
 
 								// Exclusively select the new elements, so that future gizmo interaction manipulates those items instead
 								if (DuplicatedElements.Num() > 0)
@@ -2594,6 +2604,9 @@ bool FLevelEditorViewportClient::InputWidgetDelta(FViewport* InViewport, EAxisLi
 										.SetAllowLegacyNotifications(false); // Old drag duplicate code didn't used to notify about this selection change
 
 									SelectionSet->SetSelection(DuplicatedElements, SelectionOptions);
+
+									// Force the cached manipulation list to update, as we don't notify about the change above
+									ResetElementsToManipulate();
 								}
 
 								RedrawAllViewportsIntoThisScene();
@@ -2917,11 +2930,12 @@ void FLevelEditorViewportClient::TrackingStarted( const FInputEventState& InInpu
 	bNeedToRestoreComponentBeingMovedFlag = false;
 
 	PreDragActorTransforms.Empty();
-	CacheElementsToManipulate();
 
 	// Track BSP changes
 	{
 		bIsTrackingBrushModification = false;
+
+		const UTypedElementList* ElementsToManipulate = GetElementsToManipulate();
 		TypedElementListObjectUtil::ForEachObject<ABrush>(ElementsToManipulate, [this](ABrush* InBrush)
 		{
 			bIsTrackingBrushModification = InBrush && !InBrush->IsVolumeBrush() && !FActorEditorUtils::IsABuilderBrush(InBrush);
@@ -2933,6 +2947,7 @@ void FLevelEditorViewportClient::TrackingStarted( const FInputEventState& InInpu
 	{
 		Widget->SetSnapEnabled(true);
 
+		const UTypedElementList* ElementsToManipulate = GetElementsToManipulate();
 		ViewportInteraction->BeginGizmoManipulation(ElementsToManipulate, GetWidgetMode());
 
 		if (!bDuplicateActorsInProgress)
@@ -3009,8 +3024,6 @@ void FLevelEditorViewportClient::TrackingStarted( const FInputEventState& InInpu
 			GEditor->DisableDeltaModification(true);
 		}
 	}
-
-	ElementsToManipulate->Reset();
 }
 
 void FLevelEditorViewportClient::TrackingStopped()
@@ -3057,7 +3070,7 @@ void FLevelEditorViewportClient::TrackingStopped()
 	// Don't do this if AddDelta was never called.
 	if( bDidAnythingActuallyChange && MouseDeltaTracker->HasReceivedDelta() )
 	{
-		CacheElementsToManipulate();
+		const UTypedElementList* ElementsToManipulate = GetElementsToManipulate();
 		ViewportInteraction->EndGizmoManipulation(ElementsToManipulate, GetWidgetMode());
 
 		if (!GUnrealEd->IsPivotMovedIndependently())
@@ -3068,6 +3081,7 @@ void FLevelEditorViewportClient::TrackingStopped()
 
 	if (bNeedToRestoreComponentBeingMovedFlag)
 	{
+		const UTypedElementList* ElementsToManipulate = GetElementsToManipulate();
 		TypedElementListObjectUtil::ForEachObject<AActor>(ElementsToManipulate, [this](AActor* InActor)
 		{
 			SetActorBeingMovedByEditor(InActor, false);
@@ -3104,7 +3118,6 @@ void FLevelEditorViewportClient::TrackingStopped()
 	}
 
 	PreDragActorTransforms.Empty();
-	ElementsToManipulate->Reset();
 }
 
 void FLevelEditorViewportClient::AbortTracking()
@@ -3471,28 +3484,40 @@ bool FLevelEditorViewportClient::HaveSelectedObjectsBeenChanged() const
 	return (TrackingTransaction.TransCount > 0 || TrackingTransaction.IsActive()) && (MouseDeltaTracker->HasReceivedDelta() || MouseDeltaTracker->WasExternalMovement());
 }
 
-bool FLevelEditorViewportClient::HasElementsToManipulate() const
+const UTypedElementList* FLevelEditorViewportClient::GetElementsToManipulate(const bool bForceRefresh)
+{
+	CacheElementsToManipulate(bForceRefresh);
+	return CachedElementsToManipulate;
+}
+
+void FLevelEditorViewportClient::CacheElementsToManipulate(const bool bForceRefresh)
 {
-	FLevelEditorViewportClient* MutableThis = const_cast<FLevelEditorViewportClient*>(this);
-	MutableThis->CacheElementsToManipulate();
-	const bool bHasElements = ElementsToManipulate->Num() > 0;
-	MutableThis->ElementsToManipulate->Reset();
-	return bHasElements;
+	if (bForceRefresh)
+	{
+		ResetElementsToManipulate();
+	}
+
+	if (!bHasCachedElementsToManipulate)
+	{
+		const UTypedElementSelectionSet* SelectionSet = GetSelectionSet();
+		ViewportInteraction->GetSelectedElementsToMove(SelectionSet, bIsSimulateInEditorViewport ? ETypedElementViewportInteractionWorldType::PlayInEditor : ETypedElementViewportInteractionWorldType::Editor, CachedElementsToManipulate);
+		bHasCachedElementsToManipulate = true;
+	}
 }
 
-TArray<FTypedElementHandle> FLevelEditorViewportClient::GetElementsToManipulate() const
+void FLevelEditorViewportClient::ResetElementsToManipulate()
 {
-	FLevelEditorViewportClient* MutableThis = const_cast<FLevelEditorViewportClient*>(this);
-	MutableThis->CacheElementsToManipulate();
-	TArray<FTypedElementHandle> ElementsToManipulateArray = ElementsToManipulate->GetElementHandles();
-	MutableThis->ElementsToManipulate->Reset();
-	return ElementsToManipulateArray;
+	if (bHasCachedElementsToManipulate)
+	{
+		CachedElementsToManipulate->Reset();
+		bHasCachedElementsToManipulate = false;
+	}
 }
 
-void FLevelEditorViewportClient::CacheElementsToManipulate()
+void FLevelEditorViewportClient::ResetElementsToManipulateFromSelectionChange(const UTypedElementSelectionSet* InSelectionSet)
 {
-	const UTypedElementSelectionSet* SelectionSet = GetSelectionSet();
-	ViewportInteraction->GetSelectedElementsToMove(SelectionSet, bIsSimulateInEditorViewport ? ETypedElementViewportInteractionWorldType::PlayInEditor : ETypedElementViewportInteractionWorldType::Editor, ElementsToManipulate);
+	check(InSelectionSet == GetSelectionSet());
+	ResetElementsToManipulate();
 }
 
 const UTypedElementSelectionSet* FLevelEditorViewportClient::GetSelectionSet() const
@@ -3681,8 +3706,6 @@ void FLevelEditorViewportClient::ApplyDeltaToSelectedElements(const FTransform&
 		return;
 	}
 
-	CacheElementsToManipulate();
-
 	FTransform ModifiedDeltaTransform = InDeltaTransform;
 
 	{
@@ -3700,9 +3723,8 @@ void FLevelEditorViewportClient::ApplyDeltaToSelectedElements(const FTransform&
 	FInputDeviceState InputState;
 	InputState.SetModifierKeyStates(IsShiftPressed(), IsAltPressed(), IsCtrlPressed(), IsCmdPressed());
 
+	const UTypedElementList* ElementsToManipulate = GetElementsToManipulate();
 	ViewportInteraction->UpdateGizmoManipulation(ElementsToManipulate, GetWidgetMode(), Widget ? Widget->GetCurrentAxis() : EAxisList::None, InputState, ModifiedDeltaTransform);
-
-	ElementsToManipulate->Reset();
 }
 
 void FLevelEditorViewportClient::ApplyDeltaToElement(const FTypedElementHandle& InElementHandle, const FTransform& InDeltaTransform)
@@ -4488,7 +4510,7 @@ void FLevelEditorViewportClient::AddReferencedObjects( FReferenceCollector& Coll
 {
 	FEditorViewportClient::AddReferencedObjects( Collector );
 
-	Collector.AddReferencedObject(ElementsToManipulate);
+	Collector.AddReferencedObject(CachedElementsToManipulate);
 
 	for( TSet<FViewportHoverTarget>::TIterator It( FLevelEditorViewportClient::HoveredObjects ); It; ++It )
 	{
diff --git a/Engine/Source/Editor/UnrealEd/Private/Tools/DefaultEdMode.cpp b/Engine/Source/Editor/UnrealEd/Private/Tools/DefaultEdMode.cpp
--- a/Engine/Source/Editor/UnrealEd/Private/Tools/DefaultEdMode.cpp
+++ b/Engine/Source/Editor/UnrealEd/Private/Tools/DefaultEdMode.cpp
@@ -6,14 +6,19 @@
 #include "EditorStyleSet.h"
 #include "Textures/SlateIcon.h"
 #include "LevelEditorViewport.h"
+#include "Elements/Framework/TypedElementList.h"
 
 class FLevelEditorSelectModeWidgetHelper : public FLegacyEdModeWidgetHelper
 {
 public:
 	virtual bool ShouldDrawWidget() const override
 	{
-		return GCurrentLevelEditingViewportClient
-			&& GCurrentLevelEditingViewportClient->HasElementsToManipulate();
+		if (GCurrentLevelEditingViewportClient)
+		{
+			const UTypedElementList* ElementsToManipulate = GCurrentLevelEditingViewportClient->GetElementsToManipulate();
+			return ElementsToManipulate->Num() > 0;
+		}
+		return false;
 	}
 };
 
diff --git a/Engine/Source/Editor/UnrealEd/Public/LevelEditorViewport.h b/Engine/Source/Editor/UnrealEd/Public/LevelEditorViewport.h
--- a/Engine/Source/Editor/UnrealEd/Public/LevelEditorViewport.h
+++ b/Engine/Source/Editor/UnrealEd/Public/LevelEditorViewport.h
@@ -326,14 +326,9 @@ public:
 	void MirrorSelectedElements(const FVector& InMirrorScale);
 
 	/**
-	 * Are there any elements (from the current selection set) that this viewport can manipulate (eg, via the transform gizmo).
+	 * Get the elements (from the current selection set) that this viewport can manipulate (eg, via the transform gizmo).
 	 */
-	bool HasElementsToManipulate() const;
-
-	/**
-	 * Get the array of elements (from the current selection set) that this viewport can manipulate (eg, via the transform gizmo).
-	 */
-	TArray<FTypedElementHandle> GetElementsToManipulate() const;
+	const UTypedElementList* GetElementsToManipulate(const bool bForceRefresh = false);
 
 	virtual void SetIsSimulateInEditorViewport( bool bInIsSimulateInEditorViewport ) override;
 
@@ -717,7 +712,13 @@ private:
 	bool HaveSelectedObjectsBeenChanged() const;
 
 	/** Cache the list of elements to manipulate based on the current selection set. */
-	void CacheElementsToManipulate();
+	void CacheElementsToManipulate(const bool bForceRefresh = false);
+
+	/** Reset the list of elements to manipulate */
+	void ResetElementsToManipulate();
+
+	/** Reset the list of elements to manipulate, because the selection set they were cached from has changed */
+	void ResetElementsToManipulateFromSelectionChange(const UTypedElementSelectionSet* InSelectionSet);
 
 	/** Get the selection set that associated with our level editor. */
 	const UTypedElementSelectionSet* GetSelectionSet() const;
@@ -882,8 +883,9 @@ private:
 	/** A map of actor locations before a drag operation */
 	mutable TMap<TWeakObjectPtr<const AActor>, FTransform> PreDragActorTransforms;
 
-	/** List of elements that are currently being manipulated (eg, via a drag) */
-	UTypedElementList* ElementsToManipulate;
+	/** The elements (from the current selection set) that this viewport can manipulate (eg, via the transform gizmo) */
+	bool bHasCachedElementsToManipulate = false;
+	UTypedElementList* CachedElementsToManipulate = nullptr;
 
 	/** Bit array representing the visibility of every sprite category in the current viewport */
 	TBitArray<>	SpriteCategoryVisibility;
