Added IsHLODRelevant() method to AActor & UActorComponent

diff --git a/Engine/Source/Editor/UnrealEd/Private/HierarchicalLOD.cpp b/Engine/Source/Editor/UnrealEd/Private/HierarchicalLOD.cpp
--- a/Engine/Source/Editor/UnrealEd/Private/HierarchicalLOD.cpp
+++ b/Engine/Source/Editor/UnrealEd/Private/HierarchicalLOD.cpp
@@ -668,7 +668,7 @@ bool FHierarchicalLODBuilder::ShouldGenerateCluster(AActor* Actor, const int32 H
 		return false;
 	}
 
-	if (Actor->IsHidden())
+	if (!Actor->IsHLODRelevant())
 	{
 		return false;
 	}
@@ -678,21 +678,6 @@ bool FHierarchicalLODBuilder::ShouldGenerateCluster(AActor* Actor, const int32 H
 		return false;
 	}
 
-	if( Actor->IsTemplate() )
-	{
-		return false;
-	}
-	
-	if( Actor->IsPendingKill() )
-	{
-		return false;
-	}
-
-	if (!Actor->bEnableAutoLODGeneration)
-	{
-		return false;
-	}
-
 	if (ALODActor* LODActor = Cast<ALODActor>(Actor))
 	{
 		// Ignore previous LOD actors
@@ -708,48 +693,21 @@ bool FHierarchicalLODBuilder::ShouldGenerateCluster(AActor* Actor, const int32 H
 		}
 	}
 
-	FVector Origin, Extent;
-	Actor->GetActorBounds(false, Origin, Extent);
-	if (Extent.SizeSquared() <= 0.1)
-	{
-		return false;
-	}	
-
-	// for now only consider staticmesh - I don't think skel mesh would work with simplygon merge right now @fixme
 	TArray<UStaticMeshComponent*> Components;
 	Actor->GetComponents<UStaticMeshComponent>(Components);
 
-	int32 ValidComponentCount = 0;
-
-	if (Components.Num() > 0)
+	bool bHasValidComponent = false;
+	for (UStaticMeshComponent* Component : Components)
 	{
-		for (UStaticMeshComponent* Component : Components)
-		{			
-			if (Component->bHiddenInGame)
-			{
-				continue;
-			}
-
-			if (Component->bIsEditorOnly)
-			{
-				continue;
-			}
-
-			if (!Component->GetStaticMesh())
-			{
-				continue;
-			}
-
-			// see if we should generate it
-			if (Component->ShouldGenerateAutoLOD(HLODLevelIndex))
-			{
-				++ValidComponentCount;
-				break;
-			}
+		// see if we should generate it
+		if (Component->ShouldGenerateAutoLOD(HLODLevelIndex))
+		{
+			bHasValidComponent = true;
+			break;
 		}
 	}
 
-	return (ValidComponentCount > 0);
+	return bHasValidComponent;
 }
 
 void FHierarchicalLODBuilder::ClearHLODs()
diff --git a/Engine/Source/Runtime/Engine/Classes/Components/ActorComponent.h b/Engine/Source/Runtime/Engine/Classes/Components/ActorComponent.h
--- a/Engine/Source/Runtime/Engine/Classes/Components/ActorComponent.h
+++ b/Engine/Source/Runtime/Engine/Classes/Components/ActorComponent.h
@@ -982,6 +982,9 @@ public:
 	/** Override to specify that a component is relevant to the navigation system */
 	virtual bool IsNavigationRelevant() const { return false; }
 
+	/** Override to specify that a component is relevant to the HLOD generation. */
+	virtual bool IsHLODRelevant() const { return false; }
+
 	/** Prefix used to identify template component instances */
 	static const FString ComponentTemplateNameSuffix;
 
diff --git a/Engine/Source/Runtime/Engine/Classes/Components/InstancedStaticMeshComponent.h b/Engine/Source/Runtime/Engine/Classes/Components/InstancedStaticMeshComponent.h
--- a/Engine/Source/Runtime/Engine/Classes/Components/InstancedStaticMeshComponent.h
+++ b/Engine/Source/Runtime/Engine/Classes/Components/InstancedStaticMeshComponent.h
@@ -305,6 +305,7 @@ public:
 
 	//~ Begin UActorComponent Interface
 	virtual TStructOnScope<FActorComponentInstanceData> GetComponentInstanceData() const override;
+	virtual bool IsHLODRelevant() const override;
 	//~ End UActorComponent Interface
 
 	//~ Begin UPrimitiveComponent Interface
diff --git a/Engine/Source/Runtime/Engine/Classes/Components/StaticMeshComponent.h b/Engine/Source/Runtime/Engine/Classes/Components/StaticMeshComponent.h
--- a/Engine/Source/Runtime/Engine/Classes/Components/StaticMeshComponent.h
+++ b/Engine/Source/Runtime/Engine/Classes/Components/StaticMeshComponent.h
@@ -486,6 +486,7 @@ public:
 	virtual bool IsCompiling() const override;
 #endif
 	virtual TStructOnScope<FActorComponentInstanceData> GetComponentInstanceData() const override;
+	virtual bool IsHLODRelevant() const override;
 	//~ End UActorComponent Interface.
 
 	//~ Begin UPrimitiveComponent Interface.
diff --git a/Engine/Source/Runtime/Engine/Classes/GameFramework/Actor.h b/Engine/Source/Runtime/Engine/Classes/GameFramework/Actor.h
--- a/Engine/Source/Runtime/Engine/Classes/GameFramework/Actor.h
+++ b/Engine/Source/Runtime/Engine/Classes/GameFramework/Actor.h
@@ -2037,6 +2037,9 @@ public:
 	/** Indicates whether this actor should participate in level bounds calculations */
 	virtual bool IsLevelBoundsRelevant() const { return bRelevantForLevelBounds; }
 
+	/** Indicates whether this actor contributes to the HLOD generation. */
+	virtual bool IsHLODRelevant() const;
+
 	/** 
 	 * Set LOD Parent component for all of our components, normally associated with an ALODActor. 
 	 * @param InLODParent			This component used to compute visibility when hierarchical LOD is enabled. 
diff --git a/Engine/Source/Runtime/Engine/Private/Actor.cpp b/Engine/Source/Runtime/Engine/Private/Actor.cpp
--- a/Engine/Source/Runtime/Engine/Private/Actor.cpp
+++ b/Engine/Source/Runtime/Engine/Private/Actor.cpp
@@ -48,6 +48,7 @@
 #include "LevelUtils.h"
 #include "Components/DecalComponent.h"
 #include "GameFramework/InputSettings.h"
+#include "Algo/AnyOf.h"
 
 #if WITH_EDITOR
 #include "FoliageHelper.h"
@@ -5409,6 +5410,43 @@ void AActor::PostRename(UObject* OldOuter, const FName OldName)
 	}
 }
 
+bool AActor::IsHLODRelevant() const
+{
+	if (IsHidden())
+	{
+		return false;
+	}
+
+	if (IsEditorOnly())
+	{
+		return false;
+	}
+
+	if (IsTemplate())
+	{
+		return false;
+	}
+
+	if (IsPendingKill())
+	{
+		return false;
+	}
+
+	if (!bEnableAutoLODGeneration)
+	{
+		return false;
+	}
+
+	FVector Origin, Extent;
+	GetActorBounds(false, Origin, Extent);
+	if (Extent.SizeSquared() <= 0.1)
+	{
+		return false;
+	}
+
+	return Algo::AnyOf(GetComponents(), [](const UActorComponent* Component) { return Component->IsHLODRelevant(); });
+}
+
 void AActor::SetLODParent(UPrimitiveComponent* InLODParent, float InParentDrawDistance)
 {
 	if (InLODParent && InLODParent->MinDrawDistance != InParentDrawDistance)
diff --git a/Engine/Source/Runtime/Engine/Private/Components/PrimitiveComponent.cpp b/Engine/Source/Runtime/Engine/Private/Components/PrimitiveComponent.cpp
--- a/Engine/Source/Runtime/Engine/Private/Components/PrimitiveComponent.cpp
+++ b/Engine/Source/Runtime/Engine/Private/Components/PrimitiveComponent.cpp
@@ -3903,6 +3903,11 @@ void UPrimitiveComponent::SetLastRenderTime(float InLastRenderTime)
 #if WITH_EDITOR
 const bool UPrimitiveComponent::ShouldGenerateAutoLOD(const int32 HierarchicalLevelIndex) const
 {	
+	if (!IsHLODRelevant())
+	{
+		return false;
+	}
+
 	// bAllowSpecificExclusion
 	bool bExcluded = false;
 	if (ExcludeForSpecificHLODLevels.Contains(HierarchicalLevelIndex))
@@ -3915,10 +3920,10 @@ const bool UPrimitiveComponent::ShouldGenerateAutoLOD(const int32 HierarchicalLe
 				bExcluded = true;
 			}
 		}
-	} 
-		
-	return (Mobility != EComponentMobility::Movable) && bEnableAutoLODGeneration && !bExcluded;
+	}
+
+	return !bExcluded;
 }
-#endif 
+#endif
 
 #undef LOCTEXT_NAMESPACE
diff --git a/Engine/Source/Runtime/Engine/Private/Components/StaticMeshComponent.cpp b/Engine/Source/Runtime/Engine/Private/Components/StaticMeshComponent.cpp
--- a/Engine/Source/Runtime/Engine/Private/Components/StaticMeshComponent.cpp
+++ b/Engine/Source/Runtime/Engine/Private/Components/StaticMeshComponent.cpp
@@ -2463,6 +2463,36 @@ void UStaticMeshComponent::ApplyComponentInstanceData(FStaticMeshComponentInstan
 #endif
 }
 
+bool UStaticMeshComponent::IsHLODRelevant() const
+{
+	if (!GetStaticMesh())
+	{
+		return false;
+	}
+
+	if (bHiddenInGame)
+	{
+		return false;
+	}
+
+	if (bIsEditorOnly)
+	{
+		return false;
+	}
+
+	if (Mobility == EComponentMobility::Movable)
+	{
+		return false;
+	}
+
+	if (!bEnableAutoLODGeneration)
+	{
+		return false;
+	}
+
+	return true;
+}
+
 bool UStaticMeshComponent::DoCustomNavigableGeometryExport(FNavigableGeometryExport& GeomExport) const
 {
 	const FVector Scale3D = GetComponentToWorld().GetScale3D();
diff --git a/Engine/Source/Runtime/Engine/Private/InstancedStaticMesh.cpp b/Engine/Source/Runtime/Engine/Private/InstancedStaticMesh.cpp
--- a/Engine/Source/Runtime/Engine/Private/InstancedStaticMesh.cpp
+++ b/Engine/Source/Runtime/Engine/Private/InstancedStaticMesh.cpp
@@ -1553,6 +1553,16 @@ void UInstancedStaticMeshComponent::FlushInstanceUpdateCommands()
 	PerInstanceRenderData->UpdateFromPreallocatedData(RenderInstanceData);
 }
 
+bool UInstancedStaticMeshComponent::IsHLODRelevant() const
+{
+	if (GetInstanceCount() == 0)
+	{
+		return false;
+	}
+
+	return Super::IsHLODRelevant();
+}
+
 FPrimitiveSceneProxy* UInstancedStaticMeshComponent::CreateSceneProxy()
 {
 	LLM_SCOPE(ELLMTag::InstancedMesh);
