Fix missing notification to a component when we overwrite its properties causing an ensure to trigger when converting a static mesh to a BP asset

diff --git a/Engine/Source/Editor/UnrealEd/Private/Editor.cpp b/Engine/Source/Editor/UnrealEd/Private/Editor.cpp
--- a/Engine/Source/Editor/UnrealEd/Private/Editor.cpp
+++ b/Engine/Source/Editor/UnrealEd/Private/Editor.cpp
@@ -1567,6 +1567,9 @@ namespace EditorUtilities
 
 							CopySingleProperty(SourceComponent, TargetComponent, Property);
 
+							// Notify the target one of it's properties might have changed
+							TargetComponent->PostReinitProperties();
+
 							if( Options.Flags & ECopyOptions::CallPostEditChangeProperty )
 							{
 								FPropertyChangedEvent PropertyChangedEvent( Property );
