Added missing virtual destructor on FSourceControlHelper

diff --git a/Engine/Source/Editor/UnrealEd/Private/WorldPartition/WorldPartitionHLODsBuilder.cpp b/Engine/Source/Editor/UnrealEd/Private/WorldPartition/WorldPartitionHLODsBuilder.cpp
--- a/Engine/Source/Editor/UnrealEd/Private/WorldPartition/WorldPartitionHLODsBuilder.cpp
+++ b/Engine/Source/Editor/UnrealEd/Private/WorldPartition/WorldPartitionHLODsBuilder.cpp
@@ -22,6 +22,9 @@ public:
 		: PackageHelper(InPackageHelper)
 	{}
 
+	virtual ~FSourceControlHelper()
+	{}
+
 	virtual FString GetFilename(const FString& PackageName) const override
 	{
 		return SourceControlHelpers::PackageFilename(PackageName);
