Make FComponentEditorUtils::FindMatchingComponent take in a const pointer to the component instance as it doesn't modify the instance at all. This will allow us to call it from the new subobject subsystem more safely.

#rnx

diff --git a/Engine/Source/Editor/UnrealEd/Private/Kismet2/ComponentEditorUtils.cpp b/Engine/Source/Editor/UnrealEd/Private/Kismet2/ComponentEditorUtils.cpp
--- a/Engine/Source/Editor/UnrealEd/Private/Kismet2/ComponentEditorUtils.cpp
+++ b/Engine/Source/Editor/UnrealEd/Private/Kismet2/ComponentEditorUtils.cpp
@@ -1139,7 +1139,7 @@ void FComponentEditorUtils::FillComponentContextMenuOptions(UToolMenu* Menu, con
 	}
 }
 
-UActorComponent* FComponentEditorUtils::FindMatchingComponent(UActorComponent* ComponentInstance, const TInlineComponentArray<UActorComponent*>& ComponentList)
+UActorComponent* FComponentEditorUtils::FindMatchingComponent(const UActorComponent* ComponentInstance, const TInlineComponentArray<UActorComponent*>& ComponentList)
 {
 	if (ComponentInstance == nullptr)
 	{
@@ -1169,7 +1169,7 @@ UActorComponent* FComponentEditorUtils::FindMatchingComponent(UActorComponent* C
 		return LastFoundComponent;
 	}
 
-	if (USceneComponent* CurrentSceneComponent = Cast<USceneComponent>(ComponentInstance))
+	if (const USceneComponent* CurrentSceneComponent = Cast<USceneComponent>(ComponentInstance))
 	{
 		// Sort by matching hierarchy
 		FoundComponents.Sort([&](const UActorComponent& ComponentA, const UActorComponent& ComponentB)
diff --git a/Engine/Source/Editor/UnrealEd/Public/Kismet2/ComponentEditorUtils.h b/Engine/Source/Editor/UnrealEd/Public/Kismet2/ComponentEditorUtils.h
--- a/Engine/Source/Editor/UnrealEd/Public/Kismet2/ComponentEditorUtils.h
+++ b/Engine/Source/Editor/UnrealEd/Public/Kismet2/ComponentEditorUtils.h
@@ -275,7 +275,7 @@ public:
 	 * @param ComponentList List containing possible matches
 	 * @return Valid Component pointer if match was found. nullptr otherwise.
 	 */
-	static UActorComponent* FindMatchingComponent(UActorComponent* ComponentInstance, const TInlineComponentArray<UActorComponent*>& ComponentList);
+	static UActorComponent* FindMatchingComponent(const UActorComponent* ComponentInstance, const TInlineComponentArray<UActorComponent*>& ComponentList);
 
 	/**
 	 * Make a FComponentReference from a component pointer.
