Only enable Stall Detection between Main Frame open and close, not engine startup and shutdown.

This will help ensure that the first triggered stall, the one that sends telemetry data, will represent a real stall witnessed by a person. This will also shutdown the stall detection thread sooner to alleivate any coupling with other engine systems.

diff --git a/Engine/Source/Editor/UnrealEd/Private/UnrealEdEngine.cpp b/Engine/Source/Editor/UnrealEd/Private/UnrealEdEngine.cpp
--- a/Engine/Source/Editor/UnrealEd/Private/UnrealEdEngine.cpp
+++ b/Engine/Source/Editor/UnrealEd/Private/UnrealEdEngine.cpp
@@ -70,6 +70,7 @@
 #include "Misc/MessageDialog.h"
 #include "Logging/MessageLog.h"
 #include "Subsystems/EditorActorSubsystem.h"
+#include "ProfilingDebugging/StallDetector.h"
 
 DEFINE_LOG_CATEGORY_STATIC(LogUnrealEdEngine, Log, All);
 
@@ -218,6 +219,17 @@ void UUnrealEdEngine::Init(IEngineLoop* InEngineLoop)
 	}
 
 	bPivotMovedIndependently = false;
+
+#if STALL_DETECTOR
+	// Start tracking stalls when we open the Main Frame
+	IMainFrameModule::Get().OnMainFrameCreationFinished().AddLambda([](TSharedPtr<SWindow>, bool)
+		{
+			UE::FStallDetector::Startup();
+		});
+
+	// Stop tracking stalls when we close the Main Frame, the closest event around
+	OnEditorClose().AddStatic(&UE::FStallDetector::Shutdown);
+#endif
 }
 
 bool CanCookForPlatformInThisProcess( const FString& PlatformName )
diff --git a/Engine/Source/Runtime/Launch/Private/LaunchEngineLoop.cpp b/Engine/Source/Runtime/Launch/Private/LaunchEngineLoop.cpp
--- a/Engine/Source/Runtime/Launch/Private/LaunchEngineLoop.cpp
+++ b/Engine/Source/Runtime/Launch/Private/LaunchEngineLoop.cpp
@@ -2626,13 +2626,6 @@ PRAGMA_ENABLE_DEPRECATION_WARNINGS
 	bool bSlateApplication = !IsRunningDedicatedServer() && (bIsRegularClient || bHasEditorToken);
 	if (bSlateApplication)
 	{
-#if STALL_DETECTOR
-		if (GIsEditor)
-		{
-			UE::FStallDetector::Startup();
-		}
-#endif
-
 		if (FPlatformProcess::SupportsMultithreading())
 		{
 			SCOPED_BOOT_TIMING("FPlatformSplash::Show()");
@@ -5978,13 +5971,6 @@ void FEngineLoop::AppPreExit( )
 
 void FEngineLoop::AppExit()
 {
-#if STALL_DETECTOR
-	if (UE::FStallDetector::IsRunning())
-	{
-		UE::FStallDetector::Shutdown();
-	}
-#endif
-
 #if !WITH_ENGINE
 	// when compiled WITH_ENGINE, this will happen in FEngineLoop::Exit()
 	FTaskGraphInterface::Shutdown();
