Added duplicated actor guids detection in map check.

diff --git a/Engine/Source/Editor/UnrealEd/Private/EditorServer.cpp b/Engine/Source/Editor/UnrealEd/Private/EditorServer.cpp
--- a/Engine/Source/Editor/UnrealEd/Private/EditorServer.cpp
+++ b/Engine/Source/Editor/UnrealEd/Private/EditorServer.cpp
@@ -3993,6 +3993,7 @@ bool UEditorEngine::Map_Check( UWorld* InWorld, const TCHAR* Str, FOutputDevice&
 
 	int32 LastUpdateCount = 0;
 	int32 UpdateGranularity = ProgressDenominator / 5;
+	TMap<ULevel*, TMap<FGuid, AActor*>> ActorGuidMap;
 	for( FActorIterator It(InWorld); It; ++It ) 
 	{
 		if(It.GetProgressNumerator() >= LastUpdateCount + UpdateGranularity)
@@ -4063,6 +4064,23 @@ bool UEditorEngine::Map_Check( UWorld* InWorld, const TCHAR* Str, FOutputDevice&
 					}
 				}
 			}
+
+			// Check for duplicated Guids
+			TMap<FGuid, AActor*>& PerLevelActorGuidMap = ActorGuidMap.FindOrAdd(Actor->GetLevel());
+			if (AActor** DuplicatedGuidActor = PerLevelActorGuidMap.Find(Actor->GetActorGuid()))
+			{
+				FFormatNamedArguments Arguments;
+				Arguments.Add(TEXT("Actor0"), FText::FromString(Actor->GetName()));
+				Arguments.Add(TEXT("Actor1"), FText::FromString((*DuplicatedGuidActor)->GetName()));
+				FMessageLog("MapCheck").Warning()
+					->AddToken(FUObjectToken::Create(Actor))
+					->AddToken(FTextToken::Create(FText::Format(LOCTEXT( "MapCheck_Message_MatchingActorGUID", "'{Actor0}' has same GUID as '{Actor1}' (Duplicate and replace the orig with the new one)" ), Arguments)))
+					->AddToken(FMapErrorToken::Create(FMapErrors::MatchingActorGUID));
+			}
+			else
+			{
+				PerLevelActorGuidMap.Add(Actor->GetActorGuid(), Actor);
+			}
 		}
 
 		Game_Map_Check_Actor(Str, Ar, bCheckDeprecatedOnly, Actor);
diff --git a/Engine/Source/Runtime/Core/Private/Misc/MapErrors.cpp b/Engine/Source/Runtime/Core/Private/Misc/MapErrors.cpp
--- a/Engine/Source/Runtime/Core/Private/Misc/MapErrors.cpp
+++ b/Engine/Source/Runtime/Core/Private/Misc/MapErrors.cpp
@@ -1,7 +1,8 @@
 // Copyright Epic Games, Inc. All Rights Reserved.
 
 #include "Misc/MapErrors.h"
-	
+
+FLazyName FMapErrors::MatchingActorGUID(TEXT("MatchingActorGUID"));
 FLazyName FMapErrors::MatchingLightGUID(TEXT("MatchingLightGUID"));
 FLazyName FMapErrors::ActorLargeShadowCaster(TEXT("ActorLargeShadowCaster"));
 FLazyName FMapErrors::NoDamageType(TEXT("NoDamageType"));
diff --git a/Engine/Source/Runtime/Core/Public/Misc/MapErrors.h b/Engine/Source/Runtime/Core/Public/Misc/MapErrors.h
--- a/Engine/Source/Runtime/Core/Public/Misc/MapErrors.h
+++ b/Engine/Source/Runtime/Core/Public/Misc/MapErrors.h
@@ -19,7 +19,10 @@ struct CORE_API FMapErrors
     static FLazyName MatchingLightGUID;
  
     /**  Actor */
-    
+
+	/** {Actor} has same GUID as {Actor} (Duplicate and replace the orig with the new one): Duplicate and replace the original with the new one. */
+    static FLazyName MatchingActorGUID;
+
 	/** {Actor} : Large actor casts a shadow and will cause an extreme performance hit unless bUseBooleanEnvironmentShadowing is set to true: A large actor has been set to cast shadows - this will cause extreme performance issues and should have bUseBooleanEnvironmentShadowing set to true. */
     static FLazyName ActorLargeShadowCaster;
 
