World Partition HLOD Builder
* Renamed option -CreateHLODs to -SetupHLODs
* Added progress log to the -BuildHLODs command
* Added -ForceGC option to force the GC to run between each built HLOD, as a temporary measure against GPU OOM

diff --git a/Engine/Source/Editor/UnrealEd/Private/WorldPartition/WorldPartitionHLODsBuilder.cpp b/Engine/Source/Editor/UnrealEd/Private/WorldPartition/WorldPartitionHLODsBuilder.cpp
--- a/Engine/Source/Editor/UnrealEd/Private/WorldPartition/WorldPartitionHLODsBuilder.cpp
+++ b/Engine/Source/Editor/UnrealEd/Private/WorldPartition/WorldPartitionHLODsBuilder.cpp
@@ -91,14 +91,15 @@ private:
 UWorldPartitionHLODsBuilder::UWorldPartitionHLODsBuilder(const FObjectInitializer& ObjectInitializer)
 	: Super(ObjectInitializer)
 {
-	bCreateHLODs = FParse::Param(FCommandLine::Get(), TEXT("CreateHLODs"));
+	bSetupHLODs = FParse::Param(FCommandLine::Get(), TEXT("SetupHLODs"));
 	bBuildHLODs = FParse::Param(FCommandLine::Get(), TEXT("BuildHLODs"));
 	bDeleteHLODs = FParse::Param(FCommandLine::Get(), TEXT("DeleteHLODs"));
+	bForceGC = FParse::Param(FCommandLine::Get(), TEXT("ForceGC"));
 }
 
 bool UWorldPartitionHLODsBuilder::RequiresCommandletRendering() const
 {
-	return bBuildHLODs || !(bCreateHLODs || bDeleteHLODs);
+	return bBuildHLODs || !(bSetupHLODs || bDeleteHLODs);
 }
 
 bool UWorldPartitionHLODsBuilder::Run(UWorld* World, FPackageSourceControlHelper& PackageHelper)
@@ -109,9 +110,9 @@ bool UWorldPartitionHLODsBuilder::Run(UWorld* World, FPackageSourceControlHelper
 	SourceControlHelper = new FSourceControlHelper(PackageHelper);
 
 	bool bRet;
-	if (bCreateHLODs)
+	if (bSetupHLODs)
 	{
-		bRet = CreateHLODActors(true);
+		bRet = SetupHLODActors(true);
 	}
 	else if (bBuildHLODs)
 	{
@@ -123,7 +124,7 @@ bool UWorldPartitionHLODsBuilder::Run(UWorld* World, FPackageSourceControlHelper
 	}
 	else
 	{
-		bRet = CreateHLODActors(false);
+		bRet = SetupHLODActors(false);
 	}
 
 	WorldPartition = nullptr;
@@ -132,26 +133,26 @@ bool UWorldPartitionHLODsBuilder::Run(UWorld* World, FPackageSourceControlHelper
 	return bRet;
 }
 
-bool UWorldPartitionHLODsBuilder::CreateHLODActors(bool bCreateOnly)
+bool UWorldPartitionHLODsBuilder::SetupHLODActors(bool bCreateOnly)
 {
 	WorldPartition->GenerateHLOD(SourceControlHelper, bCreateOnly);
 
 	if (bCreateOnly)
 	{
-		UE_LOG(LogWorldPartitionHLODsBuilder, Display, TEXT("#### Created HLOD actors ####"));
+		UE_LOG(LogWorldPartitionHLODsBuilder, Display, TEXT("#### World HLOD actors ####"));
 
-		int32 NumCreated = 0;
+		int32 NumActors = 0;
 		for (UActorDescContainer::TIterator<AWorldPartitionHLOD> HLODIterator(WorldPartition); HLODIterator; ++HLODIterator)
 		{
 			FWorldPartitionActorDesc* HLODActorDesc = *HLODIterator;
 			FString PackageName = HLODActorDesc->GetActorPackage().ToString();
 
-			UE_LOG(LogWorldPartitionHLODsBuilder, Display, TEXT("    [%d] %s"), NumCreated, *PackageName);
+			UE_LOG(LogWorldPartitionHLODsBuilder, Display, TEXT("    [%d] %s"), NumActors, *PackageName);
 
-			NumCreated++;
+			NumActors++;
 		}
 
-		UE_LOG(LogWorldPartitionHLODsBuilder, Display, TEXT("#### Created %d HLOD actors ####"), NumCreated);
+		UE_LOG(LogWorldPartitionHLODsBuilder, Display, TEXT("#### World contains %d HLOD actors ####"), NumActors);
 	}
 
 	return true;
@@ -175,6 +176,9 @@ bool UWorldPartitionHLODsBuilder::BuildHLODActors()
 
 		AWorldPartitionHLOD* HLODActor = CastChecked<AWorldPartitionHLOD>(ActorDesc->GetActor());
 
+		CurrentActor++;
+		UE_LOG(LogWorldPartitionHLODsBuilder, Display, TEXT("    [%d/%d] Building HLOD actor..."), CurrentActor, NumActors);
+
 		HLODActor->BuildHLOD();
 
 		if (HLODActor->GetPackage()->IsDirty())
@@ -182,7 +186,7 @@ bool UWorldPartitionHLODsBuilder::BuildHLODActors()
 			SourceControlHelper->Save(HLODActor->GetPackage());
 		}
 
-		if (HasExceededMaxMemory())
+		if (bForceGC || HasExceededMaxMemory())
 		{
 			DoCollectGarbage();
 		}
diff --git a/Engine/Source/Editor/UnrealEd/Public/WorldPartition/WorldPartitionHLODsBuilder.h b/Engine/Source/Editor/UnrealEd/Public/WorldPartition/WorldPartitionHLODsBuilder.h
--- a/Engine/Source/Editor/UnrealEd/Public/WorldPartition/WorldPartitionHLODsBuilder.h
+++ b/Engine/Source/Editor/UnrealEd/Public/WorldPartition/WorldPartitionHLODsBuilder.h
@@ -18,7 +18,7 @@ public:
 	// UWorldPartitionBuilder interface end
 
 protected:
-	bool CreateHLODActors(bool bCreateOnly);
+	bool SetupHLODActors(bool bCreateOnly);
 	bool BuildHLODActors();
 	bool DeleteHLODActors();
 
@@ -26,7 +26,8 @@ private:
 	class UWorldPartition* WorldPartition;
 	class FSourceControlHelper* SourceControlHelper;
 
-	bool bCreateHLODs;
+	bool bSetupHLODs;
 	bool bBuildHLODs;
 	bool bDeleteHLODs;
+	bool bForceGC;
 };
\ No newline at end of file
