Prevent invalid Static Mesh Instance IDs from creating element handles during selection state restoration

#rnx

diff --git a/Engine/Source/Editor/UnrealEd/Private/Elements/SMInstance/SMInstanceElementEditorSelectionInterface.cpp b/Engine/Source/Editor/UnrealEd/Private/Elements/SMInstance/SMInstanceElementEditorSelectionInterface.cpp
--- a/Engine/Source/Editor/UnrealEd/Private/Elements/SMInstance/SMInstanceElementEditorSelectionInterface.cpp
+++ b/Engine/Source/Editor/UnrealEd/Private/Elements/SMInstance/SMInstanceElementEditorSelectionInterface.cpp
@@ -19,7 +19,7 @@ private:
 	virtual FTypedElementHandle GetElementImpl() const override
 	{
 		FSMInstanceElementId SMInstanceElementId{ ISMComponentPtr.Get(/*bEvenIfPendingKill*/true), InstanceId };
-		return SMInstanceElementId
+		return FSMInstanceElementIdMap::Get().GetSMInstanceIdFromSMInstanceElementId(SMInstanceElementId)
 			? UEngineElementsLibrary::AcquireEditorSMInstanceElementHandle(SMInstanceElementId)
 			: FTypedElementHandle();
 	}
