Changed how we decide to rename a texture if the source file as an UDIM pattern when importing a file or reimporting a texture.

Now the texture doesn't get renamed if we are doing a reimport.

The texture also won't be renamed if the name provided for the texture is not the same as the source file name without the extension.

diff --git a/Engine/Source/Editor/UnrealEd/Classes/Factories/TextureFactory.h b/Engine/Source/Editor/UnrealEd/Classes/Factories/TextureFactory.h
--- a/Engine/Source/Editor/UnrealEd/Classes/Factories/TextureFactory.h
+++ b/Engine/Source/Editor/UnrealEd/Classes/Factories/TextureFactory.h
@@ -196,6 +196,10 @@ public:
 	// @todo document
 	bool InitializeFromT3DTextureCubeDataText(UTextureCube* InTextureCube, const TCHAR*& Buffer, FFeedbackContext* Warn);
 
+protected:
+	/** Keep track of if we are doing a reimport */
+	bool bIsDoingAReimport = false;
+
 private:
 	/** This variable is static because in StaticImportObject() the type of the factory is not known. */
 	static bool bSuppressImportOverwriteDialog;
diff --git a/Engine/Source/Editor/UnrealEd/Private/Factories/EditorFactories.cpp b/Engine/Source/Editor/UnrealEd/Private/Factories/EditorFactories.cpp
--- a/Engine/Source/Editor/UnrealEd/Private/Factories/EditorFactories.cpp
+++ b/Engine/Source/Editor/UnrealEd/Private/Factories/EditorFactories.cpp
@@ -3938,43 +3938,55 @@ UObject* UTextureFactory::FactoryCreateBinary
 				}
 			}
 
-			// Exclude UDIM number from the name of the UE4 texture asset we create
-			// We do this even in the case where we're only importing a single 1001 image, which won't technically be a UDIM
-			// We still want to strip the UDIM suffix in this case, as industry standard still considers this a UDIM imageset
-			const FString ShortPackageName = ObjectTools::SanitizeInvalidChars(BaseUDIMName, INVALID_LONGPACKAGE_CHARACTERS);
-			TextureName = *ShortPackageName;
-
-			// Don't try to rename the package if its the transient package
-			if ( InParent != GetTransientPackage() && InParent->IsA<UPackage>() )
-			{
-				// Need to rename the package to match the new texture name, since package was already created
-				// Package name will be the same as the object name, except will contain additional path information,
-				// so we take the existing package name, then extract the UDIM index in order to preserve the path
-				FString PackageName;
-				InParent->GetName(PackageName);
-
-				const int32 PackageUDIMIndex = UE::TextureUtilitiesCommon::ParseUDIMName(PackageName, UdimRegexPattern, PreUDIMName, PostUDIMName);
-				const FString PackageUDIMName = PreUDIMName + PostUDIMName;
-
-				// If we're re-importing UDIM texture, the package will already be correctly named after the UDIM base name
-				// In this case we'll fail to parse the UDIM name, but the package should already have the proper name
-				// Note that the package name may not match the asset name in this case, if we're reimporting new files
-				if (PackageUDIMIndex != -1)
+			// Don't try to rename the texture if we are doing a reimport.
+			if (!bIsDoingAReimport)
+			{ 
+				// Exclude UDIM number from the name of the UE4 texture asset we create
+				// We do this even in the case where we're only importing a single 1001 image, which won't technically be a UDIM
+				// We still want to strip the UDIM suffix in this case, as industry standard still considers this a UDIM imageset
+				const FString ShortPackageName = ObjectTools::SanitizeInvalidChars(BaseUDIMName, INVALID_LONGPACKAGE_CHARACTERS);
+
+				// Don't try to rename if the name of texture differ from the source file 
+				// Some factory provide a name for the texture that is different form the source file Ex:FBX.
+				const FString TextureNameAsString = TextureName.ToString();
+				if (TextureNameAsString.Equals(ObjectTools::SanitizeInvalidChars(FilenameNoExtension, INVALID_LONGPACKAGE_CHARACTERS)))
 				{
-					check(PackageUDIMIndex == BaseUDIMIndex);
-					check(PackageUDIMName.EndsWith(ShortPackageName, ESearchCase::CaseSensitive));
-
-					// In normal case, higher level code would have already checked for duplicate package name
-					// But since we're changing package name here, check to see if package with the new name already exists...
-					// If it does, code later in this method will prompt user to overwrite the existing asset
-					UPackage* ExistingPackage = FindPackage(InParent->GetOuter(), *PackageUDIMName);
-					if (ExistingPackage)
-					{
-						InParent = ExistingPackage;
-					}
-					else
+					// Don't try to rename the package if its the transient package
+					if ( InParent != GetTransientPackage() && InParent->IsA<UPackage>() )
 					{
-						verify(InParent->Rename(*PackageUDIMName, nullptr, REN_DontCreateRedirectors));
+						// Need to rename the package to match the new texture name, since package was already created
+						// Package name will be the same as the object name, except will contain additional path information,
+						// so we take the existing package name, then extract the UDIM index in order to preserve the path
+						FString PackageName;
+						InParent->GetName(PackageName);
+
+						const int32 PackageUDIMIndex = UE::TextureUtilitiesCommon::ParseUDIMName(PackageName, UdimRegexPattern, PreUDIMName, PostUDIMName);
+						const FString PackageUDIMName = PreUDIMName + PostUDIMName;
+
+						// If we're re-importing UDIM texture, the package will already be correctly named after the UDIM base name
+						// In this case we'll fail to parse the UDIM name, but the package should already have the proper name
+						// Note that the package name may not match the asset name in this case, if we're reimporting new files
+						if (PackageUDIMIndex != -1)
+						{
+							check(PackageUDIMIndex == BaseUDIMIndex);
+							check(PackageUDIMName.EndsWith(ShortPackageName, ESearchCase::CaseSensitive));
+
+							// Rename the texture object
+							TextureName = *ShortPackageName;
+
+							// In normal case, higher level code would have already checked for duplicate package name
+							// But since we're changing package name here, check to see if package with the new name already exists...
+							// If it does, code later in this method will prompt user to overwrite the existing asset
+							UPackage* ExistingPackage = FindPackage(InParent->GetOuter(), *PackageUDIMName);
+							if (ExistingPackage)
+							{
+								InParent = ExistingPackage;
+							}
+							else
+							{
+								verify(InParent->Rename(*PackageUDIMName, nullptr, REN_DontCreateRedirectors));
+							}
+						}
 					}
 				}
 			}
@@ -5504,6 +5516,8 @@ void UReimportTextureFactory::SetReimportPaths( UObject* Obj, const TArray<FStri
 */
 EReimportResult::Type UReimportTextureFactory::Reimport( UObject* Obj )
 {
+	TGuardValue<bool> IsDoingAReimportGuard(bIsDoingAReimport, true);
+
 	if(!Obj || !Obj->IsA(UTexture::StaticClass()))
 	{
 		return EReimportResult::Failed;
