Add a select mode for gizmos (effectively disables the gizmos while on). Bound to Q by default similar to other DCC tools.
Uses WM_None for now. May need a special gizmo mode to differentate between the select mode and gizmo not enabled states. Though the latter can be accomplished by overriding and  returing false from SEditorViewport::IsWidgetModeActive

diff --git a/Engine/Source/Editor/EditorStyle/Private/StarshipStyle.cpp b/Engine/Source/Editor/EditorStyle/Private/StarshipStyle.cpp
--- a/Engine/Source/Editor/EditorStyle/Private/StarshipStyle.cpp
+++ b/Engine/Source/Editor/EditorStyle/Private/StarshipStyle.cpp
@@ -4272,6 +4272,7 @@ void FStarshipEditorStyle::FStyle::SetupLevelEditorStyle()
 		Set("MainFrame.AboutUnrealEd",             new IMAGE_BRUSH("Icons/Help/icon_Help_unreal_16x", Icon16x16));
 		Set("MainFrame.CreditsUnrealEd",           new IMAGE_BRUSH_SVG("Starship/Common/Credits", Icon16x16));
 
+		Set( "EditorViewport.SelectMode", new IMAGE_BRUSH_SVG("Starship/EditorViewport/select", Icon16x16) );
 		Set( "EditorViewport.TranslateMode", new IMAGE_BRUSH_SVG( "Starship/EditorViewport/translate", Icon16x16 ) );
 		Set( "EditorViewport.RotateMode", new IMAGE_BRUSH_SVG("Starship/EditorViewport/rotate", Icon16x16 ) );
 		Set( "EditorViewport.ScaleMode", new IMAGE_BRUSH_SVG( "Starship/EditorViewport/scale", Icon16x16 ) );
diff --git a/Engine/Source/Editor/UnrealEd/Private/EditorViewportCommands.cpp b/Engine/Source/Editor/UnrealEd/Private/EditorViewportCommands.cpp
--- a/Engine/Source/Editor/UnrealEd/Private/EditorViewportCommands.cpp
+++ b/Engine/Source/Editor/UnrealEd/Private/EditorViewportCommands.cpp
@@ -131,6 +131,7 @@ void FEditorViewportCommands::RegisterCommands()
 	UI_COMMAND( DecrementRotationGridSize, "Grid Size (Rotation): Decrement", "Decreases the rotation grid size setting by one", EUserInterfaceActionType::Button, FInputChord( EModifierKey::Shift, EKeys::LeftBracket ) );
 
 	
+	UI_COMMAND(SelectMode, "Select Mode", "Select objects", EUserInterfaceActionType::ToggleButton, FInputChord(EKeys::Q));
 	UI_COMMAND( TranslateMode, "Translate Mode", "Select and translate objects", EUserInterfaceActionType::ToggleButton, FInputChord(EKeys::W) );
 	UI_COMMAND( RotateMode, "Rotate Mode", "Select and rotate objects", EUserInterfaceActionType::ToggleButton, FInputChord(EKeys::E) );
 	UI_COMMAND( ScaleMode, "Scale Mode", "Select and scale objects", EUserInterfaceActionType::ToggleButton, FInputChord(EKeys::R) );
diff --git a/Engine/Source/Editor/UnrealEd/Private/SEditorViewport.cpp b/Engine/Source/Editor/UnrealEd/Private/SEditorViewport.cpp
--- a/Engine/Source/Editor/UnrealEd/Private/SEditorViewport.cpp
+++ b/Engine/Source/Editor/UnrealEd/Private/SEditorViewport.cpp
@@ -267,6 +267,13 @@ void SEditorViewport::BindCommands()
 		);
 
 	
+	CommandListRef.MapAction(
+		Commands.SelectMode,
+		FExecuteAction::CreateSP(ClientRef, &FEditorViewportClient::SetWidgetMode, UE::Widget::WM_None),
+		FCanExecuteAction::CreateSP(ClientRef, &FEditorViewportClient::CanSetWidgetMode, UE::Widget::WM_None),
+		FIsActionChecked::CreateSP(this, &SEditorViewport::IsWidgetModeActive, UE::Widget::WM_None)
+	);
+
 	CommandListRef.MapAction(
 		Commands.TranslateMode,
 		FExecuteAction::CreateSP( ClientRef, &FEditorViewportClient::SetWidgetMode, UE::Widget::WM_Translate ),
diff --git a/Engine/Source/Editor/UnrealEd/Private/STransformViewportToolbar.cpp b/Engine/Source/Editor/UnrealEd/Private/STransformViewportToolbar.cpp
--- a/Engine/Source/Editor/UnrealEd/Private/STransformViewportToolbar.cpp
+++ b/Engine/Source/Editor/UnrealEd/Private/STransformViewportToolbar.cpp
@@ -179,7 +179,12 @@ TSharedRef< SWidget > STransformViewportToolBar::MakeTransformToolBar( const TSh
 	ToolbarBuilder.BeginSection("Transform");
 	{
 		ToolbarBuilder.BeginBlockGroup();
-		// Move Mode
+
+		// Select Mode
+		static FName SelectModeName = FName(TEXT("SelectMode"));
+		ToolbarBuilder.AddToolBarButton(FEditorViewportCommands::Get().SelectMode, NAME_None, TAttribute<FText>(), TAttribute<FText>(), TAttribute<FSlateIcon>(), SelectModeName);
+
+		// Translate Mode
 		static FName TranslateModeName = FName(TEXT("TranslateMode"));
 		ToolbarBuilder.AddToolBarButton( FEditorViewportCommands::Get().TranslateMode, NAME_None, TAttribute<FText>(), TAttribute<FText>(), TAttribute<FSlateIcon>(), TranslateModeName );
 		
diff --git a/Engine/Source/Editor/UnrealEd/Public/EditorViewportCommands.h b/Engine/Source/Editor/UnrealEd/Public/EditorViewportCommands.h
--- a/Engine/Source/Editor/UnrealEd/Public/EditorViewportCommands.h
+++ b/Engine/Source/Editor/UnrealEd/Public/EditorViewportCommands.h
@@ -159,6 +159,9 @@ public:
 	/** Captures the viewport and updates the project thumbnail png file */
 	TSharedPtr< FUICommandInfo > ScreenCaptureForProjectThumbnail;
 
+	/** Select Mode */
+	TSharedPtr< FUICommandInfo > SelectMode;
+
 	/** Translate Mode */
 	TSharedPtr< FUICommandInfo > TranslateMode;
 
