Apply "AdditionalLaunchParameters" to an app commandline when using Launch On from editor
Remove AdditionalLaunchParametersForMobile

diff --git a/Engine/Source/Editor/UnrealEd/Classes/Settings/LevelEditorPlaySettings.h b/Engine/Source/Editor/UnrealEd/Classes/Settings/LevelEditorPlaySettings.h
--- a/Engine/Source/Editor/UnrealEd/Classes/Settings/LevelEditorPlaySettings.h
+++ b/Engine/Source/Editor/UnrealEd/Classes/Settings/LevelEditorPlaySettings.h
@@ -293,10 +293,6 @@ public:
 	UPROPERTY(config, EditAnywhere, Category=PlayInStandaloneGame)
 	FString AdditionalLaunchParameters;
 
-	/** Extra parameters to be included as part of the command line for a mobile-on-PC standalone game. */
-	UPROPERTY(config, EditAnywhere, Category=PlayInStandaloneGame)
-	FString AdditionalLaunchParametersForMobile;
-
 public:
 
 	/** Whether to build the game before launching on device. */
diff --git a/Engine/Source/Editor/UnrealEd/Private/PlayLevelNewProcess.cpp b/Engine/Source/Editor/UnrealEd/Private/PlayLevelNewProcess.cpp
--- a/Engine/Source/Editor/UnrealEd/Private/PlayLevelNewProcess.cpp
+++ b/Engine/Source/Editor/UnrealEd/Private/PlayLevelNewProcess.cpp
@@ -199,13 +199,6 @@ void UEditorEngine::LaunchNewProcess(const FRequestPlaySessionParams& InParams,
 		CommandLine += FString::Printf(TEXT(" %s"), *InParams.AdditionalStandaloneCommandLineParameters.GetValue());
 	}
 
-	// Mobile uses its own set of command line arguments that can be passed. Mobile can be previewed both in-process and
-	// standalone, so we need two separate sets of arguments right now.
-	if (InParams.EditorPlaySettings->AdditionalLaunchParametersForMobile.Len() > 0)
-	{
-		CommandLine += FString::Printf(TEXT(" %s"), *InParams.EditorPlaySettings->AdditionalLaunchParametersForMobile);
-	}
-
 	// Allow servers to override which port they are launched on.
 	if (NetMode == EPlayNetMode::PIE_ListenServer)
 	{
diff --git a/Engine/Source/Editor/UnrealEd/Private/PlayLevelViaLauncher.cpp b/Engine/Source/Editor/UnrealEd/Private/PlayLevelViaLauncher.cpp
--- a/Engine/Source/Editor/UnrealEd/Private/PlayLevelViaLauncher.cpp
+++ b/Engine/Source/Editor/UnrealEd/Private/PlayLevelViaLauncher.cpp
@@ -179,6 +179,10 @@ void UEditorEngine::StartPlayUsingLauncherSession(FRequestPlaySessionParams& InR
 	ILauncherProfileRef LauncherProfile = LauncherServicesModule.CreateProfile(TEXT("Launch On Device"));
 	LauncherProfile->SetBuildMode(BuildMode);
 	LauncherProfile->SetBuildConfiguration(BuildConfiguration);
+	if (InRequestParams.EditorPlaySettings && !InRequestParams.EditorPlaySettings->AdditionalLaunchParameters.IsEmpty())
+	{
+		LauncherProfile->SetAdditionalCommandLineParameters(InRequestParams.EditorPlaySettings->AdditionalLaunchParameters);
+	}
 
 	// select the quickest cook mode based on which in editor cook mode is enabled
 	bool bIncrimentalCooking = true;
