Single place tool places static meshes as ISMs when ISM element handles are enabled.

Also fixed a bug where the ISM component was not rebuilding render data when using the placement subsystem.

diff --git a/Engine/Source/Editor/AssetPlacementEdMode/Private/Tools/PlacementBrushToolBase.cpp b/Engine/Source/Editor/AssetPlacementEdMode/Private/Tools/PlacementBrushToolBase.cpp
--- a/Engine/Source/Editor/AssetPlacementEdMode/Private/Tools/PlacementBrushToolBase.cpp
+++ b/Engine/Source/Editor/AssetPlacementEdMode/Private/Tools/PlacementBrushToolBase.cpp
@@ -280,6 +280,7 @@ TArray<FTypedElementHandle> UPlacementBrushToolBase::GetElementsInBrushRadius()
 		}
 	}
 
+#if !UE_ENABLE_SMINSTANCE_ELEMENTS
 	// Handle the IFA for the brush stroke level
 	if (UActorPartitionSubsystem* PartitionSubsystem = UWorld::GetSubsystem<UActorPartitionSubsystem>(GEditor->GetEditorWorldContext().World()))
 	{
@@ -305,6 +306,7 @@ TArray<FTypedElementHandle> UPlacementBrushToolBase::GetElementsInBrushRadius()
 			}
 		}
 	}
+#endif
 
 	return ElementHandles;
 }
diff --git a/Engine/Source/Editor/AssetPlacementEdMode/Private/Tools/PlacementPlaceSingleTool.cpp b/Engine/Source/Editor/AssetPlacementEdMode/Private/Tools/PlacementPlaceSingleTool.cpp
--- a/Engine/Source/Editor/AssetPlacementEdMode/Private/Tools/PlacementPlaceSingleTool.cpp
+++ b/Engine/Source/Editor/AssetPlacementEdMode/Private/Tools/PlacementPlaceSingleTool.cpp
@@ -62,8 +62,12 @@ void UPlacementModePlaceSingleTool::OnClickPress(const FInputDeviceRay& PressPos
 	// Place the Preview data if we managed to get to a valid handled click.
 	FPlacementOptions PlacementOptions;
 	PlacementOptions.bPreferBatchPlacement = true;
-	PlacementOptions.bPreferInstancedPlacement = false;	// Off until ISM single selection support added.
-
+#if UE_ENABLE_SMINSTANCE_ELEMENTS
+	PlacementOptions.bPreferInstancedPlacement = true;
+#else
+	PlacementOptions.bPreferInstancedPlacement = false;
+#endif
+	
 	UPlacementSubsystem* PlacementSubsystem = GEditor->GetEditorSubsystem<UPlacementSubsystem>();
 	if (PlacementSubsystem)
 	{
diff --git a/Engine/Source/Editor/EditorFramework/Private/Subsystems/PlacementSubsystem.cpp b/Engine/Source/Editor/EditorFramework/Private/Subsystems/PlacementSubsystem.cpp
--- a/Engine/Source/Editor/EditorFramework/Private/Subsystems/PlacementSubsystem.cpp
+++ b/Engine/Source/Editor/EditorFramework/Private/Subsystems/PlacementSubsystem.cpp
@@ -31,8 +31,8 @@ TArray<FTypedElementHandle> UPlacementSubsystem::PlaceAsset(const FAssetPlacemen
 TArray<FTypedElementHandle> UPlacementSubsystem::PlaceAssets(TArrayView<const FAssetPlacementInfo> InPlacementInfos, const FPlacementOptions& InPlacementOptions)
 {
 	TGuardValue<bool> bShouldCreatePreviewElements(bIsCreatingPreviewElements, InPlacementOptions.bIsCreatingPreviewElements);
+	TMap<IAssetFactoryInterface*, TArray<FTypedElementHandle>> PlacedElementsByFactory;
 
-	TArray<FTypedElementHandle> PlacedElements;
 	for (const FAssetPlacementInfo& PlacementInfo : InPlacementInfos)
 	{
 		const FAssetData& AssetData = PlacementInfo.AssetToPlace;
@@ -47,6 +47,12 @@ TArray<FTypedElementHandle> UPlacementSubsystem::PlaceAssets(TArrayView<const FA
 			continue;
 		}
 
+		if (!PlacedElementsByFactory.Contains(&*FactoryInterface))
+		{
+			FactoryInterface->BeginPlacement(InPlacementOptions);
+			PlacedElementsByFactory.Add(&*FactoryInterface);
+		}
+
 		FAssetPlacementInfo AdjustedPlacementInfo = PlacementInfo;
 		if (!FactoryInterface->PrePlaceAsset(AdjustedPlacementInfo, InPlacementOptions))
 		{
@@ -57,10 +63,17 @@ TArray<FTypedElementHandle> UPlacementSubsystem::PlaceAssets(TArrayView<const FA
 		if (PlacedHandles.Num())
 		{
 			FactoryInterface->PostPlaceAsset(PlacedHandles, PlacementInfo, InPlacementOptions);
-			PlacedElements.Append(PlacedHandles);
+			PlacedElementsByFactory[&*FactoryInterface].Append(MoveTemp(PlacedHandles));
 		}
 	}
 
+	TArray<FTypedElementHandle> PlacedElements;
+	for (TMap<IAssetFactoryInterface*, TArray<FTypedElementHandle>>::ElementType& FactoryAndPlacedElementsPair : PlacedElementsByFactory)
+	{
+		FactoryAndPlacedElementsPair.Key->EndPlacement(FactoryAndPlacedElementsPair.Value, InPlacementOptions);
+		PlacedElements.Append(MoveTemp(FactoryAndPlacedElementsPair.Value));
+	}
+
 	return PlacedElements;
 }
 
diff --git a/Engine/Source/Editor/EditorFramework/Public/Factories/AssetFactoryInterface.h b/Engine/Source/Editor/EditorFramework/Public/Factories/AssetFactoryInterface.h
--- a/Engine/Source/Editor/EditorFramework/Public/Factories/AssetFactoryInterface.h
+++ b/Engine/Source/Editor/EditorFramework/Public/Factories/AssetFactoryInterface.h
@@ -58,4 +58,20 @@ public:
 	 * @returns the FAssetData which corresponds to the placed FTypedElementHandle. The returned FAssetData may be invalid.
 	 */
 	virtual FAssetData GetAssetDataFromElementHandle(const FTypedElementHandle& InHandle) = 0;
+
+	/**
+	 * Sets up any states that the factory needs to track across a placement call.
+	 * Called once per factory during a PlaceAssets call from the placement subsystem.
+	 * 
+	 * For example, for instanced static mesh placement, we may want to disable rebuilding the parent component's tree until EndPlacement is called.
+	 */
+	virtual void BeginPlacement(const FPlacementOptions& InPlacementOptions) = 0;
+
+	/**
+	 * Tears down any state, or kicks off any rebuilds of data that this asset type may need.
+	 * Called once per factory during a PlaceAssets call from the placement subsystem.
+	 *
+	 * For example, for instanced static mesh placement, we'd rebuild the parent component's tree here.
+	 */
+	 virtual void EndPlacement(TArrayView<const FTypedElementHandle> InPlacedElements, const FPlacementOptions& InPlacementOptions) = 0;
 };
diff --git a/Engine/Source/Editor/UnrealEd/Classes/ActorFactories/ActorFactory.h b/Engine/Source/Editor/UnrealEd/Classes/ActorFactories/ActorFactory.h
--- a/Engine/Source/Editor/UnrealEd/Classes/ActorFactories/ActorFactory.h
+++ b/Engine/Source/Editor/UnrealEd/Classes/ActorFactories/ActorFactory.h
@@ -84,6 +84,8 @@ class UNREALED_API UActorFactory : public UObject, public IAssetFactoryInterface
 	virtual TArray<FTypedElementHandle> PlaceAsset(const FAssetPlacementInfo& InPlacementInfo, const FPlacementOptions& InPlacementOptions) override;
 	virtual void PostPlaceAsset(TArrayView<const FTypedElementHandle> InHandle, const FAssetPlacementInfo& InPlacementInfo, const FPlacementOptions& InPlacementOptions) override;
 	virtual FAssetData GetAssetDataFromElementHandle(const FTypedElementHandle& InHandle) override;
+	virtual void BeginPlacement(const FPlacementOptions& InPlacementOptions) override;
+	virtual void EndPlacement(TArrayView<const FTypedElementHandle> InPlacedElements, const FPlacementOptions& InPlacementOptions) override;
 	// End IAssetFactoryInterface Interface
 
 protected:
diff --git a/Engine/Source/Editor/UnrealEd/Private/Factories/ActorFactory.cpp b/Engine/Source/Editor/UnrealEd/Private/Factories/ActorFactory.cpp
--- a/Engine/Source/Editor/UnrealEd/Private/Factories/ActorFactory.cpp
+++ b/Engine/Source/Editor/UnrealEd/Private/Factories/ActorFactory.cpp
@@ -374,6 +374,14 @@ FAssetData UActorFactory::GetAssetDataFromElementHandle(const FTypedElementHandl
 	return FAssetData();
 }
 
+void UActorFactory::BeginPlacement(const FPlacementOptions& InPlacementOptions)
+{
+}
+
+void UActorFactory::EndPlacement(TArrayView<const FTypedElementHandle> InPlacedElements, const FPlacementOptions& InPlacementOptions)
+{
+}
+
 AActor* UActorFactory::CreateActor( UObject* Asset, ULevel* InLevel, FTransform SpawnTransform, EObjectFlags InObjectFlags, const FName Name )
 {
 	AActor* NewActor = nullptr;
diff --git a/Engine/Source/Editor/UnrealEd/Private/Factories/EditorStaticMeshFactory.cpp b/Engine/Source/Editor/UnrealEd/Private/Factories/EditorStaticMeshFactory.cpp
--- a/Engine/Source/Editor/UnrealEd/Private/Factories/EditorStaticMeshFactory.cpp
+++ b/Engine/Source/Editor/UnrealEd/Private/Factories/EditorStaticMeshFactory.cpp
@@ -20,7 +20,7 @@
 TArray<FTypedElementHandle> UEditorStaticMeshFactory::PlaceAsset(const FAssetPlacementInfo& InPlacementInfo, const FPlacementOptions& InPlacementOptions)
 {
 	// If we're disallowing instanced placement, or creating preview elements, don't use the ISM placement.
-	if (!InPlacementOptions.bPreferInstancedPlacement || InPlacementOptions.bIsCreatingPreviewElements)
+	if (!ShouldPlaceInstancedStaticMeshes(InPlacementOptions))
 	{
 		return Super::PlaceAsset(InPlacementInfo, InPlacementOptions);
 	}
@@ -61,10 +61,13 @@ TArray<FTypedElementHandle> UEditorStaticMeshFactory::PlaceAsset(const FAssetPla
 			FoliagePlacementInfo.Rotation = InPlacementInfo.FinalizedTransform.Rotator();
 			FoliagePlacementInfo.DrawScale3D = InPlacementInfo.FinalizedTransform.GetScale3D();
 			FoliageInfo->AddInstance(FoundOrCreatedType, FoliagePlacementInfo);
+			CurrentPlacementScopedFoliageInfos.Add(FoliageInfo);
 
-			// Todo: deal with returning instanced handles
-			// In the meantime, we return the handle of the the HISM component added.
+#if UE_ENABLE_SMINSTANCE_ELEMENTS
+			return TArray<FTypedElementHandle>({ UEngineElementsLibrary::AcquireEditorSMInstanceElementHandle(FoliageInfo->GetComponent(), FoliageInfo->GetPlacedInstanceCount() - 1) });
+#else
 			return TArray<FTypedElementHandle>({ UEngineElementsLibrary::AcquireEditorComponentElementHandle(FoliageInfo->GetComponent()) });
+#endif
 		}
 	}
 
@@ -107,3 +110,21 @@ FAssetData UEditorStaticMeshFactory::GetAssetDataFromElementHandle(const FTypedE
 
 	return Super::GetAssetDataFromElementHandle(InHandle);
 }
+
+void UEditorStaticMeshFactory::EndPlacement(TArrayView<const FTypedElementHandle> InPlacedElements, const FPlacementOptions& InPlacementOptions)
+{
+	if (ShouldPlaceInstancedStaticMeshes(InPlacementOptions))
+	{
+		for (FFoliageInfo* FoliageInfo : CurrentPlacementScopedFoliageInfos)
+		{
+			FoliageInfo->Refresh(true, false);
+		}
+	}
+
+	CurrentPlacementScopedFoliageInfos.Empty();
+}
+
+bool UEditorStaticMeshFactory::ShouldPlaceInstancedStaticMeshes(const FPlacementOptions& InPlacementOptions) const
+{
+	return !InPlacementOptions.bIsCreatingPreviewElements && InPlacementOptions.bPreferInstancedPlacement;
+}
diff --git a/Engine/Source/Editor/UnrealEd/Private/Factories/EditorStaticMeshFactory.h b/Engine/Source/Editor/UnrealEd/Private/Factories/EditorStaticMeshFactory.h
--- a/Engine/Source/Editor/UnrealEd/Private/Factories/EditorStaticMeshFactory.h
+++ b/Engine/Source/Editor/UnrealEd/Private/Factories/EditorStaticMeshFactory.h
@@ -7,6 +7,8 @@
 
 #include "EditorStaticMeshFactory.generated.h"
 
+struct FFoliageInfo;
+
 UCLASS(Transient)
 class UEditorStaticMeshFactory : public UActorFactoryStaticMesh
 {
@@ -16,5 +18,10 @@ public:
 	// Begin IAssetFactoryInterface
 	virtual TArray<FTypedElementHandle> PlaceAsset(const FAssetPlacementInfo& InPlacementInfo, const FPlacementOptions& InPlacementOptions) override;
 	virtual FAssetData GetAssetDataFromElementHandle(const FTypedElementHandle& InHandle) override;
+	virtual void EndPlacement(TArrayView<const FTypedElementHandle> InPlacedElements, const FPlacementOptions& InPlacementOptions) override;
 	// End IAssetFactoryInterface
+
+protected:
+	bool ShouldPlaceInstancedStaticMeshes(const FPlacementOptions& InPlacementOptions) const;
+	TSet<FFoliageInfo*> CurrentPlacementScopedFoliageInfos;
 };
\ No newline at end of file
