Fix mismatched startup/shutdown assert

diff --git a/Engine/Source/Editor/UnrealEd/Private/UnrealEdEngine.cpp b/Engine/Source/Editor/UnrealEd/Private/UnrealEdEngine.cpp
--- a/Engine/Source/Editor/UnrealEd/Private/UnrealEdEngine.cpp
+++ b/Engine/Source/Editor/UnrealEd/Private/UnrealEdEngine.cpp
@@ -228,7 +228,14 @@ void UUnrealEdEngine::Init(IEngineLoop* InEngineLoop)
 		});
 
 	// Stop tracking stalls when we close the Main Frame, the closest event around
-	OnEditorClose().AddStatic(&UE::FStallDetector::Shutdown);
+	OnEditorClose().AddLambda([]()
+		{
+			// We conditionally shutdown here because for now there is no UnreadEdEngine specific shutdown 
+			if (UE::FStallDetector::IsRunning())
+			{
+				UE::FStallDetector::Shutdown();
+			}
+		});
 #endif
 }
 
