Updated DerivedDataCache commandlet to accept -MapIniSection commandline parameter (like the cooker does)

(preflight was green but submission failed)

diff --git a/Engine/Source/Editor/UnrealEd/Private/Commandlets/DerivedDataCacheCommandlet.cpp b/Engine/Source/Editor/UnrealEd/Private/Commandlets/DerivedDataCacheCommandlet.cpp
--- a/Engine/Source/Editor/UnrealEd/Private/Commandlets/DerivedDataCacheCommandlet.cpp
+++ b/Engine/Source/Editor/UnrealEd/Private/Commandlets/DerivedDataCacheCommandlet.cpp
@@ -121,6 +121,36 @@ int32 UDerivedDataCacheCommandlet::Main( const FString& Params )
 			Tokens.Add(FString("*") + FPackageName::GetMapPackageExtension());
 		}
 
+		// support MapIniSection parameter
+		{
+			TArray<FString> MapIniSections;
+			FString SectionStr;
+			if (FParse::Value(*Params, TEXT("MAPINISECTION="), SectionStr))
+			{
+				if (SectionStr.Contains(TEXT("+")))
+				{
+					TArray<FString> Sections;
+					SectionStr.ParseIntoArray(Sections, TEXT("+"), true);
+					for (int32 Index = 0; Index < Sections.Num(); Index++)
+					{
+						MapIniSections.Add(Sections[Index]);
+					}
+				}
+				else
+				{
+					MapIniSections.Add(SectionStr);
+				}
+
+				TArray<FString> MapsFromIniSection;
+				for (const FString& MapIniSection : MapIniSections)
+				{
+					GEditor->LoadMapListFromIni(*MapIniSection, MapsFromIniSection);
+				}
+
+				Tokens += MapsFromIniSection;
+			}
+		}
+
 		uint8 PackageFilter = NORMALIZE_DefaultFlags;
 		if ( Switches.Contains(TEXT("MAPSONLY")) )
 		{
