Cleaning up Widgets within Menus

diff --git a/Engine/Source/Editor/TurnkeySupport/Private/TurnkeySupportModule.cpp b/Engine/Source/Editor/TurnkeySupport/Private/TurnkeySupportModule.cpp
--- a/Engine/Source/Editor/TurnkeySupport/Private/TurnkeySupportModule.cpp
+++ b/Engine/Source/Editor/TurnkeySupport/Private/TurnkeySupportModule.cpp
@@ -1013,9 +1013,13 @@ static void MakeTurnkeyPlatformMenu(FMenuBuilder& MenuBuilder, FName IniPlatform
 	}
 
 	MenuBuilder.AddWidget(
-		SNew(STextBlock)
-		.ColorAndOpacity(FSlateColor::UseSubduedForeground())
-		.Text(FormatSdkInfo(SdkInfo, true)),
+		SNew(SBox)
+		.Padding(FMargin(16.0f, 3.0f))
+		[
+			SNew(STextBlock)
+			.ColorAndOpacity(FSlateColor::UseSubduedForeground())
+			.Text(FormatSdkInfo(SdkInfo, true))
+		],
 		FText::GetEmpty()
 	);
 
diff --git a/Engine/Source/Editor/UnrealEd/Private/Kismet2/DebuggerCommands.cpp b/Engine/Source/Editor/UnrealEd/Private/Kismet2/DebuggerCommands.cpp
--- a/Engine/Source/Editor/UnrealEd/Private/Kismet2/DebuggerCommands.cpp
+++ b/Engine/Source/Editor/UnrealEd/Private/Kismet2/DebuggerCommands.cpp
@@ -877,7 +877,7 @@ TSharedRef< SWidget > FPlayWorldCommands::GeneratePlayMenuContent(TSharedRef<FUI
 			Section.AddEntry(FToolMenuEntry::InitWidget(
 				"PlayIn",
 				SNew(SBox)
-				.Padding(FMargin(10.0f, 0.0f))
+				.Padding(FMargin(16.0f, 3.0f))
 				[
 					SNew(STextBlock)
 					.ColorAndOpacity(FSlateColor::UseSubduedForeground())
diff --git a/Engine/Source/Editor/UnrealEd/Private/STransformViewportToolbar.cpp b/Engine/Source/Editor/UnrealEd/Private/STransformViewportToolbar.cpp
--- a/Engine/Source/Editor/UnrealEd/Private/STransformViewportToolbar.cpp
+++ b/Engine/Source/Editor/UnrealEd/Private/STransformViewportToolbar.cpp
@@ -95,38 +95,37 @@ TSharedRef<SWidget> STransformViewportToolBar::GenerateSurfaceSnappingMenu()
 				NAME_None,
 			EUserInterfaceActionType::RadioButton);
 
+
 		MenuBuilder.AddWidget(
-			SNew(SHorizontalBox)
-			.IsEnabled(TAttribute<bool>::Create(TAttribute<bool>::FGetter::CreateStatic(IsSnappingEnabled)))
 
-			+ SHorizontalBox::Slot()
-			.AutoWidth()
-			.Padding(FMargin(0.f, 0.f, 5.f, 0.f))
-			[
-				SNew(STextBlock)
-				.Text(LOCTEXT("SnapToSurfaceSettings_Offset", "Surface Offset"))
-			]
-			+ SHorizontalBox::Slot()
-			.VAlign(VAlign_Bottom)
-			.FillWidth(1.f)
+			SNew( SBox )
+			.Padding( FMargin(8.0f, 0.0f, 0.0f, 0.0f) )
+			.MinDesiredWidth( 100.0f )
 			[
-				SNew(SNumericEntryBox<float>)
-				.Value(
-					TAttribute<TOptional<float>>::Create(TAttribute<TOptional<float>>::FGetter::CreateStatic([] {
-						const auto& Settings = GetDefault<ULevelEditorViewportSettings>()->SnapToSurface;
-						return TOptional<float>(Settings.SnapOffsetExtent);
+				SNew ( SBorder )
+				.BorderImage(FAppStyle::Get().GetBrush("Menu.WidgetBorder"))
+				.Padding(FMargin(1.0f))
+				[
+					SNew(SNumericEntryBox<float>)
+					.IsEnabled(TAttribute<bool>::Create(TAttribute<bool>::FGetter::CreateStatic(IsSnappingEnabled)))
+					.Value(
+						TAttribute<TOptional<float>>::Create(TAttribute<TOptional<float>>::FGetter::CreateStatic([] 
+						{
+							const auto& Settings = GetDefault<ULevelEditorViewportSettings>()->SnapToSurface;
+							return TOptional<float>(Settings.SnapOffsetExtent);
 						}))
-				)
-				.OnValueChanged(
-				SNumericEntryBox<float>::FOnValueChanged::CreateStatic([](float Val) {
-					GetMutableDefault<ULevelEditorViewportSettings>()->SnapToSurface.SnapOffsetExtent = Val;
+					)
+					.OnValueChanged(SNumericEntryBox<float>::FOnValueChanged::CreateStatic([](float Val) 
+					{
+						GetMutableDefault<ULevelEditorViewportSettings>()->SnapToSurface.SnapOffsetExtent = Val;
 					}))
 					.MinValue(0.f)
 					.MaxValue(HALF_WORLD_MAX)
 					.MaxSliderValue(1000.f) // 'Sensible' range for the slider (10m)
 					.AllowSpin(true)
+				]
 			],
-			FText::GetEmpty()
+			LOCTEXT("SnapToSurfaceSettings_Offset", "Surface Offset")
 		);
 
 	}
