Fix wrong text size in viewport toolbars

diff --git a/Engine/Source/Editor/UnrealEd/Private/SEditorViewportToolBarMenu.cpp b/Engine/Source/Editor/UnrealEd/Private/SEditorViewportToolBarMenu.cpp
--- a/Engine/Source/Editor/UnrealEd/Private/SEditorViewportToolBarMenu.cpp
+++ b/Engine/Source/Editor/UnrealEd/Private/SEditorViewportToolBarMenu.cpp
@@ -61,7 +61,6 @@ void SEditorViewportToolbarMenu::Construct( const FArguments& Declaration )
 				.VAlign(VAlign_Center)
 				[
 					SNew(STextBlock)
-					.TextStyle(&FAppStyle::Get().GetWidgetStyle<FTextBlockStyle>("SmallText"))
 					.Text(Label)
 				];
 		}
@@ -74,7 +73,6 @@ void SEditorViewportToolbarMenu::Construct( const FArguments& Declaration )
 				.Padding(FMargin(4.f, 0.f))
 				[
 					SNew( STextBlock )
-					.TextStyle(&FAppStyle::Get().GetWidgetStyle<FTextBlockStyle>("SmallText"))
 					.Text(Label)
 				];
 		}
