ExternalActorsCommandlet: new commandlet to validate the integrity of external actor packages.

Options:
-repair: open files for delete instead of only reporting invalid actor files.
-loadpackages: load actor packages instead of only using the asset-registry to perform file validations.

diff --git a/Engine/Source/Editor/UnrealEd/Classes/Commandlets/ExternalActorsCommandlet.h b/Engine/Source/Editor/UnrealEd/Classes/Commandlets/ExternalActorsCommandlet.h
new file mode 100644
--- /dev/null
+++ b/Engine/Source/Editor/UnrealEd/Classes/Commandlets/ExternalActorsCommandlet.h
@@ -0,0 +1,25 @@
+// Copyright Epic Games, Inc. All Rights Reserved.
+
+#pragma once
+
+#include "CoreMinimal.h"
+#include "UObject/ObjectMacros.h"
+#include "Commandlets/Commandlet.h"
+#include "PackageSourceControlHelper.h"
+#include "ExternalActorsCommandlet.generated.h"
+
+UCLASS()
+class UNREALED_API UExternalActorsCommandlet : public UCommandlet
+{
+	GENERATED_UCLASS_BODY()
+
+public:
+	//~ Begin UCommandlet Interface
+	virtual int32 Main(const FString& Params) override;
+	//~ End UCommandlet Interface
+
+protected:
+	UWorld* LoadWorld(const FString& LevelToLoad);
+
+	FPackageSourceControlHelper PackageHelper;
+};
diff --git a/Engine/Source/Editor/UnrealEd/Private/Commandlets/ExternalActorsCommandlet.cpp b/Engine/Source/Editor/UnrealEd/Private/Commandlets/ExternalActorsCommandlet.cpp
new file mode 100644
--- /dev/null
+++ b/Engine/Source/Editor/UnrealEd/Private/Commandlets/ExternalActorsCommandlet.cpp
@@ -0,0 +1,131 @@
+// Copyright Epic Games, Inc. All Rights Reserved.
+
+#include "Commandlets/ExternalActorsCommandlet.h"
+#include "HAL/FileManager.h"
+#include "Misc/Paths.h"
+#include "Editor.h"
+#include "Engine/Level.h"
+#include "Engine/World.h"
+#include "UObject/UObjectHash.h"
+#include "AssetRegistryModule.h"
+#include "PackageHelperFunctions.h"
+#include "WorldPartition/ActorDescContainer.h"
+
+DEFINE_LOG_CATEGORY_STATIC(LogExternalActorsCommandlet, All, All);
+
+UExternalActorsCommandlet::UExternalActorsCommandlet(const FObjectInitializer& ObjectInitializer)
+: Super(ObjectInitializer)
+{
+}
+
+UWorld* UExternalActorsCommandlet::LoadWorld(const FString& LevelToLoad)
+{
+	TRACE_CPUPROFILER_EVENT_SCOPE(UWorldPartitionConvertCommandlet::LoadWorld);
+
+	SET_WARN_COLOR(COLOR_WHITE);
+	UE_LOG(LogExternalActorsCommandlet, Log, TEXT("Loading level %s."), *LevelToLoad);
+	CLEAR_WARN_COLOR();
+
+	UPackage* MapPackage = LoadPackage(nullptr, *LevelToLoad, LOAD_None);
+	if (!MapPackage)
+	{
+		UE_LOG(LogExternalActorsCommandlet, Error, TEXT("Error loading %s."), *LevelToLoad);
+		return nullptr;
+	}
+
+	return UWorld::FindWorldInPackage(MapPackage);
+}
+
+int32 UExternalActorsCommandlet::Main(const FString& Params)
+{
+	TArray<FString> Tokens, Switches;
+	ParseCommandLine(*Params, Tokens, Switches);
+
+	// Need at least the level to convert
+	if (Tokens.Num() < 1)
+	{
+		UE_LOG(LogExternalActorsCommandlet, Error, TEXT("Bad parameters"));
+		return 1;
+	}
+
+	const bool bRepair = Switches.Contains(TEXT("repair"));
+	const bool bLoadPackages = Switches.Contains(TEXT("loadpackages"));
+
+	if (!FPackageName::SearchForPackageOnDisk(Tokens[0], &Tokens[0]))
+	{
+		UE_LOG(LogExternalActorsCommandlet, Error, TEXT("Unknown level '%s'"), *Tokens[0]);
+		return 1;
+	}
+
+	// Load world
+	UWorld* MainWorld = bLoadPackages ? LoadWorld(Tokens[0]) : nullptr;
+	if (bLoadPackages && !MainWorld)
+	{
+		UE_LOG(LogExternalActorsCommandlet, Error, TEXT("Unknown world '%s'"), *Tokens[0]);
+		return 1;
+	}
+
+	// Initialize actor descriptors
+	UActorDescContainer* ActorDescContainer = NewObject<UActorDescContainer>(GetTransientPackage());
+	ActorDescContainer->Initialize(MainWorld, *Tokens[0], false);
+
+	// Validate external actors
+	FString ExternalActorsPath = ULevel::GetExternalActorsPath(Tokens[0]);
+	FString ExternalActorsFilePath = FPackageName::LongPackageNameToFilename(ExternalActorsPath);
+
+	TArray<FString> PackagesToDelete;
+	if (IFileManager::Get().DirectoryExists(*ExternalActorsFilePath))
+	{
+		bool bResult = IFileManager::Get().IterateDirectoryRecursively(*ExternalActorsFilePath, [this, ActorDescContainer, bRepair, bLoadPackages, &PackagesToDelete](const TCHAR* FilenameOrDirectory, bool bIsDirectory)
+		{
+			if (!bIsDirectory)
+			{
+				FString Filename(FilenameOrDirectory);
+				if (Filename.EndsWith(FPackageName::GetAssetPackageExtension()))
+				{
+					bool bIsValidActorPackage = false;
+					const FString PackageName = FPackageName::FilenameToLongPackageName(*Filename);
+					const FWorldPartitionActorDesc* ActorDesc = nullptr;
+
+					if (bLoadPackages)
+					{
+						if (UPackage* Package = LoadPackage(nullptr, *Filename, LOAD_None, nullptr, nullptr))
+						{						
+							ForEachObjectWithPackage(Package, [&bIsValidActorPackage](UObject* Object)
+							{
+								Object->ClearFlags(RF_Public | RF_Standalone);
+								bIsValidActorPackage = Object->IsA<AActor>();
+								return !bIsValidActorPackage;
+							});
+						}
+					}
+					else
+					{
+						ActorDesc = ActorDescContainer->GetActorDesc(PackageName);
+						bIsValidActorPackage = !!ActorDesc;
+					}
+
+					if (!bIsValidActorPackage)
+					{
+						UE_LOG(LogExternalActorsCommandlet, Error, TEXT("Invalid actor file '%s' (%s)"), *Filename, ActorDesc ? *ActorDesc->ToString() : TEXT("None"));
+
+						if (bRepair)
+						{
+							PackagesToDelete.Add(Filename);
+						}
+					}
+				}
+			}
+			return true;
+		});
+	}
+
+	CollectGarbage(GARBAGE_COLLECTION_KEEPFLAGS);
+
+	for (const FString& PackageToDelete : PackagesToDelete)
+	{
+		PackageHelper.Delete(*PackageToDelete);
+	}
+
+	return 0;
+}
\ No newline at end of file
diff --git a/Engine/Source/Runtime/Engine/Private/WorldPartition/ActorDescContainer.cpp b/Engine/Source/Runtime/Engine/Private/WorldPartition/ActorDescContainer.cpp
--- a/Engine/Source/Runtime/Engine/Private/WorldPartition/ActorDescContainer.cpp
+++ b/Engine/Source/Runtime/Engine/Private/WorldPartition/ActorDescContainer.cpp
@@ -246,6 +246,20 @@ FWorldPartitionActorDesc& UActorDescContainer::GetActorDescChecked(const FGuid&
 	return *ActorDesc->Get();
 }
 
+const FWorldPartitionActorDesc* UActorDescContainer::GetActorDesc(const FString& PackageName) const
+{
+	const FName PackageFName(*PackageName);
+	for (const TUniquePtr<FWorldPartitionActorDesc>& ActorDescPtr : ActorDescList)
+	{
+		if (ActorDescPtr->ActorPackage == PackageFName)
+		{
+			return ActorDescPtr.Get();
+		}
+	}
+
+	return nullptr;
+}
+
 void UActorDescContainer::RegisterDelegates()
 {
 	if (GEditor && !IsTemplate())
diff --git a/Engine/Source/Runtime/Engine/Public/WorldPartition/ActorDescContainer.h b/Engine/Source/Runtime/Engine/Public/WorldPartition/ActorDescContainer.h
--- a/Engine/Source/Runtime/Engine/Public/WorldPartition/ActorDescContainer.h
+++ b/Engine/Source/Runtime/Engine/Public/WorldPartition/ActorDescContainer.h
@@ -35,6 +35,8 @@ public:
 	FWorldPartitionActorDesc& GetActorDescChecked(const FGuid& Guid);
 	const FWorldPartitionActorDesc& GetActorDescChecked(const FGuid& Guid) const;
 
+	const FWorldPartitionActorDesc* GetActorDesc(const FString& PackageName) const;
+
 	int32 GetActorDescCount() const { return Actors.Num(); }
 	FName GetContainerPackage() const { return ContainerPackageName; }
 #endif
