Updated SetActorSelectionState to notify editor of selection change

diff --git a/Engine/Source/Editor/UnrealEd/Private/Subsystems/EditorActorSubsystem.cpp b/Engine/Source/Editor/UnrealEd/Private/Subsystems/EditorActorSubsystem.cpp
--- a/Engine/Source/Editor/UnrealEd/Private/Subsystems/EditorActorSubsystem.cpp
+++ b/Engine/Source/Editor/UnrealEd/Private/Subsystems/EditorActorSubsystem.cpp
@@ -434,6 +434,7 @@ void UEditorActorSubsystem::SetActorSelectionState(AActor* Actor, bool bShouldBe
 {
 	GEditor->GetSelectedActors()->Modify();
 	GEditor->SelectActor(Actor, bShouldBeSelected, /*bNotify=*/ false);
+	GEditor->NoteSelectionChange();
 }
 
 AActor* UEditorActorSubsystem::GetActorReference(FString PathToActor)
