Add LLM update calls in CookCommandlet
#rnx

diff --git a/Engine/Source/Editor/UnrealEd/Private/Commandlets/CookCommandlet.cpp b/Engine/Source/Editor/UnrealEd/Private/Commandlets/CookCommandlet.cpp
--- a/Engine/Source/Editor/UnrealEd/Private/Commandlets/CookCommandlet.cpp
+++ b/Engine/Source/Editor/UnrealEd/Private/Commandlets/CookCommandlet.cpp
@@ -5,6 +5,7 @@
 =============================================================================*/
 
 #include "Commandlets/CookCommandlet.h"
+#include "HAL/IConsoleManager.h"
 #include "HAL/PlatformFileManager.h"
 #include "HAL/PlatformApplicationMisc.h"
 #include "Misc/MessageDialog.h"
@@ -962,7 +963,18 @@ bool UCookCommandlet::CookByTheBook( const TArray<ITargetPlatform*>& Platforms)
 			DetailedCookStats::TargetPlatforms.RemoveFromEnd(TEXT("+"));
 		}
 	});
-	
+
+#if ENABLE_LOW_LEVEL_MEM_TRACKER
+	FLowLevelMemTracker& MemTracker = FLowLevelMemTracker::Get();
+	if (MemTracker.IsEnabled())
+	{
+		MemTracker.UpdateStatsPerFrame();
+	}
+	FDateTime LastMemTrackerTime = FDateTime::UtcNow();
+	double MemTrackerPeriodSeconds = 120.0;
+	FParse::Value(FCommandLine::Get(), TEXT("CookLLMPeriod="), MemTrackerPeriodSeconds);
+#endif
+
 	do
 	{
 		{
@@ -998,6 +1010,18 @@ bool UCookCommandlet::CookByTheBook( const TArray<ITargetPlatform*>& Platforms)
 				uint32 TickResults = 0;
 				static const float CookOnTheSideTimeSlice = 10.0f;
 
+#if ENABLE_LOW_LEVEL_MEM_TRACKER
+				if (MemTracker.IsEnabled())
+				{
+					FDateTime CurrentTime = FDateTime::UtcNow();
+					if (CurrentTime >= LastMemTrackerTime + FTimespan::FromSeconds(MemTrackerPeriodSeconds))
+					{
+						MemTracker.UpdateStatsPerFrame();
+						LastMemTrackerTime = CurrentTime;
+					}
+				}
+#endif
+
 				TickResults = CookOnTheFlyServer->TickCookOnTheSide( CookOnTheSideTimeSlice, NonMapPackageCountSinceLastGC, ShowProgress ? ECookTickFlags::None : ECookTickFlags::HideProgressDisplay );
 
 				{
@@ -1083,7 +1107,12 @@ bool UCookCommandlet::CookByTheBook( const TArray<ITargetPlatform*>& Platforms)
 					UE_LOG(LogCookCommandlet, Display, TEXT("GarbageCollection...%s (%s)"), (bPartialGC? TEXT(" partial gc") : TEXT("")), *GCReason);
 					GCReason = FString();
 
-
+#if ENABLE_LOW_LEVEL_MEM_TRACKER
+					if (MemTracker.IsEnabled())
+					{
+						MemTracker.UpdateStatsPerFrame();
+					}
+#endif
 					DumpMemStats();
 
 					CollectGarbage(RF_NoFlags);
@@ -1117,6 +1146,13 @@ bool UCookCommandlet::CookByTheBook( const TArray<ITargetPlatform*>& Platforms)
 		}
 	} while (bTestCook);
 
+#if ENABLE_LOW_LEVEL_MEM_TRACKER
+	if (MemTracker.IsEnabled())
+	{
+		MemTracker.UpdateStatsPerFrame();
+	}
+#endif
+
 	if (!bIterativeCooking && StartupOptions.DLCName.IsEmpty())
 	{
 		bool bFullReferencesExpected = !(CookOptions & ECookByTheBookOptions::SkipHardReferences);
