Change to using ensure() for verifying thumbnail dimensions.

I misunderstood the difference when I changed these to ensureAlways a while back.

diff --git a/Engine/Source/Editor/UnrealEd/Private/AssetThumbnail.cpp b/Engine/Source/Editor/UnrealEd/Private/AssetThumbnail.cpp
--- a/Engine/Source/Editor/UnrealEd/Private/AssetThumbnail.cpp
+++ b/Engine/Source/Editor/UnrealEd/Private/AssetThumbnail.cpp
@@ -1381,7 +1381,7 @@ void FAssetThumbnailPool::RefreshThumbnail( const TSharedPtr<FAssetThumbnail>& T
 	const uint32 Width = ThumbnailToRefresh->GetSize().X;
 	const uint32 Height = ThumbnailToRefresh->GetSize().Y;
 
-	if ( ensureAlways(AssetData.ObjectPath != NAME_None) && ensureAlways(Width > 0) && ensureAlways(Height > 0) )
+	if ( ensure(AssetData.ObjectPath != NAME_None) && ensure(Width > 0) && ensure(Height > 0) )
 	{
 		FThumbId ThumbId( AssetData.ObjectPath, Width, Height ) ;
 		const TSharedRef<FThumbnailInfo>* ThumbnailInfoPtr = ThumbnailToTextureMap.Find( ThumbId );
@@ -1398,7 +1398,7 @@ void FAssetThumbnailPool::SetRealTimeThumbnail( const TSharedPtr<FAssetThumbnail
 	const uint32 Width = Thumbnail->GetSize().X;
 	const uint32 Height = Thumbnail->GetSize().Y;
 
-	if ( ensureAlways(AssetData.ObjectPath != NAME_None) && ensureAlways(Width > 0) && ensureAlways(Height > 0) )
+	if (ensure(AssetData.ObjectPath != NAME_None) && ensure(Width > 0) && ensure(Height > 0) )
 	{
 		FThumbId ThumbId( AssetData.ObjectPath, Width, Height );
 		const TSharedRef<FThumbnailInfo>* ThumbnailInfoPtr = ThumbnailToTextureMap.Find( ThumbId );
