Global Tab Switcher Cleanup for EA

diff --git a/Engine/Source/Editor/UnrealEd/Private/Toolkits/SGlobalTabSwitchingDialog.cpp b/Engine/Source/Editor/UnrealEd/Private/Toolkits/SGlobalTabSwitchingDialog.cpp
--- a/Engine/Source/Editor/UnrealEd/Private/Toolkits/SGlobalTabSwitchingDialog.cpp
+++ b/Engine/Source/Editor/UnrealEd/Private/Toolkits/SGlobalTabSwitchingDialog.cpp
@@ -308,29 +308,10 @@ void SGlobalTabSwitchingDialog::OnMainTabListSelectionChanged(FTabListItemPtr In
 				.TextStyle(FEditorStyle::Get(), "ControlTabMenu.AssetTypeStyle")
 				.Text(SelectedItem->GetTypeString())
 			];
-
-		// Create the list of the tool tabs in the current active tab
-		TSharedPtr<FTabManager> SelectedAssetTabManager = SelectedItem->GetAssociatedTabManager();
-		if (SelectedAssetTabManager.IsValid())
-		{
-			FMenuBuilder ToolTabMenuBuilder(/*bShouldCloseAfterSelection=*/ true, /*CommandList=*/ nullptr);
-			ToolTabMenuBuilder.GetMultiBox()->SetStyle(&FEditorStyle::Get(), "ToolBar");
-			
-			// Local editor tabs
-			SelectedAssetTabManager->PopulateLocalTabSpawnerMenu(ToolTabMenuBuilder);
-			
-			// General tabs
-			const IWorkspaceMenuStructure& MenuStructure = WorkspaceMenu::GetMenuStructure();
-			SelectedAssetTabManager->PopulateTabSpawnerMenu(ToolTabMenuBuilder, MenuStructure.GetStructureRoot());
-
-			// Turn the builder into a widget
-			NewToolTabsContent = ToolTabMenuBuilder.MakeWidget();
-		}
 	}
 
 	NewTabItemToActivateDisplayBox->SetContent(NewTopContents);
 	NewTabItemToActivatePathBox->SetContent(NewBottomContents);
-	ToolTabsListBox->SetContent(NewToolTabsContent);
 }
 
 void SGlobalTabSwitchingDialog::OnMainTabListItemClicked(FTabListItemPtr InItem)
@@ -416,10 +397,6 @@ void SGlobalTabSwitchingDialog::Construct(const FArguments& InArgs, FVector2D In
 		.HeightOverride(40.0f)
 		.VAlign(VAlign_Center);
 
-	ToolTabsListBox =
-		SNew(SBox)
-		.Padding(FMargin(0.0f, 0.0f, 15.0f, 0.0f));
-
 	MainTabsListWidget = SNew(STabListWidget)
 		.ItemHeight(64)
 		.ListItemsSource(&MainTabsListDataSource)
@@ -428,20 +405,6 @@ void SGlobalTabSwitchingDialog::Construct(const FArguments& InArgs, FVector2D In
 		.OnMouseButtonClick(this, &SGlobalTabSwitchingDialog::OnMainTabListItemClicked)
 		.SelectionMode(ESelectionMode::Single);
 
-	TSharedRef<SWidget> ToolTabList = SNew(SVerticalBox)
-		+SVerticalBox::Slot()
-		.AutoHeight()
-		[
-			SNew(STextBlock)
-			.TextStyle(FEditorStyle::Get(), "ControlTabMenu.HeadingStyle")
-			.Text(LOCTEXT("ChangeToolsHeading", "Tool Windows"))
-		]
-		+SVerticalBox::Slot()
-		.FillHeight(1.0f)
-		[
-			ToolTabsListBox.ToSharedRef()
-		];
-
 	TSharedRef<SWidget> DocumentTabList = SNew(SVerticalBox)
 		+SVerticalBox::Slot()
 		.AutoHeight()
@@ -477,22 +440,10 @@ void SGlobalTabSwitchingDialog::Construct(const FArguments& InArgs, FVector2D In
 					NewTabItemToActivateDisplayBox.ToSharedRef()
 				]
 				+SVerticalBox::Slot()
+				.Padding(12.f, 24.f)
 				.FillHeight(1.0f)
 				[
-					SNew(SHorizontalBox)
-					+SHorizontalBox::Slot()
-					.FillWidth(0.4f)
-					.Padding(FMargin(0.0f, 0.0f, 0.0f, 0.0f))
-					[
-						ToolTabList
-					]
-
-					+SHorizontalBox::Slot()
-					.FillWidth(1.0f)
-					[
 						DocumentTabList
-					]
-
 				]
 				+SVerticalBox::Slot()
 				.AutoHeight()
diff --git a/Engine/Source/Editor/UnrealEd/Private/Toolkits/SGlobalTabSwitchingDialog.h b/Engine/Source/Editor/UnrealEd/Private/Toolkits/SGlobalTabSwitchingDialog.h
--- a/Engine/Source/Editor/UnrealEd/Private/Toolkits/SGlobalTabSwitchingDialog.h
+++ b/Engine/Source/Editor/UnrealEd/Private/Toolkits/SGlobalTabSwitchingDialog.h
@@ -72,8 +72,5 @@ private:
 	// The container widget for the indication of the path to the asset that will be activated when the dialog closes
 	TSharedPtr<SBox> NewTabItemToActivatePathBox;
 
-	// The container widget for the tool tabs list
-	TSharedPtr<SBox> ToolTabsListBox;
-
 	static bool bIsAlreadyOpen;
 };
