Cooker: Fix deadlock when preloads fail and don't update NumPreloadAllocated.
#rnx

diff --git a/Engine/Source/Editor/UnrealEd/Private/Cooker/CookPackageData.cpp b/Engine/Source/Editor/UnrealEd/Private/Cooker/CookPackageData.cpp
--- a/Engine/Source/Editor/UnrealEd/Private/Cooker/CookPackageData.cpp
+++ b/Engine/Source/Editor/UnrealEd/Private/Cooker/CookPackageData.cpp
@@ -678,13 +678,12 @@ namespace Cook
 			SetIsPreloadAttempted(true);
 			return true;
 		}
-		if (!PreloadableFile)
+		if (!PreloadableFile.Get())
 		{
 			TStringBuilder<NAME_SIZE> FileNameString;
 			GetFileName().ToString(FileNameString);
-			PreloadableFile = MakeShared<FPreloadableArchive>(FileNameString.ToString());
-			PackageDatas.GetMonitor().OnPreloadAllocatedChanged(*this, true);
-			PreloadableFile->InitializeAsync([this]()
+			PreloadableFile.Set(MakeShared<FPreloadableArchive>(FileNameString.ToString()), *this);
+			PreloadableFile.Get()->InitializeAsync([this]()
 				{
 					TStringBuilder<NAME_SIZE> FileNameString;
 					// Note this is an unsynchronized read of this->GetFilename. It is allowed because GetFilename does not change until this is destructed, and the destructor of this waits for PreloadableFile to finish initialization
@@ -699,34 +698,36 @@ namespace Cook
 				},
 				FPreloadableFile::Flags::PreloadHandle | FPreloadableFile::Flags::Prime);
 		}
-		if (!PreloadableFile->IsInitialized())
+		const TSharedPtr<FPreloadableArchive>& FilePtr = PreloadableFile.Get();
+		if (!FilePtr->IsInitialized())
 		{
 			if (GetIsUrgent())
 			{
 				// For urgent requests, wait on them to finish preloading rather than letting them run asynchronously and coming back to them later
-				PreloadableFile->WaitForInitialization();
-				check(PreloadableFile->IsInitialized());
+				FilePtr->WaitForInitialization();
+				check(FilePtr->IsInitialized());
 			}
 			else
 			{
 				return false;
 			}
 		}
-		if (PreloadableFile->TotalSize() < 0)
+		if (FilePtr->TotalSize() < 0)
 		{
+			UE_LOG(LogCook, Warning, TEXT("Failed to find file when preloading %s."), *GetFileName().ToString());
 			SetIsPreloadAttempted(true);
-			PreloadableFile.Reset();
+			PreloadableFile.Reset(*this);
 			PreloadableFileFormat = EPackageFormat::Binary;
 			return true;
 		}
 
 		TStringBuilder<NAME_SIZE> FileNameString;
 		GetFileName().ToString(FileNameString);
-		if (!IPackageResourceManager::TryRegisterPreloadableArchive(FPackagePath::FromLocalPath(FileNameString), PreloadableFile, PreloadableFileFormat))
+		if (!IPackageResourceManager::TryRegisterPreloadableArchive(FPackagePath::FromLocalPath(FileNameString), FilePtr, PreloadableFileFormat))
 		{
 			UE_LOG(LogCook, Warning, TEXT("Failed to register %s for preload."), *GetFileName().ToString());
 			SetIsPreloadAttempted(true);
-			PreloadableFile.Reset();
+			PreloadableFile.Reset(*this);
 			PreloadableFileFormat = EPackageFormat::Binary;
 			return true;
 		}
@@ -736,30 +737,46 @@ namespace Cook
 		return true;
 	}
 
+	void FPackageData::FTrackedPreloadableFilePtr::Set(TSharedPtr<FPreloadableArchive>&& InPtr, FPackageData& Owner)
+	{
+		Reset(Owner);
+		if (InPtr)
+		{
+			Ptr = MoveTemp(InPtr);
+			Owner.PackageDatas.GetMonitor().OnPreloadAllocatedChanged(Owner, true);
+		}
+	}
+
+	void FPackageData::FTrackedPreloadableFilePtr::Reset(FPackageData& Owner)
+	{
+		if (Ptr)
+		{
+			Owner.PackageDatas.GetMonitor().OnPreloadAllocatedChanged(Owner, false);
+			Ptr.Reset();
+		}
+	}
+
 	void FPackageData::ClearPreload()
 	{
+		const TSharedPtr<FPreloadableArchive>& FilePtr = PreloadableFile.Get();
 		if (GetIsPreloaded())
 		{
-			check(PreloadableFile);
+			check(FilePtr);
 			TStringBuilder<NAME_SIZE> FileNameString;
 			GetFileName().ToString(FileNameString);
 			if (IPackageResourceManager::UnRegisterPreloadableArchive(FPackagePath::FromLocalPath(FileNameString)))
 			{
 				UE_LOG(LogCook, Display, TEXT("PreloadableFile was created for %s but never used. This is wasteful and bad for cook performance."), *PackageName.ToString());
 			}
-			PreloadableFile->ReleaseCache(); // ReleaseCache to conserve memory if the Linker still has a pointer to it
+			FilePtr->ReleaseCache(); // ReleaseCache to conserve memory if the Linker still has a pointer to it
 		}
 		else
 		{
-			check(!PreloadableFile || !PreloadableFile->IsCacheAllocated());
+			check(!FilePtr || !FilePtr->IsCacheAllocated());
 		}
 
-		if (PreloadableFile)
-		{
-			PreloadableFile.Reset();
-			PreloadableFileFormat = EPackageFormat::Binary;
-			PackageDatas.GetMonitor().OnPreloadAllocatedChanged(*this, false);
-		}
+		PreloadableFile.Reset(*this);
+		PreloadableFileFormat = EPackageFormat::Binary;
 		SetIsPreloaded(false);
 		SetIsPreloadAttempted(false);
 	}
@@ -767,7 +784,7 @@ namespace Cook
 	void FPackageData::CheckPreloadEmpty()
 	{
 		check(!GetIsPreloadAttempted());
-		check(!PreloadableFile);
+		check(!PreloadableFile.Get());
 		check(!GetIsPreloaded());
 	}
 
diff --git a/Engine/Source/Editor/UnrealEd/Private/Cooker/CookPackageData.h b/Engine/Source/Editor/UnrealEd/Private/Cooker/CookPackageData.h
--- a/Engine/Source/Editor/UnrealEd/Private/Cooker/CookPackageData.h
+++ b/Engine/Source/Editor/UnrealEd/Private/Cooker/CookPackageData.h
@@ -351,7 +351,19 @@ namespace Cook
 		FName CookedFileName; // Override filename to be used when saving cooked package (used by generated packages)
 		TWeakObjectPtr<UPackage> Package;
 		FPackageDatas& PackageDatas; // The one-per-CookOnTheFlyServer owner of this PackageData
-		TSharedPtr<FPreloadableArchive> PreloadableFile;
+		/**
+		* The number of active PreloadableFiles is tracked globally; wrap the PreloadableFile in a struct that
+		* guarantees we always update the counter when changing it
+		*/
+		struct FTrackedPreloadableFilePtr
+		{
+			const TSharedPtr<FPreloadableArchive>& Get() { return Ptr; }
+			void Set(TSharedPtr<FPreloadableArchive>&& InPtr, FPackageData& PackageData);
+			void Reset(FPackageData& PackageData);
+		private:
+			TSharedPtr<FPreloadableArchive> Ptr;
+		};
+		FTrackedPreloadableFilePtr PreloadableFile;
 		int32 NumPendingCookedPlatformData = 0;
 		int32 CookedPlatformDataNextIndex = -1;
 		std::atomic<EPackageFormat> PreloadableFileFormat;
