Fix non-unity compile errors in MirrorDataTable.h

diff --git a/Engine/Source/Editor/UnrealEd/Private/Factories/MirrorDataTableFactory.cpp b/Engine/Source/Editor/UnrealEd/Private/Factories/MirrorDataTableFactory.cpp
--- a/Engine/Source/Editor/UnrealEd/Private/Factories/MirrorDataTableFactory.cpp
+++ b/Engine/Source/Editor/UnrealEd/Private/Factories/MirrorDataTableFactory.cpp
@@ -4,6 +4,7 @@
 #include "DataTableEditorUtils.h"
 #include "Engine/Classes/Animation/MirrorDataTable.h"
 #include "Engine/Classes/Animation/AnimationSettings.h"
+#include "Editor.h"
 #include "StructViewerModule.h"
 #include "StructViewerFilter.h"
 
diff --git a/Engine/Source/Runtime/Engine/Classes/Animation/MirrorDataTable.h b/Engine/Source/Runtime/Engine/Classes/Animation/MirrorDataTable.h
--- a/Engine/Source/Runtime/Engine/Classes/Animation/MirrorDataTable.h
+++ b/Engine/Source/Runtime/Engine/Classes/Animation/MirrorDataTable.h
@@ -3,6 +3,7 @@
 #pragma once
 
 #include "Engine/DataTable.h"
+#include "BoneContainer.h"
 #include "MirrorDataTable.generated.h"
 
 /** Type referenced by a row in the mirror data table */
