Updating styles

diff --git a/Engine/Source/Developer/Profiler/Private/Widgets/SDataGraph.cpp b/Engine/Source/Developer/Profiler/Private/Widgets/SDataGraph.cpp
--- a/Engine/Source/Developer/Profiler/Private/Widgets/SDataGraph.cpp
+++ b/Engine/Source/Developer/Profiler/Private/Widgets/SDataGraph.cpp
@@ -883,7 +883,7 @@ int32 SDataGraph::OnPaint( const FPaintArgs& Args, const FGeometry& AllottedGeom
 			OutDrawElements,
 			LayerId,
 			AllottedGeometry.ToPaintGeometry( FVector2D(LocalPositionGraphX-HalfGraphMarkerWidth,0.0f), FVector2D(GraphMarkerWidth, AllottedGeometry.GetLocalSize().Y) ),
-			FEditorStyle::GetBrush("ProgressBar.Background"),
+			FAppStyle::Get().GetBrush("Brushes.Recessed"),
 			DrawEffects,
 			FColor(255,128,128,128)
 		);
diff --git a/Engine/Source/Editor/KismetWidgets/Private/SScrubWidget.cpp b/Engine/Source/Editor/KismetWidgets/Private/SScrubWidget.cpp
--- a/Engine/Source/Editor/KismetWidgets/Private/SScrubWidget.cpp
+++ b/Engine/Source/Editor/KismetWidgets/Private/SScrubWidget.cpp
@@ -89,7 +89,7 @@ int32 SScrubWidget::OnPaint( const FPaintArgs& Args, const FGeometry& AllottedGe
 
 	const int32 TextLayer = BackgroundLayer + 1;
 
-	const FSlateBrush* StyleInfo = FEditorStyle::GetBrush( TEXT( "ProgressBar.Background" ) );
+	const FSlateBrush* StyleInfo = FAppStyle::Get().GetBrush("Brushes.Recessed");
 	const float GeomHeight = AllottedGeometry.GetLocalSize().Y;
 
 	const FTrackScaleInfo TimeScaleInfo(ViewInputMin.Get(), ViewInputMax.Get(), 0.f, 0.f, AllottedGeometry.GetLocalSize());
diff --git a/Engine/Source/Editor/MaterialEditor/Private/SMaterialEditorViewport.cpp b/Engine/Source/Editor/MaterialEditor/Private/SMaterialEditorViewport.cpp
--- a/Engine/Source/Editor/MaterialEditor/Private/SMaterialEditorViewport.cpp
+++ b/Engine/Source/Editor/MaterialEditor/Private/SMaterialEditorViewport.cpp
@@ -871,7 +871,6 @@ private:
 
 	TSharedPtr<FSlateMaterialBrush> PreviewBrush;
 	TSharedPtr<FSlateImageBrush> CheckerboardBrush;
-	const FSlateBrush* BorderBrush;
 	UTexture2D* CheckerboardTexture;
 	TSharedPtr<SImage> ImageWidget;
 	FPreviewBackgroundSettings BackgroundSettings;
@@ -890,7 +889,6 @@ void SMaterialEditorUIPreviewZoomer::Construct( const FArguments& InArgs, UMater
 	ModifyCheckerboardTextureColors(BackgroundSettings.Checkerboard);
 
 	PreviewBrush = MakeShared<FSlateMaterialBrush>(*InPreviewMaterial, InArgs._InitialPreviewSize);
-	BorderBrush = FEditorStyle::GetBrush("MaterialEditor.ViewportBorder");	
 
 	ChildSlot
 	[
diff --git a/Engine/Source/Editor/UnrealEd/Private/SCommonEditorViewportToolbarBase.cpp b/Engine/Source/Editor/UnrealEd/Private/SCommonEditorViewportToolbarBase.cpp
--- a/Engine/Source/Editor/UnrealEd/Private/SCommonEditorViewportToolbarBase.cpp
+++ b/Engine/Source/Editor/UnrealEd/Private/SCommonEditorViewportToolbarBase.cpp
@@ -132,7 +132,7 @@ void SCommonEditorViewportToolbarBase::Construct(const FArguments& InArgs, TShar
 			SNew(SEditorViewportToolbarMenu)
 			.ParentToolBar(SharedThis(this))
 			.Label(this, &SCommonEditorViewportToolbarBase::GetScalabilityWarningLabel)
-			.MenuStyle(FEditorStyle::Get(), "EditorViewportToolBar.MenuButtonWarning")
+			.MenuStyle(&FAppStyle::Get().GetWidgetStyle<FButtonStyle>("EditorViewportToolBar.WarningButton"))
 			.OnGetMenuContent(this, &SCommonEditorViewportToolbarBase::GetScalabilityWarningMenuContent)
 			.Visibility(this, &SCommonEditorViewportToolbarBase::GetScalabilityWarningVisibility)
 			.ToolTipText(LOCTEXT("ScalabilityWarning_ToolTip", "Non-default scalability settings could be affecting what is shown in this viewport.\nFor example you may experience lower visual quality, reduced particle counts, and other artifacts that don't match what the scene would look like when running outside of the editor. Click to make changes."))
