AsyncIODelete: Remove direct reference to a specific platform, and replace it with a test of the relevant platform capability.
We found that Mac was not working with AsyncIODelete because of the inability to use a FileHandle created on one thread on a different thread.
#rnx

diff --git a/Engine/Source/Editor/UnrealEd/Private/Cooker/AsyncIODelete.cpp b/Engine/Source/Editor/UnrealEd/Private/Cooker/AsyncIODelete.cpp
--- a/Engine/Source/Editor/UnrealEd/Private/Cooker/AsyncIODelete.cpp
+++ b/Engine/Source/Editor/UnrealEd/Private/Cooker/AsyncIODelete.cpp
@@ -52,23 +52,24 @@ void FAsyncIODelete::SetTempRoot(const FStringView& InOwnedTempRoot)
 void FAsyncIODelete::SetDeletesPaused(bool bInPaused)
 {
 	bPaused = bInPaused;
-#if ASYNCIODELETE_ASYNC_ENABLED
-	if (!bPaused)
+	if (AsyncEnabled())
 	{
-		IFileManager& FileManager = IFileManager::Get();
-		for (const FString& DeletePath : PausedDeletes)
+		if (!bPaused)
 		{
-			const bool IsDirectory = FileManager.DirectoryExists(*DeletePath);
-			const bool IsFile = !IsDirectory && FileManager.FileExists(*DeletePath);
-			if (!IsDirectory && !IsFile)
+			IFileManager& FileManager = IFileManager::Get();
+			for (const FString& DeletePath : PausedDeletes)
 			{
-				continue;
+				const bool IsDirectory = FileManager.DirectoryExists(*DeletePath);
+				const bool IsFile = !IsDirectory && FileManager.FileExists(*DeletePath);
+				if (!IsDirectory && !IsFile)
+				{
+					continue;
+				}
+				CreateDeleteTask(DeletePath, IsDirectory ? EPathType::Directory : EPathType::File);
 			}
-			CreateDeleteTask(DeletePath, IsDirectory ? EPathType::Directory : EPathType::File);
+			PausedDeletes.Empty();
 		}
-		PausedDeletes.Empty();
 	}
-#endif
 }
 
 bool FAsyncIODelete::Setup()
@@ -84,33 +85,34 @@ bool FAsyncIODelete::Setup()
 		return false;
 	}
 
-#if ASYNCIODELETE_ASYNC_ENABLED
-	// Delete the TempRoot directory to clear the results from any previous process using the same TempRoot that did not shut down cleanly
-	uint32 ErrorCode;
-	if (!DeleteTempRootDirectory(ErrorCode))
+	if (AsyncEnabled())
 	{
-		UE_LOG(LogCook, Error, TEXT("Could not clear asyncdelete root directory '%s'.  LastError: %i."), *TempRoot, ErrorCode);
-		return false;
-	}
+		// Delete the TempRoot directory to clear the results from any previous process using the same TempRoot that did not shut down cleanly
+		uint32 ErrorCode;
+		if (!DeleteTempRootDirectory(ErrorCode))
+		{
+			UE_LOG(LogCook, Error, TEXT("Could not clear asyncdelete root directory '%s'.  LastError: %i."), *TempRoot, ErrorCode);
+			return false;
+		}
 
-	// Create the empty directory to work in
-	if (!IFileManager::Get().MakeDirectory(*TempRoot, true))
-	{
-		UE_LOG(LogCook, Error, TEXT("Could not create asyncdelete root directory '%s'.  LastError: %i."), *TempRoot, FPlatformMisc::GetLastError());
-		return false;
-	}
+		// Create the empty directory to work in
+		if (!IFileManager::Get().MakeDirectory(*TempRoot, true))
+		{
+			UE_LOG(LogCook, Error, TEXT("Could not create asyncdelete root directory '%s'.  LastError: %i."), *TempRoot, FPlatformMisc::GetLastError());
+			return false;
+		}
 
-	// Allocate the task event
-	check(TasksComplete == nullptr);
-	TasksComplete = FPlatformProcess::GetSynchEventFromPool(true /* IsManualReset */);
-	check(ActiveTaskCount == 0);
-	TasksComplete->Trigger(); // We have 0 tasks so the event should be in the Triggered state
+		// Allocate the task event
+		check(TasksComplete == nullptr);
+		TasksComplete = FPlatformProcess::GetSynchEventFromPool(true /* IsManualReset */);
+		check(ActiveTaskCount == 0);
+		TasksComplete->Trigger(); // We have 0 tasks so the event should be in the Triggered state
 
-	// Assert that all other teardown-transient variables were cleared by the constructor or by the previous teardown
-	// TempRoot and bPaused are preserved across setup/teardown and may have any value
-	check(PausedDeletes.Num() == 0);
-	check(DeleteCounter == 0);
-#endif
+		// Assert that all other teardown-transient variables were cleared by the constructor or by the previous teardown
+		// TempRoot and bPaused are preserved across setup/teardown and may have any value
+		check(PausedDeletes.Num() == 0);
+		check(DeleteCounter == 0);
+	}
 
 	// We are now setup and ready to create DeleteTasks
 	bInitialized = true;
@@ -125,25 +127,26 @@ void FAsyncIODelete::Teardown()
 		return;
 	}
 
-#if ASYNCIODELETE_ASYNC_ENABLED
-	// Clear task variables
-	WaitForAllTasks();
-	check(ActiveTaskCount == 0 && TasksComplete != nullptr && TasksComplete->Wait(0));
-	FPlatformProcess::ReturnSynchEventToPool(TasksComplete);
-	TasksComplete = nullptr;
-
-	// Remove the temp directory from disk
-	uint32 ErrorCode;
-	if (!DeleteTempRootDirectory(ErrorCode))
+	if (AsyncEnabled())
 	{
-		// This will leave directories (and potentially files, if we were paused or if any of the asyncdeletes failed) on disk, so it is bad for users, but is not fatal for our operations.
-		UE_LOG(LogCook, Warning, TEXT("Could not delete asyncdelete root directory '%s'.  LastError: %i."), *TempRoot, ErrorCode);
-	}
+		// Clear task variables
+		WaitForAllTasks();
+		check(ActiveTaskCount == 0 && TasksComplete != nullptr && TasksComplete->Wait(0));
+		FPlatformProcess::ReturnSynchEventToPool(TasksComplete);
+		TasksComplete = nullptr;
+
+		// Remove the temp directory from disk
+		uint32 ErrorCode;
+		if (!DeleteTempRootDirectory(ErrorCode))
+		{
+			// This will leave directories (and potentially files, if we were paused or if any of the asyncdeletes failed) on disk, so it is bad for users, but is not fatal for our operations.
+			UE_LOG(LogCook, Warning, TEXT("Could not delete asyncdelete root directory '%s'.  LastError: %i."), *TempRoot, ErrorCode);
+		}
 
-	// Clear delete variables; we don't need to run the tasks for the remaining pauseddeletes because synchronously deleting the temp directory above did the work they were going to do
-	PausedDeletes.Empty();
-	DeleteCounter = 0;
-#endif
+		// Clear delete variables; we don't need to run the tasks for the remaining pauseddeletes because synchronously deleting the temp directory above did the work they were going to do
+		PausedDeletes.Empty();
+		DeleteCounter = 0;
+	}
 
 	// We are now torn down and ready for a new setup
 	bInitialized = false;
@@ -151,25 +154,26 @@ void FAsyncIODelete::Teardown()
 
 bool FAsyncIODelete::WaitForAllTasks(float TimeLimitSeconds)
 {
-#if ASYNCIODELETE_ASYNC_ENABLED
-	if (!bInitialized)
+	if (AsyncEnabled())
 	{
-		return true;
-	}
+		if (!bInitialized)
+		{
+			return true;
+		}
 
-	if (TimeLimitSeconds <= 0.f)
-	{
-		TasksComplete->Wait();
-	}
-	else
-	{
-		if (!TasksComplete->Wait(FTimespan::FromSeconds(TimeLimitSeconds)))
+		if (TimeLimitSeconds <= 0.f)
 		{
-			return false;
+			TasksComplete->Wait();
+		}
+		else
+		{
+			if (!TasksComplete->Wait(FTimespan::FromSeconds(TimeLimitSeconds)))
+			{
+				return false;
+			}
 		}
+		check(ActiveTaskCount == 0);
 	}
-	check(ActiveTaskCount == 0);
-#endif
 	return true;
 }
 
@@ -203,12 +207,13 @@ bool FAsyncIODelete::Delete(const FStringView& PathToDelete, EPathType ExpectedT
 		return false;
 	}
 
-#if ASYNCIODELETE_ASYNC_ENABLED
-	if (DeleteCounter == UINT32_MAX)
+	if (AsyncEnabled())
 	{
-		Teardown();
+		if (DeleteCounter == UINT32_MAX)
+		{
+			Teardown();
+		}
 	}
-#endif
 	if (!Setup())
 	{
 		// Setup failed; we are not able to provide asynchronous deletes; fall back to synchronous
@@ -216,36 +221,38 @@ bool FAsyncIODelete::Delete(const FStringView& PathToDelete, EPathType ExpectedT
 		return SynchronousDelete(PathToDeleteSZ, ExpectedType);
 	}
 
-#if ASYNCIODELETE_ASYNC_ENABLED
-	const FString TempPath = FPaths::Combine(TempRoot, FString::Printf(TEXT("%u"), DeleteCounter));
-	DeleteCounter++;
-
-	const bool bReplace = true;
-	const bool bEvenIfReadOnly = true;
-	const bool bMoveAttributes = false;
-	const bool bDoNotRetryOnError = true;
-	if (!IFileManager::Get().Move(*TempPath, PathToDeleteSZ, bReplace, bEvenIfReadOnly, bMoveAttributes, bDoNotRetryOnError)) // IFileManager::Move works on either files or directories
+	if (AsyncEnabled())
 	{
-		// The move failed; try a synchronous delete as backup
-		UE_LOG(LogCook, Warning, TEXT("Failed to move path '%.*s' for async delete (LastError == %i); falling back to synchronous delete."), PathToDelete.Len(), PathToDelete.GetData(), FPlatformMisc::GetLastError());
-		return SynchronousDelete(PathToDeleteSZ, ExpectedType);
-	}
+		const FString TempPath = FPaths::Combine(TempRoot, FString::Printf(TEXT("%u"), DeleteCounter));
+		DeleteCounter++;
 
-	if (bPaused)
-	{
-		PausedDeletes.Add(TempPath);
+		const bool bReplace = true;
+		const bool bEvenIfReadOnly = true;
+		const bool bMoveAttributes = false;
+		const bool bDoNotRetryOnError = true;
+		if (!IFileManager::Get().Move(*TempPath, PathToDeleteSZ, bReplace, bEvenIfReadOnly, bMoveAttributes, bDoNotRetryOnError)) // IFileManager::Move works on either files or directories
+		{
+			// The move failed; try a synchronous delete as backup
+			UE_LOG(LogCook, Warning, TEXT("Failed to move path '%.*s' for async delete (LastError == %i); falling back to synchronous delete."), PathToDelete.Len(), PathToDelete.GetData(), FPlatformMisc::GetLastError());
+			return SynchronousDelete(PathToDeleteSZ, ExpectedType);
+		}
+
+		if (bPaused)
+		{
+			PausedDeletes.Add(TempPath);
+		}
+		else
+		{
+			CreateDeleteTask(TempPath, ExpectedType);
+		}
+		return true;
 	}
 	else
 	{
-		CreateDeleteTask(TempPath, ExpectedType);
+		return SynchronousDelete(PathToDeleteSZ, ExpectedType);
 	}
-	return true;
-#else
-	return SynchronousDelete(PathToDeleteSZ, ExpectedType);
-#endif
 }
 
-#if ASYNCIODELETE_ASYNC_ENABLED
 void FAsyncIODelete::CreateDeleteTask(const FStringView& InDeletePath, EPathType PathType)
 {
 	{
@@ -270,7 +277,6 @@ void FAsyncIODelete::OnTaskComplete()
 		TasksComplete->Trigger();
 	}
 }
-#endif
 
 bool FAsyncIODelete::SynchronousDelete(const TCHAR* InDeletePath, EPathType PathType)
 {
@@ -294,7 +300,6 @@ bool FAsyncIODelete::SynchronousDelete(const TCHAR* InDeletePath, EPathType Path
 	return Result;
 }
 
-#if ASYNCIODELETE_ASYNC_ENABLED
 bool FAsyncIODelete::DeleteTempRootDirectory(uint32& OutErrorCode)
 {
 	OutErrorCode = 0;
@@ -333,7 +338,6 @@ bool FAsyncIODelete::DeleteTempRootDirectory(uint32& OutErrorCode)
 	}
 	return bDeleteSucceeded;
 }
-#endif
 
 #if WITH_ASYNCIODELETE_DEBUG
 void FAsyncIODelete::AddTempRoot(const FStringView& InTempRoot)
diff --git a/Engine/Source/Editor/UnrealEd/Private/Cooker/AsyncIODelete.h b/Engine/Source/Editor/UnrealEd/Private/Cooker/AsyncIODelete.h
--- a/Engine/Source/Editor/UnrealEd/Private/Cooker/AsyncIODelete.h
+++ b/Engine/Source/Editor/UnrealEd/Private/Cooker/AsyncIODelete.h
@@ -18,13 +18,6 @@
 	#endif
 #endif
 
-// Temporarily disable AsyncIODelete on Mac until we can diagnose why it's failing
-#if PLATFORM_MAC
-#define ASYNCIODELETE_ASYNC_ENABLED 0
-#else
-#define ASYNCIODELETE_ASYNC_ENABLED 1
-#endif
-
 class FEvent;
 
 /**
@@ -100,11 +93,14 @@ public:
 	/** Synchronously wait for all deletes to complete.  */
 	bool WaitForAllTasks(float TimeLimitSeconds = 0.0f);
 
+	static bool AsyncEnabled()
+	{
+		return FPlatformMisc::SupportsMultithreadedFileHandles();
+	}
+
 private:
-#if ASYNCIODELETE_ASYNC_ENABLED
 	/** Update the ActiveTaskCount and events for a task completing */
 	void OnTaskComplete();
-#endif
 
 	enum class EPathType
 	{
@@ -114,26 +110,20 @@ private:
 	/** Asynchronously delete a file or directory.  We handle both in the same function, but we want the interface to be explicit. */
 	bool Delete(const FStringView& PathToDelete, EPathType ExpectedType);
 
-#if ASYNCIODELETE_ASYNC_ENABLED
 	/** Create and store the task to delete the given Directory or File */
 	void CreateDeleteTask(const FStringView& InDeletePath, EPathType PathType);
-#endif
 
 	/** Delete the given path synchronously; called from a task or in error fallback cases from the public thread */
 	bool SynchronousDelete(const TCHAR* InDeletePath, EPathType PathType);
 
-#if ASYNCIODELETE_ASYNC_ENABLED
 	bool DeleteTempRootDirectory(uint32& OutErrorCode);
-#endif
 
-	FString	TempRoot;
-#if ASYNCIODELETE_ASYNC_ENABLED
+	FString	TempRoot; 
 	TArray<FString> PausedDeletes;
 	FCriticalSection CriticalSection; // We use a CriticalSection instead of a TAtomic ActiveTaskCount so that we can atomically { trigger TasksComplete if ActiveTaskCount == 0 }
 	FEvent* TasksComplete = nullptr;
 	uint32 ActiveTaskCount = 0;
 	uint32 DeleteCounter = 0;
-#endif
 	bool bInitialized = false;
 	bool bPaused = false;
 
diff --git a/Engine/Source/Editor/UnrealEd/Private/Cooker/AsyncIODeleteTest.cpp b/Engine/Source/Editor/UnrealEd/Private/Cooker/AsyncIODeleteTest.cpp
--- a/Engine/Source/Editor/UnrealEd/Private/Cooker/AsyncIODeleteTest.cpp
+++ b/Engine/Source/Editor/UnrealEd/Private/Cooker/AsyncIODeleteTest.cpp
@@ -12,11 +12,8 @@
 IMPLEMENT_SIMPLE_AUTOMATION_TEST(FAsyncIODeleteTest, "System.Core.Misc.AsyncIODelete", EAutomationTestFlags::ApplicationContextMask | EAutomationTestFlags::EngineFilter);
 bool FAsyncIODeleteTest::RunTest(const FString& Parameters)
 {
-#if ASYNCIODELETE_ASYNC_ENABLED
-	UE_LOG(LogCore, Display, TEXT("FAsyncIODeleteTest::RunTest, ASYNCIODELETE is ENABLED"));
-#else
-	UE_LOG(LogCore, Display, TEXT("FAsyncIODeleteTest::RunTest, ASYNCIODELETE is DISABLED"));
-#endif
+	bool bAsyncEnabled = FAsyncIODelete::AsyncEnabled();
+	UE_LOG(LogCore, Display, TEXT("FAsyncIODeleteTest::RunTest, ASYNCIODELETE is %s"), bAsyncEnabled ? TEXT("ENABLED") : TEXT("DISABLED"));
 
 	FString TestRoot = FPaths::CreateTempFilename(FPlatformProcess::UserTempDir(), TEXT("AsyncIODelete"), TEXT(""));
 	IFileManager& FileManager = IFileManager::Get();
@@ -102,32 +99,37 @@ bool FAsyncIODeleteTest::RunTest(const FString& Parameters)
 		AsyncIODelete2.DeleteDirectory(TestDir1);
 		TestRequestedPathsDeleted(TEXT("AsyncIODelete::Delete should have moved the deleted paths before returning even when paused."));
 		WaitForAllTasksAndVerify(AsyncIODelete2);
-#if ASYNCIODELETE_ASYNC_ENABLED
-		TestTempRootCountsEqual(TempRoot2, 1, 1, TEXT("AsyncIODelete should not have deleted the moved paths because it is paused."));
-#endif
+		if (bAsyncEnabled)
+		{
+			TestTempRootCountsEqual(TempRoot2, 1, 1, TEXT("AsyncIODelete should not have deleted the moved paths because it is paused."));
+		}
 		AsyncIODelete2.SetDeletesPaused(false);
 		WaitForAllTasksAndVerify(AsyncIODelete2);
-#if ASYNCIODELETE_ASYNC_ENABLED
-		TestTempRootCountsEqual(TempRoot2, 0, 0, TEXT("AsyncIODelete should have deleted the moved paths after unpausing."));
-#endif
+		if (bAsyncEnabled)
+		{
+			TestTempRootCountsEqual(TempRoot2, 0, 0, TEXT("AsyncIODelete should have deleted the moved paths after unpausing."));
+		}
 
 		StartSection(TEXT("Verify Teardown() deletes the TempRoot and Setup() creates it"));
 		AsyncIODelete2.Teardown();
-#if ASYNCIODELETE_ASYNC_ENABLED
-		TestTrue(TEXT("AsyncIODelete::Teardown should have deleted its TempRoot."), !FileManager.DirectoryExists(*TempRoot2));
-#endif
+		if (bAsyncEnabled)
+		{
+			TestTrue(TEXT("AsyncIODelete::Teardown should have deleted its TempRoot."), !FileManager.DirectoryExists(*TempRoot2));
+		}
 		AsyncIODelete2.Setup();
-#if ASYNCIODELETE_ASYNC_ENABLED
-		TestTrue(TEXT("AsyncIODelete::Setup should have created its TempRoot."), FileManager.DirectoryExists(*TempRoot2));
-#endif
+		if (bAsyncEnabled)
+		{
+			TestTrue(TEXT("AsyncIODelete::Setup should have created its TempRoot."), FileManager.DirectoryExists(*TempRoot2));
+		}
 
 		StartSection(TEXT("Manual setup works as long as you call SetTempRoot before Setup"));
 		FAsyncIODelete AsyncIODelete3;
 		AsyncIODelete3.SetTempRoot(TempRoot3);
 		AsyncIODelete3.Setup();
-#if ASYNCIODELETE_ASYNC_ENABLED
-		TestTrue(TEXT("Setup should have created the TempRoot."), FileManager.DirectoryExists(*TempRoot3));
-#endif
+		if (bAsyncEnabled)
+		{
+			TestTrue(TEXT("Setup should have created the TempRoot."), FileManager.DirectoryExists(*TempRoot3));
+		}
 
 		StartSection(TEXT("Check that even after Setup, waiting for tasks to complete when none have been launched should succeed"));
 		WaitForAllTasksAndVerify(AsyncIODelete);
@@ -137,20 +139,24 @@ bool FAsyncIODeleteTest::RunTest(const FString& Parameters)
 		AsyncIODelete3.DeleteFile(TestFile1);
 		AsyncIODelete3.DeleteDirectory(TestDir1);
 		AsyncIODelete3.SetTempRoot(TempRoot4);
-#if ASYNCIODELETE_ASYNC_ENABLED
-		TestTrue(TEXT("SetTempRoot should have deleted the old TempRoot."), !FileManager.DirectoryExists(*TempRoot3));
-#endif
+		if (bAsyncEnabled)
+		{
+			TestTrue(TEXT("SetTempRoot should have deleted the old TempRoot."), !FileManager.DirectoryExists(*TempRoot3));
+		}
+
 		CreateTestPathsToDelete();
 		AsyncIODelete3.DeleteFile(TestFile1);
 		AsyncIODelete3.DeleteDirectory(TestDir1);
 		TestRequestedPathsDeleted(TEXT("AsyncIODelete::Delete should have worked after changing the TempRoot."));
-#if ASYNCIODELETE_ASYNC_ENABLED
-		TestTrue(TEXT("Delete should have created the new TempRoot after SetTempRoot."), FileManager.DirectoryExists(*TempRoot4));
-#endif
+		if (bAsyncEnabled)
+		{
+			TestTrue(TEXT("Delete should have created the new TempRoot after SetTempRoot."), FileManager.DirectoryExists(*TempRoot4));
+		}
 		WaitForAllTasksAndVerify(AsyncIODelete3);
-#if ASYNCIODELETE_ASYNC_ENABLED
-		TestTempRootCountsEqual(TempRoot4, 0, 0, TEXT("AsyncIODelete::Delete should have created the tasks to delete moved files after changing the TempRoot."));
-#endif
+		if (bAsyncEnabled)
+		{
+			TestTempRootCountsEqual(TempRoot4, 0, 0, TEXT("AsyncIODelete::Delete should have created the tasks to delete moved files after changing the TempRoot."));
+		}
 
 		StartSection(TEXT("Attempting to delete a parent directory of the temproot, the temproot itself, or a child inside of it fails"));
 		FString SubDirInTempRoot4 = FPaths::Combine(TempRoot4, TEXT("SubDir"));
@@ -160,10 +166,11 @@ bool FAsyncIODeleteTest::RunTest(const FString& Parameters)
 		TestFalse(TEXT("AsyncIODelete should refuse to delete a child of its TempRoot."), AsyncIODelete3.DeleteDirectory(SubDirInTempRoot4));
 	}
 
-#if ASYNCIODELETE_ASYNC_ENABLED
-	TestTrue(TEXT("AsyncIODelete destructor should have deleted its TempRoot."),
-		!FileManager.DirectoryExists(*TempRoot) && !FileManager.DirectoryExists(*TempRoot2) && !FileManager.DirectoryExists(*TempRoot3) && !FileManager.DirectoryExists(*TempRoot4));
-#endif
+	if (bAsyncEnabled)
+	{
+		TestTrue(TEXT("AsyncIODelete destructor should have deleted its TempRoot."),
+			!FileManager.DirectoryExists(*TempRoot) && !FileManager.DirectoryExists(*TempRoot2) && !FileManager.DirectoryExists(*TempRoot3) && !FileManager.DirectoryExists(*TempRoot4));
+	}
 
 	return true;
 }
\ No newline at end of file
