Remove 4.24 deprecated mode delegate function from EditorViewportClient.

diff --git a/Engine/Source/Editor/UnrealEd/Private/EditorViewportClient.cpp b/Engine/Source/Editor/UnrealEd/Private/EditorViewportClient.cpp
--- a/Engine/Source/Editor/UnrealEd/Private/EditorViewportClient.cpp
+++ b/Engine/Source/Editor/UnrealEd/Private/EditorViewportClient.cpp
@@ -661,14 +661,6 @@ void FEditorViewportClient::RequestInvalidateHitProxy(FViewport* InViewport)
 	bNeedsInvalidateHitProxy = true;
 }
 
-void FEditorViewportClient::OnEditorModeChanged(FEdMode* EditorMode, bool bIsEntering)
-{
-	if (Viewport)
-	{
-		RequestInvalidateHitProxy(Viewport);
-	}
-}
-
 void FEditorViewportClient::OnEditorModeIDChanged(const FEditorModeID& EditorModeID, bool bIsEntering)
 {
 	if (Viewport)
diff --git a/Engine/Source/Editor/UnrealEd/Public/EditorViewportClient.h b/Engine/Source/Editor/UnrealEd/Public/EditorViewportClient.h
--- a/Engine/Source/Editor/UnrealEd/Public/EditorViewportClient.h
+++ b/Engine/Source/Editor/UnrealEd/Public/EditorViewportClient.h
@@ -545,8 +545,6 @@ public:
 	void TakeHighResScreenShot();
 
 	/** Called when an editor mode has been (de)activated */
-	UE_DEPRECATED(4.24, "Use OnEditorModeIDChanged() instead.")
-	void OnEditorModeChanged(FEdMode* EditorMode, bool bIsEntering);
 	void OnEditorModeIDChanged(const FEditorModeID& EditorModeID, bool bIsEntering);
 
 	/** FViewElementDrawer interface */
