Fix a circular strong reference with Asset Viewport Layouts and Viewports that was causing the viewports to leak when an asset editor closed.

diff --git a/Engine/Source/Editor/UnrealEd/Private/SAssetEditorViewport.cpp b/Engine/Source/Editor/UnrealEd/Private/SAssetEditorViewport.cpp
--- a/Engine/Source/Editor/UnrealEd/Private/SAssetEditorViewport.cpp
+++ b/Engine/Source/Editor/UnrealEd/Private/SAssetEditorViewport.cpp
@@ -98,7 +98,7 @@ void SAssetEditorViewport::Construct(const FArguments& InArgs, const FAssetEdito
 {
 	ConfigKey = InViewportConstructionArgs.ConfigKey;
 	Client = InArgs._EditorViewportClient;
-	ParentLayout = InViewportConstructionArgs.ParentLayout;
+	ParentTabContent = InViewportConstructionArgs.ParentLayout->GetParentTabContent();
 
 	SEditorViewport::FArguments ViewportArgs;
 	if (InArgs._ViewportSize.IsSet())
@@ -118,26 +118,18 @@ void SAssetEditorViewport::Construct(const FArguments& InArgs, const FAssetEdito
 
 void SAssetEditorViewport::OnSetViewportConfiguration(FName ConfigurationName)
 {
-	if (ParentLayout.IsValid())
+	if (TSharedPtr<FViewportTabContent> ViewportTabPinned = ParentTabContent.Pin())
 	{
-		TSharedPtr<FViewportTabContent> ViewportTabPinned = ParentLayout->GetParentTabContent().Pin();
-		if (ViewportTabPinned.IsValid())
-		{
-			ViewportTabPinned->SetViewportConfiguration(ConfigurationName);
- 			FSlateApplication::Get().DismissAllMenus();
-		}
+		ViewportTabPinned->SetViewportConfiguration(ConfigurationName);
+		FSlateApplication::Get().DismissAllMenus();
 	}
 }
 
 bool SAssetEditorViewport::IsViewportConfigurationSet(FName ConfigurationName) const
 {
-	if (ParentLayout.IsValid())
+	if (TSharedPtr<FViewportTabContent> ViewportTabPinned = ParentTabContent.Pin())
 	{
-		TSharedPtr<FViewportTabContent> ViewportTabPinned = ParentLayout->GetParentTabContent().Pin();
-		if (ViewportTabPinned.IsValid())
-		{
-			return ViewportTabPinned->IsViewportConfigurationSet(ConfigurationName);
-		}
+		return ViewportTabPinned->IsViewportConfigurationSet(ConfigurationName);
 	}
 	return false;
 }
diff --git a/Engine/Source/Editor/UnrealEd/Public/SAssetEditorViewport.h b/Engine/Source/Editor/UnrealEd/Public/SAssetEditorViewport.h
--- a/Engine/Source/Editor/UnrealEd/Public/SAssetEditorViewport.h
+++ b/Engine/Source/Editor/UnrealEd/Public/SAssetEditorViewport.h
@@ -7,7 +7,7 @@
 #include "SEditorViewport.h"
 
 class FEditorViewportClient;
-class FAssetEditorViewportLayout;
+class FViewportTabContent;
 class FMenuBuilder;
 struct FAssetEditorViewportConstructionArgs;
 
@@ -34,6 +34,5 @@ protected:
 
 	virtual TSharedRef<FEditorViewportClient> MakeEditorViewportClient() override;
 	FName ConfigKey;
-	TSharedPtr<FAssetEditorViewportLayout> ParentLayout;
-
+	TWeakPtr<FViewportTabContent> ParentTabContent;
 };
