Fix LevelInstance creating a WorldPartition with latest change

diff --git a/Engine/Source/Editor/UnrealEd/Classes/Factories/WorldFactory.h b/Engine/Source/Editor/UnrealEd/Classes/Factories/WorldFactory.h
--- a/Engine/Source/Editor/UnrealEd/Classes/Factories/WorldFactory.h
+++ b/Engine/Source/Editor/UnrealEd/Classes/Factories/WorldFactory.h
@@ -20,6 +20,7 @@ class UWorldFactory : public UFactory
 
 	TEnumAsByte<EWorldType::Type> WorldType;
 	bool bInformEngineOfWorld;
+	bool bCreateWorldPartition;
 	ERHIFeatureLevel::Type FeatureLevel;
 
 	//~ Begin UFactory Interface
diff --git a/Engine/Source/Editor/UnrealEd/Private/EditorLevelUtils.cpp b/Engine/Source/Editor/UnrealEd/Private/EditorLevelUtils.cpp
--- a/Engine/Source/Editor/UnrealEd/Private/EditorLevelUtils.cpp
+++ b/Engine/Source/Editor/UnrealEd/Private/EditorLevelUtils.cpp
@@ -674,6 +674,8 @@ ULevelStreaming* UEditorLevelUtils::CreateNewStreamingLevelForWorld(UWorld& InWo
 	{
 		// Create a new world
 		UWorldFactory* Factory = NewObject<UWorldFactory>();
+		// streaming levels shouldn't be WP
+		Factory->bCreateWorldPartition = false;
 		Factory->WorldType = EWorldType::Inactive;
 		UPackage* Pkg = CreatePackage( NULL);
 		FName WorldName(TEXT("Untitled"));
diff --git a/Engine/Source/Editor/UnrealEd/Private/Factories/WorldFactory.cpp b/Engine/Source/Editor/UnrealEd/Private/Factories/WorldFactory.cpp
--- a/Engine/Source/Editor/UnrealEd/Private/Factories/WorldFactory.cpp
+++ b/Engine/Source/Editor/UnrealEd/Private/Factories/WorldFactory.cpp
@@ -16,6 +16,7 @@ UWorldFactory::UWorldFactory(const FObjectInitializer& ObjectInitializer)
 	SupportedClass = UWorld::StaticClass();
 	WorldType = EWorldType::Inactive;
 	bInformEngineOfWorld = false;
+	bCreateWorldPartition = true;
 	FeatureLevel = ERHIFeatureLevel::Num;
 }
 
@@ -30,7 +31,7 @@ UObject* UWorldFactory::FactoryCreateNew(UClass* Class, UObject* InParent, FName
 	const bool bAddToRoot = false;
 
 	IWorldPartitionEditorModule& WorldPartitionEditorModule = FModuleManager::LoadModuleChecked<IWorldPartitionEditorModule>("WorldPartitionEditor");
-	const bool bCreateWorldPartition = WorldPartitionEditorModule.IsWorldPartitionEnabled();
+	const bool bCreateWorldPartitionSetting = WorldPartitionEditorModule.IsWorldPartitionEnabled();
 		
 	// Those are the init values taken from the default in UWorld::CreateWorld + CreateWorldPartition.
 	UWorld::InitializationValues InitValues = UWorld::InitializationValues()
@@ -38,7 +39,7 @@ UObject* UWorldFactory::FactoryCreateNew(UClass* Class, UObject* InParent, FName
 		.EnableTraceCollision(true)
 		.CreateNavigation(WorldType == EWorldType::Editor)
 		.CreateAISystem(WorldType == EWorldType::Editor)
-		.CreateWorldPartition(bCreateWorldPartition);
+		.CreateWorldPartition(bCreateWorldPartitionSetting && bCreateWorldPartition);
 
 	UWorld* NewWorld = UWorld::CreateWorld(WorldType, bInformEngineOfWorld, Name, Cast<UPackage>(InParent), bAddToRoot, FeatureLevel, &InitValues);
 	GEditor->InitBuilderBrush(NewWorld);
