Disabled cancelling empty FChange based transactions, as engine support for cancelling subsets of a transaction was disabled (see UTransBuffer::Cancel)

diff --git a/Engine/Source/Editor/UnrealEd/Public/ChangeTransactor.h b/Engine/Source/Editor/UnrealEd/Public/ChangeTransactor.h
--- a/Engine/Source/Editor/UnrealEd/Public/ChangeTransactor.h
+++ b/Engine/Source/Editor/UnrealEd/Public/ChangeTransactor.h
@@ -68,16 +68,7 @@ struct ENGINE_API FChangeTransactor
 				GUndo->StoreUndo(TransactionObject.Get(), MoveTemp(CompoundAction));
 			}
 		}
-		else
-		{
-			if (TransactedChanges == 0)
-			{
-				// If no changes were added for the compound change or any individual changes, cancel the whole transaction
-				GEngine->CancelTransaction(PendingTransactionIndex);
-				PendingTransactionIndex = -1;
-			}
-		}
-				
+	
 		if (PendingTransactionIndex != -1)
 		{
 			GEngine->EndTransaction();
