Non Unity Fix

diff --git a/Engine/Source/Editor/UnrealEd/Private/Factories/WorldFactory.cpp b/Engine/Source/Editor/UnrealEd/Private/Factories/WorldFactory.cpp
--- a/Engine/Source/Editor/UnrealEd/Private/Factories/WorldFactory.cpp
+++ b/Engine/Source/Editor/UnrealEd/Private/Factories/WorldFactory.cpp
@@ -5,6 +5,7 @@
 #include "UObject/Package.h"
 #include "ThumbnailRendering/WorldThumbnailInfo.h"
 #include "EditorClassUtils.h"
+#include "Modules/ModuleManager.h"
 #include "WorldPartition/IWorldPartitionEditorModule.h"
 
 #define LOCTEXT_NAMESPACE "WorldFactory"
