WorldPartitionConvertCommandlet - Added ability to perform changes on actor during the conversion

diff --git a/Engine/Source/Editor/UnrealEd/Classes/Commandlets/WorldPartitionConvertCommandlet.h b/Engine/Source/Editor/UnrealEd/Classes/Commandlets/WorldPartitionConvertCommandlet.h
--- a/Engine/Source/Editor/UnrealEd/Classes/Commandlets/WorldPartitionConvertCommandlet.h
+++ b/Engine/Source/Editor/UnrealEd/Classes/Commandlets/WorldPartitionConvertCommandlet.h
@@ -49,6 +49,7 @@ protected:
 	virtual bool ShouldConvertStreamingLevel(ULevelStreaming* StreamingLevel);
 	virtual bool ShouldDeleteActor(AActor* Actor, bool bMainLevel) const;
 	virtual void PerformAdditionalWorldCleanup(UWorld* World) const;
+	virtual void PerformAdditionalActorChanges(AActor* Actor) const {}
 	virtual void OutputConversionReport() const;
 	virtual void OnWorldLoaded(UWorld* World);
 	virtual void ReadAdditionalTokensAndSwitches(const TArray<FString>& Tokens, const TArray<FString>& Switches) {}
diff --git a/Engine/Source/Editor/UnrealEd/Private/Commandlets/WorldPartitionConvertCommandlet.cpp b/Engine/Source/Editor/UnrealEd/Private/Commandlets/WorldPartitionConvertCommandlet.cpp
--- a/Engine/Source/Editor/UnrealEd/Private/Commandlets/WorldPartitionConvertCommandlet.cpp
+++ b/Engine/Source/Editor/UnrealEd/Private/Commandlets/WorldPartitionConvertCommandlet.cpp
@@ -1380,6 +1380,16 @@ int32 UWorldPartitionConvertCommandlet::Main(const FString& Params)
 		// Required to clear any deleted actors from the level
 		CollectGarbage(RF_Standalone);
 
+		for (AActor* Actor : ActorList)
+		{
+			if (!Actor || Actor->IsPendingKill())
+			{
+				continue;
+			}
+
+			PerformAdditionalActorChanges(Actor);
+		}
+
 		MainLevel->bUseExternalActors = true;
 		MainWorld->WorldComposition = nullptr;
 		MainLevel->bIsPartitioned = true;
