WorldPartitionHLODsBuilder
* Added options to generate/use a build manifest
* -BuildManifest=Filename.ini
* When calling with -SetupHLODs, provide "-BuildManifest=Filename.ini -BuilderCount=N" to have N workloads in the build manifest
* When calling with -BuildHLODs, provide "-BuildManifest=Filename.ini -BuilderIdx=i" to build workload i

diff --git a/Engine/Source/Editor/UnrealEd/Private/WorldPartition/WorldPartitionHLODsBuilder.cpp b/Engine/Source/Editor/UnrealEd/Private/WorldPartition/WorldPartitionHLODsBuilder.cpp
--- a/Engine/Source/Editor/UnrealEd/Private/WorldPartition/WorldPartitionHLODsBuilder.cpp
+++ b/Engine/Source/Editor/UnrealEd/Private/WorldPartition/WorldPartitionHLODsBuilder.cpp
@@ -90,11 +90,17 @@ private:
 
 UWorldPartitionHLODsBuilder::UWorldPartitionHLODsBuilder(const FObjectInitializer& ObjectInitializer)
 	: Super(ObjectInitializer)
+	, BuilderIdx(INDEX_NONE)
+	, BuilderCount(INDEX_NONE)
 {
 	bSetupHLODs = FParse::Param(FCommandLine::Get(), TEXT("SetupHLODs"));
 	bBuildHLODs = FParse::Param(FCommandLine::Get(), TEXT("BuildHLODs"));
 	bDeleteHLODs = FParse::Param(FCommandLine::Get(), TEXT("DeleteHLODs"));
 	bForceGC = FParse::Param(FCommandLine::Get(), TEXT("ForceGC"));
+
+	FParse::Value(FCommandLine::Get(), TEXT("BuildManifest="), BuildManifest);
+	FParse::Value(FCommandLine::Get(), TEXT("BuilderIdx="), BuilderIdx);
+	FParse::Value(FCommandLine::Get(), TEXT("BuilderCount="), BuilderCount);
 }
 
 bool UWorldPartitionHLODsBuilder::RequiresCommandletRendering() const
@@ -102,8 +108,42 @@ bool UWorldPartitionHLODsBuilder::RequiresCommandletRendering() const
 	return bBuildHLODs || !(bSetupHLODs || bDeleteHLODs);
 }
 
+bool UWorldPartitionHLODsBuilder::ValidateParams() const
+{
+	if (bSetupHLODs && !BuildManifest.IsEmpty() && BuilderCount == INDEX_NONE)
+	{
+		if (BuilderCount == INDEX_NONE)
+		{
+			UE_LOG(LogWorldPartitionHLODsBuilder, Error, TEXT("Missing parameter -BuilderCount=N, exiting..."));
+			return false;
+		}
+	}
+
+	if (bBuildHLODs && !BuildManifest.IsEmpty() && BuilderIdx == INDEX_NONE)
+	{
+		if (BuilderIdx == INDEX_NONE)
+		{
+			UE_LOG(LogWorldPartitionHLODsBuilder, Error, TEXT("Missing parameter -BuilderIdx=i, exiting..."));
+			return false;
+		}
+
+		if (FPaths::FileExists(BuildManifest))
+		{
+			UE_LOG(LogWorldPartitionHLODsBuilder, Error, TEXT("Build manifest file \"%s\" not found, exiting..."), *BuildManifest);
+			return false;
+		}
+	}
+
+	return true;
+}
+
 bool UWorldPartitionHLODsBuilder::Run(UWorld* World, FPackageSourceControlHelper& PackageHelper)
 {
+	if (!ValidateParams())
+	{
+		return false;
+	}
+
 	WorldPartition = World->GetWorldPartition();
 	check(WorldPartition);
 
@@ -126,7 +166,7 @@ bool UWorldPartitionHLODsBuilder::Run(UWorld* World, FPackageSourceControlHelper
 	{
 		bRet = SetupHLODActors(false);
 	}
-
+		
 	WorldPartition = nullptr;
 	delete SourceControlHelper;
 	
@@ -135,6 +175,8 @@ bool UWorldPartitionHLODsBuilder::Run(UWorld* World, FPackageSourceControlHelper
 
 bool UWorldPartitionHLODsBuilder::SetupHLODActors(bool bCreateOnly)
 {
+	bool bRet = true;
+
 	WorldPartition->GenerateHLOD(SourceControlHelper, bCreateOnly);
 
 	if (bCreateOnly)
@@ -153,31 +195,41 @@ bool UWorldPartitionHLODsBuilder::SetupHLODActors(bool bCreateOnly)
 		}
 
 		UE_LOG(LogWorldPartitionHLODsBuilder, Display, TEXT("#### World contains %d HLOD actors ####"), NumActors);
+
+		if (!BuildManifest.IsEmpty())
+		{
+			bRet = GenerateBuildManifest();
+		}
 	}
 
-	return true;
+	return bRet;
 }
 
 bool UWorldPartitionHLODsBuilder::BuildHLODActors()
 {
-	int32 NumActors = 0;
-	for (UActorDescContainer::TIterator<AWorldPartitionHLOD> HLODIterator(WorldPartition); HLODIterator; ++HLODIterator)
+	TArray<FGuid> HLODActorsToBuild;
+	if (!GetHLODActorsToBuild(HLODActorsToBuild))
+	{
+		return false;
+	}
+
+	if (!ValidateWorkload(HLODActorsToBuild))
 	{
-		NumActors++;
+		return false;
 	}
 
-	UE_LOG(LogWorldPartitionHLODsBuilder, Display, TEXT("#### Building %d HLOD actors ####"), NumActors);
+	UE_LOG(LogWorldPartitionHLODsBuilder, Display, TEXT("#### Building %d HLOD actors ####"), HLODActorsToBuild.Num());
 
 	int32 CurrentActor = 0;
-	for (UActorDescContainer::TIterator<AWorldPartitionHLOD> HLODIterator(WorldPartition); HLODIterator; ++HLODIterator)
+	for (const FGuid& HLODActorGuid : HLODActorsToBuild)
 	{
-		FWorldPartitionReference ActorRef(WorldPartition, HLODIterator->GetGuid());
+		FWorldPartitionReference ActorRef(WorldPartition, HLODActorGuid);
 		FWorldPartitionActorDesc* ActorDesc = ActorRef.Get();
 
 		AWorldPartitionHLOD* HLODActor = CastChecked<AWorldPartitionHLOD>(ActorDesc->GetActor());
 
 		CurrentActor++;
-		UE_LOG(LogWorldPartitionHLODsBuilder, Display, TEXT("    [%d/%d] Building HLOD actor..."), CurrentActor, NumActors);
+		UE_LOG(LogWorldPartitionHLODsBuilder, Display, TEXT("    [%d/%d] Building HLOD actor..."), CurrentActor, HLODActorsToBuild.Num());
 
 		HLODActor->BuildHLOD();
 
@@ -192,7 +244,7 @@ bool UWorldPartitionHLODsBuilder::BuildHLODActors()
 		}
 	}
 
-	UE_LOG(LogWorldPartitionHLODsBuilder, Display, TEXT("#### Built %d HLOD actors ####"), NumActors);
+	UE_LOG(LogWorldPartitionHLODsBuilder, Display, TEXT("#### Built %d HLOD actors ####"), HLODActorsToBuild.Num());
 
 	return true;
 }
@@ -224,3 +276,219 @@ bool UWorldPartitionHLODsBuilder::DeleteHLODActors()
 
 	return true;
 }
+
+bool UWorldPartitionHLODsBuilder::GetHLODActorsToBuild(TArray<FGuid>& HLODActorsToBuild) const
+{
+	bool bRet = true;
+
+	if (!BuildManifest.IsEmpty())
+	{
+		// Get HLOD actors to build from the BuildManifest file
+		FConfigFile ConfigFile;
+		ConfigFile.Read(BuildManifest);
+
+		FString SectionName = FString::Printf(TEXT("HLODBuilder%d"), BuilderIdx);
+
+		const FConfigSection* ConfigSection = ConfigFile.Find(SectionName);
+		if (ConfigSection)
+		{
+			TArray<FString> HLODActorStrings;
+			ConfigSection->MultiFind(TEXT("+HLODActorGuid"), HLODActorStrings, /*bMaintainOrder=*/true);
+
+			for (const FString& HLODActorString : HLODActorStrings)
+			{
+				FString HLODActorGuidString = HLODActorString.Left(HLODActorString.Find(TEXT(",")));
+				FGuid HLODActorGuid;
+				bRet = FGuid::ParseExact(HLODActorGuidString, EGuidFormats::Digits, HLODActorGuid);
+				if (bRet)
+				{
+					HLODActorsToBuild.Add(HLODActorGuid);
+				}
+				else
+				{
+					UE_LOG(LogWorldPartitionHLODsBuilder, Error, TEXT("Error parsing section [%s] in config file \"%s\""), *SectionName, *BuildManifest);
+					break;
+				}
+			}
+		}
+		else
+		{
+			UE_LOG(LogWorldPartitionHLODsBuilder, Error, TEXT("Missing section [%s] in config file \"%s\""), *SectionName, *BuildManifest);
+			bRet = false;
+		}
+	}
+	else
+	{
+		TArray<TArray<FGuid>> HLODWorkloads = GetHLODWorldloads(1);
+		HLODActorsToBuild = MoveTemp(HLODWorkloads[0]);
+	}
+
+	return bRet;
+}
+
+TArray<TArray<FGuid>> UWorldPartitionHLODsBuilder::GetHLODWorldloads(int32 NumWorkloads) const
+{
+	// Build a mapping of 1 HLOD[Level] -> N HLOD[Level - 1]
+	TMap<FGuid, TArray<FGuid>>	HLODParenting;
+	for (UActorDescContainer::TIterator<AWorldPartitionHLOD> HLODIterator(WorldPartition); HLODIterator; ++HLODIterator)
+	{
+		TArray<FGuid>& ChildHLODs = HLODParenting.Add(HLODIterator->GetGuid());
+
+		for (const FGuid& SubActorGuid : HLODIterator->GetSubActors())
+		{
+			FWorldPartitionActorDesc* SubActorDesc = WorldPartition->GetActorDesc(SubActorGuid);
+			check(SubActorDesc);
+
+			if (SubActorDesc->GetActorClass()->IsChildOf<AWorldPartitionHLOD>())
+			{
+				ChildHLODs.Add(SubActorGuid);
+			}
+		}
+	}
+
+	// All child HLODs must be built before their parent HLOD
+	// Create groups to ensure those will be processed in the correct order, on the same builder
+	TMap<FGuid, TArray<FGuid>> HLODGroups;
+	TSet<FGuid>				   TriagedHLODs;
+
+	TFunction<void(TArray<FGuid>&, const FGuid&)> RecursiveAdd = [&TriagedHLODs, &HLODParenting, &HLODGroups, &RecursiveAdd](TArray<FGuid>& HLODGroup, const FGuid& HLODGuid)
+	{
+		if (!TriagedHLODs.Contains(HLODGuid))
+		{
+			TriagedHLODs.Add(HLODGuid);
+			HLODGroup.Insert(HLODGuid, 0); // Child will come first in the list, as they need to be built first...
+			TArray<FGuid>* ChildHLODs = HLODParenting.Find(HLODGuid);
+			if (ChildHLODs)
+			{
+				for (const auto& ChildGuid : *ChildHLODs)
+				{
+					RecursiveAdd(HLODGroup, ChildGuid);
+				}
+			}
+		}
+		else
+		{
+			HLODGroup.Insert(MoveTemp(HLODGroups.FindChecked(HLODGuid)), 0);
+			HLODGroups.Remove(HLODGuid);
+		}
+	};
+
+	for (const auto& Pair : HLODParenting)
+	{
+		if (!TriagedHLODs.Contains(Pair.Key))
+		{
+			TArray<FGuid>& HLODGroup = HLODGroups.Add(Pair.Key);
+			RecursiveAdd(HLODGroup, Pair.Key);
+		}
+	}
+
+	// Sort groups by number of HLOD actors
+	HLODGroups.ValueSort([](const TArray<FGuid>& GroupA, const TArray<FGuid>& GroupB) { return GroupA.Num() > GroupB.Num(); });
+
+	// Dispatch them in multiple lists and try to balance the workloads as much as possible
+	TArray<TArray<FGuid>> Workloads;
+	Workloads.SetNum(NumWorkloads);
+
+	int32 Idx = 0;
+	for (const auto& Pair : HLODGroups)
+	{
+		Workloads[Idx % NumWorkloads].Append(Pair.Value);
+		Idx++;
+	}
+
+	// Validate workloads to ensure our meshes are built in the correct order
+	for (const TArray<FGuid>& Workload : Workloads)
+	{
+		check(ValidateWorkload(Workload));
+	}
+
+	return Workloads;
+}
+
+bool UWorldPartitionHLODsBuilder::ValidateWorkload(const TArray<FGuid>&Workload) const
+{
+	TSet<FGuid> ProcessedHLOD;
+	ProcessedHLOD.Reserve(Workload.Num());
+
+	// For each HLOD entry in the workload, validate that its children are found before itself
+	for (const FGuid& HLODActorGuid : Workload)
+	{
+		const FWorldPartitionActorDesc* ActorDesc = WorldPartition->GetActorDesc(HLODActorGuid);
+		if(!ActorDesc)
+		{
+			UE_LOG(LogWorldPartitionHLODsBuilder, Error, TEXT("Unknown actor guid found, your HLOD actors are probably out of date. Run with -SetupHLODs to fix this. Exiting..."));
+			return false;
+		}
+
+		if (!ActorDesc->GetActorClass()->IsChildOf<AWorldPartitionHLOD>())
+		{
+			UE_LOG(LogWorldPartitionHLODsBuilder, Error, TEXT("Unexpected actor guid found in HLOD workload, exiting..."));
+			return false;
+		}
+
+		const FHLODActorDesc* HLODActorDesc = static_cast<const FHLODActorDesc*>(ActorDesc);
+
+		for (const FGuid& SubActorGuid : HLODActorDesc->GetSubActors())
+		{
+			FWorldPartitionActorDesc* SubActorDesc = WorldPartition->GetActorDesc(SubActorGuid);
+			if (!SubActorDesc)
+			{
+				UE_LOG(LogWorldPartitionHLODsBuilder, Error, TEXT("Unknown sub actor guid found, your HLOD actors are probably out of date. Run with -SetupHLODs to fix this. Exiting..."));
+				return false;
+			}
+
+			if (SubActorDesc->GetActorClass()->IsChildOf<AWorldPartitionHLOD>())
+			{
+				if(!ProcessedHLOD.Contains(SubActorGuid))
+				{
+					UE_LOG(LogWorldPartitionHLODsBuilder, Error, TEXT("Child HLOD actor missing or out of order in HLOD workload, exiting..."));
+					return false;
+				}
+			}
+		}
+
+		ProcessedHLOD.Add(HLODActorGuid);
+	}
+
+	return true;
+}
+
+bool UWorldPartitionHLODsBuilder::GenerateBuildManifest() const
+{
+	TArray<TArray<FGuid>> BuildersWorkload = GetHLODWorldloads(BuilderCount);
+
+	FConfigFile ConfigFile;
+
+	for(int32 Idx = 0; Idx < BuilderCount; Idx++)
+	{
+		FString SectionName = FString::Printf(TEXT("HLODBuilder%d"), Idx);
+
+		FConfigSection& Section = ConfigFile.Add(SectionName);
+		for(const FGuid& ActorGuid : BuildersWorkload[Idx])
+		{
+			FString ValueString = ActorGuid.ToString(EGuidFormats::Digits);
+			ValueString += TEXT(", ");
+			
+			const FWorldPartitionActorDesc* ActorDesc = WorldPartition->GetActorDesc(ActorGuid);
+			if (!ActorDesc)
+			{
+				UE_LOG(LogWorldPartitionHLODsBuilder, Error, TEXT("Invalid actor GUID found while generating the HLOD build manifest, exiting..."));
+				return false;
+			}
+
+			ValueString += ActorDesc->GetActorPackage().ToString();
+
+			Section.Add(TEXT("+HLODActorGuid"), ValueString);
+		}
+	}
+
+	ConfigFile.Dirty = true;
+
+	bool bRet = ConfigFile.Write(BuildManifest);
+	if (!bRet)
+	{
+		UE_LOG(LogWorldPartitionHLODsBuilder, Error, TEXT("Failed to write HLOD build manifest \"%s\""), *BuildManifest);
+	}
+
+	return bRet;
+}
\ No newline at end of file
diff --git a/Engine/Source/Editor/UnrealEd/Public/WorldPartition/WorldPartitionHLODsBuilder.h b/Engine/Source/Editor/UnrealEd/Public/WorldPartition/WorldPartitionHLODsBuilder.h
--- a/Engine/Source/Editor/UnrealEd/Public/WorldPartition/WorldPartitionHLODsBuilder.h
+++ b/Engine/Source/Editor/UnrealEd/Public/WorldPartition/WorldPartitionHLODsBuilder.h
@@ -18,10 +18,18 @@ public:
 	// UWorldPartitionBuilder interface end
 
 protected:
+	bool ValidateParams() const;
+
 	bool SetupHLODActors(bool bCreateOnly);
 	bool BuildHLODActors();
 	bool DeleteHLODActors();
 
+	bool GenerateBuildManifest() const;
+	bool GetHLODActorsToBuild(TArray<FGuid>& HLODActorsToBuild) const;
+
+	TArray<TArray<FGuid>> GetHLODWorldloads(int32 NumWorkloads) const;
+	bool ValidateWorkload(const TArray<FGuid>& Workload) const;
+
 private:
 	class UWorldPartition* WorldPartition;
 	class FSourceControlHelper* SourceControlHelper;
@@ -30,4 +38,8 @@ private:
 	bool bBuildHLODs;
 	bool bDeleteHLODs;
 	bool bForceGC;
+
+	FString BuildManifest;
+	int32 BuilderIdx;
+	int32 BuilderCount;
 };
\ No newline at end of file
