Removed unused settings in UHierarchicalLODSettings

diff --git a/Engine/Source/Editor/HierarchicalLODOutliner/Private/HLODOutliner.cpp b/Engine/Source/Editor/HierarchicalLODOutliner/Private/HLODOutliner.cpp
--- a/Engine/Source/Editor/HierarchicalLODOutliner/Private/HLODOutliner.cpp
+++ b/Engine/Source/Editor/HierarchicalLODOutliner/Private/HLODOutliner.cpp
@@ -685,7 +685,7 @@ namespace HLODOutliner
 
 	FText SHLODOutliner::GetBuildText() const
 	{
-		return GetDefault<UHierarchicalLODSettings>()->bInvalidateHLODClusters ? LOCTEXT("BuildMeshes", "Generate Proxy Meshes") : LOCTEXT("RebuildMeshes", "Re-generate Proxy Meshes");
+		return LOCTEXT("BuildMeshes", "Generate Proxy Meshes");
 	}
 
 	FText SHLODOutliner::GetForceBuildText() const
diff --git a/Engine/Source/Editor/UnrealEd/Private/HierarchicalLOD.cpp b/Engine/Source/Editor/UnrealEd/Private/HierarchicalLOD.cpp
--- a/Engine/Source/Editor/UnrealEd/Private/HierarchicalLOD.cpp
+++ b/Engine/Source/Editor/UnrealEd/Private/HierarchicalLOD.cpp
@@ -52,7 +52,7 @@ DEFINE_LOG_CATEGORY_STATIC(LogLODGenerator, Log, All);
 #define METER_TO_CM		100.0f
 
 UHierarchicalLODSettings::UHierarchicalLODSettings(const FObjectInitializer& ObjectInitializer)
-	: Super(ObjectInitializer), bForceSettingsInAllMaps(false), bInvalidateHLODClusters(true), bDeleteHLODAssets(true), BaseMaterial(nullptr)
+	: Super(ObjectInitializer), bForceSettingsInAllMaps(false), BaseMaterial(nullptr)
 {	
 	BaseMaterial = ConstructorHelpers::FObjectFinder<UMaterialInterface>(TEXT("/Engine/EngineMaterials/BaseFlattenMaterial")).Object;
 }
diff --git a/Engine/Source/Editor/UnrealEd/Public/HierarchicalLOD.h b/Engine/Source/Editor/UnrealEd/Public/HierarchicalLOD.h
--- a/Engine/Source/Editor/UnrealEd/Public/HierarchicalLOD.h
+++ b/Engine/Source/Editor/UnrealEd/Public/HierarchicalLOD.h
@@ -48,12 +48,6 @@ public:
 	UPROPERTY(config, EditAnywhere, Category = HLODSystem, AdvancedDisplay, meta = (DisplayName = "Map UAssets used for building HLOD data through the ", RelativeToGameContentDir, LongPackageName))
 	TArray<FFilePath> MapsToBuild;
 
-	UPROPERTY(EditAnywhere, config, Category = HLODSystem, meta = (DisplayName = "Invalidate HLOD Clusters on changes to the Sub Actors"))
-	bool bInvalidateHLODClusters;
-
-	UPROPERTY(EditAnywhere, config, Category = HLODSystem, meta = (DisplayName = "Delete (out-dated) HLOD Assets on Save", editcondition = "bInvalidateHLODClusters"))
-	bool bDeleteHLODAssets;
-	
 	/** Base material used for creating a Constant Material Instance as the Proxy Material */
 	UPROPERTY(EditAnywhere, config, Category = HLODSystem)
 	TSoftObjectPtr<class UMaterialInterface> BaseMaterial;
