Remove call to ResetLoaders in ForceReplaceReference since the call seems superfluous while forcing the load of all bulkdata on all linkers when force deleting an object
Tried reproducing old issues which added this call without success

diff --git a/Engine/Source/Editor/UnrealEd/Private/ObjectTools.cpp b/Engine/Source/Editor/UnrealEd/Private/ObjectTools.cpp
--- a/Engine/Source/Editor/UnrealEd/Private/ObjectTools.cpp
+++ b/Engine/Source/Editor/UnrealEd/Private/ObjectTools.cpp
@@ -927,8 +927,11 @@ namespace ObjectTools
 			}
 		}
 
-		// Reset linker loaders to remove the possibility that any references to 'ObjectsToReplace' exist in the loaders (these can't get picked up by the replace archives)
-		ResetLoaders(nullptr);
+		// @note FH: There shouldn't be a need to reset all loaders here to replace references. This actually currently causes all bulkdata to be force loaded since we are getting rid on the underlying loader archive
+		// Although object references in linkers aren't tracked by the GC nor can't be replaced by the reference gathering archives, they will be properly cleaned out of the linker if the linker outlives the objects for one thing,
+		// but moreover the linkers associated with the packages of the object we are replacing will also been cleaned up if we are deleting those packages after the force replace references.
+		// For both of these reasons, a call to reset all linkers here seems entirely unnecessary
+		//ResetLoaders(nullptr);
 
 		TMap<UObject*, int32> ObjToNumRefsMap;
 		if( ObjectToReplaceWith != NULL )
