Add an explicit handoff of the hierarchy timers by calling EndNetworkFileServer to allow the existence of more than COTFS at one time as needed for external use case.

diff --git a/Engine/Source/Editor/UnrealEd/Private/CookOnTheFlyServer.cpp b/Engine/Source/Editor/UnrealEd/Private/CookOnTheFlyServer.cpp
--- a/Engine/Source/Editor/UnrealEd/Private/CookOnTheFlyServer.cpp
+++ b/Engine/Source/Editor/UnrealEd/Private/CookOnTheFlyServer.cpp
@@ -471,8 +471,6 @@ UCookOnTheFlyServer::~UCookOnTheFlyServer()
 
 	delete CookByTheBookOptions;
 	CookByTheBookOptions = nullptr;
-
-	ClearHierarchyTimers();
 }
 
 // This tick only happens in the editor.  The cook commandlet directly calls tick on the side.
@@ -3363,6 +3361,11 @@ void UCookOnTheFlyServer::EndNetworkFileServer()
 		delete NetworkFileServer;
 		NetworkFileServer = NULL;
 	}
+
+	if (NetworkFileServers.Num() > 0)
+	{
+		ClearHierarchyTimers();
+	}
 	NetworkFileServers.Empty();
 }
 
