Fix crash applying instance properties

diff --git a/Engine/Source/Editor/UnrealEd/Private/Editor.cpp b/Engine/Source/Editor/UnrealEd/Private/Editor.cpp
--- a/Engine/Source/Editor/UnrealEd/Private/Editor.cpp
+++ b/Engine/Source/Editor/UnrealEd/Private/Editor.cpp
@@ -1546,8 +1546,9 @@ namespace EditorUtilities
 										// We also need to double check that either the direct archetype of the target is also identical
 										if (ComponentArchetypeInstance->GetArchetype() != TargetComponent)
 										{
+											UClass* TargetCompClass = TargetComponent->GetClass();
 											UActorComponent* CheckComponent = CastChecked<UActorComponent>(ComponentArchetypeInstance->GetArchetype());
-											while (CheckComponent != ComponentArchetypeInstance)
+											while (CheckComponent != ComponentArchetypeInstance && CheckComponent->GetClass() == TargetCompClass)
 											{
 												if (!Property->Identical_InContainer( CheckComponent, TargetComponent ))
 												{
