Fix mouse tracking making it through when alt camera oribiting in level editor viewport.

diff --git a/Engine/Source/Editor/UnrealEd/Private/LevelEditorViewport.cpp b/Engine/Source/Editor/UnrealEd/Private/LevelEditorViewport.cpp
--- a/Engine/Source/Editor/UnrealEd/Private/LevelEditorViewport.cpp
+++ b/Engine/Source/Editor/UnrealEd/Private/LevelEditorViewport.cpp
@@ -3024,7 +3024,7 @@ void FLevelEditorViewportClient::TrackingStarted( const FInputEventState& InInpu
 			// Suspend actor/component modification during each delta step to avoid recording unnecessary overhead into the transaction buffer
 			GEditor->DisableDeltaModification(true);
 		}
-		else if (!IsFlightCameraActive())
+		else if (!IsFlightCameraActive() && !ShouldOrbitCamera())
 		{
 			// Give the interactive tools context a shot at tracking.
 			ModeTools->GetInteractiveToolsContext()->StartTracking(this, Viewport);
