Mode toolkits no longer own a hard reference to their mode manager, and instead use their toolkit host's mode manager.

This prevents a lingering ref that prevented the mode manager from falling out of scope in some asset editors.

diff --git a/Engine/Source/Editor/UnrealEd/Private/Toolkits/AssetEditorToolkit.cpp b/Engine/Source/Editor/UnrealEd/Private/Toolkits/AssetEditorToolkit.cpp
--- a/Engine/Source/Editor/UnrealEd/Private/Toolkits/AssetEditorToolkit.cpp
+++ b/Engine/Source/Editor/UnrealEd/Private/Toolkits/AssetEditorToolkit.cpp
@@ -428,6 +428,17 @@ FSlateIcon FAssetEditorToolkit::GetEditorModeIcon() const
 	return FSlateIcon();
 }
 
+FEditorModeTools& FAssetEditorToolkit::GetEditorModeManager() const
+{
+	if (IsWorldCentricAssetEditor() && IsHosted())
+	{
+		return GetToolkitHost()->GetEditorModeManager();
+	}
+
+	check(EditorModeManager.IsValid());
+	return *EditorModeManager.Get();
+}
+
 const TArray< UObject* >* FAssetEditorToolkit::GetObjectsCurrentlyBeingEdited() const
 {
 	return &EditingObjects;
diff --git a/Engine/Source/Editor/UnrealEd/Private/Toolkits/BaseToolkit.cpp b/Engine/Source/Editor/UnrealEd/Private/Toolkits/BaseToolkit.cpp
--- a/Engine/Source/Editor/UnrealEd/Private/Toolkits/BaseToolkit.cpp
+++ b/Engine/Source/Editor/UnrealEd/Private/Toolkits/BaseToolkit.cpp
@@ -26,8 +26,6 @@ FBaseToolkit::FBaseToolkit()
 
 FBaseToolkit::~FBaseToolkit()
 {
-
-	EditorModeManager.Reset();
 }
 
 
@@ -110,17 +108,6 @@ bool FBaseToolkit::IsBlueprintEditor() const
 	return false;
 }
 
-FEditorModeTools& FBaseToolkit::GetEditorModeManager() const
-{
-	if (IsWorldCentricAssetEditor() && IsHosted())
-	{
-		return GetToolkitHost()->GetEditorModeManager();
-	}
-
-	check(EditorModeManager.IsValid());
-	return *EditorModeManager.Get();
-}
-
 #undef LOCTEXT_NAMESPACE
 
 
@@ -136,7 +123,6 @@ void FModeToolkit::Init(const TSharedPtr< class IToolkitHost >& InitToolkitHost,
 	ToolkitMode = EToolkitMode::Type::Standalone;
 	ToolkitHost = InitToolkitHost;
 	OwningEditorMode = InOwningMode;
-	EditorModeManager = InitToolkitHost->GetEditorModeManager().AsShared();
 
 	UInteractiveToolManager* ToolManager = GetEditorModeManager().GetInteractiveToolsContext()->ToolManager;
 	ToolManager->OnToolStarted.AddSP(this, &FModeToolkit::OnToolStarted);
@@ -174,9 +160,12 @@ void FModeToolkit::Init(const TSharedPtr< class IToolkitHost >& InitToolkitHost,
 
 FModeToolkit::~FModeToolkit()
 {
-	UInteractiveToolManager* ToolManager = GetEditorModeManager().GetInteractiveToolsContext()->ToolManager;
-	ToolManager->OnToolStarted.RemoveAll(this);
-	ToolManager->OnToolEnded.RemoveAll(this);
+	if (IsHosted())
+	{
+		UInteractiveToolManager* ToolManager = GetEditorModeManager().GetInteractiveToolsContext()->ToolManager;
+		ToolManager->OnToolStarted.RemoveAll(this);
+		ToolManager->OnToolEnded.RemoveAll(this);
+	}
 
 	OwningEditorMode.Reset();
 }
@@ -310,6 +299,12 @@ FSlateIcon FModeToolkit::GetEditorModeIcon() const
 	return FSlateIcon();
 }
 
+FEditorModeTools& FModeToolkit::GetEditorModeManager() const
+{
+	check(IsHosted());
+	return GetToolkitHost()->GetEditorModeManager();
+}
+
 TWeakObjectPtr<UEdMode> FModeToolkit::GetScriptableEditorMode() const
 {
 	return OwningEditorMode;
diff --git a/Engine/Source/Editor/UnrealEd/Public/Toolkits/AssetEditorToolkit.h b/Engine/Source/Editor/UnrealEd/Public/Toolkits/AssetEditorToolkit.h
--- a/Engine/Source/Editor/UnrealEd/Public/Toolkits/AssetEditorToolkit.h
+++ b/Engine/Source/Editor/UnrealEd/Public/Toolkits/AssetEditorToolkit.h
@@ -134,6 +134,7 @@ public:
 	virtual TWeakObjectPtr<class UEdMode> GetScriptableEditorMode() const final;
 	virtual FText GetEditorModeDisplayName() const override;
 	virtual FSlateIcon GetEditorModeIcon() const override;
+	virtual FEditorModeTools& GetEditorModeManager() const final;
 
 	/** IAssetEditorInstance interface */
 	virtual FName GetEditorName() const override;
@@ -386,7 +387,8 @@ protected:
 	/** Whether only dirty assets should be prompted about on save - otherwise all edited assets will be prompted to the user for save/check-out */
 	bool bCheckDirtyOnAssetSave;
 
-	/** The asset editor mode manager */
+	/** The editor mode manager */
+	TSharedPtr<FEditorModeTools> EditorModeManager;
 	UE_DEPRECATED(4.26, "Use EditorModeManager instead.")
 	FAssetEditorModeManager* AssetEditorModeManager;
 
diff --git a/Engine/Source/Editor/UnrealEd/Public/Toolkits/BaseToolkit.h b/Engine/Source/Editor/UnrealEd/Public/Toolkits/BaseToolkit.h
--- a/Engine/Source/Editor/UnrealEd/Public/Toolkits/BaseToolkit.h
+++ b/Engine/Source/Editor/UnrealEd/Public/Toolkits/BaseToolkit.h
@@ -43,7 +43,7 @@ public:
 	virtual bool IsBlueprintEditor() const override;
 	virtual TSharedRef<FWorkspaceItem> GetWorkspaceMenuCategory() const override { return WorkspaceMenuCategory.ToSharedRef(); }
 
-	virtual FEditorModeTools& GetEditorModeManager() const final;
+	virtual FEditorModeTools& GetEditorModeManager() const = 0;
 
 public:
 
@@ -83,8 +83,6 @@ protected:
 
 	/** The workspace menu category of this toolkit */
 	TSharedPtr<FWorkspaceItem> WorkspaceMenuCategory;
-
-	TSharedPtr<FEditorModeTools> EditorModeManager;
 };
 
 
@@ -125,6 +123,7 @@ public:
 	virtual class FEdMode* GetEditorMode() const override;
 	virtual FText GetEditorModeDisplayName() const override;
 	virtual FSlateIcon GetEditorModeIcon() const override;
+	virtual FEditorModeTools& GetEditorModeManager() const final;
 
 	virtual TWeakObjectPtr<UEdMode> GetScriptableEditorMode() const final;
 	virtual TSharedPtr<SWidget> GetInlineContent() const override;
