Fix undo/redo in geometry mode tools.

CL 15676489 added in the viewport listening to the keyboard capture state from ITF, but this has not yet been fully hooked up. Reverting back to dropping that handled value on the floor in the viewport so that we fall through to the main frame input key function.

diff --git a/Engine/Source/Editor/UnrealEd/Private/EditorViewportClient.cpp b/Engine/Source/Editor/UnrealEd/Private/EditorViewportClient.cpp
--- a/Engine/Source/Editor/UnrealEd/Private/EditorViewportClient.cpp
+++ b/Engine/Source/Editor/UnrealEd/Private/EditorViewportClient.cpp
@@ -2933,7 +2933,7 @@ bool FEditorViewportClient::InputKey(FViewport* InViewport, int32 ControllerId,
 	// Give the interactive tools context a shot at input if it hasn't been handled already.
 	if (!bHandled)
 	{
-		bHandled = ModeTools->GetInteractiveToolsContext()->InputKey(this, Viewport, Key, Event);
+		ModeTools->GetInteractiveToolsContext()->InputKey(this, Viewport, Key, Event);
 	}
 
 	//apply the visibility and set the cursor positions
