Sync the pivot location immediately when clicking on a viewport element

#rnx

diff --git a/Engine/Source/Editor/UnrealEd/Private/LevelViewportClickHandlers.cpp b/Engine/Source/Editor/UnrealEd/Private/LevelViewportClickHandlers.cpp
--- a/Engine/Source/Editor/UnrealEd/Private/LevelViewportClickHandlers.cpp
+++ b/Engine/Source/Editor/UnrealEd/Private/LevelViewportClickHandlers.cpp
@@ -203,6 +203,10 @@ namespace LevelViewportClickHandlers
 						}
 						LevelEditorElementSelectionSet->SelectElement(ResolvedElement, SelectionOptions);
 					}
+
+					// Notify any pending selection change now, as this avoids the visual pivot location "lagging" behind the actual selection,
+					// and also ensures that the pivot is at the correct location prior to opening any context menus (which block the update)
+					LevelEditorElementSelectionSet->NotifyPendingChanges();
 				}
 
 				if (bNeedViewportRefresh)
