Fix Static Analysis / Non-Unity Win64.

diff --git a/Engine/Source/Editor/UnrealEd/Private/WorldPartition/WorldPartitionHLODsBuilder.cpp b/Engine/Source/Editor/UnrealEd/Private/WorldPartition/WorldPartitionHLODsBuilder.cpp
--- a/Engine/Source/Editor/UnrealEd/Private/WorldPartition/WorldPartitionHLODsBuilder.cpp
+++ b/Engine/Source/Editor/UnrealEd/Private/WorldPartition/WorldPartitionHLODsBuilder.cpp
@@ -7,6 +7,7 @@
 #include "Editor.h"
 #include "Logging/LogMacros.h"
 #include "FileHelpers.h"
+#include "Misc/ConfigCacheIni.h"
 
 #include "WorldPartition/WorldPartition.h"
 #include "WorldPartition/WorldPartitionSubsystem.h"
@@ -110,7 +111,7 @@ bool UWorldPartitionHLODsBuilder::RequiresCommandletRendering() const
 
 bool UWorldPartitionHLODsBuilder::ValidateParams() const
 {
-	if (bSetupHLODs && !BuildManifest.IsEmpty() && BuilderCount == INDEX_NONE)
+	if (bSetupHLODs && !BuildManifest.IsEmpty())
 	{
 		if (BuilderCount == INDEX_NONE)
 		{
@@ -119,7 +120,7 @@ bool UWorldPartitionHLODsBuilder::ValidateParams() const
 		}
 	}
 
-	if (bBuildHLODs && !BuildManifest.IsEmpty() && BuilderIdx == INDEX_NONE)
+	if (bBuildHLODs && !BuildManifest.IsEmpty())
 	{
 		if (BuilderIdx == INDEX_NONE)
 		{
