Update Class Picker Dialog Button Styles

diff --git a/Engine/Source/Editor/UnrealEd/Private/SClassPickerDialog.cpp b/Engine/Source/Editor/UnrealEd/Private/SClassPickerDialog.cpp
--- a/Engine/Source/Editor/UnrealEd/Private/SClassPickerDialog.cpp
+++ b/Engine/Source/Editor/UnrealEd/Private/SClassPickerDialog.cpp
@@ -152,6 +152,8 @@ void SClassPickerDialog::Construct(const FArguments& InArgs)
 					+SUniformGridPanel::Slot(0,0)
 					[
 						SNew(SButton)
+						.ButtonStyle(FAppStyle::Get(), "PrimaryButton")
+						.TextStyle(FAppStyle::Get(), "DialogButtonText")
 						.Text(NSLOCTEXT("SClassPickerDialog", "ClassPickerSelectButton", "Select"))
 						.HAlign(HAlign_Center)
 						.Visibility( this, &SClassPickerDialog::GetSelectButtonVisibility )
@@ -161,6 +163,7 @@ void SClassPickerDialog::Construct(const FArguments& InArgs)
 					+SUniformGridPanel::Slot(1,0)
 					[
 						SNew(SButton)
+						.TextStyle(FAppStyle::Get(), "DialogButtonText")
 						.Text(NSLOCTEXT("SClassPickerDialog", "ClassPickerCancelButton", "Cancel"))
 						.HAlign(HAlign_Center)
 						.OnClicked(this, &SClassPickerDialog::OnClassPickerCanceled)
