Added UWorldPartitionBuilder::PreWorldInitialization()

diff --git a/Engine/Source/Editor/UnrealEd/Classes/Commandlets/WorldPartitionBuilderCommandlet.h b/Engine/Source/Editor/UnrealEd/Classes/Commandlets/WorldPartitionBuilderCommandlet.h
--- a/Engine/Source/Editor/UnrealEd/Classes/Commandlets/WorldPartitionBuilderCommandlet.h
+++ b/Engine/Source/Editor/UnrealEd/Classes/Commandlets/WorldPartitionBuilderCommandlet.h
@@ -17,6 +17,8 @@ class UWorldPartitionBuilderCommandlet : public UCommandlet, public FPackageSour
 	virtual int32 Main(const FString& Params) override;
 	//~ End UCommandlet Interface
 
+	UWorldPartitionBuilder* CreateBuilder(const FString& WorldConfigFilename);
+
 private:
 	UPROPERTY()
 	TArray<TObjectPtr<UWorldPartitionBuilder>> Builders;
diff --git a/Engine/Source/Editor/UnrealEd/Private/Commandlets/WorldPartitionBuilderCommandlet.cpp b/Engine/Source/Editor/UnrealEd/Private/Commandlets/WorldPartitionBuilderCommandlet.cpp
--- a/Engine/Source/Editor/UnrealEd/Private/Commandlets/WorldPartitionBuilderCommandlet.cpp
+++ b/Engine/Source/Editor/UnrealEd/Private/Commandlets/WorldPartitionBuilderCommandlet.cpp
@@ -7,6 +7,7 @@
 #include "EngineUtils.h"
 #include "Editor.h"
 #include "Logging/LogMacros.h"
+#include "Misc/CommandLine.h"
 #include "Misc/ConfigCacheIni.h"
 #include "HAL/PlatformFileManager.h"
 #include "ProfilingDebugging/ScopedTimers.h"
@@ -21,6 +22,52 @@
 	: Super(ObjectInitializer)
 {}
 
+UWorldPartitionBuilder* UWorldPartitionBuilderCommandlet::CreateBuilder(const FString& WorldConfigFilename)
+{
+	// Parse builder class name
+	FString BuilderClassName;
+	if (!FParse::Value(FCommandLine::Get(), TEXT("Builder="), BuilderClassName, false))
+	{
+		UE_LOG(LogWorldPartitionBuilderCommandlet, Error, TEXT("Invalid builder name."));
+		return nullptr;
+	}
+
+	UClass* BuilderClass = FindObject<UClass>(ANY_PACKAGE, *BuilderClassName);
+
+	if (!BuilderClass)
+	{
+		UE_LOG(LogWorldPartitionBuilderCommandlet, Error, TEXT("Unknown builder %s."), *BuilderClassName);
+		return nullptr;
+	}
+
+	// Create builder instance
+	UWorldPartitionBuilder* Builder = NewObject<UWorldPartitionBuilder>(this, BuilderClass);
+	check(Builder);
+
+	Builder->AddToRoot();
+
+	// Validate builder settings
+	if (Builder->RequiresCommandletRendering() && !IsAllowCommandletRendering())
+	{
+		UE_LOG(LogWorldPartitionBuilderCommandlet, Error, TEXT("The option \"-AllowCommandletRendering\" must be provided for the %s process to work"), *BuilderClassName);
+		return nullptr;
+	}
+
+	// Load configuration file
+	if (FPlatformFileManager::Get().GetPlatformFile().FileExists(*WorldConfigFilename))
+	{
+		LoadConfig(GetClass(), *WorldConfigFilename);
+	}
+
+	// Load builder configuration
+	if (FPlatformFileManager::Get().GetPlatformFile().FileExists(*WorldConfigFilename))
+	{
+		Builder->LoadConfig(BuilderClass, *WorldConfigFilename);
+	}
+
+	return Builder;
+}
+
 int32 UWorldPartitionBuilderCommandlet::Main(const FString& Params)
 {
 	TRACE_CPUPROFILER_EVENT_SCOPE(UWorldPartitionBuilderCommandlet::Main);
@@ -49,13 +96,19 @@
 		return 1;
 	}
 
-	// Load configuration file
 	FString WorldConfigFilename = FPaths::ChangeExtension(WorldFilename, TEXT("ini"));
-	if (FPlatformFileManager::Get().GetPlatformFile().FileExists(*WorldConfigFilename))
+
+	// Create builder instance
+	UWorldPartitionBuilder* Builder = CreateBuilder(WorldConfigFilename);
+	if (!Builder)
 	{
-		LoadConfig(GetClass(), *WorldConfigFilename);
+		UE_LOG(LogWorldPartitionBuilderCommandlet, Error, TEXT("Failed to create builder, exiting."));
+		return 1;
 	}
 
+	// Perform builder pre world initialisation
+	Builder->PreWorldInitialization(*this);
+
 	// Load the map package
 	UPackage* MapPackage = LoadPackage(NULL, *Tokens[0], LOAD_None);
 	if (!MapPackage)
@@ -112,47 +165,12 @@
 	WorldContext.SetCurrentWorld(World);
 	GWorld = World;
 
-	// Parse builder class name
-	FString BuilderClassName;
-	if (!FParse::Value(*Params, TEXT("Builder="), BuilderClassName, false))
-	{
-		UE_LOG(LogWorldPartitionBuilderCommandlet, Error, TEXT("Invalid builder name."));
-		return 1;
-	}
-
-	UClass* BuilderClass = FindObject<UClass>(ANY_PACKAGE, *BuilderClassName);
-
-	if (!BuilderClass)
-	{
-		UE_LOG(LogWorldPartitionBuilderCommandlet, Error, TEXT("Unknown builder %s."), *BuilderClassName);
-		return 1;
-	}
-
-	// Create builder instance
-	UWorldPartitionBuilder* Builder = NewObject<UWorldPartitionBuilder>(this, BuilderClass);
-	check(Builder);
-
-	Builder->AddToRoot();
-	
-	// Validate builder settings
-	if (Builder->RequiresCommandletRendering() && !IsAllowCommandletRendering())
-	{
-		UE_LOG(LogWorldPartitionBuilderCommandlet, Error, TEXT("The option \"-AllowCommandletRendering\" must be provided for the %s process to work"), *BuilderClassName);
-		return 1;
-	}
-
 	if (Builder->RequiresEntireWorldLoading())
 	{
 		const FBox LoadBox(FVector(-WORLD_MAX, -WORLD_MAX, -WORLD_MAX), FVector(WORLD_MAX, WORLD_MAX, WORLD_MAX));
 		WorldPartition->LoadEditorCells(LoadBox);
 	}
 
-	// Load builder configuration
-	if (FPlatformFileManager::Get().GetPlatformFile().FileExists(*WorldConfigFilename))
-	{
-		Builder->LoadConfig(BuilderClass, *WorldConfigFilename);
-	}
-
 	// Run builder
 	if (!Builder->Run(World, *this))
 	{
diff --git a/Engine/Source/Editor/UnrealEd/Public/WorldPartition/WorldPartitionBuilder.h b/Engine/Source/Editor/UnrealEd/Public/WorldPartition/WorldPartitionBuilder.h
--- a/Engine/Source/Editor/UnrealEd/Public/WorldPartition/WorldPartitionBuilder.h
+++ b/Engine/Source/Editor/UnrealEd/Public/WorldPartition/WorldPartitionBuilder.h
@@ -15,6 +15,8 @@ public:
 	virtual bool RequiresEntireWorldLoading() const PURE_VIRTUAL(UWorldPartitionBuilder::RequiresEntireWorldLoading, return false;);
 	virtual bool Run(UWorld* World, FPackageSourceControlHelper& PackageHelper) PURE_VIRTUAL(UWorldPartitionBuilder::Run, return false;);
 
+	virtual bool PreWorldInitialization(FPackageSourceControlHelper& PackageHelper) { return true; }
+
 protected:
 	bool HasExceededMaxMemory() const;
 	void DoCollectGarbage() const;
