Added support for distributed World Partition HLOD generation
* Runs multiple steps on different machines
* Setup: Done on a single machine, will create/update/delete HLOD actors as required. The distributed setup step takes a "BuilderCount" argument, which will be used to write a build manifest that specify the work that needs to be done by the N builders
* Build: Performed on multiple machines, read work to be done from the build manifest
* Submit: Final step, perform on a single machine. Gather files modified in all the previous steps and submit them in a single changelist

Intermediary files that are created/modified/deleted are transfered between the jobs as build products, and submitted in the last step.

To run with UAT:
RunUAT.bat BuildGraph -Script=Collaboration\Frosty\Frosty\Build\FrostyHLOD.xml -Target="HLOD Generation" -set:MapName=YourMapName -set:HLODBuilderCount=N

To run steps individually:
* Setup (N = Builder count to use)
    YourProject YourMapName -run=WorldPartitionBuilderCommandlet -SCCProvider=Perforce -AllowCommandletRendering -Builder=WorldPartitionHLODsBuilder -SetupHLODs -DistributedBuild -BuilderCount=N
* Build (i from 0 to N-1):
    YourProject YourMapName -run=WorldPartitionBuilderCommandlet -SCCProvider=Perforce -AllowCommandletRendering -Builder=WorldPartitionHLODsBuilder -DistributedBuild -BuildHLODs -BuilderIdx=i
* Submit:
    YourProject YourMapName -run=WorldPartitionBuilderCommandlet -SCCProvider=Perforce -AllowCommandletRendering -Builder=WorldPartitionHLODsBuilder -DistributedBuild -SubmitHLODs

diff --git a/Engine/Source/Editor/UnrealEd/Private/WorldPartition/WorldPartitionHLODsBuilder.cpp b/Engine/Source/Editor/UnrealEd/Private/WorldPartition/WorldPartitionHLODsBuilder.cpp
--- a/Engine/Source/Editor/UnrealEd/Private/WorldPartition/WorldPartitionHLODsBuilder.cpp
+++ b/Engine/Source/Editor/UnrealEd/Private/WorldPartition/WorldPartitionHLODsBuilder.cpp
@@ -3,11 +3,19 @@
 #include "WorldPartition/WorldPartitionHLODsBuilder.h"
 
 #include "CoreMinimal.h"
-#include "EngineUtils.h"
-#include "Editor.h"
+#include "HAL/FileManager.h"
 #include "Logging/LogMacros.h"
-#include "FileHelpers.h"
 #include "Misc/ConfigCacheIni.h"
+#include "Misc/FileHelper.h"
+
+#include "EngineUtils.h"
+#include "SourceControlHelpers.h"
+
+#include "ISourceControlModule.h"
+#include "ISourceControlProvider.h"
+#include "SourceControlOperations.h"
+
+#include "DerivedDataCacheInterface.h"
 
 #include "WorldPartition/WorldPartition.h"
 #include "WorldPartition/WorldPartitionSubsystem.h"
@@ -38,26 +46,32 @@ public:
 
 	virtual bool Checkout(UPackage* Package) const override
 	{
+		ModifiedFiles.Add(GetFilename(Package));
 		return PackageHelper.Checkout(Package);
 	}
 
 	virtual bool Add(UPackage* Package) const override
 	{
+		ModifiedFiles.Add(GetFilename(Package));
 		return PackageHelper.AddToSourceControl(Package);
 	}
 
 	virtual bool Delete(const FString& PackageName) const override
 	{
+		ModifiedFiles.Add(PackageName);
 		return PackageHelper.Delete(PackageName);
 	}
 
 	virtual bool Delete(UPackage* Package) const override
 	{
+		ModifiedFiles.Add(GetFilename(Package));
 		return PackageHelper.Delete(Package);
 	}
 
 	virtual bool Save(UPackage* Package) const override
 	{
+		ModifiedFiles.Add(GetFilename(Package));
+
 		// Checkout package
 		Package->MarkAsFullyLoaded();
 
@@ -85,23 +99,46 @@ public:
 		return true;
 	}
 
+	const TSet<FString> GetModifiedFiles() const
+	{
+		return ModifiedFiles;
+	}
+
 private:
 	FPackageSourceControlHelper& PackageHelper;
+	mutable TSet<FString> ModifiedFiles;
 };
 
+static const FString DistributedBuildWorkingDirName = TEXT("HLODTemp");
+static const FString DistributedBuildManifestName = TEXT("HLODBuildManifest.ini");
+
 UWorldPartitionHLODsBuilder::UWorldPartitionHLODsBuilder(const FObjectInitializer& ObjectInitializer)
 	: Super(ObjectInitializer)
 	, BuilderIdx(INDEX_NONE)
 	, BuilderCount(INDEX_NONE)
+	, DistributedBuildWorkingDir(FPaths::RootDir() / DistributedBuildWorkingDirName)
+	, DistributedBuildManifest(DistributedBuildWorkingDir / DistributedBuildManifestName)
 {
 	bSetupHLODs = FParse::Param(FCommandLine::Get(), TEXT("SetupHLODs"));
 	bBuildHLODs = FParse::Param(FCommandLine::Get(), TEXT("BuildHLODs"));
 	bDeleteHLODs = FParse::Param(FCommandLine::Get(), TEXT("DeleteHLODs"));
-	bForceGC = FParse::Param(FCommandLine::Get(), TEXT("ForceGC"));
+	bSubmitHLODs = FParse::Param(FCommandLine::Get(), TEXT("SubmitHLODs"));
+
+	bDistributedBuild = FParse::Param(FCommandLine::Get(), TEXT("DistributedBuild"));
 
 	FParse::Value(FCommandLine::Get(), TEXT("BuildManifest="), BuildManifest);
 	FParse::Value(FCommandLine::Get(), TEXT("BuilderIdx="), BuilderIdx);
 	FParse::Value(FCommandLine::Get(), TEXT("BuilderCount="), BuilderCount);
+
+	if (bDistributedBuild)
+	{
+		if (!BuildManifest.IsEmpty())
+		{
+			UE_LOG(LogWorldPartitionHLODsBuilder, Warning, TEXT("Ignoring parameter -BuildManifest when a distributed build is performed"));
+		}
+
+		BuildManifest = DistributedBuildManifest;
+	}
 }
 
 bool UWorldPartitionHLODsBuilder::RequiresCommandletRendering() const
@@ -111,40 +148,86 @@ bool UWorldPartitionHLODsBuilder::RequiresCommandletRendering() const
 
 bool UWorldPartitionHLODsBuilder::ValidateParams() const
 {
-	if (bSetupHLODs && !BuildManifest.IsEmpty())
+	if (bSetupHLODs && IsUsingBuildManifest())
 	{
-		if (BuilderCount == INDEX_NONE)
+		if (BuilderCount <= 0)
 		{
-			UE_LOG(LogWorldPartitionHLODsBuilder, Error, TEXT("Missing parameter -BuilderCount=N, exiting..."));
+			UE_LOG(LogWorldPartitionHLODsBuilder, Error, TEXT("Missing parameter -BuilderCount=N (where N > 0), exiting..."));
 			return false;
 		}
 	}
 
-	if (bBuildHLODs && !BuildManifest.IsEmpty())
+	if (bBuildHLODs && IsUsingBuildManifest())
 	{
-		if (BuilderIdx == INDEX_NONE)
+		if (BuilderIdx < 0)
 		{
 			UE_LOG(LogWorldPartitionHLODsBuilder, Error, TEXT("Missing parameter -BuilderIdx=i, exiting..."));
 			return false;
 		}
 
-		if (FPaths::FileExists(BuildManifest))
+		if (!FPaths::FileExists(BuildManifest))
 		{
 			UE_LOG(LogWorldPartitionHLODsBuilder, Error, TEXT("Build manifest file \"%s\" not found, exiting..."), *BuildManifest);
 			return false;
 		}
+
+		FString CurrentEngineVersion = FEngineVersion::Current().ToString();
+		FString ManifestEngineVersion = TEXT("unknown");
+
+		FConfigFile ConfigFile;
+		ConfigFile.Read(BuildManifest);
+		const FConfigSection* ConfigSection = ConfigFile.Find(TEXT("General"));
+		if (ConfigSection)
+		{
+			const FConfigValue* ConfigValue = ConfigSection->Find(TEXT("EngineVersion"));
+			if (ConfigValue)
+			{
+				ManifestEngineVersion = ConfigValue->GetValue();
+			}
+		}
+		if (ManifestEngineVersion != CurrentEngineVersion)
+		{
+			UE_LOG(LogWorldPartitionHLODsBuilder, Error, TEXT("Build manifest engine version doesn't match current engine version (%s vs %s), exiting..."), *ManifestEngineVersion, *CurrentEngineVersion);
+			return false;
+		}
+	}
+
+	if (bSubmitHLODs && !IsDistributedBuild())
+	{
+		UE_LOG(LogWorldPartitionHLODsBuilder, Error, TEXT("-SubmitHLODs argument only valid for distributed builds, exiting..."), *BuildManifest);
+		return false;
+	}
+
+	if (IsDistributedBuild() && !ISourceControlModule::Get().GetProvider().IsEnabled())
+	{
+		UE_LOG(LogWorldPartitionHLODsBuilder, Error, TEXT("Distributed builds requires that a valid source control provider is enabled, exiting..."), *BuildManifest);
+		return false;
 	}
 
 	return true;
 }
 
-bool UWorldPartitionHLODsBuilder::Run(UWorld* World, FPackageSourceControlHelper& PackageHelper)
+bool UWorldPartitionHLODsBuilder::PreWorldInitialization(FPackageSourceControlHelper& PackageHelper)
 {
 	if (!ValidateParams())
 	{
 		return false;
 	}
 
+	bool bRet = true;
+
+	// When running a distributed build, retrieve relevant build products from the previous steps
+	if (IsDistributedBuild() && (bBuildHLODs || bSubmitHLODs))
+	{
+		FString WorkingDirFolder = bBuildHLODs ? FString::Printf(TEXT("HLODBuilder%d"), BuilderIdx) : TEXT("ToSubmit");
+		bRet = CopyFilesFromWorkingDir(WorkingDirFolder);
+	}
+
+	return bRet;
+}
+
+bool UWorldPartitionHLODsBuilder::Run(UWorld* World, FPackageSourceControlHelper& PackageHelper)
+{
 	WorldPartition = World->GetWorldPartition();
 	check(WorldPartition);
 
@@ -163,25 +246,36 @@ bool UWorldPartitionHLODsBuilder::Run(UWorld* World, FPackageSourceControlHelper
 	{
 		bRet = DeleteHLODActors();
 	}
+	else if (bSubmitHLODs)
+	{
+		bRet = SubmitHLODActors();
+	}
 	else
 	{
 		bRet = SetupHLODActors(false);
 	}
-		
+
 	WorldPartition = nullptr;
 	delete SourceControlHelper;
+
+	// TODO: DDC shutdown crash workaround - Wait for any DDC writes to complete
+	GetDerivedDataCacheRef().WaitForQuiescence(true);
 	
 	return bRet;
 }
 
 bool UWorldPartitionHLODsBuilder::SetupHLODActors(bool bCreateOnly)
 {
-	bool bRet = true;
-
 	WorldPartition->GenerateHLOD(SourceControlHelper, bCreateOnly);
 
 	if (bCreateOnly)
 	{
+		// When performing a distributed build, ensure our work folder is empty
+		if (IsDistributedBuild())
+		{
+			IFileManager::Get().DeleteDirectory(*DistributedBuildWorkingDir, false, true);
+		}
+
 		UE_LOG(LogWorldPartitionHLODsBuilder, Display, TEXT("#### World HLOD actors ####"));
 
 		int32 NumActors = 0;
@@ -197,13 +291,74 @@ bool UWorldPartitionHLODsBuilder::SetupHLODActors(bool bCreateOnly)
 
 		UE_LOG(LogWorldPartitionHLODsBuilder, Display, TEXT("#### World contains %d HLOD actors ####"), NumActors);
 
-		if (!BuildManifest.IsEmpty())
+		if (IsUsingBuildManifest())
 		{
-			bRet = GenerateBuildManifest();
+			TMap<FString, int32> FilesToBuilderMap;
+			bool bGenerated = GenerateBuildManifest(FilesToBuilderMap);
+			if (!bGenerated)
+			{
+				return false;
+			}
+
+			// When performing a distributed build, move modified files to the temporary working dir, to be submitted later in the last "submit" step
+			if (IsDistributedBuild())
+			{
+				ModifiedFiles.Append(SourceControlHelper->GetModifiedFiles());
+
+				TArray<TArray<FString>> BuildersFiles;
+				BuildersFiles.SetNum(BuilderCount);
+
+				TArray<FString> FilesToSubmit;
+
+				for (const FString& ModifiedFile : ModifiedFiles)
+				{
+					int32* Idx = FilesToBuilderMap.Find(ModifiedFile);
+					if (Idx)
+					{
+						BuildersFiles[*Idx].Add(ModifiedFile);
+					}
+					else
+					{
+						FilesToSubmit.Add(ModifiedFile);
+					}
+				}
+
+				// Gather build product to ensure intermediary files are copied between the different HLOD generation steps
+				TArray<FString> BuildProducts;
+
+				// Copy files that will be handled by the different builders
+				for (int32 Idx = 0; Idx < BuilderCount; Idx++)
+				{
+					if (!CopyFilesToWorkingDir(FString::Printf(TEXT("HLODBuilder%d"), Idx), BuildersFiles[Idx], BuildProducts))
+					{
+						return false;
+					}
+				}
+
+				// Copy files that won't be handled by builders but must be submitted in the last step
+				if (!CopyFilesToWorkingDir(TEXT("ToSubmit"), FilesToSubmit, BuildProducts))
+				{
+					return false;
+				}
+
+				// The build manifest must also be included as a build product to be available in the next steps
+				BuildProducts.Add(BuildManifest);
+
+				// Write build products to a file
+				FString BuildProductsFile = DistributedBuildWorkingDir / TEXT("BuildProducts.txt");
+				bool bRet = FFileHelper::SaveStringArrayToFile(BuildProducts, *BuildProductsFile);
+				if (!bRet)
+				{
+					UE_LOG(LogWorldPartitionHLODsBuilder, Error, TEXT("Error writing build product file %s"), *BuildProductsFile);
+					return false;
+				}
+
+				ModifiedFiles.Empty();
+			}
 		}
 	}
 
-	return bRet;
+	return true;
 }
 
 bool UWorldPartitionHLODsBuilder::BuildHLODActors()
@@ -234,12 +389,18 @@ bool UWorldPartitionHLODsBuilder::BuildHLODActors()
 
 		HLODActor->BuildHLOD();
 
-		if (HLODActor->GetPackage()->IsDirty())
+		UPackage* ActorPackage = HLODActor->GetPackage();
+		if (ActorPackage->IsDirty())
 		{
-			SourceControlHelper->Save(HLODActor->GetPackage());
+			bool bSaved = SourceControlHelper->Save(ActorPackage);
+			if (!bSaved)
+			{
+				UE_LOG(LogWorldPartitionHLODsBuilder, Error, TEXT("Failed to save %s, exiting..."), *USourceControlHelpers::PackageFilename(ActorPackage));
+				return false;
+			}
 		}
 
-		if (bForceGC || HasExceededMaxMemory())
+		if (HasExceededMaxMemory())
 		{
 			DoCollectGarbage();
 		}
@@ -247,6 +408,30 @@ bool UWorldPartitionHLODsBuilder::BuildHLODActors()
 
 	UE_LOG(LogWorldPartitionHLODsBuilder, Display, TEXT("#### Built %d HLOD actors ####"), HLODActorsToBuild.Num());
 
+	// Move modified files to the temporary working dir, to be submitted later in the final "submit" pass, from a single machine.
+	if (IsDistributedBuild())
+	{
+		ModifiedFiles.Append(SourceControlHelper->GetModifiedFiles());
+
+		TArray<FString> BuildProducts;
+
+		if (!CopyFilesToWorkingDir("ToSubmit", ModifiedFiles.Array(), BuildProducts))
+		{
+			return false;
+		}
+
+		// Write build products to a file
+		FString BuildProductsFile = DistributedBuildWorkingDir / TEXT("BuildProducts.txt");
+		bool bRet = FFileHelper::SaveStringArrayToFile(BuildProducts, *BuildProductsFile);
+		if (!bRet)
+		{
+			UE_LOG(LogWorldPartitionHLODsBuilder, Error, TEXT("Error writing build product file %s"), *BuildProductsFile);
+			return false;
+		}
+
+		ModifiedFiles.Empty();
+	}
+
 	return true;
 }
 
@@ -278,6 +463,35 @@ bool UWorldPartitionHLODsBuilder::DeleteHLODActors()
 	return true;
 }
 
+bool UWorldPartitionHLODsBuilder::SubmitHLODActors()
+{
+	bool bRet = true;
+
+	// Check in all modified files
+	if (ModifiedFiles.Num() > 0)
+	{
+		FText ChangelistDescription = FText::FromString(FString::Printf(TEXT("Rebuilt HLODs for \"%s\" at %s"), *WorldPartition->GetWorld()->GetName(), *FEngineVersion::Current().ToString()));
+
+		TSharedRef<FCheckIn, ESPMode::ThreadSafe> CheckInOperation = ISourceControlOperation::Create<FCheckIn>();
+		CheckInOperation->SetDescription(ChangelistDescription);
+		bRet = ISourceControlModule::Get().GetProvider().Execute(CheckInOperation, ModifiedFiles.Array()) == ECommandResult::Succeeded;
+		if (!bRet)
+		{
+			UE_LOG(LogWorldPartitionHLODsBuilder, Error, TEXT("Failed to submit %d files to source control."), ModifiedFiles.Num());
+		}
+		else
+		{
+			UE_LOG(LogWorldPartitionHLODsBuilder, Display, TEXT("#### Submitted %d files to source control ####"), ModifiedFiles.Num());
+		}
+	}
+	else
+	{
+		UE_LOG(LogWorldPartitionHLODsBuilder, Display, TEXT("#### No files to submit ####"));
+	}
+
+	return bRet;
+}
+
 bool UWorldPartitionHLODsBuilder::GetHLODActorsToBuild(TArray<FGuid>& HLODActorsToBuild) const
 {
 	bool bRet = true;
@@ -293,14 +507,13 @@ bool UWorldPartitionHLODsBuilder::GetHLODActorsToBuild(TArray<FGuid>& HLODActors
 		const FConfigSection* ConfigSection = ConfigFile.Find(SectionName);
 		if (ConfigSection)
 		{
-			TArray<FString> HLODActorStrings;
-			ConfigSection->MultiFind(TEXT("+HLODActorGuid"), HLODActorStrings, /*bMaintainOrder=*/true);
+			TArray<FString> HLODActorGuidStrings;
+			ConfigSection->MultiFind(TEXT("+HLODActorGuid"), HLODActorGuidStrings, /*bMaintainOrder=*/true);
 
-			for (const FString& HLODActorString : HLODActorStrings)
+			for (const FString& HLODActorGuidString : HLODActorGuidStrings)
 			{
-				FString HLODActorGuidString = HLODActorString.Left(HLODActorString.Find(TEXT(",")));
 				FGuid HLODActorGuid;
-				bRet = FGuid::ParseExact(HLODActorGuidString, EGuidFormats::Digits, HLODActorGuid);
+				bRet = FGuid::Parse(HLODActorGuidString, HLODActorGuid);
 				if (bRet)
 				{
 					HLODActorsToBuild.Add(HLODActorGuid);
@@ -454,12 +667,16 @@ bool UWorldPartitionHLODsBuilder::ValidateWorkload(const TArray<FGuid>&Workload)
 	return true;
 }
 
-bool UWorldPartitionHLODsBuilder::GenerateBuildManifest() const
+bool UWorldPartitionHLODsBuilder::GenerateBuildManifest(TMap<FString, int32>& FilesToBuilderMap) const
 {
 	TArray<TArray<FGuid>> BuildersWorkload = GetHLODWorldloads(BuilderCount);
 
 	FConfigFile ConfigFile;
 
+	FConfigSection& GeneralSection = ConfigFile.Add("General");
+	GeneralSection.Add(TEXT("BuilderCount"), FString::FromInt(BuilderCount));
+	GeneralSection.Add(TEXT("EngineVersion"), FEngineVersion::Current().ToString());
+
 	for(int32 Idx = 0; Idx < BuilderCount; Idx++)
 	{
 		FString SectionName = FString::Printf(TEXT("HLODBuilder%d"), Idx);
@@ -467,19 +684,17 @@ bool UWorldPartitionHLODsBuilder::GenerateBuildManifest() const
 		FConfigSection& Section = ConfigFile.Add(SectionName);
 		for(const FGuid& ActorGuid : BuildersWorkload[Idx])
 		{
-			FString ValueString = ActorGuid.ToString(EGuidFormats::Digits);
-			ValueString += TEXT(", ");
-			
+			Section.Add(TEXT("+HLODActorGuid"), ActorGuid.ToString(EGuidFormats::Digits));
+
+			// Track which builder is responsible to handle each actor
 			const FWorldPartitionActorDesc* ActorDesc = WorldPartition->GetActorDesc(ActorGuid);
 			if (!ActorDesc)
 			{
 				UE_LOG(LogWorldPartitionHLODsBuilder, Error, TEXT("Invalid actor GUID found while generating the HLOD build manifest, exiting..."));
 				return false;
 			}
-
-			ValueString += ActorDesc->GetActorPackage().ToString();
-
-			Section.Add(TEXT("+HLODActorGuid"), ValueString);
+			FString ActorPackageFilename = USourceControlHelpers::PackageFilename(ActorDesc->GetActorPackage().ToString());
+			FilesToBuilderMap.Emplace(ActorPackageFilename, Idx);
 		}
 	}
 
@@ -492,4 +707,217 @@ bool UWorldPartitionHLODsBuilder::GenerateBuildManifest() const
 	}
 
 	return bRet;
+}
+
+/*
+	Working Dir structure
+		/HLODBuilder0
+			/Add
+				NewFileA
+				NewFileB
+			/Delete
+				DeletedFileA
+				DeletedFileB
+			/Edit
+				EditedFileA
+				EditedFileB
+
+		/HLODBuilder1
+			...
+		/ToSubmit
+			...
+
+	Distributed mode
+		* Distributed mode is ran into 3 steps
+			* Setup (1 job)		
+			* Build (N jobs)	
+			* Submit (1 job)	
+		
+		* The Setup step will place files under the "HLODBuilder[0-N]" folder. Those files could be new or modified HLOD actors that will be built in the Build step. The setup step will also place files into the "ToSubmit" folder (deleted HLOD actors for example).
+		* Each parallel job in the Build step will retrieve files from the "HLODBuilder[0-N]" folder. They will then proceed to build the HLOD actors as specified in the build manifest file. All built HLOD actor files will then be placed in the /ToSubmit folder.
+		* The Submit step will gather all files under /ToSubmit and submit them.
+		
+
+		|			Setup			|					Build					  |		   Submit			|
+		/Content -----------> /HLODBuilder -----------> /Content -----------> /ToSubmit -----------> /Content
+*/
+
+const FName FileAction_Add(TEXT("Add"));
+const FName FileAction_Edit(TEXT("Edit"));
+const FName FileAction_Delete(TEXT("Delete"));
+
+bool UWorldPartitionHLODsBuilder::CopyFilesToWorkingDir(const FString& TargetDir, const TArray<FString>& FilesToCopy, TArray<FString>& BuildProducts)
+{
+	FString AbsoluteTargetDir = DistributedBuildWorkingDir / TargetDir / TEXT("");
+
+	TArray<FString> FilesToDelete;
+
+	for (const FString& SourceFilename : FilesToCopy)
+	{
+		FSourceControlState FileState = USourceControlHelpers::QueryFileState(SourceFilename);
+		if (FileState.bIsValid)
+		{
+			bool bShouldCopyToWorkingDir = false;
+			FName FileAction;
+
+			if (FileState.bIsAdded)
+			{
+				FileAction = FileAction_Add;
+				FilesToDelete.Add(SourceFilename);
+			}
+			else if (FileState.bIsCheckedOut)
+			{
+				FileAction = FileAction_Edit;
+			}
+			else if (FileState.bIsDeleted)
+			{
+				FileAction = FileAction_Delete;
+			}
+
+			if (!FileAction.IsNone())
+			{
+				FString SourceFilenameRelativeToRoot = SourceFilename;
+				FPaths::MakePathRelativeTo(SourceFilenameRelativeToRoot, *FPaths::RootDir());
+
+				FString TargetFilename = AbsoluteTargetDir / FileAction.ToString() / SourceFilenameRelativeToRoot;
+
+				BuildProducts.Add(TargetFilename);
+
+				if (FileAction != FileAction_Delete)
+				{
+					bool bRet = IFileManager::Get().Copy(*TargetFilename, *SourceFilename, false) == COPY_OK;
+					if (!bRet)
+					{
+						UE_LOG(LogWorldPartitionHLODsBuilder, Error, TEXT("Failed to copy file from \"%s\" to \"%s\""), *SourceFilename, *TargetFilename);
+						return false;
+					}
+				}
+				else
+				{
+					bool bRet = FFileHelper::SaveStringToFile(TEXT(""), *TargetFilename);
+					if (!bRet)
+					{
+						UE_LOG(LogWorldPartitionHLODsBuilder, Error, TEXT("Failed to create empty file at \"%s\""), *TargetFilename);
+						return false;
+					}
+				}
+			}
+		}
+	}
+
+	bool bRet = USourceControlHelpers::RevertFiles(FilesToCopy);
+	if (!bRet)
+	{
+		UE_LOG(LogWorldPartitionHLODsBuilder, Error, TEXT("Failed to revert modified files: %s"), *USourceControlHelpers::LastErrorMsg().ToString());
+		return false;
+	}
+
+	// Delete files we added
+	for (const FString& FileToDelete : FilesToDelete)
+	{
+		if (!IFileManager::Get().Delete(*FileToDelete, false, true))
+		{
+			UE_LOG(LogWorldPartitionHLODsBuilder, Error, TEXT("Error deleting file %s locally"), *FileToDelete);
+			return false;
+		}
+	}
+
+	return true;
+}
+
+bool UWorldPartitionHLODsBuilder::CopyFilesFromWorkingDir(const FString& SourceDir)
+{
+	FString AbsoluteSourceDir = DistributedBuildWorkingDir / SourceDir / TEXT("");
+
+	TArray<FString> Files;
+	IFileManager::Get().FindFilesRecursive(Files, *AbsoluteSourceDir, TEXT("*.*"), true, false);
+
+	TArray<FString>			FilesToAdd;
+	TMap<FString, FString>	FilesToCopy;
+	TArray<FString>			FilesToDelete;
+
+	bool bRet = true;
+
+	for(const FString& File : Files)
+	{
+		FString PathRelativeToWorkingDir = File;
+		FPaths::MakePathRelativeTo(PathRelativeToWorkingDir, *AbsoluteSourceDir);
+
+		FString FileActionString;
+		const int32 SlashIndex = PathRelativeToWorkingDir.Find(TEXT("/"));
+		if (SlashIndex != INDEX_NONE)
+		{
+			FileActionString = PathRelativeToWorkingDir.Mid(0, SlashIndex);
+		}
+
+		FPaths::MakePathRelativeTo(PathRelativeToWorkingDir, *(FileActionString / TEXT("")));
+		FString FullPathInRootDirectory =  FPaths::RootDir() / PathRelativeToWorkingDir;
+
+		FName FileAction(FileActionString);
+		if (FileAction == FileAction_Add)
+		{
+			bRet = IFileManager::Get().Copy(*FullPathInRootDirectory, *File, false) == COPY_OK;
+			if (!bRet)
+			{
+				UE_LOG(LogWorldPartitionHLODsBuilder, Error, TEXT("Failed to copy file from \"%s\" to \"%s\""), *File, *FullPathInRootDirectory);
+				return false;
+			}
+			FilesToAdd.Add(*FullPathInRootDirectory);
+		}
+		else if (FileAction == FileAction_Edit)
+		{
+			FilesToCopy.Add(FullPathInRootDirectory, File);
+		}
+		else if (FileAction == FileAction_Delete)
+		{
+			FilesToDelete.Add(*FullPathInRootDirectory);
+		}
+		else
+		{
+			UE_LOG(LogWorldPartitionHLODsBuilder, Error, TEXT("Unsupported file action %s for file %s"), *FileActionString, *FullPathInRootDirectory);
+		}
+	}
+
+	TArray<FString> FilesToEdit;
+	FilesToCopy.GetKeys(FilesToEdit);
+
+	bRet = USourceControlHelpers::MarkFilesForAdd(FilesToAdd);
+	if (!bRet)
+	{
+		UE_LOG(LogWorldPartitionHLODsBuilder, Error, TEXT("Adding files to source control failed: %s"), *USourceControlHelpers::LastErrorMsg().ToString());
+		while (!FPlatformMisc::IsDebuggerPresent())
+		{
+			FPlatformProcess::Sleep(0.1f);
+		}
+		return false;
+	}
+	ModifiedFiles.Append(FilesToAdd);
+
+	bRet = USourceControlHelpers::MarkFilesForDelete(FilesToDelete);
+	if (!bRet)
+	{
+		UE_LOG(LogWorldPartitionHLODsBuilder, Error, TEXT("Deleting files from source control failed: %s"), *USourceControlHelpers::LastErrorMsg().ToString());
+		return false;
+	}
+	ModifiedFiles.Append(FilesToDelete);
+
+	bRet = USourceControlHelpers::CheckOutFiles(FilesToEdit);
+	if (!bRet)
+	{
+		UE_LOG(LogWorldPartitionHLODsBuilder, Error, TEXT("Checking out files from source control failed: %s"), *USourceControlHelpers::LastErrorMsg().ToString());
+		return false;
+	}
+
+	for(const auto& Pair : FilesToCopy)
+	{
+		bRet = IFileManager::Get().Copy(*Pair.Key, *Pair.Value) == COPY_OK;
+		if (!bRet)
+		{
+			UE_LOG(LogWorldPartitionHLODsBuilder, Error, TEXT("Failed to copy file from \"%s\" to \"%s\""), *Pair.Value, *Pair.Key);
+			return false;
+		}
+		ModifiedFiles.Add(Pair.Key);
+	}
+
+	return true;
 }
\ No newline at end of file
diff --git a/Engine/Source/Editor/UnrealEd/Public/WorldPartition/WorldPartitionHLODsBuilder.h b/Engine/Source/Editor/UnrealEd/Public/WorldPartition/WorldPartitionHLODsBuilder.h
--- a/Engine/Source/Editor/UnrealEd/Public/WorldPartition/WorldPartitionHLODsBuilder.h
+++ b/Engine/Source/Editor/UnrealEd/Public/WorldPartition/WorldPartitionHLODsBuilder.h
@@ -14,32 +14,45 @@ public:
 	// UWorldPartitionBuilder interface begin
 	virtual bool RequiresCommandletRendering() const override;
 	virtual bool RequiresEntireWorldLoading() const override { return false; }
+	virtual bool PreWorldInitialization(FPackageSourceControlHelper& PackageHelper) override;
 	virtual bool Run(UWorld* World, FPackageSourceControlHelper& PackageHelper) override;
 	// UWorldPartitionBuilder interface end
 
 protected:
+	bool IsDistributedBuild() const { return bDistributedBuild; }
+	bool IsUsingBuildManifest() const { return !BuildManifest.IsEmpty(); }
 	bool ValidateParams() const;
 
 	bool SetupHLODActors(bool bCreateOnly);
 	bool BuildHLODActors();
 	bool DeleteHLODActors();
+	bool SubmitHLODActors();
 
-	bool GenerateBuildManifest() const;
+	bool GenerateBuildManifest(TMap<FString, int32>& FilesToBuilderMap) const;
 	bool GetHLODActorsToBuild(TArray<FGuid>& HLODActorsToBuild) const;
 
 	TArray<TArray<FGuid>> GetHLODWorldloads(int32 NumWorkloads) const;
 	bool ValidateWorkload(const TArray<FGuid>& Workload) const;
 
+	bool CopyFilesToWorkingDir(const FString& TargetDir, const TArray<FString>& FilesToCopy, TArray<FString>& BuildProducts);
+	bool CopyFilesFromWorkingDir(const FString& SourceDir);
+
 private:
 	class UWorldPartition* WorldPartition;
 	class FSourceControlHelper* SourceControlHelper;
 
+	// Options --
 	bool bSetupHLODs;
 	bool bBuildHLODs;
 	bool bDeleteHLODs;
-	bool bForceGC;
-
+	bool bSubmitHLODs;
+	bool bDistributedBuild;
 	FString BuildManifest;
 	int32 BuilderIdx;
 	int32 BuilderCount;
+
+	TSet<FString> ModifiedFiles;
+
+	const FString DistributedBuildWorkingDir;
+	const FString DistributedBuildManifest;
 };
\ No newline at end of file
diff --git a/Engine/Source/Programs/AutomationTool/BuildGraph/Tasks/GatherBuildProductsFromFileTask.cs b/Engine/Source/Programs/AutomationTool/BuildGraph/Tasks/GatherBuildProductsFromFileTask.cs
new file mode 100644
--- /dev/null
+++ b/Engine/Source/Programs/AutomationTool/BuildGraph/Tasks/GatherBuildProductsFromFileTask.cs
@@ -0,0 +1,66 @@
+// Copyright Epic Games, Inc. All Rights Reserved.
+
+using System;
+using System.Collections.Generic;
+using System.IO;
+using System.Text;
+using System.Xml;
+
+using EpicGames.Core;
+using UnrealBuildTool;
+
+namespace AutomationTool.Tasks
+{	
+	public class GatherBuildProductsFromFileTaskParameters
+	{
+		[TaskParameter]
+		public string BuildProductsFile;
+	}
+
+	[TaskElement("GatherBuildProductsFromFile", typeof(GatherBuildProductsFromFileTaskParameters))]
+	class GatherBuildProductsFromFileTask : CustomTask
+	{
+		public GatherBuildProductsFromFileTask(GatherBuildProductsFromFileTaskParameters InParameters)
+		{
+			Parameters = InParameters;
+		}
+
+		public override void Execute(JobContext Job, HashSet<FileReference> BuildProducts, Dictionary<string, HashSet<FileReference>> TagNameToFileSet)
+		{
+			List<FileReference> CleanupFiles = new List<FileReference>();
+
+			CommandUtils.LogInformation("Gathering BuildProducts from {0}...", Parameters.BuildProductsFile);
+
+			try
+			{
+				var FileBuildProducts = File.ReadAllLines(Parameters.BuildProductsFile);
+				foreach(var BuildProduct in FileBuildProducts)
+				{
+					CommandUtils.LogInformation("Adding file to build products: {0}", BuildProduct);
+					BuildProducts.Add(new FileReference(BuildProduct));
+				}
+			}
+			catch (Exception Ex)
+			{
+				CommandUtils.LogInformation("Failed to gather build products: {0}", Ex.Message);
+			}
+		}
+		
+		public override void Write(XmlWriter Writer)
+		{
+			Write(Writer, Parameters);
+		}
+		
+		public override IEnumerable<string> FindConsumedTagNames()
+		{
+			yield break;
+		}
+
+		public override IEnumerable<string> FindProducedTagNames()
+		{
+			yield break;
+		}
+
+		public GatherBuildProductsFromFileTaskParameters Parameters;
+	}
+}
