Run HLOD "build" step that require rendering on "Win64_NoAWS" agents
Ensure -AllowCommandletRendering is only needed for the "build" step

diff --git a/Engine/Source/Editor/UnrealEd/Private/WorldPartition/WorldPartitionHLODsBuilder.cpp b/Engine/Source/Editor/UnrealEd/Private/WorldPartition/WorldPartitionHLODsBuilder.cpp
--- a/Engine/Source/Editor/UnrealEd/Private/WorldPartition/WorldPartitionHLODsBuilder.cpp
+++ b/Engine/Source/Editor/UnrealEd/Private/WorldPartition/WorldPartitionHLODsBuilder.cpp
@@ -143,7 +143,9 @@ UWorldPartitionHLODsBuilder::UWorldPartitionHLODsBuilder(const FObjectInitialize
 
 bool UWorldPartitionHLODsBuilder::RequiresCommandletRendering() const
 {
-	return bBuildHLODs || !(bSetupHLODs || bDeleteHLODs);
+	// Commandlet requires rendering only for building HLODs
+	// Building will occur either if -BuildHLODs is provided or no explicit step arguments are provided
+	return bBuildHLODs || !(bSetupHLODs || bDeleteHLODs || bSubmitHLODs);
 }
 
 bool UWorldPartitionHLODsBuilder::ValidateParams() const
