WorldPartitionHLODsBuilder - Added a few constants

diff --git a/Engine/Source/Editor/UnrealEd/Private/WorldPartition/WorldPartitionHLODsBuilder.cpp b/Engine/Source/Editor/UnrealEd/Private/WorldPartition/WorldPartitionHLODsBuilder.cpp
--- a/Engine/Source/Editor/UnrealEd/Private/WorldPartition/WorldPartitionHLODsBuilder.cpp
+++ b/Engine/Source/Editor/UnrealEd/Private/WorldPartition/WorldPartitionHLODsBuilder.cpp
@@ -111,6 +111,10 @@ private:
 
 static const FString DistributedBuildWorkingDirName = TEXT("HLODTemp");
 static const FString DistributedBuildManifestName = TEXT("HLODBuildManifest.ini");
+static const FString BuildProductsFileName = TEXT("BuildProducts.txt");
+
+FString GetHLODBuilderFolderName(uint32 BuilderIndex) { return FString::Printf(TEXT("HLODBuilder%d"), BuilderIndex); }
+FString GetToSubmitFolderName() { return TEXT("ToSubmit"); }
 
 UWorldPartitionHLODsBuilder::UWorldPartitionHLODsBuilder(const FObjectInitializer& ObjectInitializer)
 	: Super(ObjectInitializer)
@@ -221,7 +225,7 @@ bool UWorldPartitionHLODsBuilder::PreWorldInitialization(FPackageSourceControlHe
 	// When running a distributed build, retrieve relevant build products from the previous steps
 	if (IsDistributedBuild() && (bBuildHLODs || bSubmitHLODs))
 	{
-		FString WorkingDirFolder = bBuildHLODs ? FString::Printf(TEXT("HLODBuilder%d"), BuilderIdx) : TEXT("ToSubmit");
+		FString WorkingDirFolder = bBuildHLODs ? GetHLODBuilderFolderName(BuilderIdx) : GetToSubmitFolderName();
 		bRet = CopyFilesFromWorkingDir(WorkingDirFolder);
 	}
 
@@ -331,14 +335,14 @@ bool UWorldPartitionHLODsBuilder::SetupHLODActors(bool bCreateOnly)
 				// Copy files that will be handled by the different builders
 				for (int32 Idx = 0; Idx < BuilderCount; Idx++)
 				{
-					if (!CopyFilesToWorkingDir(FString::Printf(TEXT("HLODBuilder%d"), Idx), BuildersFiles[Idx], BuildProducts))
+					if (!CopyFilesToWorkingDir(GetHLODBuilderFolderName(Idx), BuildersFiles[Idx], BuildProducts))
 					{
 						return false;
 					}
 				}
 
 				// Copy files that won't be handled by builders but must be submitted in the last step
-				if (!CopyFilesToWorkingDir(TEXT("ToSubmit"), FilesToSubmit, BuildProducts))
+				if (!CopyFilesToWorkingDir(GetToSubmitFolderName(), FilesToSubmit, BuildProducts))
 				{
 					return false;
 				}
@@ -347,7 +351,7 @@ bool UWorldPartitionHLODsBuilder::SetupHLODActors(bool bCreateOnly)
 				BuildProducts.Add(BuildManifest);
 
 				// Write build products to a file
-				FString BuildProductsFile = DistributedBuildWorkingDir / TEXT("BuildProducts.txt");
+				FString BuildProductsFile = DistributedBuildWorkingDir / BuildProductsFileName;
 				bool bRet = FFileHelper::SaveStringArrayToFile(BuildProducts, *BuildProductsFile);
 				if (!bRet)
 				{
@@ -423,7 +427,7 @@ bool UWorldPartitionHLODsBuilder::BuildHLODActors()
 		}
 
 		// Write build products to a file
-		FString BuildProductsFile = DistributedBuildWorkingDir / TEXT("BuildProducts.txt");
+		FString BuildProductsFile = DistributedBuildWorkingDir / BuildProductsFileName;
 		bool bRet = FFileHelper::SaveStringArrayToFile(BuildProducts, *BuildProductsFile);
 		if (!bRet)
 		{
@@ -504,7 +508,7 @@ bool UWorldPartitionHLODsBuilder::GetHLODActorsToBuild(TArray<FGuid>& HLODActors
 		FConfigFile ConfigFile;
 		ConfigFile.Read(BuildManifest);
 
-		FString SectionName = FString::Printf(TEXT("HLODBuilder%d"), BuilderIdx);
+		FString SectionName = GetHLODBuilderFolderName(BuilderIdx);
 
 		const FConfigSection* ConfigSection = ConfigFile.Find(SectionName);
 		if (ConfigSection)
@@ -681,7 +685,7 @@ bool UWorldPartitionHLODsBuilder::GenerateBuildManifest(TMap<FString, int32>& Fi
 
 	for(int32 Idx = 0; Idx < BuilderCount; Idx++)
 	{
-		FString SectionName = FString::Printf(TEXT("HLODBuilder%d"), Idx);
+		FString SectionName = GetHLODBuilderFolderName(Idx);
 
 		FConfigSection& Section = ConfigFile.Add(SectionName);
 		for(const FGuid& ActorGuid : BuildersWorkload[Idx])
