Fix New map save hang: re-introduce REN_ForceNoResetLoaders in SaveWorld and call ResetLoaders explicitely to avoid UObject::Rename behavior of calling ResetLoaders(nullptr)

diff --git a/Engine/Source/Editor/UnrealEd/Private/FileHelpers.cpp b/Engine/Source/Editor/UnrealEd/Private/FileHelpers.cpp
--- a/Engine/Source/Editor/UnrealEd/Private/FileHelpers.cpp
+++ b/Engine/Source/Editor/UnrealEd/Private/FileHelpers.cpp
@@ -712,8 +712,10 @@ static bool SaveWorld(UWorld* World,
 
 				if (!DuplicatedWorld)
 				{
+					// Explict Reset Loaders of Package here because we want to avoid resetting of all loaders which is the current behavior of UObject::Rename when passing in a UPackage
+					ResetLoaders(Package);
 					// Duplicate failed or not needed. Just do a rename.
-					Package->Rename(*NewPackageName, NULL, REN_NonTransactional | REN_DontCreateRedirectors);
+					Package->Rename(*NewPackageName, NULL, REN_NonTransactional | REN_DontCreateRedirectors | REN_ForceNoResetLoaders);
 					
 					if (bWorldNeedsRename)
 					{
@@ -730,7 +732,7 @@ static bool SaveWorld(UWorld* World,
 							}
 						}
 
-						World->Rename(*NewWorldAssetName, NULL, REN_NonTransactional | REN_DontCreateRedirectors);
+						World->Rename(*NewWorldAssetName, NULL, REN_NonTransactional | REN_DontCreateRedirectors | REN_ForceNoResetLoaders);
 					}
 
 					// We're changing the world path, add a path redirector so that soft object paths get fixed on save
