WorldPartition: Fix New Empty Level inability to save

diff --git a/Engine/Source/Editor/UnrealEd/Private/EditorServer.cpp b/Engine/Source/Editor/UnrealEd/Private/EditorServer.cpp
--- a/Engine/Source/Editor/UnrealEd/Private/EditorServer.cpp
+++ b/Engine/Source/Editor/UnrealEd/Private/EditorServer.cpp
@@ -2331,6 +2331,7 @@ UWorld* UEditorEngine::NewMap()
 	Factory->bInformEngineOfWorld = true;
 	Factory->FeatureLevel = DefaultWorldFeatureLevel;
 	UPackage* Pkg = CreatePackage(nullptr);
+	Pkg->SetPackageFlags(PKG_NewlyCreated);
 	EObjectFlags Flags = RF_Public | RF_Standalone;
 	UWorld* NewWorld = CastChecked<UWorld>(Factory->FactoryCreateNew(UWorld::StaticClass(), Pkg, TEXT("Untitled"), Flags, NULL, GWarn));
 	Context.SetCurrentWorld(NewWorld);
diff --git a/Engine/Source/Editor/UnrealEd/Private/FileHelpers.cpp b/Engine/Source/Editor/UnrealEd/Private/FileHelpers.cpp
--- a/Engine/Source/Editor/UnrealEd/Private/FileHelpers.cpp
+++ b/Engine/Source/Editor/UnrealEd/Private/FileHelpers.cpp
@@ -3718,7 +3718,7 @@ bool FEditorFileUtils::SaveCurrentLevel()
 		{
 			UPackage* LevelPackage = Level->GetPackage();
 			// Save the level
-			if (!bCheckDirty || LevelPackage->IsDirty())
+			if (!bCheckDirty || LevelPackage->IsDirty() || LevelPackage->HasAnyPackageFlags(PKG_NewlyCreated))
 			{
 				bReturnCode &= FEditorFileUtils::SaveLevel(Level);
 			}
