Do not crash when loading a material that calls a missing function (UE-109104).

diff --git a/Engine/Source/Editor/UnrealEd/Classes/MaterialGraph/MaterialGraphNode_Base.h b/Engine/Source/Editor/UnrealEd/Classes/MaterialGraph/MaterialGraphNode_Base.h
--- a/Engine/Source/Editor/UnrealEd/Classes/MaterialGraph/MaterialGraphNode_Base.h
+++ b/Engine/Source/Editor/UnrealEd/Classes/MaterialGraph/MaterialGraphNode_Base.h
@@ -45,8 +45,10 @@ class UMaterialGraphNode_Base : public UEdGraphNode
 	virtual bool IsRootNode() const {return false;}
 	/** Get a single Input Pin via its index */
 	class UEdGraphPin* GetInputPin(int32 InputIndex) const { return InputPins[InputIndex]; }
-	/** Get a single Output Pin via its index */
+	/** Get a single Output Pin via its index (trusting the caller to provide a valid index) */
 	class UEdGraphPin* GetOutputPin(int32 OutputIndex) const { return OutputPins[OutputIndex]; }
+	/** Get a single Output Pin via its index (can return null if an invalid index is provided) */
+	class UEdGraphPin* TryGetOutputPin(int32 OutputIndex) const { return (OutputIndex >= 0 && OutputIndex < OutputPins.Num()) ? OutputPins[OutputIndex] : nullptr; }
 	/** Gets the exec input pin */
 	class UEdGraphPin* GetExecInputPin() const { return ExecInputPin; }
 	/** Replace a given node with this one, changing all pin links */
diff --git a/Engine/Source/Editor/UnrealEd/Private/MaterialGraph.cpp b/Engine/Source/Editor/UnrealEd/Private/MaterialGraph.cpp
--- a/Engine/Source/Editor/UnrealEd/Private/MaterialGraph.cpp
+++ b/Engine/Source/Editor/UnrealEd/Private/MaterialGraph.cpp
@@ -337,7 +337,12 @@ void UMaterialGraph::LinkGraphNodesFromMaterial()
 					// Unclear why this is null sometimes outside of composite reroute, but this is safer than crashing
 					if (UMaterialGraphNode* GraphNode = Cast<UMaterialGraphNode>(ExpressionInputs[PinInfo.Index]->Expression->GraphNode))
 					{
-						Pin->MakeLinkTo(GraphNode->GetOutputPin(GetValidOutputIndex(ExpressionInputs[PinInfo.Index])));
+						// if GraphNode is a material function call for a missing material function, it may not have any output pins
+						UEdGraphPin* OutputPin = GraphNode->TryGetOutputPin(GetValidOutputIndex(ExpressionInputs[PinInfo.Index]));
+						if (LIKELY(OutputPin))
+						{
+							Pin->MakeLinkTo(OutputPin);
+						}
 					}
 					else if (UMaterialExpressionReroute* CompositeReroute = Cast<UMaterialExpressionReroute>(ExpressionInputs[PinInfo.Index]->Expression))
 					{
