Sequencer: Make sure to export the bound skeletal mesh component and not the actor's first skeletal mesh component only. This fixes a bug where child skeletal mesh components weren't being exported.

diff --git a/Engine/Source/Editor/UnrealEd/Private/Fbx/FbxAnimationExport.cpp b/Engine/Source/Editor/UnrealEd/Private/Fbx/FbxAnimationExport.cpp
--- a/Engine/Source/Editor/UnrealEd/Private/Fbx/FbxAnimationExport.cpp
+++ b/Engine/Source/Editor/UnrealEd/Private/Fbx/FbxAnimationExport.cpp
@@ -641,7 +641,7 @@ void FFbxExporter::ExportMatineeGroup(class AMatineeActor* MatineeActor, USkelet
 	ExportAnimTrack(AnimTrackAdapter, Owner, SkeletalMeshComponent, SamplingRate);
 }
 
-void FFbxExporter::ExportAnimTrack(IAnimTrackAdapter& AnimTrackAdapter, AActor* Actor, USkeletalMeshComponent* SkeletalMeshComponent, float SamplingRate)
+void FFbxExporter::ExportAnimTrack(IAnimTrackAdapter& AnimTrackAdapter, AActor* Actor, USkeletalMeshComponent* InSkeletalMeshComponent, float SamplingRate)
 {
 	// show a status update every 1 second worth of samples
 	const float UpdateFrequency = 1.0f;
@@ -650,16 +650,16 @@ void FFbxExporter::ExportAnimTrack(IAnimTrackAdapter& AnimTrackAdapter, AActor*
 	// find root and find the bone array
 	TArray<FbxNode*> BoneNodes;
 
-	if ( FindSkeleton(SkeletalMeshComponent, BoneNodes)==false )
+	if ( FindSkeleton(InSkeletalMeshComponent, BoneNodes)==false )
 	{
 		UE_LOG(LogFbx, Warning, TEXT("Error FBX Animation Export, no root skeleton found."));
 		return;		
 	}
 	//if we have no allocated bone space transforms something wrong so try to recalc them
-	if (SkeletalMeshComponent->GetBoneSpaceTransforms().Num() <= 0 )
+	if (InSkeletalMeshComponent->GetBoneSpaceTransforms().Num() <= 0 )
 	{
-		SkeletalMeshComponent->RecalcRequiredBones(0);
-		if (SkeletalMeshComponent->GetBoneSpaceTransforms().Num() <= 0)
+		InSkeletalMeshComponent->RecalcRequiredBones(0);
+		if (InSkeletalMeshComponent->GetBoneSpaceTransforms().Num() <= 0)
 		{
 			UE_LOG(LogFbx, Warning, TEXT("Error FBX Animation Export, no bone transforms."));
 			return;
@@ -675,6 +675,9 @@ void FFbxExporter::ExportAnimTrack(IAnimTrackAdapter& AnimTrackAdapter, AActor*
 	int32 AnimationLength = AnimTrackAdapter.GetLength();
 	float FrameRate = AnimTrackAdapter.GetFrameRate();
 
+	TArray<USkeletalMeshComponent*> SkeletalMeshComponents;
+	Actor->GetComponents(SkeletalMeshComponents);
+
 	for (int32 FrameCount = 0; FrameCount <= AnimationLength; ++FrameCount)
 	{
 		if (FrameCount == 0)
@@ -694,14 +697,17 @@ void FFbxExporter::ExportAnimTrack(IAnimTrackAdapter& AnimTrackAdapter, AActor*
 
 		// Update space bases so new animation position has an effect.
 		// @todo - hack - this will be removed at some point
-		SkeletalMeshComponent->TickAnimation(0.03f, false);
-
-		SkeletalMeshComponent->RefreshBoneTransforms();
-		SkeletalMeshComponent->RefreshSlaveComponents();
-		SkeletalMeshComponent->UpdateComponentToWorld();
-		SkeletalMeshComponent->FinalizeBoneTransform();
-		SkeletalMeshComponent->MarkRenderTransformDirty();
-		SkeletalMeshComponent->MarkRenderDynamicDataDirty();
+		for (USkeletalMeshComponent* SkeletalMeshComponent : SkeletalMeshComponents)
+		{
+			SkeletalMeshComponent->TickAnimation(0.03f, false);
+
+			SkeletalMeshComponent->RefreshBoneTransforms();
+			SkeletalMeshComponent->RefreshSlaveComponents();
+			SkeletalMeshComponent->UpdateComponentToWorld();
+			SkeletalMeshComponent->FinalizeBoneTransform();
+			SkeletalMeshComponent->MarkRenderTransformDirty();
+			SkeletalMeshComponent->MarkRenderDynamicDataDirty();
+		}
 
 		FbxTime ExportTime; 
 		ExportTime.SetSecondDouble(GetExportOptions()->bExportLocalTime ? LocalFrame / FrameRate : SampleTime);
@@ -714,7 +720,7 @@ void FFbxExporter::ExportAnimTrack(IAnimTrackAdapter& AnimTrackAdapter, AActor*
 			GWarn->StatusUpdate( FMath::RoundToInt( SampleTime ), FMath::RoundToInt(AnimationLength), NSLOCTEXT("FbxExporter", "ExportingToFbxStatus", "Exporting to FBX") );
 		}
 
-		TArray<FTransform> LocalBoneTransforms = SkeletalMeshComponent->GetBoneSpaceTransforms();
+		TArray<FTransform> LocalBoneTransforms = InSkeletalMeshComponent->GetBoneSpaceTransforms();
 
 		if (LocalBoneTransforms.Num() == 0)
 		{
@@ -724,7 +730,7 @@ void FFbxExporter::ExportAnimTrack(IAnimTrackAdapter& AnimTrackAdapter, AActor*
 		// Add the animation data to the bone nodes
 		for(int32 BoneIndex = 0; BoneIndex < BoneNodes.Num(); ++BoneIndex)
 		{
-			FName BoneName = SkeletalMeshComponent->SkeletalMesh->GetRefSkeleton().GetBoneName(BoneIndex);
+			FName BoneName = InSkeletalMeshComponent->SkeletalMesh->GetRefSkeleton().GetBoneName(BoneIndex);
 			FbxNode* CurrentBoneNode = BoneNodes[BoneIndex];
 
 			// Create the AnimCurves
@@ -746,7 +752,7 @@ void FFbxExporter::ExportAnimTrack(IAnimTrackAdapter& AnimTrackAdapter, AActor*
 
 			if (GetExportOptions()->MapSkeletalMotionToRoot && BoneIndex == 0)
 			{
-				BoneTransform = SkeletalMeshComponent->GetSocketTransform(BoneName) * InitialInvParentTransform;
+				BoneTransform = InSkeletalMeshComponent->GetSocketTransform(BoneName) * InitialInvParentTransform;
 			}
 
 			FbxVector4 Translation = Converter.ConvertToFbxPos(BoneTransform.GetLocation());
diff --git a/Engine/Source/Editor/UnrealEd/Private/Fbx/FbxMainExport.cpp b/Engine/Source/Editor/UnrealEd/Private/Fbx/FbxMainExport.cpp
--- a/Engine/Source/Editor/UnrealEd/Private/Fbx/FbxMainExport.cpp
+++ b/Engine/Source/Editor/UnrealEd/Private/Fbx/FbxMainExport.cpp
@@ -1912,7 +1912,11 @@ bool FFbxExporter::ExportLevelSequenceTracks(UMovieScene* MovieScene, IMovieScen
 		}
 	}
 
-	USkeletalMeshComponent* SkeletalMeshComp = Actor ? Cast<USkeletalMeshComponent>(Actor->GetComponentByClass(USkeletalMeshComponent::StaticClass())) : nullptr;
+	USkeletalMeshComponent* SkeletalMeshComp = Cast<USkeletalMeshComponent>(BoundObject);
+	if (!SkeletalMeshComp)
+	{
+		SkeletalMeshComp = Actor ? Cast<USkeletalMeshComponent>(Actor->GetComponentByClass(USkeletalMeshComponent::StaticClass())) : nullptr;	
+	}
 
 	FFrameRate DisplayRate = MovieScene->GetDisplayRate();
 
