Fixes some "not initialized properly" warnings

diff --git a/Engine/Plugins/Experimental/Landmass/Source/Runtime/Public/BrushEffectsList.h b/Engine/Plugins/Experimental/Landmass/Source/Runtime/Public/BrushEffectsList.h
--- a/Engine/Plugins/Experimental/Landmass/Source/Runtime/Public/BrushEffectsList.h
+++ b/Engine/Plugins/Experimental/Landmass/Source/Runtime/Public/BrushEffectsList.h
@@ -122,10 +122,10 @@ struct FBrushEffectSmoothBlending
 	GENERATED_BODY()
 
 	UPROPERTY(EditAnywhere, BlueprintReadWrite, Category = BrushEffects)
-	float InnerSmoothDistance;
+	float InnerSmoothDistance = 0.01f;
 
 	UPROPERTY(EditAnywhere, BlueprintReadWrite, Category = BrushEffects)
-	float OuterSmoothDistance;
+	float OuterSmoothDistance = 0.01f;
 };
 
 USTRUCT(BlueprintType)
diff --git a/Engine/Source/Editor/UnrealEd/Classes/Editor/PropertyEditorTestObject.h b/Engine/Source/Editor/UnrealEd/Classes/Editor/PropertyEditorTestObject.h
--- a/Engine/Source/Editor/UnrealEd/Classes/Editor/PropertyEditorTestObject.h
+++ b/Engine/Source/Editor/UnrealEd/Classes/Editor/PropertyEditorTestObject.h
@@ -164,10 +164,10 @@ struct FPropertyEditorTestEditCondition
 	int32 HasInlineEditCondition = 0;
 
 	UPROPERTY(EditAnywhere, Category=Default)
-	ETestEnumFlags Flags;
+	ETestEnumFlags Flags = ETestEnumFlags::None;
 
 	UPROPERTY(EditAnywhere, Category=Default, meta=(EditCondition="Flags == ETestEnumFlags::One", EditConditionHides))
-	int32 EnabledAndVisibleWhenOne;
+	int32 EnabledAndVisibleWhenOne = 0;
 };
 
 UCLASS(EditInlineNew)
diff --git a/Engine/Source/Editor/UnrealEd/Classes/Settings/EditorLoadingSavingSettings.h b/Engine/Source/Editor/UnrealEd/Classes/Settings/EditorLoadingSavingSettings.h
--- a/Engine/Source/Editor/UnrealEd/Classes/Settings/EditorLoadingSavingSettings.h
+++ b/Engine/Source/Editor/UnrealEd/Classes/Settings/EditorLoadingSavingSettings.h
@@ -32,7 +32,7 @@ struct FAutoReimportWildcard
 	
 	/** When true, files that match this wildcard will be included (if it doesn't fail any other filters), when false, matches will be excluded from the reimporter */
 	UPROPERTY(EditAnywhere, config, Category=AutoReimport)
-	bool bInclude;
+	bool bInclude = false;
 };
 
 
diff --git a/Engine/Source/Editor/UnrealEd/Public/EditorWorldExtension.h b/Engine/Source/Editor/UnrealEd/Public/EditorWorldExtension.h
--- a/Engine/Source/Editor/UnrealEd/Public/EditorWorldExtension.h
+++ b/Engine/Source/Editor/UnrealEd/Public/EditorWorldExtension.h
@@ -35,7 +35,7 @@ struct FEditorWorldExtensionActorData
 	TObjectPtr<AActor> Actor = nullptr;
 
 	UPROPERTY()
-	bool bValidForPIE;
+	bool bValidForPIE = false;
 };
 
 UCLASS()
