SEditorHeaderButton now defaults to having a "+" icon

diff --git a/Engine/Source/Editor/UnrealEd/Public/SEditorHeaderButton.h b/Engine/Source/Editor/UnrealEd/Public/SEditorHeaderButton.h
--- a/Engine/Source/Editor/UnrealEd/Public/SEditorHeaderButton.h
+++ b/Engine/Source/Editor/UnrealEd/Public/SEditorHeaderButton.h
@@ -12,7 +12,7 @@ class UNREALED_API SEditorHeaderButton : public SCompoundWidget
 public:
 
 	SLATE_BEGIN_ARGS(SEditorHeaderButton) :
-		_Icon(nullptr)
+		_Icon(FAppStyle::Get().GetBrush("Icons.Plus"))
 		{}
 
 		/** The text to display in the button. */
