make MoveViewportCamerasToComponent take in a const pointer because it does not modify anything on it

#rnx

diff --git a/Engine/Source/Editor/UnrealEd/Classes/Editor/EditorEngine.h b/Engine/Source/Editor/UnrealEd/Classes/Editor/EditorEngine.h
--- a/Engine/Source/Editor/UnrealEd/Classes/Editor/EditorEngine.h
+++ b/Engine/Source/Editor/UnrealEd/Classes/Editor/EditorEngine.h
@@ -997,7 +997,7 @@ public:
 	* @param	Component				Target component
 	* @param	bActiveViewportOnly		If true, move/reorient only the active viewport.
 	*/
-	void MoveViewportCamerasToComponent(USceneComponent* Component, bool bActiveViewportOnly);
+	void MoveViewportCamerasToComponent(const USceneComponent* Component, bool bActiveViewportOnly);
 
 	/**
 	 * Moves all viewport cameras to focus on the provided bounding box.
diff --git a/Engine/Source/Editor/UnrealEd/Private/EditorServer.cpp b/Engine/Source/Editor/UnrealEd/Private/EditorServer.cpp
--- a/Engine/Source/Editor/UnrealEd/Private/EditorServer.cpp
+++ b/Engine/Source/Editor/UnrealEd/Private/EditorServer.cpp
@@ -4921,7 +4921,7 @@ void UEditorEngine::MoveViewportCamerasToActor(const TArray<AActor*> &Actors, co
 	FEditorDelegates::OnFocusViewportOnActors.Broadcast(Actors);
 }
 
-void UEditorEngine::MoveViewportCamerasToComponent(USceneComponent* Component, bool bActiveViewportOnly)
+void UEditorEngine::MoveViewportCamerasToComponent(const USceneComponent* Component, bool bActiveViewportOnly)
 {
 	if (Component != nullptr)
 	{
