Cooker: Fix deadlock when preloads fail and don't update NumPreloadAllocated.
#rnx

diff --git a/Engine/Source/Editor/UnrealEd/Private/Cooker/CookPackageData.cpp b/Engine/Source/Editor/UnrealEd/Private/Cooker/CookPackageData.cpp
--- a/Engine/Source/Editor/UnrealEd/Private/Cooker/CookPackageData.cpp
+++ b/Engine/Source/Editor/UnrealEd/Private/Cooker/CookPackageData.cpp
@@ -678,6 +678,11 @@ namespace Cook
 			SetIsPreloadAttempted(true);
 			return true;
 		}
+		{
+			// Use of preloaded files is not yet implemented when GAllowCookedDataInEditorBuilds is on, see FLinkerLoad::CreateLoader
+			SetIsPreloadAttempted(true);
+			return true;
+		}
 		if (!PreloadableFile.Get())
 		{
 			TStringBuilder<NAME_SIZE> FileNameString;
