Disable deprecation warnings for PackageGuid in EditorDomainUtils.cpp (to prepare for ue4 merge)

diff --git a/Engine/Source/Editor/UnrealEd/Private/EditorDomain/EditorDomainUtils.cpp b/Engine/Source/Editor/UnrealEd/Private/EditorDomain/EditorDomainUtils.cpp
--- a/Engine/Source/Editor/UnrealEd/Private/EditorDomain/EditorDomainUtils.cpp
+++ b/Engine/Source/Editor/UnrealEd/Private/EditorDomain/EditorDomainUtils.cpp
@@ -29,7 +29,9 @@ FPackageDigest GetPackageDigest(const FAssetPackageData& PackageData, FName Pack
 	int32 CurrentFileVersionUE4 = GPackageFileUE4Version;
 	int32 CurrentFileVersionLicenseeUE4 = GPackageFileLicenseeUE4Version;
 	Writer << EditorDomainVersion;
+	PRAGMA_DISABLE_DEPRECATION_WARNINGS
 	Writer << const_cast<FGuid&>(PackageData.PackageGuid);
+	PRAGMA_ENABLE_DEPRECATION_WARNINGS
 	Writer << CurrentFileVersionUE4;
 	Writer << CurrentFileVersionLicenseeUE4;
 	check(Algo::IsSorted(PackageData.GetCustomVersions()));
