Cooking: optimize precooking of textures in editor idle time by not triggering BeginCacheForCookedPlatformData if IsCachedCookedPlatformDataLoaded is true.

diff --git a/Engine/Source/Editor/UnrealEd/Private/CookOnTheFlyServer.cpp b/Engine/Source/Editor/UnrealEd/Private/CookOnTheFlyServer.cpp
--- a/Engine/Source/Editor/UnrealEd/Private/CookOnTheFlyServer.cpp
+++ b/Engine/Source/Editor/UnrealEd/Private/CookOnTheFlyServer.cpp
@@ -3115,7 +3115,10 @@ void UCookOnTheFlyServer::TickPrecacheObjectsForPlatforms(const float TimeSlice,
 			{
 				continue;
 			}
-			Texture->BeginCacheForCookedPlatformData(TargetPlatform);
+			if (!Texture->IsCachedCookedPlatformDataLoaded(TargetPlatform))
+			{
+				Texture->BeginCacheForCookedPlatformData(TargetPlatform);
+			}
 		}
 		if (Timer.IsTimeUp())
 		{
