LiveCoding Re-instancing

LIMITATIONS:

1) Re-instancing will only update UClass instance data.
2) Adding and removing properties should only be done towards the end of a class or structure and can not be followed by complex data types.
3) Adding and removing properties from a base class should not be done if a derived class contains complex data types.

KNOWN ISSUES:

1) Changes to enumerations and structures will not be reflected in existing blueprints.  However, adding new nodes to the blueprint will show the updated enumeration or structure.
2) If a class contains an enumeration or structure as a member, the class will not be re-instanced if enumeration or structure is changed.

CHANGES:

1) LiveCodingServer
1a) Modified to always execute certain static instances during load.
1b) Modified to exclude the _Statics static structures to avoid patching to old copies.

2) Added support for LiveCoding reinstancing
2a) Refactored deferred registration system for UClass, UEnum, and UScriptStruct to use a common system that works for normal game, hot reload and live coding.
2b) Type specific version check data is possible (i.e. enum doesn't have a size)
2c) Single registration static for UClass
2d) Single registration class for all types that is just a blind forward to API.
2e) Static and dynamic registrations use different API entry points to avoid having overloaded argument lists that just apply to one or the other.
2f) Shims for older API

3) New common "Reload" system to avoid using HotReload code.
3a) Support common delegates regardless of who is reloading/reinstancing.
3b) Re-instancing code moved from HotReload to Kismet2 (where the bulk of the re-instance code already existed).
3c) Modified PyWrapper to use new helper class instead of depending on HotRelaod
3d) Added WITH_RELOAD which is defined if HotReload or LiveCoding is enabled.
3e) Modifed existing code to use new #define and delegates.

Robert did the review on the changes covered by Part 2.  Remaining changes are all straightforward.

diff --git a/Engine/Plugins/Editor/AssetManagerEditor/Source/AssetManagerEditor/Private/AssetManagerEditorModule.cpp b/Engine/Plugins/Editor/AssetManagerEditor/Source/AssetManagerEditor/Private/AssetManagerEditorModule.cpp
--- a/Engine/Plugins/Editor/AssetManagerEditor/Source/AssetManagerEditor/Private/AssetManagerEditorModule.cpp
+++ b/Engine/Plugins/Editor/AssetManagerEditor/Source/AssetManagerEditor/Private/AssetManagerEditorModule.cpp
@@ -61,7 +61,6 @@
 #include "Misc/FileHelper.h"
 #include "Misc/MessageDialog.h"
 #include "UObject/ConstructorHelpers.h"
-#include "Misc/HotReloadInterface.h"
 #include "Misc/ScopedSlowTask.h"
 #include "Editor/StatsViewer/Public/IStatsViewer.h"
 #include "Editor/StatsViewer/Public/StatsViewerModule.h"
@@ -384,7 +383,7 @@ private:
 	TSharedRef<FExtender> OnExtendAssetEditor(const TSharedRef<FUICommandList> CommandList, const TArray<UObject*> ContextSensitiveObjects);
 	void ExtendAssetEditorMenu();
 	void ExtendLevelEditorActorContextMenu();
-	void OnHotReload(bool bWasTriggeredAutomatically);
+	void OnReloadComplete(EReloadCompleteReason Reason);
 	void OnMarkPackageDirty(UPackage* Pkg, bool bWasDirty);
 	void OnEditAssetIdentifiers(TArray<FAssetIdentifier> AssetIdentifiers);
 
@@ -511,8 +510,7 @@ void FAssetManagerEditorModule::StartupModule()
 			.SetMenuType(ETabSpawnerMenuType::Hidden);
 
 		// Register for hot reload and package dirty to invalidate data
-		IHotReloadInterface& HotReloadSupport = FModuleManager::LoadModuleChecked<IHotReloadInterface>("HotReload");
-		HotReloadSupport.OnHotReload().AddRaw(this, &FAssetManagerEditorModule::OnHotReload);
+		FCoreUObjectDelegates::ReloadCompleteDelegate.AddRaw(this, &FAssetManagerEditorModule::OnReloadComplete);
 
 		UPackage::PackageMarkedDirtyEvent.AddRaw(this, &FAssetManagerEditorModule::OnMarkPackageDirty);
 
@@ -585,11 +583,7 @@ void FAssetManagerEditorModule::ShutdownModule()
 			SizeMapTab.Pin()->RequestCloseTab();
 		}
 
-		if (FModuleManager::Get().IsModuleLoaded("HotReload"))
-		{
-			IHotReloadInterface& HotReloadSupport = FModuleManager::GetModuleChecked<IHotReloadInterface>("HotReload");
-			HotReloadSupport.OnHotReload().RemoveAll(this);
-		}
+		FCoreUObjectDelegates::ReloadCompleteDelegate.RemoveAll(this);
 
 		UPackage::PackageMarkedDirtyEvent.RemoveAll(this);
 		FEditorDelegates::OnOpenReferenceViewer.RemoveAll(this);
@@ -985,7 +979,7 @@ void FAssetManagerEditorModule::ExtendLevelEditorActorContextMenu()
 	}));
 }
 
-void FAssetManagerEditorModule::OnHotReload(bool bWasTriggeredAutomatically)
+void FAssetManagerEditorModule::OnReloadComplete(EReloadCompleteReason Reason)
 {
 	UAssetManager* AssetManager = UAssetManager::GetIfValid();
 
diff --git a/Engine/Plugins/Editor/ContentBrowser/ContentBrowserClassDataSource/Source/ContentBrowserClassDataSource/Private/NativeClassHierarchy.cpp b/Engine/Plugins/Editor/ContentBrowser/ContentBrowserClassDataSource/Source/ContentBrowserClassDataSource/Private/NativeClassHierarchy.cpp
--- a/Engine/Plugins/Editor/ContentBrowser/ContentBrowserClassDataSource/Source/ContentBrowserClassDataSource/Private/NativeClassHierarchy.cpp
+++ b/Engine/Plugins/Editor/ContentBrowser/ContentBrowserClassDataSource/Source/ContentBrowserClassDataSource/Private/NativeClassHierarchy.cpp
@@ -10,7 +10,6 @@
 #include "Misc/PackageName.h"
 #include "Kismet2/KismetEditorUtilities.h"
 #include "GameProjectGenerationModule.h"
-#include "Misc/HotReloadInterface.h"
 #include "SourceCodeNavigation.h"
 #include "Interfaces/IPluginManager.h"
 #include "Interfaces/IProjectManager.h"
@@ -53,20 +52,15 @@ FNativeClassHierarchy::FNativeClassHierarchy()
 	// Register to be notified of module changes
 	FModuleManager::Get().OnModulesChanged().AddRaw(this, &FNativeClassHierarchy::OnModulesChanged);
 
-	// Register to be notified of hot reloads
-	IHotReloadInterface& HotReloadSupport = FModuleManager::LoadModuleChecked<IHotReloadInterface>("HotReload");
-	HotReloadSupport.OnHotReload().AddRaw(this, &FNativeClassHierarchy::OnHotReload);
+	// Register to be notified of reloads
+	FCoreUObjectDelegates::ReloadCompleteDelegate.AddRaw(this, &FNativeClassHierarchy::OnReloadComplete);
 }
 
 FNativeClassHierarchy::~FNativeClassHierarchy()
 {
 	FModuleManager::Get().OnModulesChanged().RemoveAll(this);
 
-	if(FModuleManager::Get().IsModuleLoaded("HotReload"))
-	{
-		IHotReloadInterface& HotReloadSupport = FModuleManager::GetModuleChecked<IHotReloadInterface>("HotReload");
-		HotReloadSupport.OnHotReload().RemoveAll(this);
-	}
+	FCoreUObjectDelegates::ReloadCompleteDelegate.RemoveAll(this);
 }
 
 TSharedPtr<const FNativeClassHierarchyNode> FNativeClassHierarchy::FindNode(const FName InClassPath, const ENativeClassHierarchyNodeType InType) const
@@ -543,7 +537,7 @@ void FNativeClassHierarchy::OnModulesChanged(FName InModuleName, EModuleChangeRe
 	}
 }
 
-void FNativeClassHierarchy::OnHotReload(bool bWasTriggeredAutomatically)
+void FNativeClassHierarchy::OnReloadComplete(EReloadCompleteReason Reason)
 {
 	PopulateHierarchy();
 }
diff --git a/Engine/Plugins/Editor/ContentBrowser/ContentBrowserClassDataSource/Source/ContentBrowserClassDataSource/Private/NativeClassHierarchy.h b/Engine/Plugins/Editor/ContentBrowser/ContentBrowserClassDataSource/Source/ContentBrowserClassDataSource/Private/NativeClassHierarchy.h
--- a/Engine/Plugins/Editor/ContentBrowser/ContentBrowserClassDataSource/Source/ContentBrowserClassDataSource/Private/NativeClassHierarchy.h
+++ b/Engine/Plugins/Editor/ContentBrowser/ContentBrowserClassDataSource/Source/ContentBrowserClassDataSource/Private/NativeClassHierarchy.h
@@ -348,11 +348,11 @@ private:
 	void OnModulesChanged(FName InModuleName, EModuleChangeReason InModuleChangeReason);
 
 	/**
-	 * Called when we're notified that a module has been hot-reloaded
+	 * Called when we're notified that a module has been reloaded
 	 *
-	 * @param bWasTriggeredAutomatically - True if the hot-reload was automatically triggered, or false if it was from a user action
+	 * @param Reason - The reason why the module was reloaded
 	 */
-	void OnHotReload(bool bWasTriggeredAutomatically);
+	void OnReloadComplete(EReloadCompleteReason Reason);
 
 	/**
 	 * Given a class, work out which module it belongs to
diff --git a/Engine/Plugins/Enterprise/VariantManager/Source/VariantManager/Private/VariantManagerUtils.cpp b/Engine/Plugins/Enterprise/VariantManager/Source/VariantManager/Private/VariantManagerUtils.cpp
--- a/Engine/Plugins/Enterprise/VariantManager/Source/VariantManager/Private/VariantManagerUtils.cpp
+++ b/Engine/Plugins/Enterprise/VariantManager/Source/VariantManager/Private/VariantManagerUtils.cpp
@@ -26,12 +26,12 @@ FDelegateHandle FVariantManagerUtils::OnHotReloadHandle;
 
 void FVariantManagerUtils::RegisterForHotReload()
 {
-	OnHotReloadHandle = FCoreUObjectDelegates::RegisterClassForHotReloadReinstancingDelegate.AddStatic(&FVariantManagerUtils::InvalidateCache);
+	OnHotReloadHandle = FCoreUObjectDelegates::ReloadReinstancingCompleteDelegate.AddStatic(&FVariantManagerUtils::InvalidateCache);
 }
 
 void FVariantManagerUtils::UnregisterForHotReload()
 {
-	FCoreUObjectDelegates::RegisterClassForHotReloadReinstancingDelegate.Remove(OnHotReloadHandle);
+	FCoreUObjectDelegates::ReloadReinstancingCompleteDelegate.Remove(OnHotReloadHandle);
 	OnHotReloadHandle.Reset();
 }
 
@@ -155,7 +155,7 @@ FStructProperty* FVariantManagerUtils::GetDefaultLightColorProperty()
 	return DefaultLightColorProperty;
 }
 
-void FVariantManagerUtils::InvalidateCache(UClass* OldClass, UClass* NewClass, EHotReloadedClassFlags Flags)
+void FVariantManagerUtils::InvalidateCache()
 {
 	OverrideMaterialsProperty = nullptr;
 	RelativeLocationProperty = nullptr;
diff --git a/Engine/Plugins/Enterprise/VariantManager/Source/VariantManager/Private/VariantManagerUtils.h b/Engine/Plugins/Enterprise/VariantManager/Source/VariantManager/Private/VariantManagerUtils.h
--- a/Engine/Plugins/Enterprise/VariantManager/Source/VariantManager/Private/VariantManagerUtils.h
+++ b/Engine/Plugins/Enterprise/VariantManager/Source/VariantManager/Private/VariantManagerUtils.h
@@ -45,7 +45,7 @@ public:
 
 private:
 	// Invalidates all of our cached FProperty pointers
-	static void InvalidateCache(UClass* OldClass, UClass* NewClass, EHotReloadedClassFlags Flags);
+	static void InvalidateCache();
 
 	static FArrayProperty* OverrideMaterialsProperty;
 	static FStructProperty* RelativeLocationProperty;
diff --git a/Engine/Plugins/Experimental/PythonScriptPlugin/Source/PythonScriptPlugin/Private/PyWrapperTypeRegistry.cpp b/Engine/Plugins/Experimental/PythonScriptPlugin/Source/PythonScriptPlugin/Private/PyWrapperTypeRegistry.cpp
--- a/Engine/Plugins/Experimental/PythonScriptPlugin/Source/PythonScriptPlugin/Private/PyWrapperTypeRegistry.cpp
+++ b/Engine/Plugins/Experimental/PythonScriptPlugin/Source/PythonScriptPlugin/Private/PyWrapperTypeRegistry.cpp
@@ -27,6 +27,9 @@
 #include "UObject/Package.h"
 #include "UObject/EnumProperty.h"
 #include "UObject/StructOnScope.h"
+#if WITH_EDITOR
+#include "Kismet2/ReloadUtilities.h"
+#endif
 
 #if WITH_PYTHON
 
@@ -366,10 +369,10 @@ void FPyWrapperTypeReinstancer::AddPendingStruct(UPythonGeneratedStruct* OldStru
 
 void FPyWrapperTypeReinstancer::ProcessPending()
 {
-	const bool bRunGC = ClassesToReinstance.Num() > 0 || StructsToReinstance.Num() > 0;
-
-	if (ClassesToReinstance.Num() > 0)
+	if (ClassesToReinstance.Num() > 0 || StructsToReinstance.Num() > 0)
 	{
+		TUniquePtr<FReload> Reload(new FReload(EActiveReloadType::Reinstancing, TEXT(""), *GLog));
+
 		for (const auto& ClassToReinstancePair : ClassesToReinstance)
 		{
 			if (ClassToReinstancePair.Key && ClassToReinstancePair.Value)
@@ -377,20 +380,26 @@ void FPyWrapperTypeReinstancer::ProcessPending()
 				if (!ClassToReinstancePair.Value->HasAnyClassFlags(CLASS_NewerVersionExists))
 				{
 					// Assume the classes have changed
-					FCoreUObjectDelegates::RegisterClassForHotReloadReinstancingDelegate.Broadcast(ClassToReinstancePair.Key, ClassToReinstancePair.Value, EHotReloadedClassFlags::Changed);
+					Reload->NotifyChange(ClassToReinstancePair.Value, ClassToReinstancePair.Key);
 				}
 			}
 		}
-		FCoreUObjectDelegates::ReinstanceHotReloadedClassesDelegate.Broadcast();
 
-		ClassesToReinstance.Reset();
-	}
+		// This doesn't do anything ATM
+		for (const auto& StructToReinstancePair : StructsToReinstance)
+		{
+			if (StructToReinstancePair.Key && StructToReinstancePair.Value)
+			{
+				// Assume the structures have changed
+				Reload->NotifyChange(StructToReinstancePair.Value, StructToReinstancePair.Key);
+			}
+		}
 
-	// todo: need support for re-instancing structs
-	StructsToReinstance.Reset();
+		Reload->Reinstance();
+
+		ClassesToReinstance.Reset();
+		StructsToReinstance.Reset();
 
-	if (bRunGC)
-	{
 		CollectGarbage(GARBAGE_COLLECTION_KEEPFLAGS);
 	}
 }
diff --git a/Engine/Plugins/Runtime/GameplayAbilities/Source/GameplayAbilitiesEditor/Private/GameplayAbilitiesEditorModule.cpp b/Engine/Plugins/Runtime/GameplayAbilities/Source/GameplayAbilitiesEditor/Private/GameplayAbilitiesEditorModule.cpp
--- a/Engine/Plugins/Runtime/GameplayAbilities/Source/GameplayAbilitiesEditor/Private/GameplayAbilitiesEditorModule.cpp
+++ b/Engine/Plugins/Runtime/GameplayAbilities/Source/GameplayAbilitiesEditor/Private/GameplayAbilitiesEditorModule.cpp
@@ -180,15 +180,15 @@ void FGameplayAbilitiesEditorModule::StartupModule()
 		
 	ApplyGameplayModEvaluationChannelAliasesToEnumMetadata();
 
-#if WITH_HOT_RELOAD
+#if WITH_RELOAD
 	// This code attempts to relaunch the GameplayCueEditor tab when you hotreload this module
-	if (GIsHotReload && FSlateApplication::IsInitialized())
+	if (IsReloadActive() && FSlateApplication::IsInitialized())
 	{
 		FLevelEditorModule& LevelEditorModule = FModuleManager::GetModuleChecked<FLevelEditorModule>( TEXT("LevelEditor") );
 		TSharedPtr<FTabManager> LevelEditorTabManager = LevelEditorModule.GetLevelEditorTabManager();
 		LevelEditorTabManager->TryInvokeTab(FName("GameplayCueApp"));
 	}
-#endif // WITH_HOT_RELOAD
+#endif // WITH_RELOAD
 
 	IGameplayAbilitiesModule::Get().CallOrRegister_OnAbilitySystemGlobalsReady(FSimpleMulticastDelegate::FDelegate::CreateLambda([] {
 		FGameplayAbilitiesEditorModule::RegisterDebuggingCallbacks();
diff --git a/Engine/Source/Developer/HotReload/Private/HotReload.cpp b/Engine/Source/Developer/HotReload/Private/HotReload.cpp
--- a/Engine/Source/Developer/HotReload/Private/HotReload.cpp
+++ b/Engine/Source/Developer/HotReload/Private/HotReload.cpp
@@ -35,10 +35,11 @@
 #include "DesktopPlatformModule.h"
 #if WITH_ENGINE
 #include "Engine/Engine.h"
-#include "Kismet2/KismetReinstanceUtilities.h"
-#include "HotReloadClassReinstancer.h"
 #include "EngineAnalytics.h"
 #endif
+#if WITH_EDITOR
+#include "Kismet2/ReloadUtilities.h"
+#endif
 #include "Misc/ScopeExit.h"
 #include "Algo/Transform.h"
 
@@ -84,6 +85,30 @@ namespace EThreeStateBool
 	}
 };
 
+class FScopedHotReload
+{
+public:
+	FScopedHotReload(TUniquePtr<FReload>& InUniquePtr, const TArray<UPackage*>& InPackages)
+		: UniquePtr(InUniquePtr)
+	{
+		UniquePtr.Reset(new FReload(EActiveReloadType::HotReload, TEXT("HOTRELOAD"), InPackages, *GLog));
+	}
+
+	FScopedHotReload(TUniquePtr<FReload>& InUniquePtr)
+		: UniquePtr(InUniquePtr)
+	{
+		UniquePtr.Reset(new FReload(EActiveReloadType::HotReload, TEXT("HOTRELOAD"), *GLog));
+	}
+
+	~FScopedHotReload()
+	{
+		UniquePtr.Reset();
+	}
+
+private:
+	TUniquePtr<FReload>& UniquePtr;
+};
+
 /**
  * Module for HotReload support
  */
@@ -111,10 +136,14 @@ public:
 	virtual bool RecompileModule(const FName InModuleName, FOutputDevice &Ar, ERecompileModuleFlags Flags) override;
 	virtual bool IsCurrentlyCompiling() const override { return ModuleCompileProcessHandle.IsValid(); }
 	virtual void RequestStopCompilation() override { bRequestCancelCompilation = true; }
-	virtual void AddHotReloadFunctionRemap(FNativeFuncPtr NewFunctionPointer, FNativeFuncPtr OldFunctionPointer) override;	
+	PRAGMA_DISABLE_DEPRECATION_WARNINGS
+	virtual void AddHotReloadFunctionRemap(FNativeFuncPtr NewFunctionPointer, FNativeFuncPtr OldFunctionPointer) override;
+	PRAGMA_ENABLE_DEPRECATION_WARNINGS
 	virtual ECompilationResult::Type RebindPackages(const TArray<UPackage*>& Packages, EHotReloadFlags Flags, FOutputDevice &Ar) override;
 	virtual ECompilationResult::Type DoHotReloadFromEditor(EHotReloadFlags Flags) override;
-	virtual FHotReloadEvent& OnHotReload() override { return HotReloadEvent; }	
+	PRAGMA_DISABLE_DEPRECATION_WARNINGS
+	virtual FHotReloadEvent& OnHotReload() override { return HotReloadEvent; }
+	PRAGMA_ENABLE_DEPRECATION_WARNINGS
 	virtual FModuleCompilerStartedEvent& OnModuleCompilerStarted() override { return ModuleCompilerStartedEvent; }
 	virtual FModuleCompilerFinishedEvent& OnModuleCompilerFinished() override { return ModuleCompilerFinishedEvent; }
 	virtual FString GetModuleCompileMethod(FName InModuleName) override;
@@ -196,21 +225,9 @@ private:
 	 */
 	ECompilationResult::Type DoHotReloadInternal(const TMap<FName, FString>& ChangedModuleNames, const TArray<UPackage*>& Packages, const TArray<FName>& InDependentModules, FOutputDevice& HotReloadAr);
 
-	/**
-	 * Finds all references to old CDOs and replaces them with the new ones.
-	 * Skipping UBlueprintGeneratedClass::OverridenArchetypeForCDO as it's the
-	 * only one needed.
-	 */
-	void ReplaceReferencesToReconstructedCDOs();
-
 #if WITH_ENGINE
 	void RegisterForReinstancing(UClass* OldClass, UClass* NewClass, EHotReloadedClassFlags Flags);
 	void ReinstanceClasses();
-
-	/**
-	 * Called from CoreUObject to re-instance hot-reloaded classes
-	 */
-	void ReinstanceClass(UClass* OldClass, UClass* NewClass, const TMap<UClass*, UClass*>& OldToNewClassesMap);
 #endif
 
 	/**
@@ -348,11 +365,11 @@ private:
 	/** True if the directory watcher has been successfully initialized */
 	bool bDirectoryWatcherInitialized;
 
-	/** Reconstructed CDOs map during hot-reload. */
-	TMap<UObject*, UObject*> ReconstructedCDOsMap;
-
 	/** Keeps record of hot-reload session starting time. */
 	double HotReloadStartTime;
+
+	/** Current reload object */
+	TUniquePtr<FReload> Reload;
 };
 
 namespace HotReloadDefs
@@ -519,8 +536,10 @@ void FHotReloadModule::StartupModule()
 
 #if WITH_ENGINE
 	// Register re-instancing delegate (Core)
+	PRAGMA_DISABLE_DEPRECATION_WARNINGS
 	FCoreUObjectDelegates::RegisterClassForHotReloadReinstancingDelegate.AddRaw(this, &FHotReloadModule::RegisterForReinstancing);
 	FCoreUObjectDelegates::ReinstanceHotReloadedClassesDelegate.AddRaw(this, &FHotReloadModule::ReinstanceClasses);
+	PRAGMA_ENABLE_DEPRECATION_WARNINGS
 #endif
 
 	// Register directory watcher delegate
@@ -694,7 +713,8 @@ bool FHotReloadModule::RecompileModule(const FName InModuleName, FOutputDevice &
 	// Reload the module if it was loaded before we recompiled
 	if ((bWasModuleLoaded || !!(Flags & ERecompileModuleFlags::ForceCodeProject)) && !!(Flags & ERecompileModuleFlags::ReloadAfterRecompile))
 	{
-		TGuardValue<bool> GuardIsHotReload(GIsHotReload, true);
+		FScopedHotReload Guard(Reload);
+		Reload->SetSendReloadCompleteNotification(false);
 		Ar.Logf( TEXT( "Reloading module %s after successful compile." ), *InModuleName.ToString() );
 		if (!ModuleManager.LoadModuleWithCallback( InModuleName, Ar ))
 		{
@@ -707,6 +727,7 @@ bool FHotReloadModule::RecompileModule(const FName InModuleName, FOutputDevice &
 	if (!!(Flags & ERecompileModuleFlags::ForceCodeProject))
 	{
 		HotReloadEvent.Broadcast( false );
+		FCoreUObjectDelegates::ReloadCompleteDelegate.Broadcast(EReloadCompleteReason::HotReloadManual);
 	}
 
 	return true;
@@ -715,26 +736,10 @@ bool FHotReloadModule::RecompileModule(const FName InModuleName, FOutputDevice &
 #endif // WITH_HOT_RELOAD
 }
 
-/** Type hash for a UObject Function Pointer, maybe not a great choice, but it should be sufficient for the needs here. **/
-inline uint32 GetTypeHash(FNativeFuncPtr A)
-{
-	return *(uint32*)&A;
-}
-
-/** Map from old function pointer to new function pointer for hot reload. */
-static TMap<FNativeFuncPtr, FNativeFuncPtr> HotReloadFunctionRemap;
-
-static TSet<UBlueprint*> HotReloadBPSetToRecompile;
-static TSet<UBlueprint*> HotReloadBPSetToRecompileBytecodeOnly;
-
 /** Adds and entry for the UFunction native pointer remap table */
 void FHotReloadModule::AddHotReloadFunctionRemap(FNativeFuncPtr NewFunctionPointer, FNativeFuncPtr OldFunctionPointer)
 {
-	FNativeFuncPtr OtherNewFunction = HotReloadFunctionRemap.FindRef(OldFunctionPointer);
-	check(!OtherNewFunction || OtherNewFunction == NewFunctionPointer);
-	check(NewFunctionPointer);
-	check(OldFunctionPointer);
-	HotReloadFunctionRemap.Add(OldFunctionPointer, NewFunctionPointer);
+	ReloadNotifyFunctionRemap(NewFunctionPointer, OldFunctionPointer);
 }
 
 ECompilationResult::Type FHotReloadModule::DoHotReloadFromEditor(EHotReloadFlags Flags)
@@ -772,9 +777,9 @@ ECompilationResult::Type FHotReloadModule::DoHotReloadInternal(const TMap<FName,
 	ErrorsFC.ClearWarningsAndErrors();
 
 	// Rebind the hot reload DLL 
-	TGuardValue<bool> GuardIsHotReload(GIsHotReload, true);
+	FScopedHotReload Guard(Reload, Packages);
+	Reload->SetSendReloadCompleteNotification(false);
 	TGuardValue<bool> GuardIsInitialLoad(GIsInitialLoad, true);
-	HotReloadFunctionRemap.Empty(); // redundant
 
 	CollectGarbage(GARBAGE_COLLECTION_KEEPFLAGS); // we create a new CDO in the transient package...this needs to go away before we try again.
 
@@ -837,55 +842,14 @@ ECompilationResult::Type FHotReloadModule::DoHotReloadInternal(const TMap<FName,
 
 	if (bReloadSucceeded)
 	{
-		int32 NumFunctionsRemapped = 0;
-		// Remap all native functions (and gather scriptstructs)
-		TArray<UScriptStruct*> ScriptStructs;
-		for (FRawObjectIterator It; It; ++It)
-		{
-			if (UFunction* Function = Cast<UFunction>(static_cast<UObject*>(It->Object)))
-			{
-				if (FNativeFuncPtr NewFunction = HotReloadFunctionRemap.FindRef(Function->GetNativeFunc()))
-				{
-					++NumFunctionsRemapped;
-					Function->SetNativeFunc(NewFunction);
-				}
-			}
-
-			if (UScriptStruct* ScriptStruct = Cast<UScriptStruct>(static_cast<UObject*>(It->Object)))
-			{
-				if (!ScriptStruct->HasAnyFlags(RF_ClassDefaultObject) && ScriptStruct->GetCppStructOps() && Packages.ContainsByPredicate([=](UPackage* Package) { return ScriptStruct->IsIn(Package); }))
-				{
-					ScriptStructs.Add(ScriptStruct);
-				}
-			}
-		}
-		// now let's set up the script structs...this relies on super behavior, so null them all, then set them all up. Internally this sets them up hierarchically.
-		for (UScriptStruct* Script : ScriptStructs)
-		{
-			Script->ClearCppStructOps();
-		}
-		for (UScriptStruct* Script : ScriptStructs)
-		{
-			Script->PrepareCppStructOps();
-			check(Script->GetCppStructOps());
-		}
-		// Make sure new classes have the token stream assembled
-		UClass::AssembleReferenceTokenStreams();
-
-		HotReloadAr.Logf(ELogVerbosity::Display, TEXT("HotReload successful (%d functions remapped  %d scriptstructs remapped)"), NumFunctionsRemapped, ScriptStructs.Num());
-
-		HotReloadFunctionRemap.Empty();
-
-		ReplaceReferencesToReconstructedCDOs();
-
-		// Force GC to collect reinstanced objects
-		CollectGarbage(GARBAGE_COLLECTION_KEEPFLAGS, true);
+		Reload->Finalize();
 
 		Result = ECompilationResult::Succeeded;
 	}
 
 
-	HotReloadEvent.Broadcast( !bIsHotReloadingFromEditor );
+	HotReloadEvent.Broadcast( !bIsHotReloadingFromEditor);
+	FCoreUObjectDelegates::ReloadCompleteDelegate.Broadcast(bIsHotReloadingFromEditor ? EReloadCompleteReason::HotReloadManual : EReloadCompleteReason::HotReloadAutomatic);
 
 	HotReloadAr.Logf(ELogVerbosity::Display, TEXT("HotReload took %4.1fs."), FPlatformTime::Seconds() - HotReloadStartTime);
 
@@ -900,144 +864,6 @@ ECompilationResult::Type FHotReloadModule::DoHotReloadInternal(const TMap<FName,
 #endif
 }
 
-void FHotReloadModule::ReplaceReferencesToReconstructedCDOs()
-{
-	if (ReconstructedCDOsMap.Num() == 0)
-	{
-		return;
-	}
-
-	// Thread pool manager. We need new thread pool with increased
-	// amount of stack size. Standard GThreadPool was encountering
-	// stack overflow error during serialization.
-	static struct FReplaceReferencesThreadPool
-	{
-		FReplaceReferencesThreadPool()
-		{
-			Pool = FQueuedThreadPool::Allocate();
-			int32 NumThreadsInThreadPool = FPlatformMisc::NumberOfWorkerThreadsToSpawn();
-			verify(Pool->Create(NumThreadsInThreadPool, 256 * 1024));
-		}
-
-		~FReplaceReferencesThreadPool()
-		{
-			Pool->Destroy();
-		}
-
-		FQueuedThreadPool* GetPool() { return Pool; }
-
-	private:
-		FQueuedThreadPool* Pool;
-	} ThreadPoolManager;
-
-	// Async task to enable multithreaded CDOs reference search.
-	class FFindRefTask : public FNonAbandonableTask
-	{
-	public:
-		explicit FFindRefTask(const TMap<UObject*, UObject*>& InReconstructedCDOsMap, int32 ReserveElements)
-			: ReconstructedCDOsMap(InReconstructedCDOsMap)
-		{
-			ObjectsArray.Reserve(ReserveElements);
-		}
-
-		void DoWork()
-		{
-			for (UObject* Object : ObjectsArray)
-			{
-				class FReplaceCDOReferencesArchive : public FArchiveUObject
-				{
-				public:
-					FReplaceCDOReferencesArchive(UObject* InPotentialReferencer, const TMap<UObject*, UObject*>& InReconstructedCDOsMap)
-						: ReconstructedCDOsMap(InReconstructedCDOsMap)
-						, PotentialReferencer(InPotentialReferencer)
-					{
-						ArIsObjectReferenceCollector = true;
-						ArIgnoreOuterRef = true;
-					}
-
-					virtual FString GetArchiveName() const override
-					{
-						return TEXT("FReplaceCDOReferencesArchive");
-					}
-
-					FArchive& operator<<(UObject*& ObjRef)
-					{
-						UObject* Obj = ObjRef;
-
-						if (Obj && Obj != PotentialReferencer)
-						{
-							if (UObject* const* FoundObj = ReconstructedCDOsMap.Find(Obj))
-							{
-								ObjRef = *FoundObj;
-							}
-						}
-
-						return *this;
-					}
-
-					const TMap<UObject*, UObject*>& ReconstructedCDOsMap;
-					UObject* PotentialReferencer;
-				};
-
-				FReplaceCDOReferencesArchive FindRefsArchive(Object, ReconstructedCDOsMap);
-				Object->Serialize(FindRefsArchive);
-			}
-		}
-
-		FORCEINLINE TStatId GetStatId() const
-		{
-			RETURN_QUICK_DECLARE_CYCLE_STAT(FFindRefTask, STATGROUP_ThreadPoolAsyncTasks);
-		}
-
-		TArray<UObject*> ObjectsArray;
-
-	private:
-		const TMap<UObject*, UObject*>& ReconstructedCDOsMap;
-	};
-
-	const int32 NumberOfThreads = FPlatformMisc::NumberOfWorkerThreadsToSpawn();
-	const int32 NumObjects = GUObjectArray.GetObjectArrayNum();
-	const int32 ObjectsPerTask = FMath::CeilToInt((float)NumObjects / NumberOfThreads);
-
-	// Create tasks.
-	TArray<FAsyncTask<FFindRefTask>> Tasks;
-	Tasks.Reserve(NumberOfThreads);
-
-	for (int32 TaskId = 0; TaskId < NumberOfThreads; ++TaskId)
-	{
-		Tasks.Emplace(ReconstructedCDOsMap, ObjectsPerTask);
-	}
-
-	// Distribute objects uniformly between tasks.
-	int32 CurrentTaskId = 0;
-	for (FThreadSafeObjectIterator ObjIter; ObjIter; ++ObjIter)
-	{
-		UObject* CurObject = *ObjIter;
-
-		if (CurObject->IsPendingKill())
-		{
-			continue;
-		}
-
-		Tasks[CurrentTaskId].GetTask().ObjectsArray.Add(CurObject);
-		CurrentTaskId = (CurrentTaskId + 1) % NumberOfThreads;
-	}
-
-	// Run async tasks in worker threads.
-	for (FAsyncTask<FFindRefTask>& Task : Tasks)
-	{
-		Task.StartBackgroundTask(ThreadPoolManager.GetPool());
-	}
-
-	// Wait until tasks are finished
-	for (FAsyncTask<FFindRefTask>& AsyncTask : Tasks)
-	{
-		AsyncTask.EnsureCompletion();
-	}
-
-	ReconstructedCDOsMap.Empty();
-}
-
 ECompilationResult::Type FHotReloadModule::RebindPackages(const TArray<UPackage*>& InPackages, EHotReloadFlags Flags, FOutputDevice &Ar)
 {
 	ECompilationResult::Type Result = ECompilationResult::Unknown;
@@ -1208,54 +1034,32 @@ ECompilationResult::Type FHotReloadModule::RebindPackagesInternal(const TArray<U
 #if WITH_ENGINE
 void FHotReloadModule::RegisterForReinstancing(UClass* OldClass, UClass* NewClass, EHotReloadedClassFlags Flags)
 {
-	// Don't allow reinstancing of UEngine classes
-	if (!OldClass->IsChildOf(UEngine::StaticClass()))
+	
+	// For compatibility, monitor this broadcast.  If we don't have an active reload then we create one assuming
+	// the broadcaster wanted to reinstance (i.e. python wrapper).  
+	IReload* TempReload = GetActiveReloadInterface();
+	if (TempReload == nullptr)
 	{
-		TMap<UClass*, UClass*>& ClassesToReinstance = GetClassesToReinstanceForHotReload();
-		checkf(!ClassesToReinstance.Contains(OldClass) || ClassesToReinstance[OldClass] == NewClass, TEXT("Attempting to hot reload a class which is already being hot reloaded as a different class"));
-		ClassesToReinstance.Add(OldClass, NewClass);
+		Reload.Reset(new FReload(EActiveReloadType::Reinstancing, TEXT(""), *GLog));
+		BeginReload(Reload->GetType(), *Reload);
+		TempReload = Reload.Get();
 	}
-	else if (EnumHasAnyFlags(Flags, EHotReloadedClassFlags::Changed))
+
+	// Only invoke the notification if we own the reload object and it is the temporary reinstancing object.
+	if (TempReload == Reload.Get() && TempReload->GetType() == EActiveReloadType::Reinstancing)
 	{
-		UE_LOG(LogHotReload, Warning, TEXT("Engine class '%s' has changed but will be ignored for hot reload"), *NewClass->GetName());
+		TempReload->NotifyChange(EnumHasAnyFlags(Flags, EHotReloadedClassFlags::Changed) ? NewClass : OldClass, OldClass);
 	}
 }
 
 void FHotReloadModule::ReinstanceClasses()
 {
-#if WITH_HOT_RELOAD
-	if (GIsHotReload)
-	{
-		UClass::AssembleReferenceTokenStreams();
-	}
-#endif // WITH_HOT_RELOAD
-
-	TMap<UClass*, UClass*>& ClassesToReinstance = GetClassesToReinstanceForHotReload();
-
-	TMap<UClass*, UClass*> OldToNewClassesMap;
-	for (const TPair<UClass*, UClass*>& Pair : ClassesToReinstance)
-	{
-		if (Pair.Value != nullptr)
-		{
-			OldToNewClassesMap.Add(Pair.Key, Pair.Value);
-		}
-	}
-
-	for (const TPair<UClass*, UClass*>& Pair : ClassesToReinstance)
-	{
-		ReinstanceClass(Pair.Key, Pair.Value, OldToNewClassesMap);
-	}
-
-	ClassesToReinstance.Empty();
-}
-
-void FHotReloadModule::ReinstanceClass(UClass* OldClass, UClass* NewClass, const TMap<UClass*, UClass*>& OldToNewClassesMap)
-{	
-	TSharedPtr<FHotReloadClassReinstancer> ReinstanceHelper = FHotReloadClassReinstancer::Create(NewClass, OldClass, OldToNewClassesMap, ReconstructedCDOsMap, HotReloadBPSetToRecompile, HotReloadBPSetToRecompileBytecodeOnly);
-	if (ReinstanceHelper->ClassNeedsReinstancing())
+	// Only invoke the notification if we own the reload object and it is the temporary reinstancing object.
+	IReload* TempReload = GetActiveReloadInterface();
+	if (TempReload == Reload.Get() && TempReload->GetType() == EActiveReloadType::Reinstancing)
 	{
-		UE_LOG(LogHotReload, Log, TEXT("Re-instancing %s after hot-reload."), NewClass ? *NewClass->GetName() : *OldClass->GetName());
-		ReinstanceHelper->ReinstanceObjectsAndUpdateDefaults();
+		TempReload->Reinstance();
+		Reload.Reset();
 	}
 }
 #endif
diff --git a/Engine/Source/Developer/HotReload/Private/HotReloadClassReinstancer.cpp b/Engine/Source/Developer/HotReload/Private/HotReloadClassReinstancer.cpp
deleted file mode 100644
diff --git a/Engine/Source/Developer/HotReload/Private/HotReloadClassReinstancer.h b/Engine/Source/Developer/HotReload/Private/HotReloadClassReinstancer.h
deleted file mode 100644
diff --git a/Engine/Source/Developer/Windows/LiveCoding/LiveCoding.Build.cs b/Engine/Source/Developer/Windows/LiveCoding/LiveCoding.Build.cs
--- a/Engine/Source/Developer/Windows/LiveCoding/LiveCoding.Build.cs
+++ b/Engine/Source/Developer/Windows/LiveCoding/LiveCoding.Build.cs
@@ -11,6 +11,16 @@ public class LiveCoding : ModuleRules
 		PrivateDependencyModuleNames.Add("CoreUObject");
 		PrivateDependencyModuleNames.Add("Settings");
 
+		if (Target.bCompileAgainstEngine)
+		{
+			PrivateDependencyModuleNames.Add("Engine");
+		}
+
+		if (Target.bBuildEditor == true)
+        {
+            PrivateDependencyModuleNames.Add("UnrealEd");
+        }        
+
 		if(Target.bUseDebugLiveCodingConsole)
         {
             PrivateDefinitions.Add("USE_DEBUG_LIVE_CODING_CONSOLE=1");
diff --git a/Engine/Source/Developer/Windows/LiveCoding/Private/External/LC_ClientCommandActions.cpp b/Engine/Source/Developer/Windows/LiveCoding/Private/External/LC_ClientCommandActions.cpp
--- a/Engine/Source/Developer/Windows/LiveCoding/Private/External/LC_ClientCommandActions.cpp
+++ b/Engine/Source/Developer/Windows/LiveCoding/Private/External/LC_ClientCommandActions.cpp
@@ -14,10 +14,6 @@
 #include "LC_Logging.h"
 // END EPIC MOD
 
-// BEGIN EPIC MOD - Support for object reinstancing. We need to call a global post-patch handler, rather than just getting individual callbacks for modified modules.
-extern bool GHasLoadedPatch;
-// END EPIC MOD
-
 bool actions::RegisterProcessFinished::Execute(const CommandType* command, const DuplexPipe* pipe, void* context, const void*, size_t)
 {
 	bool* successfullyRegisteredProcess = static_cast<bool*>(context);
@@ -139,7 +135,8 @@ bool actions::LoadPatch::Execute(const CommandType* command, const DuplexPipe* p
 	}
 	// END EPIC MOD
 	// BEGIN EPIC MOD - Support for object reinstancing. We need to call a global post-patch handler, rather than just getting individual callbacks for modified modules.
-	GHasLoadedPatch = true;
+	extern void LiveCodingBeginPatch();
+	LiveCodingBeginPatch();
 	// END EPIC MOD
 
 	pipe->SendAck();
@@ -180,16 +177,13 @@ bool actions::LogOutput::Execute(const CommandType*, const DuplexPipe* pipe, voi
 	return true;
 }
 
-// BEGIN EPIC MOD - Notification that compilation has finished
-extern bool GIsCompileActive;
-// END EPIC MOD
-
 bool actions::CompilationFinished::Execute(const CommandType*, const DuplexPipe* pipe, void*, const void*, size_t)
 {
 	pipe->SendAck();
 
 	// BEGIN EPIC MOD - Notification that compilation has finished
-	GIsCompileActive = false;
+	extern void LiveCodingEndCompile();
+	LiveCodingEndCompile();
 	// END EPIC MOD
 
 	// don't continue execution
diff --git a/Engine/Source/Developer/Windows/LiveCoding/Private/External/LC_StringUtil.cpp b/Engine/Source/Developer/Windows/LiveCoding/Private/External/LC_StringUtil.cpp
--- a/Engine/Source/Developer/Windows/LiveCoding/Private/External/LC_StringUtil.cpp
+++ b/Engine/Source/Developer/Windows/LiveCoding/Private/External/LC_StringUtil.cpp
@@ -75,6 +75,86 @@ namespace detail
 
 		return true;
 	}
+
+	// BEGIN EPIC MOD
+	template <typename T>
+	static const T* StartsWithEx(const T* str, const T* subString)
+	{
+		for (;;)
+		{
+			// Reached the end of the substring, return the remaining part of original string
+			T c2 = *subString++;
+			if (c2 == 0)
+			{
+				return str;
+			}
+
+			// If the characters don't match, we are done.  We don't have to do a check specifically
+			// for 0 since we know that c2 must not be zero at this point.
+			T c1 = *str++;
+			if (c1 != c2)
+			{
+				return nullptr;
+			}
+		}
+	}
+
+	template <typename T>
+	bool MatchesWildcardRecursive(const T* target, const T* wildcard)
+	{
+		for (;;)
+		{
+			T w = *wildcard++;
+
+			// We have reached the end of the wildcard string, this is successful if have no more target string
+			if (w == 0)
+			{
+				return *target == 0;
+			}
+
+			// If we have a wildcard character, then 
+			else if (w == '*')
+			{
+
+				// Skip any multiple wildcards 
+				for (; *wildcard == '*'; ++wildcard)
+				{
+				}
+
+				w = *wildcard++;
+
+				// If wildcard is at the end of the wildcards, then we have a match
+				if (w == 0)
+				{
+					return true;
+				}
+
+				// Look through the target string for the given character.  
+				// If we reach the end of the target string without finding a match
+				// then this is a failed match.  If we find the character, recurse 
+				// to match the remaining part of the string
+				for (;;)
+				{
+					T t = *target++;
+					if (t == 0)
+					{
+						return false;
+					}
+					else if (t == w && MatchesWildcardRecursive(target, wildcard))
+					{
+						return true;
+					}
+				}
+			}
+			
+			// If we don't match the character, then we don't have a match
+			else if (w != *target++)
+			{
+				return false;
+			}
+		}
+	}
+	// END EPIC MOD
 }
 
 
@@ -285,6 +365,30 @@ namespace string
 	}
 
 
+	// BEGIN EPIC MOD
+	const char* StartsWithEx(const char* str, const char* subString)
+	{
+		return detail::StartsWithEx(str, subString);
+	}
+
+	const wchar_t* StartsWithEx(const wchar_t* str, const wchar_t* subString)
+	{
+		return detail::StartsWithEx(str, subString);
+	}
+
+
+	bool MatchWildcard(const char* target, const char* wildcard)
+	{
+		return detail::MatchesWildcardRecursive(target, wildcard);
+	}
+
+	bool MatchWildcard(const wchar_t* target, const wchar_t* wildcard)
+	{
+		return detail::MatchesWildcardRecursive(target, wildcard);
+	}
+	// END EPIC MOD
+
+
 	std::string ToUpper(const char* str)
 	{
 		std::string upperStr(str);
diff --git a/Engine/Source/Developer/Windows/LiveCoding/Private/External/LC_StringUtil.h b/Engine/Source/Developer/Windows/LiveCoding/Private/External/LC_StringUtil.h
--- a/Engine/Source/Developer/Windows/LiveCoding/Private/External/LC_StringUtil.h
+++ b/Engine/Source/Developer/Windows/LiveCoding/Private/External/LC_StringUtil.h
@@ -37,6 +37,14 @@ namespace string
 	bool StartsWith(const char* str, const char* subString);
 	bool StartsWith(const wchar_t* str, const wchar_t* subString);
 
+	// BEGIN EPIC MOD
+	const char* StartsWithEx(const char* str, const char* subString);
+	const wchar_t* StartsWithEx(const wchar_t* str, const wchar_t* subString);
+
+	bool MatchWildcard(const char* target, const char* wildcard);
+	bool MatchWildcard(const wchar_t* target, const wchar_t* wildcard);
+	// END EPIC MOD
+
 	inline char ToLower(char c)
 	{
 		return static_cast<char>(::tolower(c));
diff --git a/Engine/Source/Developer/Windows/LiveCoding/Private/LiveCodingModule.cpp b/Engine/Source/Developer/Windows/LiveCoding/Private/LiveCodingModule.cpp
--- a/Engine/Source/Developer/Windows/LiveCoding/Private/LiveCodingModule.cpp
+++ b/Engine/Source/Developer/Windows/LiveCoding/Private/LiveCodingModule.cpp
@@ -16,6 +16,14 @@
 #include "ISettingsModule.h"
 #include "ISettingsSection.h"
 #include "Windows/WindowsHWrapper.h"
+#if WITH_EDITOR
+	#include "Kismet2/ReloadUtilities.h"
+#else
+	#include "UObject/Reload.h"
+#endif
+#if WITH_ENGINE
+	#include "Engine/Engine.h"
+#endif
 
 IMPLEMENT_MODULE(FLiveCodingModule, LiveCoding)
 
@@ -25,6 +33,7 @@ bool GIsCompileActive = false;
 bool GHasLoadedPatch = false;
 FString GLiveCodingConsolePath;
 FString GLiveCodingConsoleArguments;
+FLiveCodingModule* GLiveCodingModule = nullptr;
 
 #if IS_MONOLITHIC
 extern const TCHAR* GLiveCodingEngineDir;
@@ -43,6 +52,71 @@ LPP_POSTCOMPILE_HOOK(FLiveCodingModule::PostCompileHook);
 #pragma clang diagnostic pop
 #endif
 
+#if !WITH_EDITOR
+class FNullReload : public IReload
+{
+public:
+	FNullReload(FLiveCodingModule& InLiveCodingModule)
+		: LiveCodingModule(InLiveCodingModule)
+	{
+		BeginReload(EActiveReloadType::LiveCoding, *this);
+	}
+
+	~FNullReload()
+	{
+		EndReload();
+	}
+
+	virtual EActiveReloadType GetType() const
+	{
+		return EActiveReloadType::LiveCoding;
+	}
+
+
+	virtual const TCHAR* GetPrefix() const
+	{
+		return TEXT("LIVECODING");
+	}
+
+	virtual void NotifyFunctionRemap(FNativeFuncPtr NewFunctionPointer, FNativeFuncPtr OldFunctionPointer)
+	{
+	}
+
+	virtual void NotifyChange(UClass* New, UClass* Old) override
+	{
+		bGotClassChange |= Old != nullptr && Old != New;
+	}
+
+	virtual void NotifyChange(UEnum* New, UEnum* Old) override
+	{
+	}
+
+	virtual void NotifyChange(UScriptStruct* New, UScriptStruct* Old) override
+	{
+		bGotClassChange |= Old != nullptr && Old != New;
+	}
+
+	virtual void Reinstance()
+	{
+		if (bGotClassChange)
+		{
+			static const TCHAR* Message = TEXT("Object structure changes detected.  LiveCoding re-instancing isn't supported in builds without the editor");
+			UE_LOG(LogLiveCoding, Error, TEXT("%s"), Message);
+#if WITH_ENGINE
+			if (GEngine)
+			{
+				GEngine->AddOnScreenDebugMessage(uint64(uintptr_t(&LiveCodingModule)), 5.f, FColor::Red, Message);
+			}
+#endif
+		}
+	}
+
+private:
+	FLiveCodingModule& LiveCodingModule;
+	bool bGotClassChange = false;
+};
+#endif
+
 FLiveCodingModule::FLiveCodingModule()
 	: bEnabledLastTick(false)
 	, bEnabledForSession(false)
@@ -52,6 +126,12 @@ FLiveCodingModule::FLiveCodingModule()
 	, FullProjectDir(FPaths::ConvertRelativePathToFull(FPaths::ProjectDir()))
 	, FullProjectPluginsDir(FPaths::ConvertRelativePathToFull(FPaths::ProjectPluginsDir()))
 {
+	GLiveCodingModule = this;
+}
+
+FLiveCodingModule::~FLiveCodingModule()
+{
+	GLiveCodingModule = nullptr;
 }
 
 void FLiveCodingModule::StartupModule()
@@ -275,13 +355,34 @@ void FLiveCodingModule::Tick()
 	extern void LppSyncPoint();
 	LppSyncPoint();
 
-	if (GHasLoadedPatch)
+	if (!GIsCompileActive && Reload.IsValid())
 	{
-		OnPatchCompleteDelegate.Broadcast();
-		GHasLoadedPatch = false;
+		if (GHasLoadedPatch)
+		{
+			CheckForClassChanges();
+			OnPatchCompleteDelegate.Broadcast();
+			GHasLoadedPatch = false;
+		}
+		Reload.Reset();
 	}
 }
 
+void FLiveCodingModule::CheckForClassChanges()
+{
+#if WITH_COREUOBJECT && WITH_ENGINE
+
+	// During the module loading process, the list of changed classes will be recorded.  Invoking this method will 
+	// result in the RegisterForReinstancing method being invoked which in turn records the classes in the ClassesToReinstance
+	// member variable being populated.
+	ProcessNewlyLoadedUObjects();
+
+#if WITH_EDITOR
+	Reload->Finalize();
+#endif
+
+#endif
+}
+
 ILiveCodingModule::FOnPatchCompleteDelegate& FLiveCodingModule::GetOnPatchCompleteDelegate()
 {
 	return OnPatchCompleteDelegate;
@@ -499,6 +600,7 @@ void FLiveCodingModule::PreCompileHook()
 {
 	UE_LOG(LogLiveCoding, Display, TEXT("Starting Live Coding compile."));
 	GIsCompileActive = true;
+	BeginReload();
 }
 
 void FLiveCodingModule::PostCompileHook()
@@ -507,4 +609,33 @@ void FLiveCodingModule::PostCompileHook()
 	GIsCompileActive = false;
 }
 
+void FLiveCodingModule::BeginReload()
+{
+	if (GLiveCodingModule != nullptr)
+	{
+		if (!GLiveCodingModule->Reload.IsValid())
+		{
+#if WITH_EDITOR
+			GLiveCodingModule->Reload.Reset(new FReload(EActiveReloadType::LiveCoding, TEXT("LIVECODING"), *GLog));
+#else
+			GLiveCodingModule->Reload.Reset(new FNullReload(*GLiveCodingModule));
+#endif
+		}
+	}
+}
+
+// Invoked from LC_ClientCommandActions
+void LiveCodingBeginPatch()
+{
+	GHasLoadedPatch = true;
+	// If we are beginning a patch from a restart from the console, we need to create the reload object
+	FLiveCodingModule::BeginReload();
+}
+
+// Invoked from LC_ClientCommandActions
+void LiveCodingEndCompile()
+{
+	GIsCompileActive = false;
+}
+
 #undef LOCTEXT_NAMESPACE
diff --git a/Engine/Source/Developer/Windows/LiveCoding/Private/LiveCodingModule.h b/Engine/Source/Developer/Windows/LiveCoding/Private/LiveCodingModule.h
--- a/Engine/Source/Developer/Windows/LiveCoding/Private/LiveCodingModule.h
+++ b/Engine/Source/Developer/Windows/LiveCoding/Private/LiveCodingModule.h
@@ -12,10 +12,17 @@ class IConsoleVariable;
 class ISettingsSection;
 class ULiveCodingSettings;
 
+#if WITH_EDITOR
+class FReload;
+#else
+class IReload;
+#endif
+
 class FLiveCodingModule final : public ILiveCodingModule
 {
 public:
 	FLiveCodingModule();
+	~FLiveCodingModule();
 
 	// IModuleInterface implementation
 	virtual void StartupModule() override;
@@ -34,9 +41,12 @@ public:
 	virtual void Tick() override;
 	virtual FOnPatchCompleteDelegate& GetOnPatchCompleteDelegate() override;
 
+	// Methods invoked by the LiveCodingConsole
 	static void PreCompileHook();
 	static void PostCompileHook();
 
+	static void BeginReload();
+
 private:
 	ULiveCodingSettings* Settings;
 	TSharedPtr<ISettingsSection> SettingsSection;
@@ -58,6 +68,14 @@ private:
 	FDelegateHandle ModulesChangedDelegateHandle;
 	FOnPatchCompleteDelegate OnPatchCompleteDelegate;
 
+#if WITH_EDITOR
+	TUniquePtr<FReload> Reload;
+#else
+	TUniquePtr<IReload> Reload;
+#endif
+
+	void CheckForClassChanges();
+
 	bool StartLiveCoding();
 
 	void OnModulesChanged(FName ModuleName, EModuleChangeReason Reason);
diff --git a/Engine/Source/Developer/Windows/LiveCodingServer/Private/External/LC_LiveModule.cpp b/Engine/Source/Developer/Windows/LiveCodingServer/Private/External/LC_LiveModule.cpp
--- a/Engine/Source/Developer/Windows/LiveCodingServer/Private/External/LC_LiveModule.cpp
+++ b/Engine/Source/Developer/Windows/LiveCodingServer/Private/External/LC_LiveModule.cpp
@@ -2473,6 +2473,12 @@ LiveModule::ErrorType::Enum LiveModule::Update(FileAttributeCache* fileCache, Di
 					*/
 					continue;
 				}
+				// BEGIN EPIC MOD
+				else if (symbols::IsUENoStripSymbol(symbolName))
+				{
+					continue;
+				}
+				// END EPIC MOD
 
 				LC_LOG_DEV("Considering symbol %s for stripping", symbolName.c_str());
 				LC_LOG_INDENT_DEV;
@@ -3931,6 +3937,15 @@ LiveModule::ErrorType::Enum LiveModule::Update(FileAttributeCache* fileCache, Di
 		{
 			const symbols::Symbol* initializerSymbol = initializerDb.dynamicInitializers[i];
 			const ImmutableString& name = initializerSymbol->name;
+
+			// BEGIN EPIC MOD
+			if (symbols::IsUEInitializerSymbol(name))
+			{
+				LC_LOG_DEV("Skipping dynamic initializer symbol %s", name.c_str());
+				continue;
+			}
+			// END EPIC MOD
+
 			const ModuleCache::FindSymbolData& originalData = m_moduleCache->FindSymbolByName(token, name);
 			if (originalData.symbol)
 			{
diff --git a/Engine/Source/Developer/Windows/LiveCodingServer/Private/External/LC_ModuleCache.cpp b/Engine/Source/Developer/Windows/LiveCodingServer/Private/External/LC_ModuleCache.cpp
--- a/Engine/Source/Developer/Windows/LiveCodingServer/Private/External/LC_ModuleCache.cpp
+++ b/Engine/Source/Developer/Windows/LiveCodingServer/Private/External/LC_ModuleCache.cpp
@@ -83,6 +83,33 @@ ModuleCache::FindSymbolData ModuleCache::FindSymbolByName(size_t ignoreToken, co
 }
 
 
+// BEGIN EPIC MOD
+ModuleCache::FindSymbolData ModuleCache::FindSymbolByNameBackwards(size_t ignoreToken, const ImmutableString& symbolName) const
+{
+	CriticalSection::ScopedLock lock(&m_cs);
+
+	const size_t count = m_cache.size();
+	for (size_t i = 0u; i < count; ++i)
+	{
+		const size_t index = count - 1u - i;
+		if (index == ignoreToken)
+		{
+			continue;
+		}
+
+		const Data& data = m_cache[index];
+		const symbols::Symbol* symbol = symbols::FindSymbolByName(data.symbolDb, symbolName);
+		if (symbol)
+		{
+			return FindSymbolData{ &data, symbol };
+		}
+	}
+
+	return FindSymbolData{};
+}
+// END EPIC MOD
+
+
 ModuleCache::FindHookData ModuleCache::FindHooksInSectionBackwards(size_t ignoreToken, const ImmutableString& sectionName) const
 {
 	CriticalSection::ScopedLock lock(&m_cs);
diff --git a/Engine/Source/Developer/Windows/LiveCodingServer/Private/External/LC_ModuleCache.h b/Engine/Source/Developer/Windows/LiveCodingServer/Private/External/LC_ModuleCache.h
--- a/Engine/Source/Developer/Windows/LiveCodingServer/Private/External/LC_ModuleCache.h
+++ b/Engine/Source/Developer/Windows/LiveCodingServer/Private/External/LC_ModuleCache.h
@@ -73,6 +73,11 @@ public:
 	// tries finding a symbol by name, starting from the first module, walking to the latest, excluding the module with the given token
 	FindSymbolData FindSymbolByName(size_t ignoreToken, const ImmutableString& symbolName) const;
 
+	// BEGIN EPIC MOD
+	// tries finding a symbol by name, starting from the last module, walking to the first, excluding the module with the given token
+	FindSymbolData FindSymbolByNameBackwards(size_t ignoreToken, const ImmutableString& symbolName) const;
+	// END EPIC MOD
+
 	// tries finding the first and last hook in a given section, starting from the newest module, walking to the first, excluding the module with the given token
 	FindHookData FindHooksInSectionBackwards(size_t ignoreToken, const ImmutableString& sectionName) const;
 
diff --git a/Engine/Source/Developer/Windows/LiveCodingServer/Private/External/LC_RelocationPatcher.cpp b/Engine/Source/Developer/Windows/LiveCodingServer/Private/External/LC_RelocationPatcher.cpp
--- a/Engine/Source/Developer/Windows/LiveCodingServer/Private/External/LC_RelocationPatcher.cpp
+++ b/Engine/Source/Developer/Windows/LiveCodingServer/Private/External/LC_RelocationPatcher.cpp
@@ -192,9 +192,15 @@ relocations::Record relocations::PatchRelocation
 		}
 	}
 
+	// BEGIN EPIC MOD
+	bool backwards = symbols::IsUEReversePatchSymbol(dstSymbolName) || symbols::IsUEReversePatchSymbol(srcSymbolName);
+	// END EPIC MOD
+
 	// find the relocation's destination symbol in the original .exe, and patch the relocation
 	// to point to this symbol.
-	const ModuleCache::FindSymbolData& originalData = moduleCache->FindSymbolByName(newModuleIndex, dstSymbolName);
+	// BEGIN EPIC MOD
+	const ModuleCache::FindSymbolData originalData = backwards ? moduleCache->FindSymbolByNameBackwards(ModuleCache::SEARCH_ALL_MODULES, dstSymbolName) : moduleCache->FindSymbolByName(newModuleIndex, dstSymbolName);
+	// END EPIC MOD
 	const symbols::Symbol* originalSymbol = originalData.symbol;
 	if (!originalSymbol)
 	{
diff --git a/Engine/Source/Developer/Windows/LiveCodingServer/Private/External/LC_SymbolPatterns.cpp b/Engine/Source/Developer/Windows/LiveCodingServer/Private/External/LC_SymbolPatterns.cpp
--- a/Engine/Source/Developer/Windows/LiveCodingServer/Private/External/LC_SymbolPatterns.cpp
+++ b/Engine/Source/Developer/Windows/LiveCodingServer/Private/External/LC_SymbolPatterns.cpp
@@ -220,4 +220,23 @@ namespace symbolPatterns
 	{
 		"@?A0x"
 	};
+
+	// BEGIN EPIC MOD
+	extern const char* const UE_STATICS_BLOCK_PATTERNS[5] =
+	{
+		"*Z_Construct_UClass_*_Statics*",
+		"*Z_Construct_UEnum_*_Statics*",
+		"*Z_Construct_UFunction_*_Statics*",
+		"*Z_Construct_UDelegateFunction_*_Statics*",
+		"*Z_Construct_UScriptStruct_*_Statics*",
+	};
+
+	extern const char* const UE_REGISTER_PATTERNS[4] =
+	{
+		"?AutoInitialize",
+		"?Z_CompiledInDeferEnum_UEnum_",
+		"?Z_CompiledInDeferStruct_UScriptStruct_",
+		"?Z_CompiledInDeferCppStructOps_UScriptStruct_",
+	};
+	// END EPIC MOD
 }
diff --git a/Engine/Source/Developer/Windows/LiveCodingServer/Private/External/LC_SymbolPatterns.h b/Engine/Source/Developer/Windows/LiveCodingServer/Private/External/LC_SymbolPatterns.h
--- a/Engine/Source/Developer/Windows/LiveCodingServer/Private/External/LC_SymbolPatterns.h
+++ b/Engine/Source/Developer/Windows/LiveCodingServer/Private/External/LC_SymbolPatterns.h
@@ -41,4 +41,9 @@ namespace symbolPatterns
 	extern const char* const TLS_STATICS_PATTERNS[1];
 
 	extern const char* const ANONYMOUS_NAMESPACE_PATTERN;
+
+	// BEGIN EPIC MOD
+	extern const char* const UE_STATICS_BLOCK_PATTERNS[5];
+	extern const char* const UE_REGISTER_PATTERNS[4];
+	// END EPIC MOD
 }
diff --git a/Engine/Source/Developer/Windows/LiveCodingServer/Private/External/LC_Symbols.cpp b/Engine/Source/Developer/Windows/LiveCodingServer/Private/External/LC_Symbols.cpp
--- a/Engine/Source/Developer/Windows/LiveCodingServer/Private/External/LC_Symbols.cpp
+++ b/Engine/Source/Developer/Windows/LiveCodingServer/Private/External/LC_Symbols.cpp
@@ -2174,6 +2174,66 @@ namespace symbols
 	}
 
 
+	// BEGIN EPIC MOD
+	template <typename T, size_t N>
+	static inline const char* StartsWithPatternsEx(const char* name, const T(&patterns)[N])
+	{
+		for (size_t i = 0u; i < N; ++i)
+		{
+			const char* out = string::StartsWithEx(name, patterns[i]);
+			if (out != nullptr)
+			{
+				return out;
+			}
+		}
+
+		return nullptr;
+	}
+
+
+	template <typename T, size_t N>
+	static inline const char* StartsWithPatternsEx(const char* name, const char*& outName, size_t& outNameLength, const T(&patterns)[N])
+	{
+		name = StartsWithPatternsEx(name, patterns);
+		if (name == nullptr)
+		{
+			return nullptr;
+		}
+
+		outName = name;
+		for (;;)
+		{
+			char c = *name++;
+			if (c == 0)
+			{
+				return nullptr;
+			}
+			if (c == '$' || c == '@')
+			{
+				break;
+			}
+		}
+
+		outNameLength = name - outName - 1;
+		return name;
+	}
+
+	template <typename T, size_t N>
+	static inline bool MatchWildcardPatterns(const char* name, const T(&patterns)[N])
+	{
+		for (size_t i = 0u; i < N; ++i)
+		{
+			if (string::MatchWildcard(name, patterns[i]))
+			{
+				return true;
+			}
+		}
+
+		return false;
+	}
+	// END EPIC MOD
+
+
 	bool IsPchSymbol(const ImmutableString& symbolName)
 	{
 		return ContainsPatterns(symbolName.c_str(), symbolPatterns::PCH_SYMBOL_PATTERNS);
@@ -2299,6 +2359,42 @@ namespace symbols
 	{
 		return symbolName.c_str()[0] == '.';
 	}
+
+	// BEGIN EPIC MOD
+	bool IsUERegisterSymbol(const ImmutableString& symbolName, const char*& name, size_t& nameLength)
+	{
+		const char* remaining = StartsWithPatternsEx(symbolName.c_str(), name, nameLength, symbolPatterns::UE_REGISTER_PATTERNS);
+		return remaining != nullptr && string::StartsWith(remaining, "initializer$");
+	}
+
+	bool IsUERegisterSymbol(const ImmutableString& symbolName)
+	{
+		const char* name;
+		size_t nameLength;
+		return IsUERegisterSymbol(symbolName, name, nameLength);
+	}
+
+	bool IsUENoStripSymbol(const ImmutableString& symbolName)
+	{
+		return IsUEStaticsSymbol(symbolName);
+	}
+
+	bool IsUEInitializerSymbol(const ImmutableString& symbolName)
+	{
+		return IsUERegisterSymbol(symbolName) || IsUEStaticsSymbol(symbolName);
+	}
+
+	bool IsUEReversePatchSymbol(const ImmutableString& symbolName)
+	{
+		return IsUEStaticsSymbol(symbolName);
+	}
+
+	bool IsUEStaticsSymbol(const ImmutableString& symbolName)
+	{
+		return MatchWildcardPatterns(symbolName.c_str(), symbolPatterns::UE_STATICS_BLOCK_PATTERNS);
+	}
+	// END EPIC MOD
+
 }
 
 // BEGIN EPIC MOD
diff --git a/Engine/Source/Developer/Windows/LiveCodingServer/Private/External/LC_Symbols.h b/Engine/Source/Developer/Windows/LiveCodingServer/Private/External/LC_Symbols.h
--- a/Engine/Source/Developer/Windows/LiveCodingServer/Private/External/LC_Symbols.h
+++ b/Engine/Source/Developer/Windows/LiveCodingServer/Private/External/LC_Symbols.h
@@ -412,4 +412,13 @@ namespace symbols
 	bool IsTlsStaticsRelatedSymbol(const ImmutableString& symbolName);
 
 	bool IsSectionSymbol(const ImmutableString& symbolName);
+
+	// BEGIN EPIC MOD
+	bool IsUERegisterSymbol(const ImmutableString& symbolName, const char*& name, size_t& nameLength);
+	bool IsUERegisterSymbol(const ImmutableString& symbolName);
+	bool IsUENoStripSymbol(const ImmutableString& symbolName);
+	bool IsUEInitializerSymbol(const ImmutableString& symbolName);
+	bool IsUEReversePatchSymbol(const ImmutableString& symbolName);
+	bool IsUEStaticsSymbol(const ImmutableString& symbolName);
+	// END EPIC MOD
 }
diff --git a/Engine/Source/Editor/AIGraph/Classes/AIGraphTypes.h b/Engine/Source/Editor/AIGraph/Classes/AIGraphTypes.h
--- a/Engine/Source/Editor/AIGraph/Classes/AIGraphTypes.h
+++ b/Engine/Source/Editor/AIGraph/Classes/AIGraphTypes.h
@@ -82,7 +82,7 @@ struct AIGRAPH_API FGraphNodeClassHelper
 	void OnAssetAdded(const struct FAssetData& AssetData);
 	void OnAssetRemoved(const struct FAssetData& AssetData);
 	void InvalidateCache();
-	void OnHotReload(bool bWasTriggeredAutomatically);
+	void OnReloadComplete(EReloadCompleteReason Reason);
 
 	static void AddUnknownClass(const FGraphNodeClassData& ClassData);
 	static bool IsClassKnown(const FGraphNodeClassData& ClassData);
diff --git a/Engine/Source/Editor/AIGraph/Private/AIGraphTypes.cpp b/Engine/Source/Editor/AIGraph/Private/AIGraphTypes.cpp
--- a/Engine/Source/Editor/AIGraph/Private/AIGraphTypes.cpp
+++ b/Engine/Source/Editor/AIGraph/Private/AIGraphTypes.cpp
@@ -16,7 +16,6 @@
 #include "Logging/MessageLog.h"
 #include "ARFilter.h"
 #include "AssetRegistryModule.h"
-#include "Misc/HotReloadInterface.h"
 
 #define LOCTEXT_NAMESPACE "AIGraph"
 
@@ -150,9 +149,8 @@ FGraphNodeClassHelper::FGraphNodeClassHelper(UClass* InRootClass)
 	AssetRegistryModule.Get().OnAssetAdded().AddRaw(this, &FGraphNodeClassHelper::OnAssetAdded);
 	AssetRegistryModule.Get().OnAssetRemoved().AddRaw(this, &FGraphNodeClassHelper::OnAssetRemoved);
 
-	// Register to have Populate called when doing a Hot Reload.
-	IHotReloadInterface& HotReloadSupport = FModuleManager::LoadModuleChecked<IHotReloadInterface>("HotReload");
-	HotReloadSupport.OnHotReload().AddRaw(this, &FGraphNodeClassHelper::OnHotReload);
+	// Register to have Populate called when doing a Reload.
+	FCoreUObjectDelegates::ReloadCompleteDelegate.AddRaw(this, &FGraphNodeClassHelper::OnReloadComplete);
 
 	// Register to have Populate called when a Blueprint is compiled.
 	GEditor->OnBlueprintCompiled().AddRaw(this, &FGraphNodeClassHelper::InvalidateCache);
@@ -171,12 +169,8 @@ FGraphNodeClassHelper::~FGraphNodeClassHelper()
 		AssetRegistryModule.Get().OnAssetAdded().RemoveAll(this);
 		AssetRegistryModule.Get().OnAssetRemoved().RemoveAll(this);
 
-		// Unregister to have Populate called when doing a Hot Reload.
-		if (FModuleManager::Get().IsModuleLoaded(TEXT("HotReload")))
-		{
-			IHotReloadInterface& HotReloadSupport = FModuleManager::GetModuleChecked<IHotReloadInterface>("HotReload");
-			HotReloadSupport.OnHotReload().RemoveAll(this);
-		}
+		// Unregister to have Populate called when doing a Reload.
+		FCoreUObjectDelegates::ReloadCompleteDelegate.RemoveAll(this);
 
 		// Unregister to have Populate called when a Blueprint is compiled.
 		if (UObjectInitialized())
@@ -328,7 +322,7 @@ void FGraphNodeClassHelper::InvalidateCache()
 	UpdateAvailableBlueprintClasses();
 }
 
-void FGraphNodeClassHelper::OnHotReload(bool bWasTriggeredAutomatically)
+void FGraphNodeClassHelper::OnReloadComplete(EReloadCompleteReason Reason)
 {
 	InvalidateCache();
 }
diff --git a/Engine/Source/Editor/BlueprintGraph/Private/BlueprintActionDatabase.cpp b/Engine/Source/Editor/BlueprintGraph/Private/BlueprintActionDatabase.cpp
--- a/Engine/Source/Editor/BlueprintGraph/Private/BlueprintActionDatabase.cpp
+++ b/Engine/Source/Editor/BlueprintGraph/Private/BlueprintActionDatabase.cpp
@@ -66,8 +66,6 @@
 // used below in BlueprintActionDatabaseImpl::GetNodeSpectificActions()
 #include "EdGraph/EdGraphNode_Documentation.h"
 
-#include "Misc/HotReloadInterface.h"
-
 #include "BlueprintTypePromotion.h"
 
 #define LOCTEXT_NAMESPACE "BlueprintActionDatabase"
@@ -518,7 +516,7 @@ namespace BlueprintActionDatabaseImpl
 	/**
 	 * Refreshes database after project was hot-reloaded.
 	 */
-	static void OnProjectHotReloaded(bool bWasTriggeredAutomatically);
+	static void OnReloadComplete(EReloadCompleteReason Reason);
 
 	/** 
 	 * Assets that we cleared from the database (to remove references, and make 
@@ -544,7 +542,7 @@ static void BlueprintActionDatabaseImpl::OnModulesChanged(FName InModuleName, EM
 }
 
 //------------------------------------------------------------------------------
-static void BlueprintActionDatabaseImpl::OnProjectHotReloaded(bool bWasTriggeredAutomatically)
+static void BlueprintActionDatabaseImpl::OnReloadComplete(EReloadCompleteReason Reason)
 {
 	BlueprintActionDatabaseImpl::bRefreshAllRequested = true;
 }
@@ -1169,8 +1167,7 @@ FBlueprintActionDatabase::FBlueprintActionDatabase()
 
 	OnModulesChangedDelegateHandle = FModuleManager::Get().OnModulesChanged().AddStatic(&BlueprintActionDatabaseImpl::OnModulesChanged);
 
-	IHotReloadInterface& HotReloadSupport = FModuleManager::LoadModuleChecked<IHotReloadInterface>("HotReload");
-	OnHotReloadDelegateHandle = HotReloadSupport.OnHotReload().AddStatic(&BlueprintActionDatabaseImpl::OnProjectHotReloaded);
+	OnReloadCompleteDelegateHandle = FCoreUObjectDelegates::ReloadCompleteDelegate.AddStatic(&BlueprintActionDatabaseImpl::OnReloadComplete);
 }
 
 //------------------------------------------------------------------------------
@@ -1198,10 +1195,8 @@ FBlueprintActionDatabase::~FBlueprintActionDatabase()
 	FWorldDelegates::RefreshLevelScriptActions.Remove(RefreshLevelScriptActionsDelegateHandle);
 	FModuleManager::Get().OnModulesChanged().Remove(OnModulesChangedDelegateHandle);
 
-	if (IHotReloadInterface* HotReloadSupport = FModuleManager::GetModulePtr<IHotReloadInterface>("HotReload"))
-	{
-		HotReloadSupport->OnHotReload().Remove(OnHotReloadDelegateHandle);
-	}
+
+	FCoreUObjectDelegates::ReloadCompleteDelegate.Remove(OnReloadCompleteDelegateHandle);
 }
 
 //------------------------------------------------------------------------------
diff --git a/Engine/Source/Editor/BlueprintGraph/Private/EdGraphSchema_K2.cpp b/Engine/Source/Editor/BlueprintGraph/Private/EdGraphSchema_K2.cpp
--- a/Engine/Source/Editor/BlueprintGraph/Private/EdGraphSchema_K2.cpp
+++ b/Engine/Source/Editor/BlueprintGraph/Private/EdGraphSchema_K2.cpp
@@ -87,7 +87,6 @@
 #include "BlueprintNodeBinder.h"
 #include "BlueprintComponentNodeSpawner.h"
 #include "AssetRegistryModule.h"
-#include "Misc/HotReloadInterface.h"
 
 #include "K2Node_CastByteToEnum.h"
 #include "K2Node_ClassDynamicCast.h"
@@ -2234,7 +2233,7 @@ private:
 	static FAutocastFunctionMap* AutocastFunctionMap;
 
 	TMap<FString, TWeakObjectPtr<UFunction>> InnerMap;
-	FDelegateHandle OnHotReloadDelegateHandle;
+	FDelegateHandle OnReloadCompleteDelegateHandle;
 	FDelegateHandle OnModulesChangedDelegateHandle;
 
 	static FString GenerateTypeData(const FEdGraphPinType& PinType)
@@ -2370,7 +2369,7 @@ public:
 		AutocastFunctionMap = nullptr;
 	}
 
-	static void OnProjectHotReloaded(bool bWasTriggeredAutomatically)
+	static void OnReloadComplete(EReloadCompleteReason Reaosn)
 	{
 		if (AutocastFunctionMap)
 		{
@@ -2390,18 +2389,14 @@ public:
 	{
 		Refresh();
 
-		IHotReloadInterface& HotReloadSupport = FModuleManager::LoadModuleChecked<IHotReloadInterface>("HotReload");
-		OnHotReloadDelegateHandle = HotReloadSupport.OnHotReload().AddStatic(&FAutocastFunctionMap::OnProjectHotReloaded);
+		OnReloadCompleteDelegateHandle = FCoreUObjectDelegates::ReloadCompleteDelegate.AddStatic(&FAutocastFunctionMap::OnReloadComplete);
 
 		OnModulesChangedDelegateHandle = FModuleManager::Get().OnModulesChanged().AddStatic(&OnModulesChanged);
 	}
 
 	~FAutocastFunctionMap()
 	{
-		if (IHotReloadInterface* HotReloadSupport = FModuleManager::GetModulePtr<IHotReloadInterface>("HotReload"))
-		{
-			HotReloadSupport->OnHotReload().Remove(OnHotReloadDelegateHandle);
-		}
+		FCoreUObjectDelegates::ReloadCompleteDelegate.Remove(OnReloadCompleteDelegateHandle);
 
 		FModuleManager::Get().OnModulesChanged().Remove(OnModulesChangedDelegateHandle); 
 	}
diff --git a/Engine/Source/Editor/BlueprintGraph/Public/BlueprintActionDatabase.h b/Engine/Source/Editor/BlueprintGraph/Public/BlueprintActionDatabase.h
--- a/Engine/Source/Editor/BlueprintGraph/Public/BlueprintActionDatabase.h
+++ b/Engine/Source/Editor/BlueprintGraph/Public/BlueprintActionDatabase.h
@@ -205,7 +205,7 @@ private:
 	FDelegateHandle OnWorldDestroyedDelegateHandle;
 	FDelegateHandle RefreshLevelScriptActionsDelegateHandle;
 	FDelegateHandle OnModulesChangedDelegateHandle;
-	FDelegateHandle OnHotReloadDelegateHandle;
+	FDelegateHandle OnReloadCompleteDelegateHandle;
 
 	/** Pointer to the shared list of currently existing component types */
 	const TArray<struct FComponentTypeEntry>* ComponentTypes;
diff --git a/Engine/Source/Editor/ClassViewer/Private/SClassViewer.cpp b/Engine/Source/Editor/ClassViewer/Private/SClassViewer.cpp
--- a/Engine/Source/Editor/ClassViewer/Private/SClassViewer.cpp
+++ b/Engine/Source/Editor/ClassViewer/Private/SClassViewer.cpp
@@ -68,7 +68,6 @@
 #include "GameProjectGenerationModule.h"
 
 #include "SourceCodeNavigation.h"
-#include "Misc/HotReloadInterface.h"
 #include "Widgets/Input/SSearchBox.h"
 #include "Misc/TextFilterExpressionEvaluator.h"
 
@@ -158,8 +157,8 @@ private:
 	 */	
 	void AddChildren_NoFilter( TSharedPtr< FClassViewerNode >& InOutRootNode, TMap<FName, TSharedPtr< FClassViewerNode >>& InOutClassPathToNode );
 
-	/** Called when hot reload has finished */
-	void OnHotReload( bool bWasTriggeredAutomatically );
+	/** Called when reload has finished */
+	void OnReloadComplete( EReloadCompleteReason Reason );
 
 	/** 
 	 * Loads the tag data for an unloaded blueprint asset.
@@ -1012,9 +1011,8 @@ FClassHierarchy::FClassHierarchy()
 	AssetRegistryModule.Get().OnAssetAdded().AddRaw( this, &FClassHierarchy::AddAsset);
 	AssetRegistryModule.Get().OnAssetRemoved().AddRaw( this, &FClassHierarchy::RemoveAsset );
 
-	// Register to have Populate called when doing a Hot Reload.
-	IHotReloadInterface& HotReloadSupport = FModuleManager::LoadModuleChecked<IHotReloadInterface>("HotReload");
-	HotReloadSupport.OnHotReload().AddRaw( this, &FClassHierarchy::OnHotReload );
+	// Register to have Populate called when doing a Reload.
+	FCoreUObjectDelegates::ReloadCompleteDelegate.AddRaw( this, &FClassHierarchy::OnReloadComplete );
 
 	// Register to have Populate called when a Blueprint is compiled.
 	OnBlueprintCompiledRequestPopulateClassHierarchyDelegateHandle            = GEditor->OnBlueprintCompiled().AddStatic(ClassViewer::Helpers::RequestPopulateClassHierarchy);
@@ -1033,12 +1031,8 @@ FClassHierarchy::~FClassHierarchy()
 		AssetRegistryModule.Get().OnAssetAdded().RemoveAll( this );
 		AssetRegistryModule.Get().OnAssetRemoved().RemoveAll( this );
 
-		// Unregister to have Populate called when doing a Hot Reload.
-		if(FModuleManager::Get().IsModuleLoaded("HotReload"))
-		{
-			IHotReloadInterface& HotReloadSupport = FModuleManager::GetModuleChecked<IHotReloadInterface>("HotReload");
-			HotReloadSupport.OnHotReload().RemoveAll(this);
-		}
+		// Unregister to have Populate called when doing a Reload.
+		FCoreUObjectDelegates::ReloadCompleteDelegate.RemoveAll(this);
 
 		if (GEditor)
 		{
@@ -1051,7 +1045,7 @@ FClassHierarchy::~FClassHierarchy()
 	FModuleManager::Get().OnModulesChanged().RemoveAll(this);
 }
 
-void FClassHierarchy::OnHotReload(bool bWasTriggeredAutomatically)
+void FClassHierarchy::OnReloadComplete(EReloadCompleteReason Reason)
 {
 	ClassViewer::Helpers::RequestPopulateClassHierarchy();
 }
diff --git a/Engine/Source/Editor/Kismet/Private/FindInBlueprintManager.cpp b/Engine/Source/Editor/Kismet/Private/FindInBlueprintManager.cpp
--- a/Engine/Source/Editor/Kismet/Private/FindInBlueprintManager.cpp
+++ b/Engine/Source/Editor/Kismet/Private/FindInBlueprintManager.cpp
@@ -46,7 +46,6 @@
 #include "ImaginaryBlueprintData.h"
 #include "FiBSearchInstance.h"
 #include "Misc/ConfigCacheIni.h"
-#include "Misc/HotReloadInterface.h"
 #include "BlueprintAssetHandler.h"
 
 #include "JsonObjectConverter.h"
@@ -1996,12 +1995,7 @@ FFindInBlueprintSearchManager::~FFindInBlueprintSearchManager()
 	FCoreUObjectDelegates::GetPreGarbageCollectDelegate().RemoveAll(this);
 	FCoreUObjectDelegates::GetPostGarbageCollect().RemoveAll(this);
 	FCoreUObjectDelegates::OnAssetLoaded.RemoveAll(this);
-
-	if(FModuleManager::Get().IsModuleLoaded("HotReload"))
-	{
-		IHotReloadInterface& HotReloadSupport = FModuleManager::GetModuleChecked<IHotReloadInterface>("HotReload");
-		HotReloadSupport.OnHotReload().RemoveAll(this);
-	}
+	FCoreUObjectDelegates::ReloadCompleteDelegate.RemoveAll(this);
 
 	// Shut down the global find results tab feature.
 	EnableGlobalFindResults(false);
@@ -2054,9 +2048,8 @@ void FFindInBlueprintSearchManager::Initialize()
 	FCoreUObjectDelegates::GetPostGarbageCollect().AddRaw(this, &FFindInBlueprintSearchManager::UnpauseFindInBlueprintSearch);
 	FCoreUObjectDelegates::OnAssetLoaded.AddRaw(this, &FFindInBlueprintSearchManager::OnAssetLoaded);
 	
-	// Register to be notified of hot reloads
-	IHotReloadInterface& HotReloadSupport = FModuleManager::LoadModuleChecked<IHotReloadInterface>("HotReload");
-	HotReloadSupport.OnHotReload().AddRaw(this, &FFindInBlueprintSearchManager::OnHotReload);
+	// Register to be notified of reloads
+	FCoreUObjectDelegates::ReloadCompleteDelegate.AddRaw(this, &FFindInBlueprintSearchManager::OnReloadComplete);
 
 	if(!GIsSavingPackage && AssetRegistryModule)
 	{
@@ -2436,7 +2429,7 @@ void FFindInBlueprintSearchManager::OnBlueprintUnloaded(UBlueprint* InBlueprint)
 	}
 }
 
-void FFindInBlueprintSearchManager::OnHotReload(bool bWasTriggeredAutomatically)
+void FFindInBlueprintSearchManager::OnReloadComplete(EReloadCompleteReason Reason)
 {
 	CachedAssetClasses.Reset();
 }
diff --git a/Engine/Source/Editor/Kismet/Public/FindInBlueprintManager.h b/Engine/Source/Editor/Kismet/Public/FindInBlueprintManager.h
--- a/Engine/Source/Editor/Kismet/Public/FindInBlueprintManager.h
+++ b/Engine/Source/Editor/Kismet/Public/FindInBlueprintManager.h
@@ -681,8 +681,8 @@ private:
 	/** Callback from Kismet when a Blueprint is unloaded */
 	void OnBlueprintUnloaded(class UBlueprint* InBlueprint);
 
-	/** Callback hook from the Hot Reload manager that indicates that a module has been hot-reloaded */
-	void OnHotReload(bool bWasTriggeredAutomatically);
+	/** Callback hook from the Reload manager that indicates that a module has been reloaded */
+	void OnReloadComplete(EReloadCompleteReason Reason);
 
 	/** Returns a copy of the search data that's cached at the given index. Will return invalid (empty) search data if the index is out of range */
 	FSearchData GetSearchDataForIndex(int32 CacheIndex);
diff --git a/Engine/Source/Editor/MainFrame/Private/MainFrameModule.cpp b/Engine/Source/Editor/MainFrame/Private/MainFrameModule.cpp
--- a/Engine/Source/Editor/MainFrame/Private/MainFrameModule.cpp
+++ b/Engine/Source/Editor/MainFrame/Private/MainFrameModule.cpp
@@ -563,7 +563,7 @@ void FMainFrameModule::StartupModule( )
 	IHotReloadModule& HotReloadModule = IHotReloadModule::Get();
 	HotReloadModule.OnModuleCompilerStarted().AddRaw( this, &FMainFrameModule::HandleLevelEditorModuleCompileStarted );
 	HotReloadModule.OnModuleCompilerFinished().AddRaw( this, &FMainFrameModule::HandleLevelEditorModuleCompileFinished );
-	HotReloadModule.OnHotReload().AddRaw( this, &FMainFrameModule::HandleHotReloadFinished );
+	FCoreUObjectDelegates::ReloadCompleteDelegate.AddRaw( this, &FMainFrameModule::HandleReloadFinished );
 
 #if WITH_EDITOR
 	ISourceCodeAccessModule& SourceCodeAccessModule = FModuleManager::LoadModuleChecked<ISourceCodeAccessModule>("SourceCodeAccess");
@@ -597,7 +597,7 @@ void FMainFrameModule::ShutdownModule( )
 	if( IHotReloadModule::IsAvailable() )
 	{
 		IHotReloadModule& HotReloadModule = IHotReloadModule::Get();
-		HotReloadModule.OnHotReload().RemoveAll( this );
+		FCoreUObjectDelegates::ReloadCompleteDelegate.RemoveAll( this );
 		HotReloadModule.OnModuleCompilerStarted().RemoveAll( this );
 		HotReloadModule.OnModuleCompilerFinished().RemoveAll( this );
 	}
@@ -752,11 +752,11 @@ void FMainFrameModule::HandleLevelEditorModuleCompileFinished(const FString& Log
 }
 
 
-void FMainFrameModule::HandleHotReloadFinished( bool bWasTriggeredAutomatically )
+void FMainFrameModule::HandleReloadFinished( EReloadCompleteReason Reason )
 {
 	// Only play the notification for hot reloads that were triggered automatically.  If the user triggered the hot reload, they'll
 	// have a different visual cue for that, such as the "Compiling Complete!" notification
-	if( bWasTriggeredAutomatically )
+	if( Reason == EReloadCompleteReason::HotReloadAutomatic )
 	{
 		FNotificationInfo Info( LOCTEXT("HotReloadFinished", "Hot Reload Complete!") );
 		Info.Image = FEditorStyle::GetBrush(TEXT("LevelEditor.RecompileGameCode"));
diff --git a/Engine/Source/Editor/MainFrame/Private/MainFrameModule.h b/Engine/Source/Editor/MainFrame/Private/MainFrameModule.h
--- a/Engine/Source/Editor/MainFrame/Private/MainFrameModule.h
+++ b/Engine/Source/Editor/MainFrame/Private/MainFrameModule.h
@@ -204,8 +204,8 @@ private:
 	// Handles the level editor module finishing to recompile.
 	void HandleLevelEditorModuleCompileFinished( const FString& LogDump, ECompilationResult::Type CompilationResult, bool bShowLog );
 
-	/** Called when Hot Reload completes */
-	void HandleHotReloadFinished( bool bWasTriggeredAutomatically );
+	/** Called when Reload completes */
+	void HandleReloadFinished( EReloadCompleteReason Reason );
 
 	// Handles the code accessor having finished launching its editor
 	void HandleCodeAccessorLaunched( const bool WasSuccessful );
diff --git a/Engine/Source/Editor/SourceControlWindows/Private/SourceControlWindowsModule.cpp b/Engine/Source/Editor/SourceControlWindows/Private/SourceControlWindowsModule.cpp
--- a/Engine/Source/Editor/SourceControlWindows/Private/SourceControlWindowsModule.cpp
+++ b/Engine/Source/Editor/SourceControlWindows/Private/SourceControlWindowsModule.cpp
@@ -72,13 +72,13 @@ void FSourceControlWindowsModule::StartupModule()
 		.SetGroup(SourceControlGroup)
 		.SetIcon(SourceControlIcon);
 
-#if WITH_HOT_RELOAD
-	// This code attempts to relaunch the GameplayCueEditor tab when you hotreload this module
-	if (GIsHotReload && FSlateApplication::IsInitialized())
+#if WITH_RELOAD
+	// This code attempts to relaunch the GameplayCueEditor tab when you reload this module
+	if (IsReloadActive() && FSlateApplication::IsInitialized())
 	{
 		ShowChangelistsTab();
 	}
-#endif // WITH_HOT_RELOAD
+#endif // WITH_RELOAD
 }
 
 void FSourceControlWindowsModule::ShutdownModule()
diff --git a/Engine/Source/Editor/StructViewer/Private/SStructViewer.cpp b/Engine/Source/Editor/StructViewer/Private/SStructViewer.cpp
--- a/Engine/Source/Editor/StructViewer/Private/SStructViewer.cpp
+++ b/Engine/Source/Editor/StructViewer/Private/SStructViewer.cpp
@@ -9,7 +9,6 @@
 #include "Misc/ScopedSlowTask.h"
 #include "Modules/ModuleManager.h"
 #include "UObject/UObjectIterator.h"
-#include "Misc/HotReloadInterface.h"
 #include "Misc/TextFilterExpressionEvaluator.h"
 
 #include "Editor.h"
@@ -183,8 +182,8 @@ private:
 	/** Callback registered to the Asset Registry to be notified when an asset is removed. */
 	void RemoveAsset(const FAssetData& InRemovedAssetData);
 
-	/** Called when hot reload has finished */
-	void OnHotReload(bool bWasTriggeredAutomatically);
+	/** Called when reload has finished */
+	void OnReloadComplete(EReloadCompleteReason Reason);
 
 	/** Called when modules are loaded or unloaded */
 	void OnModulesChanged(FName ModuleThatChanged, EModuleChangeReason ReasonForChange);
@@ -794,9 +793,8 @@ FStructHierarchy::FStructHierarchy()
 	AssetRegistryModule.Get().OnAssetAdded().AddRaw(this, &FStructHierarchy::AddAsset);
 	AssetRegistryModule.Get().OnAssetRemoved().AddRaw(this, &FStructHierarchy::RemoveAsset);
 
-	// Register to have Populate called when doing a Hot Reload.
-	IHotReloadInterface& HotReloadSupport = FModuleManager::LoadModuleChecked<IHotReloadInterface>("HotReload");
-	HotReloadSupport.OnHotReload().AddRaw(this, &FStructHierarchy::OnHotReload);
+	// Register to have Populate called when doing a Reload.
+	FCoreUObjectDelegates::ReloadCompleteDelegate.AddRaw(this, &FStructHierarchy::OnReloadComplete);
 
 	FModuleManager::Get().OnModulesChanged().AddRaw(this, &FStructHierarchy::OnModulesChanged);
 
@@ -813,12 +811,8 @@ FStructHierarchy::~FStructHierarchy()
 		AssetRegistryModule.Get().OnAssetAdded().RemoveAll(this);
 		AssetRegistryModule.Get().OnAssetRemoved().RemoveAll(this);
 
-		// Unregister to have Populate called when doing a Hot Reload.
-		if (FModuleManager::Get().IsModuleLoaded("HotReload"))
-		{
-			IHotReloadInterface& HotReloadSupport = FModuleManager::GetModuleChecked<IHotReloadInterface>("HotReload");
-			HotReloadSupport.OnHotReload().RemoveAll(this);
-		}
+		// Unregister to have Populate called when doing a Reload.
+		FCoreUObjectDelegates::ReloadCompleteDelegate.RemoveAll(this);
 	}
 
 	FModuleManager::Get().OnModulesChanged().RemoveAll(this);
@@ -1004,7 +998,7 @@ void FStructHierarchy::RemoveAsset(const FAssetData& InRemovedAssetData)
 	}
 }
 
-void FStructHierarchy::OnHotReload(bool bWasTriggeredAutomatically)
+void FStructHierarchy::OnReloadComplete(EReloadCompleteReason Reason)
 {
 	DirtyStructHierarchy();
 }
diff --git a/Engine/Source/Editor/UMGEditor/Private/Palette/SPaletteView.cpp b/Engine/Source/Editor/UMGEditor/Private/Palette/SPaletteView.cpp
--- a/Engine/Source/Editor/UMGEditor/Private/Palette/SPaletteView.cpp
+++ b/Engine/Source/Editor/UMGEditor/Private/Palette/SPaletteView.cpp
@@ -21,8 +21,6 @@
 #include "Templates/WidgetTemplateClass.h"
 #include "Templates/WidgetTemplateBlueprintClass.h"
 
-#include "Developer/HotReload/Public/IHotReload.h"
-
 #include "AssetRegistryModule.h"
 #include "Widgets/Input/SSearchBox.h"
 #include "Widgets/Input/SCheckBox.h"
diff --git a/Engine/Source/Editor/UMGEditor/Private/Palette/SPaletteViewModel.cpp b/Engine/Source/Editor/UMGEditor/Private/Palette/SPaletteViewModel.cpp
--- a/Engine/Source/Editor/UMGEditor/Private/Palette/SPaletteViewModel.cpp
+++ b/Engine/Source/Editor/UMGEditor/Private/Palette/SPaletteViewModel.cpp
@@ -16,8 +16,6 @@
 #include "Templates/WidgetTemplateClass.h"
 #include "Templates/WidgetTemplateBlueprintClass.h"
 
-#include "Developer/HotReload/Public/IHotReload.h"
-
 #include "AssetRegistryModule.h"
 #include "WidgetBlueprintEditorUtils.h"
 
@@ -114,7 +112,7 @@ void FPaletteViewModel::RegisterToEvents()
 	// Register for events that can trigger a palette rebuild
 	GEditor->OnBlueprintReinstanced().AddRaw(this, &FPaletteViewModel::OnBlueprintReinstanced);
 	FEditorDelegates::OnAssetsDeleted.AddSP(this, &FPaletteViewModel::HandleOnAssetsDeleted);
-	IHotReloadModule::Get().OnHotReload().AddSP(this, &FPaletteViewModel::HandleOnHotReload);
+	FCoreUObjectDelegates::ReloadCompleteDelegate.AddSP(this, &FPaletteViewModel::OnReloadComplete);
 
 	// register for any objects replaced
 	FCoreUObjectDelegates::OnObjectsReplaced.AddRaw(this, &FPaletteViewModel::OnObjectsReplaced);
@@ -128,7 +126,7 @@ FPaletteViewModel::~FPaletteViewModel()
 {
 	GEditor->OnBlueprintReinstanced().RemoveAll(this);
 	FEditorDelegates::OnAssetsDeleted.RemoveAll(this);
-	IHotReloadModule::Get().OnHotReload().RemoveAll(this);
+	FCoreUObjectDelegates::ReloadCompleteDelegate.RemoveAll(this);
 	FCoreUObjectDelegates::OnObjectsReplaced.RemoveAll(this);
 
 	UWidgetPaletteFavorites* Favorites = GetDefault<UWidgetDesignerSettings>()->Favorites;
@@ -479,7 +477,7 @@ void FPaletteViewModel::OnFavoritesUpdated()
 	bRebuildRequested = true;
 }
 
-void FPaletteViewModel::HandleOnHotReload(bool bWasTriggeredAutomatically)
+void FPaletteViewModel::OnReloadComplete(EReloadCompleteReason Reason)
 {
 	bRebuildRequested = true;
 }
diff --git a/Engine/Source/Editor/UMGEditor/Private/Palette/SPaletteViewModel.h b/Engine/Source/Editor/UMGEditor/Private/Palette/SPaletteViewModel.h
--- a/Engine/Source/Editor/UMGEditor/Private/Palette/SPaletteViewModel.h
+++ b/Engine/Source/Editor/UMGEditor/Private/Palette/SPaletteViewModel.h
@@ -177,7 +177,7 @@ private:
 	void OnFavoritesUpdated();
 
 	/** Requests a rebuild of the widget list */
-	void HandleOnHotReload(bool bWasTriggeredAutomatically);
+	void OnReloadComplete(EReloadCompleteReason Reason);
 
 	/** Requests a rebuild of the widget list if a widget blueprint was deleted */
 	void HandleOnAssetsDeleted(const TArray<UClass*>& DeletedAssetClasses);
diff --git a/Engine/Source/Editor/UnrealEd/Private/ComponentTypeRegistry.cpp b/Engine/Source/Editor/UnrealEd/Private/ComponentTypeRegistry.cpp
--- a/Engine/Source/Editor/UnrealEd/Private/ComponentTypeRegistry.cpp
+++ b/Engine/Source/Editor/UnrealEd/Private/ComponentTypeRegistry.cpp
@@ -17,7 +17,6 @@
 #include "ClassIconFinder.h"
 #include "EdGraphSchema_K2.h"
 #include "Kismet2/KismetEditorUtilities.h"
-#include "Misc/HotReloadInterface.h"
 #include "SComponentClassCombo.h"
 #include "Settings/ClassViewerSettings.h"
 
@@ -498,20 +497,15 @@ FComponentTypeRegistry::FComponentTypeRegistry()
 	// This will load the assets on next tick. It's not safe to do right now because we could be deep in a stack
 	Data->Invalidate();
 
-	IHotReloadInterface& HotReloadSupport = FModuleManager::LoadModuleChecked<IHotReloadInterface>("HotReload");
-	HotReloadSupport.OnHotReload().AddRaw(this, &FComponentTypeRegistry::OnProjectHotReloaded);
+	FCoreUObjectDelegates::ReloadCompleteDelegate.AddRaw(this, &FComponentTypeRegistry::OnReloadComplete);
 }
 
 FComponentTypeRegistry::~FComponentTypeRegistry()
 {
-	if( FModuleManager::Get().IsModuleLoaded("HotReload") )
-	{
-		IHotReloadInterface& HotReloadSupport = FModuleManager::GetModuleChecked<IHotReloadInterface>("HotReload");
-		HotReloadSupport.OnHotReload().RemoveAll(this);
-	}
+	FCoreUObjectDelegates::ReloadCompleteDelegate.RemoveAll(this);
 }
 
-void FComponentTypeRegistry::OnProjectHotReloaded( bool bWasTriggeredAutomatically )
+void FComponentTypeRegistry::OnReloadComplete(EReloadCompleteReason Reason)
 {
 	Data->ForceRefreshComponentList();
 }
diff --git a/Engine/Source/Editor/UnrealEd/Private/EditorEngine.cpp b/Engine/Source/Editor/UnrealEd/Private/EditorEngine.cpp
--- a/Engine/Source/Editor/UnrealEd/Private/EditorEngine.cpp
+++ b/Engine/Source/Editor/UnrealEd/Private/EditorEngine.cpp
@@ -757,7 +757,7 @@ void UEditorEngine::InitEditor(IEngineLoop* InEngineLoop)
 			}
 		}
 
-		FCoreUObjectDelegates::RegisterHotReloadAddedClassesDelegate.AddUObject(this, &UEditorEngine::CreateVolumeFactoriesForNewClasses);
+		FCoreUObjectDelegates::ReloadAddedClassesDelegate.AddUObject(this, &UEditorEngine::CreateVolumeFactoriesForNewClasses);
 	}
 
 	// Used for sorting ActorFactory classes.
diff --git a/Engine/Source/Editor/UnrealEd/Private/Kismet2/KismetReinstanceUtilities.cpp b/Engine/Source/Editor/UnrealEd/Private/Kismet2/KismetReinstanceUtilities.cpp
--- a/Engine/Source/Editor/UnrealEd/Private/Kismet2/KismetReinstanceUtilities.cpp
+++ b/Engine/Source/Editor/UnrealEd/Private/Kismet2/KismetReinstanceUtilities.cpp
@@ -2166,11 +2166,7 @@ void FBlueprintCompileReinstancer::ReplaceInstancesOfClass_Inner(TMap<UClass*, U
 			UClass* OldClass = OldToNewClass.Key;
 			UClass* NewClass = OldToNewClass.Value;
 			check(OldClass && NewClass);
-#if WITH_HOT_RELOAD
-			check(OldClass != NewClass || GIsHotReload);
-#else
-			check(OldClass != NewClass);
-#endif
+			check(OldClass != NewClass || IsReloadActive());
 			{
 				const bool bIsComponent = NewClass->IsChildOf<UActorComponent>();
 
@@ -2328,11 +2324,7 @@ void FBlueprintCompileReinstancer::ReplaceInstancesOfClass_Inner(TMap<UClass*, U
 			UClass* OldClass = OldToNewClass.Key;
 			UClass* NewClass = OldToNewClass.Value;
 			check(OldClass && NewClass);
-#if WITH_HOT_RELOAD
-			check(OldClass != NewClass || GIsHotReload);
-#else
-			check(OldClass != NewClass);
-#endif
+			check(OldClass != NewClass || IsReloadActive());
 
 			FReplaceReferenceHelper::IncludeCDO(OldClass, NewClass, OldToNewInstanceMap, SourceObjects, InOriginalCDO);
 
diff --git a/Engine/Source/Editor/UnrealEd/Private/Kismet2/ReloadUtilities.cpp b/Engine/Source/Editor/UnrealEd/Private/Kismet2/ReloadUtilities.cpp
new file mode 100644
--- /dev/null
+++ b/Engine/Source/Editor/UnrealEd/Private/Kismet2/ReloadUtilities.cpp
@@ -0,0 +1,1085 @@
+// Copyright Epic Games, Inc. All Rights Reserved.
+
+/*=============================================================================
+	ReloadUtilities.cpp: Helpers for reloading
+=============================================================================*/
+
+#include "Kismet2/ReloadUtilities.h"
+#include "Async/AsyncWork.h"
+#include "Engine/Engine.h"
+#include "Engine/EngineTypes.h"
+#include "Kismet2/KismetReinstanceUtilities.h"
+#include "Misc/QueuedThreadPool.h"
+#include "Serialization/ArchiveReplaceObjectRef.h"
+#include "UObject/Class.h"
+#include "UObject/UObjectIterator.h"
+#include "Misc/StringBuilder.h"
+
+/**
+ * Helper class used for re-instancing native and blueprint classes after hot-reload
+ */
+class FReloadClassReinstancer : public FBlueprintCompileReinstancer
+{
+	/** Holds a property and its offset in the serialized properties data array */
+	struct FCDOProperty
+	{
+		FCDOProperty()
+			: Property(nullptr)
+			, SubobjectName(NAME_None)
+			, SerializedValueOffset(0)
+			, SerializedValueSize(0)
+		{}
+
+		FProperty* Property;
+		FName SubobjectName;
+		int64 SerializedValueOffset;
+		int64 SerializedValueSize;
+	};
+
+	/** Contains all serialized CDO property data and the map of all serialized properties */
+	struct FCDOPropertyData
+	{
+		TArray<uint8> Bytes;
+		TMap<FName, FCDOProperty> Properties;
+	};
+
+	/** Hot-reloaded version of the old class */
+	UClass* NewClass;
+
+	/** Serialized properties of the original CDO (before hot-reload) */
+	FCDOPropertyData OriginalCDOProperties;
+
+	/** Serialized properties of the new CDO (after hot-reload) */
+	FCDOPropertyData ReconstructedCDOProperties;
+
+	/** True if the provided native class needs re-instancing */
+	bool bNeedsReinstancing;
+
+	/** Necessary for delta serialization */
+	UObject* CopyOfPreviousCDO;
+
+	/**
+	 * Sets the re-instancer up for new class re-instancing
+	 *
+	 * @param InNewClass Class that has changed after hot-reload
+	 * @param InOldClass Class before it was hot-reloaded
+	 */
+	void SetupNewClassReinstancing(UClass* InNewClass, UClass* InOldClass);
+
+	/**
+	* Sets the re-instancer up for old class re-instancing. Always re-creates the CDO.
+	*
+	* @param InOldClass Class that has NOT changed after hot-reload
+	*/
+	void RecreateCDOAndSetupOldClassReinstancing(UClass* InOldClass);
+
+	/**
+	* Creates a mem-comparable array of data containing CDO property values.
+	*
+	* @param InObject CDO
+	* @param OutData Data containing all of the CDO property values
+	*/
+	void SerializeCDOProperties(UObject* InObject, FCDOPropertyData& OutData);
+
+	/**
+	 * Re-creates class default object.
+	 *
+	 * @param InClass Class that has NOT changed after hot-reload.
+	 * @param InOuter Outer for the new CDO.
+	 * @param InName Name of the new CDO.
+	 * @param InFlags Flags of the new CDO.
+	 */
+	void ReconstructClassDefaultObject(UClass* InClass, UObject* InOuter, FName InName, EObjectFlags InFlags);
+
+	/** Updates property values on instances of the hot-reloaded class */
+	void UpdateDefaultProperties();
+
+	/** Returns true if the properties of the CDO have changed during hot-reload */
+	FORCEINLINE bool DefaultPropertiesHaveChanged() const
+	{
+		return OriginalCDOProperties.Bytes.Num() != ReconstructedCDOProperties.Bytes.Num() ||
+			FMemory::Memcmp(OriginalCDOProperties.Bytes.GetData(), ReconstructedCDOProperties.Bytes.GetData(), OriginalCDOProperties.Bytes.Num());
+	}
+
+public:
+
+	/** Sets the re-instancer up to re-instance native classes */
+	FReloadClassReinstancer(UClass* InNewClass, UClass* InOldClass, const TMap<UClass*, UClass*>& OldToNewClassesMap, TMap<UObject*, UObject*>& OutReconstructedCDOsMap, TSet<UBlueprint*>& InBPSetToRecompile, TSet<UBlueprint*>& InBPSetToRecompileBytecodeOnly);
+
+	/** Destructor */
+	virtual ~FReloadClassReinstancer();
+
+	/** If true, the class needs re-instancing */
+	FORCEINLINE bool ClassNeedsReinstancing() const
+	{
+		return bNeedsReinstancing;
+	}
+
+	/** Reinstances all objects of the hot-reloaded class and update their properties to match the new CDO */
+	void ReinstanceObjectsAndUpdateDefaults();
+
+	/** Creates the reinstancer as a sharable object */
+	static TSharedPtr<FReloadClassReinstancer> Create(UClass* InNewClass, UClass* InOldClass, const TMap<UClass*, UClass*>& OldToNewClassesMap, TMap<UObject*, UObject*>& OutReconstructedCDOsMap, TSet<UBlueprint*>& InBPSetToRecompile, TSet<UBlueprint*>& InBPSetToRecompileBytecodeOnly)
+	{
+		return MakeShareable(new FReloadClassReinstancer(InNewClass, InOldClass, OldToNewClassesMap, OutReconstructedCDOsMap, InBPSetToRecompile, InBPSetToRecompileBytecodeOnly));
+	}
+
+	// FSerializableObject interface
+	virtual void AddReferencedObjects(FReferenceCollector& Collector) override;
+	// End of FSerializableObject interface
+
+	virtual bool IsClassObjectReplaced() const override { return true; }
+
+	virtual void EnlistDependentBlueprintToRecompile(UBlueprint* BP, bool bBytecodeOnly) override;
+	virtual void BlueprintWasRecompiled(UBlueprint* BP, bool bBytecodeOnly) override;
+
+protected:
+
+	// FBlueprintCompileReinstancer interface
+	virtual bool ShouldPreserveRootComponentOfReinstancedActor() const override { return false; }
+	// End of FBlueprintCompileReinstancer interface
+
+private:
+	/** Reference to reconstructed CDOs map in this hot-reload session. */
+	TMap<UObject*, UObject*>& ReconstructedCDOsMap;
+	TSet<UBlueprint*>& BPSetToRecompile;
+	TSet<UBlueprint*>& BPSetToRecompileBytecodeOnly;
+	const TMap<UClass*, UClass*>& OldToNewClassesMap;
+};
+
+void FReloadClassReinstancer::SetupNewClassReinstancing(UClass* InNewClass, UClass* InOldClass)
+{
+	// Set base class members to valid values
+	ClassToReinstance = InNewClass;
+	DuplicatedClass = InOldClass;
+	OriginalCDO = InOldClass->GetDefaultObject();
+	bHasReinstanced = false;
+	bNeedsReinstancing = true;
+	NewClass = InNewClass;
+
+	// Collect the original CDO property values
+	SerializeCDOProperties(InOldClass->GetDefaultObject(), OriginalCDOProperties);
+	// Collect the property values of the new CDO
+	SerializeCDOProperties(InNewClass->GetDefaultObject(), ReconstructedCDOProperties);
+
+	SaveClassFieldMapping(InOldClass);
+
+	ObjectsThatShouldUseOldStuff.Add(InOldClass); //CDO of REINST_ class can be used as archetype
+
+	TArray<UClass*> ChildrenOfClass;
+	GetDerivedClasses(InOldClass, ChildrenOfClass);
+	for (auto ClassIt = ChildrenOfClass.CreateConstIterator(); ClassIt; ++ClassIt)
+	{
+		UClass* ChildClass = *ClassIt;
+		UBlueprint* ChildBP = Cast<UBlueprint>(ChildClass->ClassGeneratedBy);
+		if (ChildBP && !ChildBP->HasAnyFlags(RF_BeingRegenerated))
+		{
+			// If this is a direct child, change the parent and relink so the property chain is valid for reinstancing
+			if (!ChildBP->HasAnyFlags(RF_NeedLoad))
+			{
+				if (ChildClass->GetSuperClass() == InOldClass)
+				{
+					ReparentChild(ChildBP);
+				}
+
+				Children.AddUnique(ChildBP);
+				if (ChildBP->ParentClass == InOldClass)
+				{
+					ChildBP->ParentClass = NewClass;
+				}
+			}
+			else
+			{
+				// If this is a child that caused the load of their parent, relink to the REINST class so that we can still serialize in the CDO, but do not add to later processing
+				ReparentChild(ChildClass);
+			}
+		}
+	}
+
+	// Finally, remove the old class from Root so that it can get GC'd and mark it as CLASS_NewerVersionExists
+	InOldClass->RemoveFromRoot();
+	InOldClass->ClassFlags |= CLASS_NewerVersionExists;
+}
+
+void FReloadClassReinstancer::SerializeCDOProperties(UObject* InObject, FReloadClassReinstancer::FCDOPropertyData& OutData)
+{
+	// Creates a mem-comparable CDO data
+	class FCDOWriter : public FMemoryWriter
+	{
+		/** Objects already visited by this archive */
+		TSet<UObject*>& VisitedObjects;
+		/** Output property data */
+		FCDOPropertyData& PropertyData;
+		/** Current subobject being serialized */
+		FName SubobjectName;
+
+	public:
+		/** Serializes all script properties of the provided DefaultObject */
+		FCDOWriter(FCDOPropertyData& InOutData, TSet<UObject*>& InVisitedObjects, FName InSubobjectName)
+			: FMemoryWriter(InOutData.Bytes, /* bIsPersistent = */ false, /* bSetOffset = */ true)
+			, VisitedObjects(InVisitedObjects)
+			, PropertyData(InOutData)
+			, SubobjectName(InSubobjectName)
+		{
+			// Disable delta serialization, we want to serialize everything
+			ArNoDelta = true;
+		}
+		virtual void Serialize(void* Data, int64 Num) override
+		{
+			// Collect serialized properties so we can later update their values on instances if they change
+			FProperty* SerializedProperty = GetSerializedProperty();
+			if (SerializedProperty != nullptr)
+			{
+				FCDOProperty& PropertyInfo = PropertyData.Properties.FindOrAdd(SerializedProperty->GetFName());
+				if (PropertyInfo.Property == nullptr)
+				{
+					PropertyInfo.Property = SerializedProperty;
+					PropertyInfo.SubobjectName = SubobjectName;
+					PropertyInfo.SerializedValueOffset = Tell();
+					PropertyInfo.SerializedValueSize = Num;
+				}
+				else
+				{
+					PropertyInfo.SerializedValueSize += Num;
+				}
+			}
+			FMemoryWriter::Serialize(Data, Num);
+		}
+		/** Serializes an object. Only name and class for normal references, deep serialization for DSOs */
+		virtual FArchive& operator<<(class UObject*& InObj) override
+		{
+			FArchive& Ar = *this;
+			if (InObj)
+			{
+				FName ClassName = InObj->GetClass()->GetFName();
+				FName ObjectName = InObj->GetFName();
+				Ar << ClassName;
+				Ar << ObjectName;
+				if (!VisitedObjects.Contains(InObj))
+				{
+					VisitedObjects.Add(InObj);
+					if (Ar.GetSerializedProperty() && Ar.GetSerializedProperty()->ContainsInstancedObjectProperty())
+					{
+						// Serialize all DSO properties too
+						FCDOWriter DefaultSubobjectWriter(PropertyData, VisitedObjects, InObj->GetFName());
+						InObj->SerializeScriptProperties(DefaultSubobjectWriter);
+						Seek(PropertyData.Bytes.Num());
+					}
+				}
+			}
+			else
+			{
+				FName UnusedName = NAME_None;
+				Ar << UnusedName;
+				Ar << UnusedName;
+			}
+
+			return *this;
+		}
+		virtual FArchive& operator<<(FObjectPtr& InObj) override
+		{
+			// Invoke the method above
+			return FArchiveUObject::SerializeObjectPtr(*this, InObj);
+		}
+		/** Serializes an FName as its index and number */
+		virtual FArchive& operator<<(FName& InName) override
+		{
+			FArchive& Ar = *this;
+			FNameEntryId ComparisonIndex = InName.GetComparisonIndex();
+			FNameEntryId DisplayIndex = InName.GetDisplayIndex();
+			int32 Number = InName.GetNumber();
+			Ar << ComparisonIndex;
+			Ar << DisplayIndex;
+			Ar << Number;
+			return Ar;
+		}
+		virtual FArchive& operator<<(FLazyObjectPtr& LazyObjectPtr) override
+		{
+			FArchive& Ar = *this;
+			FUniqueObjectGuid UniqueID = LazyObjectPtr.GetUniqueID();
+			Ar << UniqueID;
+			return *this;
+		}
+		virtual FArchive& operator<<(FSoftObjectPtr& Value) override
+		{
+			FArchive& Ar = *this;
+			FSoftObjectPath UniqueID = Value.GetUniqueID();
+			Ar << UniqueID;
+			return Ar;
+		}
+		virtual FArchive& operator<<(FSoftObjectPath& Value) override
+		{
+			FArchive& Ar = *this;
+
+			FString Path = Value.ToString();
+
+			Ar << Path;
+
+			if (IsLoading())
+			{
+				Value.SetPath(MoveTemp(Path));
+			}
+
+			return Ar;
+		}
+		FArchive& operator<<(FWeakObjectPtr& WeakObjectPtr) override
+		{
+			return FArchiveUObject::SerializeWeakObjectPtr(*this, WeakObjectPtr);
+		}
+		/** Archive name, for debugging */
+		virtual FString GetArchiveName() const override { return TEXT("FCDOWriter"); }
+	};
+	TSet<UObject*> VisitedObjects;
+	VisitedObjects.Add(InObject);
+	FCDOWriter Ar(OutData, VisitedObjects, NAME_None);
+	InObject->SerializeScriptProperties(Ar);
+}
+
+void FReloadClassReinstancer::ReconstructClassDefaultObject(UClass* InClass, UObject* InOuter, FName InName, EObjectFlags InFlags)
+{
+	// Get the parent CDO
+	UClass* ParentClass = InClass->GetSuperClass();
+	UObject* ParentDefaultObject = NULL;
+	if (ParentClass != NULL)
+	{
+		ParentDefaultObject = ParentClass->GetDefaultObject(); // Force the default object to be constructed if it isn't already
+	}
+
+	// Re-create
+	InClass->ClassDefaultObject = StaticAllocateObject(InClass, InOuter, InName, InFlags, EInternalObjectFlags::None, false);
+	check(InClass->ClassDefaultObject);
+	const bool bShouldInitializeProperties = false;
+	const bool bCopyTransientsFromClassDefaults = false;
+	(*InClass->ClassConstructor)(FObjectInitializer(InClass->ClassDefaultObject, ParentDefaultObject, bCopyTransientsFromClassDefaults, bShouldInitializeProperties));
+}
+
+void FReloadClassReinstancer::RecreateCDOAndSetupOldClassReinstancing(UClass* InOldClass)
+{
+	// Set base class members to valid values
+	ClassToReinstance = InOldClass;
+	DuplicatedClass = InOldClass;
+	OriginalCDO = InOldClass->GetDefaultObject();
+	bHasReinstanced = false;
+	bNeedsReinstancing = false;
+	NewClass = InOldClass; // The class doesn't change in this case
+
+	// Collect the original property values
+	SerializeCDOProperties(InOldClass->GetDefaultObject(), OriginalCDOProperties);
+
+	// Remember all the basic info about the object before we rename it
+	EObjectFlags CDOFlags = OriginalCDO->GetFlags();
+	UObject* CDOOuter = OriginalCDO->GetOuter();
+	FName CDOName = OriginalCDO->GetFName();
+
+	// Rename original CDO, so we can store this one as OverridenArchetypeForCDO
+	// and create new one with the same name and outer.
+	OriginalCDO->Rename(
+		*MakeUniqueObjectName(
+			GetTransientPackage(),
+			OriginalCDO->GetClass(),
+			*FString::Printf(TEXT("BPGC_ARCH_FOR_CDO_%s"), *InOldClass->GetName())
+		).ToString(),
+		GetTransientPackage(),
+		REN_DoNotDirty | REN_DontCreateRedirectors | REN_NonTransactional | REN_SkipGeneratedClasses | REN_ForceNoResetLoaders);
+
+	// Re-create the CDO, re-running its constructor
+	ReconstructClassDefaultObject(InOldClass, CDOOuter, CDOName, CDOFlags);
+
+	ReconstructedCDOsMap.Add(OriginalCDO, InOldClass->GetDefaultObject());
+
+	// Collect the property values after re-constructing the CDO
+	SerializeCDOProperties(InOldClass->GetDefaultObject(), ReconstructedCDOProperties);
+
+	// We only want to re-instance the old class if its CDO's values have changed or any of its DSOs' property values have changed
+	if (DefaultPropertiesHaveChanged())
+	{
+		bNeedsReinstancing = true;
+		SaveClassFieldMapping(InOldClass);
+
+		TArray<UClass*> ChildrenOfClass;
+		GetDerivedClasses(InOldClass, ChildrenOfClass);
+		for (auto ClassIt = ChildrenOfClass.CreateConstIterator(); ClassIt; ++ClassIt)
+		{
+			UClass* ChildClass = *ClassIt;
+			UBlueprint* ChildBP = Cast<UBlueprint>(ChildClass->ClassGeneratedBy);
+			if (ChildBP && !ChildBP->HasAnyFlags(RF_BeingRegenerated))
+			{
+				if (!ChildBP->HasAnyFlags(RF_NeedLoad))
+				{
+					Children.AddUnique(ChildBP);
+					UBlueprintGeneratedClass* BPGC = Cast<UBlueprintGeneratedClass>(ChildBP->GeneratedClass);
+					UObject* CurrentCDO = BPGC ? BPGC->GetDefaultObject(false) : nullptr;
+					if (CurrentCDO && (OriginalCDO == CurrentCDO->GetArchetype()))
+					{
+						BPGC->OverridenArchetypeForCDO = OriginalCDO;
+					}
+				}
+			}
+		}
+	}
+}
+
+FReloadClassReinstancer::FReloadClassReinstancer(UClass* InNewClass, UClass* InOldClass, const TMap<UClass*, UClass*>& InOldToNewClassesMap, TMap<UObject*, UObject*>& OutReconstructedCDOsMap, TSet<UBlueprint*>& InBPSetToRecompile, TSet<UBlueprint*>& InBPSetToRecompileBytecodeOnly)
+	: NewClass(nullptr)
+	, bNeedsReinstancing(false)
+	, CopyOfPreviousCDO(nullptr)
+	, ReconstructedCDOsMap(OutReconstructedCDOsMap)
+	, BPSetToRecompile(InBPSetToRecompile)
+	, BPSetToRecompileBytecodeOnly(InBPSetToRecompileBytecodeOnly)
+	, OldToNewClassesMap(InOldToNewClassesMap)
+{
+	ensure(InOldClass);
+	ensure(!HotReloadedOldClass && !HotReloadedNewClass);
+	HotReloadedOldClass = InOldClass;
+	HotReloadedNewClass = InNewClass ? InNewClass : InOldClass;
+
+	for (const TPair<UClass*, UClass*>& OldToNewClass : OldToNewClassesMap)
+	{
+		ObjectsThatShouldUseOldStuff.Add(OldToNewClass.Key);
+	}
+
+	// If InNewClass is NULL, then the old class has not changed after hot-reload.
+	// However, we still need to check for changes to its constructor code (CDO values).
+	if (InNewClass)
+	{
+		SetupNewClassReinstancing(InNewClass, InOldClass);
+
+		TMap<UObject*, UObject*> ClassRedirects;
+		ClassRedirects.Add(InOldClass, InNewClass);
+
+		for (TObjectIterator<UBlueprint> BlueprintIt; BlueprintIt; ++BlueprintIt)
+		{
+			FArchiveReplaceObjectRef<UObject> ReplaceObjectArch(*BlueprintIt, ClassRedirects, false, true, true);
+			if (ReplaceObjectArch.GetCount())
+			{
+				EnlistDependentBlueprintToRecompile(*BlueprintIt, false);
+			}
+		}
+	}
+	else
+	{
+		RecreateCDOAndSetupOldClassReinstancing(InOldClass);
+	}
+}
+
+FReloadClassReinstancer::~FReloadClassReinstancer()
+{
+	// Make sure the base class does not remove the DuplicatedClass from root, we not always want it.
+	// For example when we're just reconstructing CDOs. Other cases are handled by HotReloadClassReinstancer.
+	DuplicatedClass = nullptr;
+
+	ensure(HotReloadedOldClass);
+	HotReloadedOldClass = nullptr;
+	HotReloadedNewClass = nullptr;
+}
+
+/** Helper for finding subobject in an array. Usually there's not that many subobjects on a class to justify a TMap */
+FORCEINLINE static UObject* FindDefaultSubobject(TArray<UObject*>& InDefaultSubobjects, FName SubobjectName)
+{
+	for (UObject* Subobject : InDefaultSubobjects)
+	{
+		if (Subobject->GetFName() == SubobjectName)
+		{
+			return Subobject;
+		}
+	}
+	return nullptr;
+}
+
+void FReloadClassReinstancer::UpdateDefaultProperties()
+{
+	struct FPropertyToUpdate
+	{
+		FProperty* Property;
+		FName SubobjectName;
+		uint8* OldSerializedValuePtr;
+		uint8* NewValuePtr;
+		int64 OldSerializedSize;
+	};
+	/** Memory writer archive that supports UObject values the same way as FCDOWriter. */
+	class FPropertyValueMemoryWriter : public FMemoryWriter
+	{
+	public:
+		FPropertyValueMemoryWriter(TArray<uint8>& OutData)
+			: FMemoryWriter(OutData)
+		{}
+		virtual FArchive& operator<<(class UObject*& InObj) override
+		{
+			FArchive& Ar = *this;
+			if (InObj)
+			{
+				FName ClassName = InObj->GetClass()->GetFName();
+				FName ObjectName = InObj->GetFName();
+				Ar << ClassName;
+				Ar << ObjectName;
+			}
+			else
+			{
+				FName UnusedName = NAME_None;
+				Ar << UnusedName;
+				Ar << UnusedName;
+			}
+			return *this;
+		}
+		virtual FArchive& operator<<(FObjectPtr& InObj) override
+		{
+			// Invoke the method above
+			return FArchiveUObject::SerializeObjectPtr(*this, InObj);
+		}
+		virtual FArchive& operator<<(FName& InName) override
+		{
+			FArchive& Ar = *this;
+			FNameEntryId ComparisonIndex = InName.GetComparisonIndex();
+			FNameEntryId DisplayIndex = InName.GetDisplayIndex();
+			int32 Number = InName.GetNumber();
+			Ar << ComparisonIndex;
+			Ar << DisplayIndex;
+			Ar << Number;
+			return Ar;
+		}
+		virtual FArchive& operator<<(FLazyObjectPtr& LazyObjectPtr) override
+		{
+			FArchive& Ar = *this;
+			FUniqueObjectGuid UniqueID = LazyObjectPtr.GetUniqueID();
+			Ar << UniqueID;
+			return *this;
+		}
+		virtual FArchive& operator<<(FSoftObjectPtr& Value) override
+		{
+			FArchive& Ar = *this;
+			FSoftObjectPath UniqueID = Value.GetUniqueID();
+			Ar << UniqueID;
+			return Ar;
+		}
+		virtual FArchive& operator<<(FSoftObjectPath& Value) override
+		{
+			FArchive& Ar = *this;
+
+			FString Path = Value.ToString();
+
+			Ar << Path;
+
+			if (IsLoading())
+			{
+				Value.SetPath(MoveTemp(Path));
+			}
+
+			return Ar;
+		}
+		FArchive& operator<<(FWeakObjectPtr& WeakObjectPtr) override
+		{
+			return FArchiveUObject::SerializeWeakObjectPtr(*this, WeakObjectPtr);
+		}
+	};
+
+	// Collect default subobjects to update their properties too
+	const int32 DefaultSubobjectArrayCapacity = 16;
+	TArray<UObject*> DefaultSubobjectArray;
+	DefaultSubobjectArray.Empty(DefaultSubobjectArrayCapacity);
+	NewClass->GetDefaultObject()->CollectDefaultSubobjects(DefaultSubobjectArray);
+
+	TArray<FPropertyToUpdate> PropertiesToUpdate;
+	// Collect all properties that have actually changed
+	for (const TPair<FName, FCDOProperty>& Pair : ReconstructedCDOProperties.Properties)
+	{
+		FCDOProperty* OldPropertyInfo = OriginalCDOProperties.Properties.Find(Pair.Key);
+		if (OldPropertyInfo)
+		{
+			const FCDOProperty& NewPropertyInfo = Pair.Value;
+
+			uint8* OldSerializedValuePtr = OriginalCDOProperties.Bytes.GetData() + OldPropertyInfo->SerializedValueOffset;
+			uint8* NewSerializedValuePtr = ReconstructedCDOProperties.Bytes.GetData() + NewPropertyInfo.SerializedValueOffset;
+			if (OldPropertyInfo->SerializedValueSize != NewPropertyInfo.SerializedValueSize ||
+				FMemory::Memcmp(OldSerializedValuePtr, NewSerializedValuePtr, OldPropertyInfo->SerializedValueSize) != 0)
+			{
+				// Property value has changed so add it to the list of properties that need updating on instances
+				FPropertyToUpdate PropertyToUpdate;
+				PropertyToUpdate.Property = NewPropertyInfo.Property;
+				PropertyToUpdate.NewValuePtr = nullptr;
+				PropertyToUpdate.SubobjectName = NewPropertyInfo.SubobjectName;
+
+				if (NewPropertyInfo.Property->GetOwner<UObject>() == NewClass)
+				{
+					PropertyToUpdate.NewValuePtr = PropertyToUpdate.Property->ContainerPtrToValuePtr<uint8>(NewClass->GetDefaultObject());
+				}
+				else if (NewPropertyInfo.SubobjectName != NAME_None)
+				{
+					UObject* DefaultSubobjectPtr = FindDefaultSubobject(DefaultSubobjectArray, NewPropertyInfo.SubobjectName);
+					if (DefaultSubobjectPtr && NewPropertyInfo.Property->GetOwner<UObject>() == DefaultSubobjectPtr->GetClass())
+					{
+						PropertyToUpdate.NewValuePtr = PropertyToUpdate.Property->ContainerPtrToValuePtr<uint8>(DefaultSubobjectPtr);
+					}
+				}
+				if (PropertyToUpdate.NewValuePtr)
+				{
+					PropertyToUpdate.OldSerializedValuePtr = OldSerializedValuePtr;
+					PropertyToUpdate.OldSerializedSize = OldPropertyInfo->SerializedValueSize;
+
+					PropertiesToUpdate.Add(PropertyToUpdate);
+				}
+			}
+		}
+	}
+	if (PropertiesToUpdate.Num())
+	{
+		TArray<uint8> CurrentValueSerializedData;
+
+		// Update properties on all existing instances of the class
+		const UPackage* TransientPackage = GetTransientPackage();
+		for (FThreadSafeObjectIterator It(NewClass); It; ++It)
+		{
+			UObject* ObjectPtr = *It;
+			if (ObjectPtr->IsPendingKill() || ObjectPtr->GetOutermost() == TransientPackage)
+			{
+				continue;
+			}
+
+			DefaultSubobjectArray.Empty(DefaultSubobjectArrayCapacity);
+			ObjectPtr->CollectDefaultSubobjects(DefaultSubobjectArray);
+
+			for (auto& PropertyToUpdate : PropertiesToUpdate)
+			{
+				uint8* InstanceValuePtr = nullptr;
+				if (PropertyToUpdate.SubobjectName == NAME_None)
+				{
+					InstanceValuePtr = PropertyToUpdate.Property->ContainerPtrToValuePtr<uint8>(ObjectPtr);
+				}
+				else
+				{
+					UObject* DefaultSubobjectPtr = FindDefaultSubobject(DefaultSubobjectArray, PropertyToUpdate.SubobjectName);
+					if (DefaultSubobjectPtr && PropertyToUpdate.Property->GetOwner<UObject>() == DefaultSubobjectPtr->GetClass())
+					{
+						InstanceValuePtr = PropertyToUpdate.Property->ContainerPtrToValuePtr<uint8>(DefaultSubobjectPtr);
+					}
+				}
+
+				if (InstanceValuePtr)
+				{
+					// Serialize current value to a byte array as we don't have the previous CDO to compare against, we only have its serialized property data
+					CurrentValueSerializedData.Empty(CurrentValueSerializedData.Num() + CurrentValueSerializedData.GetSlack());
+					FPropertyValueMemoryWriter CurrentValueWriter(CurrentValueSerializedData);
+					PropertyToUpdate.Property->SerializeItem(FStructuredArchiveFromArchive(CurrentValueWriter).GetSlot(), InstanceValuePtr);
+
+					// Update only when the current value on the instance is identical to the original CDO
+					if (CurrentValueSerializedData.Num() == PropertyToUpdate.OldSerializedSize &&
+						FMemory::Memcmp(CurrentValueSerializedData.GetData(), PropertyToUpdate.OldSerializedValuePtr, CurrentValueSerializedData.Num()) == 0)
+					{
+						// Update with the new value
+						PropertyToUpdate.Property->CopyCompleteValue(InstanceValuePtr, PropertyToUpdate.NewValuePtr);
+					}
+				}
+			}
+		}
+	}
+}
+
+void FReloadClassReinstancer::ReinstanceObjectsAndUpdateDefaults()
+{
+	ReinstanceObjects(true);
+	UpdateDefaultProperties();
+}
+
+void FReloadClassReinstancer::AddReferencedObjects(FReferenceCollector& Collector)
+{
+	FBlueprintCompileReinstancer::AddReferencedObjects(Collector);
+	Collector.AllowEliminatingReferences(false);
+	Collector.AddReferencedObject(CopyOfPreviousCDO);
+	Collector.AllowEliminatingReferences(true);
+}
+
+void FReloadClassReinstancer::EnlistDependentBlueprintToRecompile(UBlueprint* BP, bool bBytecodeOnly)
+{
+	if (IsValid(BP))
+	{
+		if (bBytecodeOnly)
+		{
+			if (!BPSetToRecompile.Contains(BP) && !BPSetToRecompileBytecodeOnly.Contains(BP))
+			{
+				BPSetToRecompileBytecodeOnly.Add(BP);
+			}
+		}
+		else
+		{
+			if (!BPSetToRecompile.Contains(BP))
+			{
+				if (BPSetToRecompileBytecodeOnly.Contains(BP))
+				{
+					BPSetToRecompileBytecodeOnly.Remove(BP);
+				}
+
+				BPSetToRecompile.Add(BP);
+			}
+		}
+	}
+}
+
+void FReloadClassReinstancer::BlueprintWasRecompiled(UBlueprint* BP, bool bBytecodeOnly)
+{
+	BPSetToRecompile.Remove(BP);
+	BPSetToRecompileBytecodeOnly.Remove(BP);
+
+	FBlueprintCompileReinstancer::BlueprintWasRecompiled(BP, bBytecodeOnly);
+}
+
+FReload::FReload(EActiveReloadType InType, const TCHAR* InPrefix, const TArray<UPackage*>& InPackages, FOutputDevice& InAr)
+	: Type(InType)
+	, Prefix(InPrefix)
+	, Packages(InPackages)
+	, Ar(InAr)
+	, bCollectPackages(false)
+{
+	BeginReload(Type, *this);
+}
+
+FReload::FReload(EActiveReloadType InType, const TCHAR* InPrefix, FOutputDevice& InAr)
+	: Type(InType)
+	, Prefix(InPrefix)
+	, Ar(InAr)
+	, bCollectPackages(true)
+{
+	BeginReload(Type, *this);
+}
+
+FReload::~FReload()
+{
+	EndReload();
+
+	TStringBuilder<256> Builder;
+	if (ClassStats.HasValues() || StructStats.HasValues() || EnumStats.HasValues() || NumFunctionsRemapped != 0 || NumScriptStructsRemapped != 0)
+	{
+		FormatStats(Builder, TEXT("class"), TEXT("classes"), ClassStats);
+		FormatStats(Builder, TEXT("enum"), TEXT("enums"), EnumStats);
+		FormatStats(Builder, TEXT("scriptstruct"), TEXT("scriptstructs"), StructStats);
+		FormatStat(Builder, TEXT("function"), TEXT("functions"), TEXT("remapped"), NumFunctionsRemapped);
+		FormatStat(Builder, TEXT("scriptstruct"), TEXT("scriptstructs"), TEXT("remapped"), NumScriptStructsRemapped);
+	}
+	else
+	{
+		Builder << TEXT("No object changes detected");
+	}
+	Ar.Logf(ELogVerbosity::Display, TEXT("Reload/Re-instancing Complete: %s"), *Builder);
+
+	if (bSendReloadComplete)
+	{
+		FCoreUObjectDelegates::ReloadCompleteDelegate.Broadcast(EReloadCompleteReason::None);
+	}
+}
+
+void FReload::Reset()
+{
+	FunctionRemap.Empty();
+	BPSetToRecompile.Empty();
+	BPSetToRecompileBytecodeOnly.Empty();
+	ReconstructedCDOsMap.Empty();
+	ReinstancedClasses.Empty();
+}
+
+void FReload::UpdateStats(ReinstanceStats& Stats, void* New, void* Old)
+{
+	if (Old == nullptr)
+	{
+		++Stats.New;
+	}
+	else if (Old != New)
+	{
+		++Stats.Changed;
+	}
+	else
+	{
+		++Stats.Unchanged;
+	}
+}
+
+void FReload::FormatStats(FStringBuilderBase& Out, const TCHAR* Singular, const TCHAR* Plural, const ReinstanceStats& Stats)
+{
+	FormatStat(Out, Singular, Plural, TEXT("new"), Stats.New);
+	FormatStat(Out, Singular, Plural, TEXT("changed"), Stats.Changed);
+	FormatStat(Out, Singular, Plural, TEXT("unchanged"), Stats.Unchanged);
+}
+
+void FReload::FormatStat(FStringBuilderBase& Out, const TCHAR* Singular, const TCHAR* Plural, const TCHAR* What, int32 Value)
+{
+	if (Value == 0)
+	{
+		return;
+	}
+
+	if (Out.Len() != 0)
+	{
+		Out << TEXT(", ");
+	}
+	Out << Value << TEXT(" ") << (Value > 1 ? Plural : Singular) << TEXT(" ") << What;
+}
+
+void FReload::NotifyFunctionRemap(FNativeFuncPtr NewFunctionPointer, FNativeFuncPtr OldFunctionPointer)
+{
+	FNativeFuncPtr OtherNewFunction = FunctionRemap.FindRef(OldFunctionPointer);
+	check(!OtherNewFunction || OtherNewFunction == NewFunctionPointer);
+	check(NewFunctionPointer);
+	check(OldFunctionPointer);
+	FunctionRemap.Add(OldFunctionPointer, NewFunctionPointer);
+}
+
+void FReload::NotifyChange(UClass* New, UClass* Old)
+{
+	UpdateStats(ClassStats, New, Old);
+
+	// Ignore new classes
+	if (Old != nullptr)
+	{
+		// Don't allow re-instancing of UEngine classes
+		if (!Old->IsChildOf(UEngine::StaticClass()))
+		{
+			UClass* NewIfChanged = Old != New ? New : nullptr; // supporting code detects unchanged based on null new pointer
+			TMap<UClass*, UClass*>& ClassesToReinstance = GetClassesToReinstanceForHotReload();
+			checkf(!ClassesToReinstance.Contains(Old) || ClassesToReinstance[Old] == NewIfChanged, TEXT("Attempting to reload a class which is already being reloaded as a different class"));
+			ClassesToReinstance.Add(Old, NewIfChanged);
+		}
+		else if (Old != New) // This has changed
+		{
+			Ar.Logf(ELogVerbosity::Warning, TEXT("Engine class '%s' has changed but will be ignored for reload"), *New->GetName());
+		}
+	}
+}
+
+void FReload::NotifyChange(UEnum* New, UEnum* Old)
+{
+	UpdateStats(EnumStats, New, Old);
+}
+
+void FReload::NotifyChange(UScriptStruct* New, UScriptStruct* Old)
+{
+	UpdateStats(StructStats, New, Old);
+}
+
+void FReload::Reinstance()
+{
+	if (Type != EActiveReloadType::Reinstancing)
+	{
+		UClass::AssembleReferenceTokenStreams();
+	}
+
+	TMap<UClass*, UClass*>& ClassesToReinstance = GetClassesToReinstanceForHotReload();
+
+	TMap<UClass*, UClass*> OldToNewClassesMap;
+	for (const TPair<UClass*, UClass*>& Pair : ClassesToReinstance)
+	{
+		if (Pair.Value != nullptr)
+		{
+			OldToNewClassesMap.Add(Pair.Key, Pair.Value);
+		}
+	}
+
+	for (const TPair<UClass*, UClass*>& Pair : ClassesToReinstance)
+	{
+		ReinstanceClass(Pair.Value, Pair.Key, OldToNewClassesMap);
+	}
+
+	ReinstancedClasses = MoveTemp(ClassesToReinstance);
+
+	FCoreUObjectDelegates::ReloadReinstancingCompleteDelegate.Broadcast();
+}
+
+void FReload::ReinstanceClass(UClass* NewClass, UClass* OldClass, const TMap<UClass*, UClass*>& OldToNewClassesMap)
+{
+	TSharedPtr<FReloadClassReinstancer> ReinstanceHelper = FReloadClassReinstancer::Create(NewClass, OldClass, OldToNewClassesMap, ReconstructedCDOsMap, BPSetToRecompile, BPSetToRecompileBytecodeOnly);
+	if (ReinstanceHelper->ClassNeedsReinstancing())
+	{
+		Ar.Logf(ELogVerbosity::Log, TEXT("Re-instancing %s after reload."), NewClass ? *NewClass->GetName() : *OldClass->GetName());
+		ReinstanceHelper->ReinstanceObjectsAndUpdateDefaults();
+	}
+}
+
+void FReload::Finalize()
+{
+	// If we have to collect the packages
+	if (bCollectPackages)
+	{
+		for (const TPair<UClass*, UClass*>& Pair : ReinstancedClasses)
+		{
+			UClass* OldClass = Pair.Key;
+			UClass* NewClass = Pair.Value;
+			Packages.AddUnique(NewClass ? NewClass->GetPackage() : OldClass->GetPackage());
+		}
+	}
+
+	// Remap all native functions (and gather scriptstructs)
+	TArray<UScriptStruct*> ScriptStructs;
+	for (FRawObjectIterator It; It; ++It)
+	{
+		UObject* a = static_cast<UObject*>(It->Object);
+		UFunction* f = Cast<UFunction>(a);
+
+		if (UFunction* Function = Cast<UFunction>(static_cast<UObject*>(It->Object)))
+		{
+			if (FNativeFuncPtr NewFunction = FunctionRemap.FindRef(Function->GetNativeFunc()))
+			{
+				++NumFunctionsRemapped;
+				Function->SetNativeFunc(NewFunction);
+			}
+		}
+
+		if (UScriptStruct* ScriptStruct = Cast<UScriptStruct>(static_cast<UObject*>(It->Object)))
+		{
+			if (!ScriptStruct->HasAnyFlags(RF_ClassDefaultObject) && ScriptStruct->GetCppStructOps() && Packages.ContainsByPredicate([=](UPackage* Package) { return ScriptStruct->IsIn(Package); }))
+			{
+				ScriptStructs.Add(ScriptStruct);
+			}
+		}
+	}
+	NumScriptStructsRemapped = ScriptStructs.Num();
+
+	// now let's set up the script structs...this relies on super behavior, so null them all, then set them all up. Internally this sets them up hierarchically.
+	for (UScriptStruct* Script : ScriptStructs)
+	{
+		Script->ClearCppStructOps();
+	}
+	for (UScriptStruct* Script : ScriptStructs)
+	{
+		Script->PrepareCppStructOps();
+		check(Script->GetCppStructOps());
+	}
+	// Make sure new classes have the token stream assembled
+	UClass::AssembleReferenceTokenStreams();
+
+	ReplaceReferencesToReconstructedCDOs();
+
+	// Force GC to collect reinstanced objects
+	CollectGarbage(GARBAGE_COLLECTION_KEEPFLAGS, true);
+}
+
+void FReload::ReplaceReferencesToReconstructedCDOs()
+{
+	if (ReconstructedCDOsMap.Num() == 0)
+	{
+		return;
+	}
+
+	// Thread pool manager. We need new thread pool with increased
+	// amount of stack size. Standard GThreadPool was encountering
+	// stack overflow error during serialization.
+	static struct FReplaceReferencesThreadPool
+	{
+		FReplaceReferencesThreadPool()
+		{
+			Pool = FQueuedThreadPool::Allocate();
+			int32 NumThreadsInThreadPool = FPlatformMisc::NumberOfWorkerThreadsToSpawn();
+			verify(Pool->Create(NumThreadsInThreadPool, 256 * 1024));
+		}
+
+		~FReplaceReferencesThreadPool()
+		{
+			Pool->Destroy();
+		}
+
+		FQueuedThreadPool* GetPool() { return Pool; }
+
+	private:
+		FQueuedThreadPool* Pool;
+	} ThreadPoolManager;
+
+	// Async task to enable multithreaded CDOs reference search.
+	class FFindRefTask : public FNonAbandonableTask
+	{
+	public:
+		explicit FFindRefTask(const TMap<UObject*, UObject*>& InReconstructedCDOsMap, int32 ReserveElements)
+			: ReconstructedCDOsMap(InReconstructedCDOsMap)
+		{
+			ObjectsArray.Reserve(ReserveElements);
+		}
+
+		void DoWork()
+		{
+			for (UObject* Object : ObjectsArray)
+			{
+				class FReplaceCDOReferencesArchive : public FArchiveUObject
+				{
+				public:
+					FReplaceCDOReferencesArchive(UObject* InPotentialReferencer, const TMap<UObject*, UObject*>& InReconstructedCDOsMap)
+						: ReconstructedCDOsMap(InReconstructedCDOsMap)
+						, PotentialReferencer(InPotentialReferencer)
+					{
+						ArIsObjectReferenceCollector = true;
+						ArIgnoreOuterRef = true;
+					}
+
+					virtual FString GetArchiveName() const override
+					{
+						return TEXT("FReplaceCDOReferencesArchive");
+					}
+
+					FArchive& operator<<(UObject*& ObjRef)
+					{
+						UObject* Obj = ObjRef;
+
+						if (Obj && Obj != PotentialReferencer)
+						{
+							if (UObject* const* FoundObj = ReconstructedCDOsMap.Find(Obj))
+							{
+								ObjRef = *FoundObj;
+							}
+						}
+
+						return *this;
+					}
+
+					const TMap<UObject*, UObject*>& ReconstructedCDOsMap;
+					UObject* PotentialReferencer;
+				};
+
+				FReplaceCDOReferencesArchive FindRefsArchive(Object, ReconstructedCDOsMap);
+				Object->Serialize(FindRefsArchive);
+			}
+		}
+
+		FORCEINLINE TStatId GetStatId() const
+		{
+			RETURN_QUICK_DECLARE_CYCLE_STAT(FFindRefTask, STATGROUP_ThreadPoolAsyncTasks);
+		}
+
+		TArray<UObject*> ObjectsArray;
+
+	private:
+		const TMap<UObject*, UObject*>& ReconstructedCDOsMap;
+	};
+
+	const int32 NumberOfThreads = FPlatformMisc::NumberOfWorkerThreadsToSpawn();
+	const int32 NumObjects = GUObjectArray.GetObjectArrayNum();
+	const int32 ObjectsPerTask = FMath::CeilToInt((float)NumObjects / NumberOfThreads);
+
+	// Create tasks.
+	TArray<FAsyncTask<FFindRefTask>> Tasks;
+	Tasks.Reserve(NumberOfThreads);
+
+	for (int32 TaskId = 0; TaskId < NumberOfThreads; ++TaskId)
+	{
+		Tasks.Emplace(ReconstructedCDOsMap, ObjectsPerTask);
+	}
+
+	// Distribute objects uniformly between tasks.
+	int32 CurrentTaskId = 0;
+	for (FThreadSafeObjectIterator ObjIter; ObjIter; ++ObjIter)
+	{
+		UObject* CurObject = *ObjIter;
+
+		if (CurObject->IsPendingKill())
+		{
+			continue;
+		}
+
+		Tasks[CurrentTaskId].GetTask().ObjectsArray.Add(CurObject);
+		CurrentTaskId = (CurrentTaskId + 1) % NumberOfThreads;
+	}
+
+	// Run async tasks in worker threads.
+	for (FAsyncTask<FFindRefTask>& Task : Tasks)
+	{
+		Task.StartBackgroundTask(ThreadPoolManager.GetPool());
+	}
+
+	// Wait until tasks are finished
+	for (FAsyncTask<FFindRefTask>& AsyncTask : Tasks)
+	{
+		AsyncTask.EnsureCompletion();
+	}
+}
diff --git a/Engine/Source/Editor/UnrealEd/Public/ComponentTypeRegistry.h b/Engine/Source/Editor/UnrealEd/Public/ComponentTypeRegistry.h
--- a/Engine/Source/Editor/UnrealEd/Public/ComponentTypeRegistry.h
+++ b/Engine/Source/Editor/UnrealEd/Public/ComponentTypeRegistry.h
@@ -43,7 +43,7 @@ struct UNREALED_API FComponentTypeRegistry
 	 */
 	void InvalidateClass(TSubclassOf<UActorComponent> ClassToUpdate);
 private:
-	void OnProjectHotReloaded( bool bWasTriggeredAutomatically );
+	void OnReloadComplete(EReloadCompleteReason Reason);
 
 private:
 	FComponentTypeRegistry();
diff --git a/Engine/Source/Editor/UnrealEd/Public/Kismet2/ReloadUtilities.h b/Engine/Source/Editor/UnrealEd/Public/Kismet2/ReloadUtilities.h
new file mode 100644
--- /dev/null
+++ b/Engine/Source/Editor/UnrealEd/Public/Kismet2/ReloadUtilities.h
@@ -0,0 +1,138 @@
+// Copyright Epic Games, Inc. All Rights Reserved.
+
+/*=============================================================================
+	ReloadUtilities.h: Helper utilities and classes for reloading
+=============================================================================*/
+
+#pragma once
+
+#include "UObject/Reload.h"
+#include "Containers/StringFwd.h"
+
+class UBlueprint;
+
+class UNREALED_API FReload : public IReload
+{
+public:
+	struct FNativeFuncPtrMapKeyFuncs : TDefaultMapKeyFuncs<FNativeFuncPtr, FNativeFuncPtr, false>
+	{
+		static FORCEINLINE uint32 GetKeyHash(FNativeFuncPtr Key)
+		{
+			return *(uint32*)&Key;
+		}
+	};
+
+	using TFunctionRemap = TMap<FNativeFuncPtr, FNativeFuncPtr, FDefaultSetAllocator, FNativeFuncPtrMapKeyFuncs>;
+
+public:
+
+	FReload(EActiveReloadType InType, const TCHAR* InPrefix, const TArray<UPackage*>& InPackages, FOutputDevice& InAr);
+	FReload(EActiveReloadType InType, const TCHAR* InPrefix, FOutputDevice& InAr);
+
+	virtual ~FReload();
+
+	// IReload interface 
+	virtual EActiveReloadType GetType() const override { return Type; }
+	virtual const TCHAR* GetPrefix() const override { return Prefix; };
+	virtual void NotifyFunctionRemap(FNativeFuncPtr NewFunctionPointer, FNativeFuncPtr OldFunctionPointer) override;
+	virtual void NotifyChange(UClass* New, UClass* Old) override;
+	virtual void NotifyChange(UEnum* New, UEnum* Old) override;
+	virtual void NotifyChange(UScriptStruct* New, UScriptStruct* Old) override;
+	virtual void Reinstance() override;
+
+	/**
+	 * If you wish to reuse the same reload object, invoke this method to reset the state
+	 */
+	void Reset();
+
+	/**
+	 * Perform any finalize processing for reloading.
+	 */
+	void Finalize();
+
+	/**
+	 * Set the sending of the complete notification
+	 */
+	void SetSendReloadCompleteNotification(bool bSend)
+	{
+		bSendReloadComplete = bSend;
+	}
+
+private:
+
+	struct ReinstanceStats
+	{
+		int32 New = 0;
+		int32 Changed = 0;
+		int32 Unchanged = 0;
+
+		bool HasValues() const
+		{
+			return New + Changed + Unchanged != 0;
+		}
+	};
+
+	/**
+	 * Finds all references to old CDOs and replaces them with the new ones.
+	 * Skipping UBlueprintGeneratedClass::OverridenArchetypeForCDO as it's the
+	 * only one needed.
+	 */
+	void ReplaceReferencesToReconstructedCDOs();
+
+	/**
+	 * Given a new/old class pair, re-instance the class
+	 */
+	void ReinstanceClass(UClass* NewClass, UClass* OldClass, const TMap<UClass*, UClass*>& OldToNewClassesMap);
+
+	/**
+	 * Based on the pointers, update the given stat
+	 */
+	void UpdateStats(ReinstanceStats& Stats, void* New, void* Old);
+
+	/**
+	 * Helper method to format all the stats
+	 */
+	static void FormatStats(FStringBuilderBase& Out, const TCHAR* Singular, const TCHAR* Plural, const ReinstanceStats& Stats);
+
+	/**
+	 * Helper method to format a specific stat
+	 */
+	static void FormatStat(FStringBuilderBase& Out, const TCHAR* Singular, const TCHAR* Plural, const TCHAR* What, int32 Value);
+
+	/** Type of the active reload */
+	EActiveReloadType Type = EActiveReloadType::None;
+
+	/** Prefix applied when renaming objects */
+	const TCHAR* Prefix = nullptr;
+
+	/** List of packages affected by the reload */
+	TArray<UPackage*> Packages;
+	
+	/** Output device for any logging */
+	FOutputDevice& Ar;
+
+	/** Map from old function pointer to new function pointer for hot reload. */
+	TFunctionRemap FunctionRemap;
+
+	/* Unused */
+	TSet<UBlueprint*> BPSetToRecompile;
+	TSet<UBlueprint*> BPSetToRecompileBytecodeOnly;
+
+	/** Map of the reconstructed CDOs during the reinstancing process */
+	TMap<UObject*, UObject*> ReconstructedCDOsMap;
+
+	/** Map from old class to new class.  New class may be null */
+	TMap<UClass*, UClass*> ReinstancedClasses;
+
+	/** If true, we have to collect the package list from the context */
+	bool bCollectPackages;
+
+	/** If true, send reload complete notification */
+	bool bSendReloadComplete = true;
+
+	ReinstanceStats ClassStats;
+	ReinstanceStats EnumStats;
+	ReinstanceStats StructStats;
+	int32 NumFunctionsRemapped = 0;
+	int32 NumScriptStructsRemapped = 0;
+};
diff --git a/Engine/Source/Programs/UnrealHeaderTool/Private/CodeGenerator.cpp b/Engine/Source/Programs/UnrealHeaderTool/Private/CodeGenerator.cpp
--- a/Engine/Source/Programs/UnrealHeaderTool/Private/CodeGenerator.cpp
+++ b/Engine/Source/Programs/UnrealHeaderTool/Private/CodeGenerator.cpp
@@ -2430,7 +2430,7 @@ void FNativeClassHeaderGenerator::ExportNativeGeneratedInitCode(FOutputDevice& O
 		GeneratedClassRegisterFunctionText.Logf(TEXT("\t{\r\n"));
 		if (!bIsDynamic)
 		{
-			GeneratedClassRegisterFunctionText.Logf(TEXT("\t\tstatic UClass* OuterClass = nullptr;\r\n"));
+			GeneratedClassRegisterFunctionText.Logf(TEXT("\t\tstatic UClass*& OuterClass = %s::StaticRegistrationInfo().OuterSingleton;\r\n"), *ClassNameCPP);
 			GeneratedClassRegisterFunctionText.Logf(TEXT("\t\tif (!OuterClass)\r\n"));
 		}
 		else
@@ -2551,27 +2551,16 @@ void FNativeClassHeaderGenerator::ExportNativeGeneratedInitCode(FOutputDevice& O
 	if (bIsDynamic)
 	{
 		const FString& ClassPackageName = FClass::GetTypePackageName(Class);
-		Out.Logf(TEXT("\tstatic FCompiledInDefer Z_CompiledInDefer_UClass_%s(%s, &%s::StaticClass, TEXT(\"%s\"), TEXT(\"%s\"), true, %s, %s, %s);\r\n"),
+		Out.Logf(TEXT("\tstatic FRegisterCompiledInInfo Z_CompiledInDefer_UClass_%s(%s, &%s::StaticClass, TEXT(\"%s\"), TEXT(\"%s\"), %s, %s, %s);\r\n"),
 			*ClassNameCPP,
 			*SingletonName,
-			*ClassNameCPP,
 			*ClassPackageName,
+			*ClassNameCPP,
 			*OverriddenClassName,
 			*AsTEXT(ClassPackageName),
 			*AsTEXT(FNativeClassHeaderGenerator::GetOverriddenPathName(Class)),
 			*InitSearchableValuesFunctionParam);
 	}
-	else
-	{
-		Out.Logf(TEXT("\tstatic FCompiledInDefer Z_CompiledInDefer_UClass_%s(%s, &%s::StaticClass, TEXT(\"%s\"), TEXT(\"%s\"), false, nullptr, nullptr, %s);\r\n"),
-			*ClassNameCPP,
-			*SingletonName,
-			*ClassNameCPP,
-			*Class->GetOutermost()->GetName(),
-			*ClassNameCPP,
-			*InitSearchableValuesFunctionParam);
-	}
-
 
 	if (ClassHasReplicatedProperties(Class))
 	{
@@ -2637,7 +2626,8 @@ void FNativeClassHeaderGenerator::ExportFunction(FOutputDevice& Out, FReferenceG
 	CurrentFunctionText.Logf(TEXT("\tstruct %s\r\n"), *StaticsStructName);
 	CurrentFunctionText.Log (TEXT("\t{\r\n"));
 
-	if (bIsNoExport || !(Function->FunctionFlags&FUNC_Event))  // non-events do not export a params struct, so lets do that locally for offset determination
+	bool bParamsInStatic = bIsNoExport || !(Function->FunctionFlags & FUNC_Event); // non-events do not export a params struct, so lets do that locally for offset determination
+	if (bParamsInStatic)
 	{
 		TArray<UScriptStruct*> Structs = FindNoExportStructs(Function);
 		for (UScriptStruct* Struct : Structs)
@@ -2678,7 +2668,14 @@ void FNativeClassHeaderGenerator::ExportFunction(FOutputDevice& Out, FReferenceG
 			FunctionName.LeftChopInline(HEADER_GENERATED_DELEGATE_SIGNATURE_SUFFIX_LENGTH, false);
 		}
 
-		StructureSize = FString::Printf(TEXT("sizeof(%s)"), *GetEventStructParamsName(TempFunction->GetOuter(), *FunctionName));
+		if (bParamsInStatic)
+		{
+			StructureSize = FString::Printf(TEXT("sizeof(%s::%s)"), *StaticsStructName, *GetEventStructParamsName(TempFunction->GetOuter(), *FunctionName));
+		}
+		else
+		{
+			StructureSize = FString::Printf(TEXT("sizeof(%s)"), *GetEventStructParamsName(TempFunction->GetOuter(), *FunctionName));
+		}
 	}
 	else
 	{
@@ -2735,7 +2732,7 @@ void FNativeClassHeaderGenerator::ExportFunction(FOutputDevice& Out, FReferenceG
 
 	CurrentFunctionText.Logf(TEXT("\t\tif (!ReturnFunction)\r\n"));
 	CurrentFunctionText.Logf(TEXT("\t\t{\r\n"));
-	CurrentFunctionText.Logf(TEXT("\t\t\tUE4CodeGen_Private::ConstructUFunction(ReturnFunction, %s::FuncParams);\r\n"), *StaticsStructName);
+	CurrentFunctionText.Logf(TEXT("\t\t\tUE4CodeGen_Private::ConstructUFunction(ReturnFunction, %s::FuncParams, &ReturnFunction);\r\n"), *StaticsStructName);
 	CurrentFunctionText.Log (TEXT("\t\t}\r\n"));
 	CurrentFunctionText.Log (TEXT("\t\treturn ReturnFunction;\r\n"));
 	CurrentFunctionText.Log (TEXT("\t}\r\n"));
@@ -3873,6 +3870,17 @@ void FNativeClassHeaderGenerator::ExportGeneratedStructBodyMacros(FOutputDevice&
 
 		FString GetHashName = FString::Printf(TEXT("Get_%s_Hash"), *ChoppedSingletonName);
 
+		if (!bIsDynamic)
+		{
+			Out.Logf(TEXT("\tstatic FStructRegistrationInfo& Z_Registration_Info_UScriptStruct_%s()\r\n"), *Struct->GetName());
+			Out.Logf(TEXT("\t{\r\n"));
+			Out.Logf(TEXT("\t\tstatic FStructRegistrationInfo info;\r\n"));
+			Out.Logf(TEXT("\t\treturn info;\r\n"));
+			Out.Logf(TEXT("\t}\r\n"));
+		}
+
+		Out.Logf(TEXT("\textern %suint32 %s();\r\n"), *FriendApiString, *GetHashName);
+
 		Out.Logf(TEXT("class UScriptStruct* %s::StaticStruct()\r\n"), *StructNameCPP);
 		Out.Logf(TEXT("{\r\n"));
 
@@ -3880,7 +3888,7 @@ void FNativeClassHeaderGenerator::ExportGeneratedStructBodyMacros(FOutputDevice&
 		const FString& OuterName (bIsDynamic ? STRING_StructPackage : GetPackageSingletonName(CastChecked<UPackage>(Struct->GetOuter()), OutReferenceGatherers.UniqueCrossModuleReferences));
 		if (!bIsDynamic)
 		{
-			Out.Logf(TEXT("\tstatic class UScriptStruct* Singleton = NULL;\r\n"));
+			Out.Logf(TEXT("\tstatic class UScriptStruct*& Singleton = Z_Registration_Info_UScriptStruct_%s().OuterSingleton;\r\n"), *Struct->GetName());
 		}
 		else
 		{
@@ -3890,10 +3898,9 @@ void FNativeClassHeaderGenerator::ExportGeneratedStructBodyMacros(FOutputDevice&
 
 		Out.Logf(TEXT("\tif (!Singleton)\r\n"));
 		Out.Logf(TEXT("\t{\r\n"));
-		Out.Logf(TEXT("\t\textern %suint32 %s();\r\n"), *FriendApiString, *GetHashName);
 
-		Out.Logf(TEXT("\t\tSingleton = GetStaticStruct(%s, %s, TEXT(\"%s\"), sizeof(%s), %s());\r\n"),
-			*ChoppedSingletonName, *OuterName, *ActualStructName, *StructNameCPP, *GetHashName);
+		Out.Logf(TEXT("\t\tSingleton = GetStaticStruct(%s, %s, TEXT(\"%s\"));\r\n"),
+			*ChoppedSingletonName, *OuterName, *ActualStructName);
 
 		// if this struct has RigVM methods - we need to register the method to our central
 		// registry on construction of the static struct
@@ -3922,7 +3929,7 @@ void FNativeClassHeaderGenerator::ExportGeneratedStructBodyMacros(FOutputDevice&
 		if (bIsDynamic)
 		{
 			const FString& StructPackageName = FClass::GetTypePackageName(Struct);
-			Out.Logf(TEXT("static FCompiledInDeferStruct Z_CompiledInDeferStruct_UScriptStruct_%s(%s::StaticStruct, TEXT(\"%s\"), TEXT(\"%s\"), true, %s, %s);\r\n"),
+			Out.Logf(TEXT("static FRegisterCompiledInInfo Z_CompiledInDeferStruct_UScriptStruct_%s(%s::StaticStruct, TEXT(\"%s\"), TEXT(\"%s\"),  %s, %s);\r\n"),
 				*StructNameCPP,
 				*StructNameCPP,
 				*StructPackageName,
@@ -3932,16 +3939,22 @@ void FNativeClassHeaderGenerator::ExportGeneratedStructBodyMacros(FOutputDevice&
 		}
 		else
 		{
-			Out.Logf(TEXT("static FCompiledInDeferStruct Z_CompiledInDeferStruct_UScriptStruct_%s(%s::StaticStruct, TEXT(\"%s\"), TEXT(\"%s\"), false, nullptr, nullptr);\r\n"),
+			// Do not change the Z_CompiledInDeferStruct_UScriptStruct_ without changing LC_SymbolPatterns
+			Out.Logf(TEXT("static FRegisterCompiledInInfo Z_CompiledInDeferStruct_UScriptStruct_%s(%s::StaticStruct, TEXT(\"%s\"), TEXT(\"%s\"), Z_Registration_Info_UScriptStruct_%s(), CONSTUCT_RELOAD_VERSION_INFO(FStructReloadVersionInfo, sizeof(%s), %s()));\r\n"),
 				*StructNameCPP,
 				*StructNameCPP,
 				*Struct->GetOutermost()->GetName(),
-				*ActualStructName);
+				*ActualStructName,
+				*Struct->GetName(),
+				*StructNameCPP, 
+				*GetHashName
+			);
 		}
 
 		// Generate StaticRegisterNatives equivalent for structs without classes.
 		if (!Struct->GetOuter()->IsA(UStruct::StaticClass()))
 		{
+			// Do not change the Z_CompiledInDeferCppStructOps_UScriptStruct_ without changing LC_SymbolPatterns
 			const FString ShortPackageName = FPackageName::GetShortName(Struct->GetOuter()->GetName());
 			Out.Logf(TEXT("static struct FScriptStruct_%s_StaticRegisterNatives%s\r\n"), *ShortPackageName, *StructNameCPP);
 			Out.Logf(TEXT("{\r\n"));
@@ -3949,7 +3962,7 @@ void FNativeClassHeaderGenerator::ExportGeneratedStructBodyMacros(FOutputDevice&
 			Out.Logf(TEXT("\t{\r\n"));
 			Out.Logf(TEXT("\t\tUScriptStruct::DeferCppStructOps<%s>(FName(TEXT(\"%s\")));\r\n"), *StructNameCPP, *ActualStructName);
 			Out.Logf(TEXT("\t}\r\n"));
-			Out.Logf(TEXT("} ScriptStruct_%s_StaticRegisterNatives%s;\r\n"), *ShortPackageName, *StructNameCPP);
+			Out.Logf(TEXT("} Z_CompiledInDeferCppStructOps_UScriptStruct_%s;\r\n"), *StructNameCPP);
 		}
 	}
 
@@ -4054,17 +4067,17 @@ void FNativeClassHeaderGenerator::ExportGeneratedStructBodyMacros(FOutputDevice&
 	// Structs can either have a UClass or UPackage as outer (if declared in non-UClass header).
 	if (!bIsDynamic)
 	{
-		GeneratedStructRegisterFunctionText.Log (TEXT("#if WITH_HOT_RELOAD\r\n"));
-		GeneratedStructRegisterFunctionText.Logf(TEXT("\t\textern uint32 %s();\r\n"), *HashFuncName);
-		GeneratedStructRegisterFunctionText.Logf(TEXT("\t\tUPackage* Outer = %s;\r\n"), *GetPackageSingletonName(CastChecked<UPackage>(Struct->GetOuter()), OutReferenceGatherers.UniqueCrossModuleReferences));
-		GeneratedStructRegisterFunctionText.Logf(TEXT("\t\tstatic UScriptStruct* ReturnStruct = FindExistingStructIfHotReloadOrDynamic(Outer, TEXT(\"%s\"), sizeof(%s), %s(), false);\r\n"), *ActualStructName, *NoExportStructNameCPP, *HashFuncName);
-		GeneratedStructRegisterFunctionText.Log (TEXT("#else\r\n"));
-		GeneratedStructRegisterFunctionText.Logf(TEXT("\t\tstatic UScriptStruct* ReturnStruct = nullptr;\r\n"));
-		GeneratedStructRegisterFunctionText.Log (TEXT("#endif\r\n"));
+		if (Struct->StructFlags & STRUCT_Native)
+		{
+			GeneratedStructRegisterFunctionText.Logf(TEXT("\t\tstatic UScriptStruct*& ReturnStruct = Z_Registration_Info_UScriptStruct_%s().InnerSingleton;\r\n"), *Struct->GetName());
+		}
+		else
+		{
+			GeneratedStructRegisterFunctionText.Logf(TEXT("\t\tstatic UScriptStruct* ReturnStruct = nullptr;\r\n"));
+		}
 	}
 	else
 	{
-		GeneratedStructRegisterFunctionText.Logf(TEXT("\t\textern uint32 %s();\r\n"), *HashFuncName);
 		GeneratedStructRegisterFunctionText.Logf(TEXT("\t\tUPackage* Outer = FindOrConstructDynamicTypePackage(TEXT(\"%s\"));\r\n"), *FClass::GetTypePackageName(Struct));
 		GeneratedStructRegisterFunctionText.Logf(TEXT("\t\tUScriptStruct* ReturnStruct = FindExistingStructIfHotReloadOrDynamic(Outer, TEXT(\"%s\"), sizeof(%s), %s(), true);\r\n"), *ActualStructName, *NoExportStructNameCPP, *HashFuncName);
 	}
@@ -4164,6 +4177,7 @@ void FNativeClassHeaderGenerator::ExportGeneratedEnumInitCode(FOutputDevice& Out
 	const FString SingletonName         = GetSingletonNameFuncAddr(Enum, OutReferenceGatherers.UniqueCrossModuleReferences);
 	const FString EnumNameCpp           = Enum->GetName(); //UserDefinedEnum should already have a valid cpp name.
 	const FString OverriddenEnumNameCpp = FNativeClassHeaderGenerator::GetOverriddenName(Enum);
+	const FString StaticsStructName     = SingletonName/*.LeftChop(2)*/ + TEXT("_Statics");
 
 	const bool bIsEditorOnlyDataType = GEditorOnlyDataTypes.Contains(Enum);
 
@@ -4172,136 +4186,171 @@ void FNativeClassHeaderGenerator::ExportGeneratedEnumInitCode(FOutputDevice& Out
 
 	const FString& PackageSingletonName = (bIsDynamic ? FClass::GetTypePackageName(static_cast<UField*>(Enum)) : GetPackageSingletonName(CastChecked<UPackage>(Enum->GetOuter()), OutReferenceGatherers.UniqueCrossModuleReferences));
 
+	if (!bIsDynamic)
+	{
+		Out.Logf(TEXT("\tstatic FEnumRegistrationInfo& Z_Registration_Info_UEnum_%s()\r\n"), *Enum->GetName());
+		Out.Logf(TEXT("\t{\r\n"));
+		Out.Logf(TEXT("\t\tstatic FEnumRegistrationInfo info;\r\n"));
+		Out.Logf(TEXT("\t\treturn info;\r\n"));
+		Out.Logf(TEXT("\t}\r\n"));
+	}
+
 	Out.Logf(TEXT("\tstatic UEnum* %s_StaticEnum()\r\n"), *Enum->GetName());
 	Out.Logf(TEXT("\t{\r\n"));
 
 	if (!bIsDynamic)
 	{
-		Out.Logf(TEXT("\t\tstatic UEnum* Singleton = nullptr;\r\n"));
+		Out.Logf(TEXT("\t\tstatic UEnum*& OuterSingleton = Z_Registration_Info_UEnum_%s().OuterSingleton;\r\n"), *Enum->GetName());
+		Out.Logf(TEXT("\t\tif (!OuterSingleton)\r\n"));
+		Out.Logf(TEXT("\t\t{\r\n"));
+		Out.Logf(TEXT("\t\t\tOuterSingleton = GetStaticEnum(%s, %s, TEXT(\"%s\"));\r\n"), *SingletonName, *PackageSingletonName, *Enum->GetName());
+		Out.Logf(TEXT("\t\t}\r\n"));
+		Out.Logf(TEXT("\t\treturn OuterSingleton;\r\n"));
 	}
 	else
 	{
 		Out.Logf(TEXT("\t\tclass UPackage* EnumPackage = FindOrConstructDynamicTypePackage(TEXT(\"%s\"));\r\n"), *PackageSingletonName);
 		Out.Logf(TEXT("\t\tclass UEnum* Singleton = Cast<UEnum>(StaticFindObjectFast(UEnum::StaticClass(), EnumPackage, TEXT(\"%s\")));\r\n"), *OverriddenEnumNameCpp);
-	}
-	Out.Logf(TEXT("\t\tif (!Singleton)\r\n"));
-	Out.Logf(TEXT("\t\t{\r\n"));
-	if (!bIsDynamic)
-	{
-		Out.Logf(TEXT("\t\t\tSingleton = GetStaticEnum(%s, %s, TEXT(\"%s\"));\r\n"), *SingletonName, *PackageSingletonName, *Enum->GetName());
-	}
-	else
-	{
+		Out.Logf(TEXT("\t\tif (!Singleton)\r\n"));
+		Out.Logf(TEXT("\t\t{\r\n"));
 		Out.Logf(TEXT("\t\t\tSingleton = GetStaticEnum(%s, EnumPackage, TEXT(\"%s\"));\r\n"), *SingletonName, *OverriddenEnumNameCpp);
+		Out.Logf(TEXT("\t\t}\r\n"));
+		Out.Logf(TEXT("\t\treturn Singleton;\r\n"));
 	}
 
-	Out.Logf(TEXT("\t\t}\r\n"));
-	Out.Logf(TEXT("\t\treturn Singleton;\r\n"));
 	Out.Logf(TEXT("\t}\r\n"));
 
+	const FString& EnumSingletonName = GetSingletonName(Enum, OutReferenceGatherers.UniqueCrossModuleReferences);
+	const FString HashFuncName = FString::Printf(TEXT("Get_%s_Hash"), *SingletonName);
+
 	Out.Logf(TEXT("\ttemplate<> %sUEnum* StaticEnum<%s>()\r\n"), *GetAPIString(), *Enum->CppType);
 	Out.Logf(TEXT("\t{\r\n"));
 	Out.Logf(TEXT("\t\treturn %s_StaticEnum();\r\n"), *Enum->GetName());
 	Out.Logf(TEXT("\t}\r\n"));
 
-	if (bIsDynamic)
-	{
-		const FString& EnumPackageName = FClass::GetTypePackageName(static_cast<UField*>(Enum));
-		Out.Logf(
-			TEXT("\tstatic FCompiledInDeferEnum Z_CompiledInDeferEnum_UEnum_%s(%s_StaticEnum, TEXT(\"%s\"), TEXT(\"%s\"), true, %s, %s);\r\n"),
-			*EnumNameCpp,
-			*EnumNameCpp,
-			*EnumPackageName,
-			*OverriddenEnumNameCpp,
-			*AsTEXT(EnumPackageName),
-			*AsTEXT(FNativeClassHeaderGenerator::GetOverriddenPathName(static_cast<UField*>(Enum)))
-		);
-	}
-	else
+	FUHTStringBuilder StaticDefinitions;
+	FUHTStringBuilder StaticDeclarations;
+
+	// Generate the static declarations and definitions
 	{
-		Out.Logf(
-			TEXT("\tstatic FCompiledInDeferEnum Z_CompiledInDeferEnum_UEnum_%s(%s_StaticEnum, TEXT(\"%s\"), TEXT(\"%s\"), false, nullptr, nullptr);\r\n"),
-			*EnumNameCpp,
-			*EnumNameCpp,
-			*Enum->GetOutermost()->GetName(),
-			*OverriddenEnumNameCpp
-		);
+
+		// Enums can either have a UClass or UPackage as outer (if declared in non-UClass header).
+		FString OuterString;
+		if (!bIsDynamic)
+		{
+			OuterString = PackageSingletonName;
+		}
+		else
+		{
+			OuterString = FString::Printf(TEXT("[](){ return (UObject*)FindOrConstructDynamicTypePackage(TEXT(\"%s\")); }()"), *PackageSingletonName);
+		}
+
+		const TCHAR* UEnumObjectFlags = bIsDynamic ? TEXT("RF_Public|RF_Transient") : TEXT("RF_Public|RF_Transient|RF_MarkAsNative");
+		const TCHAR* EnumFlags = Enum->HasAnyEnumFlags(EEnumFlags::Flags) ? TEXT("EEnumFlags::Flags") : TEXT("EEnumFlags::None");
+
+		const TCHAR* EnumFormStr = TEXT("");
+		switch (Enum->GetCppForm())
+		{
+		case UEnum::ECppForm::Regular:    EnumFormStr = TEXT("UEnum::ECppForm::Regular");    break;
+		case UEnum::ECppForm::Namespaced: EnumFormStr = TEXT("UEnum::ECppForm::Namespaced"); break;
+		case UEnum::ECppForm::EnumClass:  EnumFormStr = TEXT("UEnum::ECppForm::EnumClass");  break;
+		}
+
+		const FString& EnumDisplayNameFn = Enum->GetMetaData(TEXT("EnumDisplayNameFn"));
+
+		StaticDeclarations.Logf(TEXT("\tstruct %s\r\n"), *StaticsStructName);
+		StaticDeclarations.Logf(TEXT("\t{\r\n"));
+
+		StaticDeclarations.Logf(TEXT("\t\tstatic const UE4CodeGen_Private::FEnumeratorParam Enumerators[];\r\n"));
+		StaticDefinitions.Logf(TEXT("\tconst UE4CodeGen_Private::FEnumeratorParam %s::Enumerators[] = {\r\n"), *StaticsStructName);
+		for (int32 Index = 0; Index != Enum->NumEnums(); ++Index)
+		{
+			const TCHAR* OverridenNameMetaDatakey = TEXT("OverrideName");
+			const FString KeyName = Enum->HasMetaData(OverridenNameMetaDatakey, Index) ? Enum->GetMetaData(OverridenNameMetaDatakey, Index) : Enum->GetNameByIndex(Index).ToString();
+			StaticDefinitions.Logf(TEXT("\t\t{ %s, (int64)%s },\r\n"), *CreateUTF8LiteralString(KeyName), *Enum->GetNameByIndex(Index).ToString());
+		}
+		StaticDefinitions.Logf(TEXT("\t};\r\n"));
+
+		FString MetaDataParamsName = StaticsStructName + TEXT("::Enum_MetaDataParams");
+		FString MetaDataParams = OutputMetaDataCodeForObject(StaticDeclarations, StaticDefinitions, Enum, *MetaDataParamsName, TEXT("\t\t"), TEXT("\t"));
+
+		StaticDeclarations.Logf(TEXT("\t\tstatic const UE4CodeGen_Private::FEnumParams EnumParams;\r\n"));
+		StaticDefinitions.Logf(TEXT("\tconst UE4CodeGen_Private::FEnumParams %s::EnumParams = {\r\n"), *StaticsStructName);
+		StaticDefinitions.Logf(TEXT("\t\t(UObject*(*)())%s,\r\n"), *OuterString.LeftChop(2));
+		StaticDefinitions.Logf(TEXT("\t\t%s,\r\n"), EnumDisplayNameFn.IsEmpty() ? TEXT("nullptr") : *EnumDisplayNameFn);
+		StaticDefinitions.Logf(TEXT("\t\t%s,\r\n"), *CreateUTF8LiteralString(OverriddenEnumNameCpp));
+		StaticDefinitions.Logf(TEXT("\t\t%s,\r\n"), *CreateUTF8LiteralString(Enum->CppType));
+		StaticDefinitions.Logf(TEXT("\t\t%s::Enumerators,\r\n"), *StaticsStructName);
+		StaticDefinitions.Logf(TEXT("\t\tUE_ARRAY_COUNT(%s::Enumerators),\r\n"), *StaticsStructName);
+		StaticDefinitions.Logf(TEXT("\t\t%s,\r\n"), UEnumObjectFlags);
+		StaticDefinitions.Logf(TEXT("\t\t%s,\r\n"), EnumFlags);
+		StaticDefinitions.Logf(TEXT("\t\tUE4CodeGen_Private::EDynamicType::%s,\r\n"), bIsDynamic ? TEXT("Dynamic") : TEXT("NotDynamic"));
+		StaticDefinitions.Logf(TEXT("\t\t(uint8)%s,\r\n"), EnumFormStr);
+		StaticDefinitions.Logf(TEXT("\t\t%s\r\n"), *MetaDataParams);
+		StaticDefinitions.Logf(TEXT("\t};\r\n"));
+
+		StaticDeclarations.Logf(TEXT("\t};\r\n"));
 	}
 
-	const FString& EnumSingletonName = GetSingletonName(Enum, OutReferenceGatherers.UniqueCrossModuleReferences);
-	const FString HashFuncName       = FString::Printf(TEXT("Get_%s_Hash"), *SingletonName);
+	//////////////////////////////////////
 
 	FUHTStringBuilder GeneratedEnumRegisterFunctionText;
 
+	GeneratedEnumRegisterFunctionText.Logf(TEXT("%s"), *StaticDeclarations);
+	GeneratedEnumRegisterFunctionText.Logf(TEXT("%s"), *StaticDefinitions);
+
 	GeneratedEnumRegisterFunctionText.Logf(TEXT("\tUEnum* %s\r\n"), *EnumSingletonName);
 	GeneratedEnumRegisterFunctionText.Logf(TEXT("\t{\r\n"));
 
 	// Enums can either have a UClass or UPackage as outer (if declared in non-UClass header).
-	FString OuterString;
 	if (!bIsDynamic)
 	{
-		OuterString = PackageSingletonName;
-		GeneratedEnumRegisterFunctionText.Logf(TEXT("#if WITH_HOT_RELOAD\r\n"));
-		GeneratedEnumRegisterFunctionText.Logf(TEXT("\t\tUPackage* Outer = %s;\r\n"), *OuterString);
-		GeneratedEnumRegisterFunctionText.Logf(TEXT("\t\tstatic UEnum* ReturnEnum = FindExistingEnumIfHotReloadOrDynamic(Outer, TEXT(\"%s\"), 0, %s(), false);\r\n"), *EnumNameCpp, *HashFuncName);
-		GeneratedEnumRegisterFunctionText.Logf(TEXT("#else\r\n"));
-		GeneratedEnumRegisterFunctionText.Logf(TEXT("\t\tstatic UEnum* ReturnEnum = nullptr;\r\n"));
-		GeneratedEnumRegisterFunctionText.Logf(TEXT("#endif // WITH_HOT_RELOAD\r\n"));
+		GeneratedEnumRegisterFunctionText.Logf(TEXT("\t\tstatic UEnum*& ReturnEnum = Z_Registration_Info_UEnum_%s().InnerSingleton;\r\n"), *Enum->GetName());
 	}
 	else
 	{
-		OuterString = FString::Printf(TEXT("[](){ return (UObject*)FindOrConstructDynamicTypePackage(TEXT(\"%s\")); }()"), *PackageSingletonName);
 		GeneratedEnumRegisterFunctionText.Logf(TEXT("\t\tUPackage* Outer = FindOrConstructDynamicTypePackage(TEXT(\"%s\"));"), *PackageSingletonName);
 		GeneratedEnumRegisterFunctionText.Logf(TEXT("\t\tUEnum* ReturnEnum = FindExistingEnumIfHotReloadOrDynamic(Outer, TEXT(\"%s\"), 0, %s(), true);\r\n"), *OverriddenEnumNameCpp, *HashFuncName);
 	}
 	GeneratedEnumRegisterFunctionText.Logf(TEXT("\t\tif (!ReturnEnum)\r\n"));
 	GeneratedEnumRegisterFunctionText.Logf(TEXT("\t\t{\r\n"));
 
-	const TCHAR* UEnumObjectFlags = bIsDynamic ? TEXT("RF_Public|RF_Transient") : TEXT("RF_Public|RF_Transient|RF_MarkAsNative");
-	const TCHAR* EnumFlags        = Enum->HasAnyEnumFlags(EEnumFlags::Flags) ? TEXT("EEnumFlags::Flags") : TEXT("EEnumFlags::None");
-
-	const TCHAR* EnumFormStr = TEXT("");
-	switch (Enum->GetCppForm())
-	{
-		case UEnum::ECppForm::Regular:    EnumFormStr = TEXT("UEnum::ECppForm::Regular");    break;
-		case UEnum::ECppForm::Namespaced: EnumFormStr = TEXT("UEnum::ECppForm::Namespaced"); break;
-		case UEnum::ECppForm::EnumClass:  EnumFormStr = TEXT("UEnum::ECppForm::EnumClass");  break;
-	}
-
-	const FString& EnumDisplayNameFn = Enum->GetMetaData(TEXT("EnumDisplayNameFn"));
-
-	GeneratedEnumRegisterFunctionText.Logf(TEXT("\t\t\tstatic const UE4CodeGen_Private::FEnumeratorParam Enumerators[] = {\r\n"));
-	for (int32 Index = 0; Index != Enum->NumEnums(); ++Index)
-	{
-		const TCHAR* OverridenNameMetaDatakey = TEXT("OverrideName");
-		const FString KeyName = Enum->HasMetaData(OverridenNameMetaDatakey, Index) ? Enum->GetMetaData(OverridenNameMetaDatakey, Index) : Enum->GetNameByIndex(Index).ToString();
-		GeneratedEnumRegisterFunctionText.Logf(TEXT("\t\t\t\t{ %s, (int64)%s },\r\n"), *CreateUTF8LiteralString(KeyName), *Enum->GetNameByIndex(Index).ToString());
-	}
-	GeneratedEnumRegisterFunctionText.Logf(TEXT("\t\t\t};\r\n"));
-
-	FOutputDeviceNull OutputDeviceNull;
-	FString MetaDataParams = OutputMetaDataCodeForObject(OutputDeviceNull, GeneratedEnumRegisterFunctionText, Enum, TEXT("Enum_MetaDataParams"), TEXT(""), TEXT("\t\t\t"));
-
-	GeneratedEnumRegisterFunctionText.Logf(TEXT("\t\t\tstatic const UE4CodeGen_Private::FEnumParams EnumParams = {\r\n"));
-	GeneratedEnumRegisterFunctionText.Logf(TEXT("\t\t\t\t(UObject*(*)())%s,\r\n"), *OuterString.LeftChop(2));
-	GeneratedEnumRegisterFunctionText.Logf(TEXT("\t\t\t\t%s,\r\n"), EnumDisplayNameFn.IsEmpty() ? TEXT("nullptr") : *EnumDisplayNameFn);
-	GeneratedEnumRegisterFunctionText.Logf(TEXT("\t\t\t\t%s,\r\n"), *CreateUTF8LiteralString(OverriddenEnumNameCpp));
-	GeneratedEnumRegisterFunctionText.Logf(TEXT("\t\t\t\t%s,\r\n"), *CreateUTF8LiteralString(Enum->CppType));
-	GeneratedEnumRegisterFunctionText.Logf(TEXT("\t\t\t\tEnumerators,\r\n"));
-	GeneratedEnumRegisterFunctionText.Logf(TEXT("\t\t\t\tUE_ARRAY_COUNT(Enumerators),\r\n"));
-	GeneratedEnumRegisterFunctionText.Logf(TEXT("\t\t\t\t%s,\r\n"), UEnumObjectFlags);
-	GeneratedEnumRegisterFunctionText.Logf(TEXT("\t\t\t\t%s,\r\n"), EnumFlags);
-	GeneratedEnumRegisterFunctionText.Logf(TEXT("\t\t\t\tUE4CodeGen_Private::EDynamicType::%s,\r\n"), bIsDynamic ? TEXT("Dynamic") : TEXT("NotDynamic"));
-	GeneratedEnumRegisterFunctionText.Logf(TEXT("\t\t\t\t(uint8)%s,\r\n"), EnumFormStr);
-	GeneratedEnumRegisterFunctionText.Logf(TEXT("\t\t\t\t%s\r\n"), *MetaDataParams);
-	GeneratedEnumRegisterFunctionText.Logf(TEXT("\t\t\t};\r\n"));
-	GeneratedEnumRegisterFunctionText.Logf(TEXT("\t\t\tUE4CodeGen_Private::ConstructUEnum(ReturnEnum, EnumParams);\r\n"));
+	GeneratedEnumRegisterFunctionText.Logf(TEXT("\t\t\tUE4CodeGen_Private::ConstructUEnum(ReturnEnum, %s::EnumParams);\r\n"), *StaticsStructName);
 	GeneratedEnumRegisterFunctionText.Logf(TEXT("\t\t}\r\n"));
 	GeneratedEnumRegisterFunctionText.Logf(TEXT("\t\treturn ReturnEnum;\r\n"));
 	GeneratedEnumRegisterFunctionText.Logf(TEXT("\t}\r\n"));
 
 	uint32 EnumHash = GenerateTextHash(*GeneratedEnumRegisterFunctionText);
+
 	Out.Logf(TEXT("\tuint32 %s() { return %uU; }\r\n"), *HashFuncName, EnumHash);
+
+	if (bIsDynamic)
+	{
+		const FString& EnumPackageName = FClass::GetTypePackageName(static_cast<UField*>(Enum));
+		Out.Logf(
+			TEXT("\tstatic FRegisterCompiledInInfo Z_CompiledInDeferEnum_UEnum_%s(%s_StaticEnum, TEXT(\"%s\"), TEXT(\"%s\"), %s, %s);\r\n"),
+			*EnumNameCpp,
+			*EnumNameCpp,
+			*EnumPackageName,
+			*OverriddenEnumNameCpp,
+			*AsTEXT(EnumPackageName),
+			*AsTEXT(FNativeClassHeaderGenerator::GetOverriddenPathName(static_cast<UField*>(Enum)))
+		);
+	}
+	else
+	{
+		// Do not change the Z_CompiledInDeferEnum_UEnum_ without changing LC_SymbolPatterns
+		Out.Logf(
+			TEXT("\tstatic FRegisterCompiledInInfo Z_CompiledInDeferEnum_UEnum_%s(%s_StaticEnum, TEXT(\"%s\"), TEXT(\"%s\"), Z_Registration_Info_UEnum_%s(), CONSTUCT_RELOAD_VERSION_INFO(FEnumReloadVersionInfo, %s()));\r\n"),
+			*EnumNameCpp,
+			*EnumNameCpp,
+			*Enum->GetOutermost()->GetName(),
+			*OverriddenEnumNameCpp,
+			*EnumNameCpp,
+			*HashFuncName
+		);
+	}
+
 	Out.Log(GeneratedEnumRegisterFunctionText);
 }
 
diff --git a/Engine/Source/Runtime/Core/Private/HAL/ConsoleManager.cpp b/Engine/Source/Runtime/Core/Private/HAL/ConsoleManager.cpp
--- a/Engine/Source/Runtime/Core/Private/HAL/ConsoleManager.cpp
+++ b/Engine/Source/Runtime/Core/Private/HAL/ConsoleManager.cpp
@@ -1512,8 +1512,8 @@ IConsoleObject* FConsoleManager::AddConsoleObject(const TCHAR* Name, IConsoleObj
 		// An existing console object was found that has the same name as the object being registered.
 		// In most cases this is not allowed, but if there is a variable with the same name and is
 		// in an 'unregistered' state or we're hot reloading dlls, we may be able to replace or update that variable.
-#if WITH_HOT_RELOAD
-		const bool bCanUpdateOrReplaceObj = (ExistingObj->AsVariable()||ExistingObj->AsCommand()) && (GIsHotReload || ExistingObj->TestFlags(ECVF_Unregistered));
+#if WITH_RELOAD
+		const bool bCanUpdateOrReplaceObj = (ExistingObj->AsVariable()||ExistingObj->AsCommand()) && (IsReloadActive() || ExistingObj->TestFlags(ECVF_Unregistered));
 #else
 		const bool bCanUpdateOrReplaceObj = ExistingObj->AsVariable() && ExistingObj->TestFlags(ECVF_Unregistered);
 #endif
@@ -1521,7 +1521,7 @@ IConsoleObject* FConsoleManager::AddConsoleObject(const TCHAR* Name, IConsoleObj
 		{
 			// NOTE: The reason we don't assert here is because when using HotReload, locally-initialized static console variables will be
 			//       re-registered, and it's desirable for the new variables to clobber the old ones.  Because this happen outside of the
-			//       hot reload stack frame (GIsHotReload=true), we can't detect and handle only those cases, so we opt to warn instead.
+			//       reload stack frame (IsActiveReload()=true), we can't detect and handle only those cases, so we opt to warn instead.
 			UE_LOG(LogConsoleManager, Warning, TEXT( "Console object named '%s' already exists but is being registered again, but we weren't expected it to be! (FConsoleManager::AddConsoleObject)"), Name );
 		}
 
@@ -1558,8 +1558,8 @@ IConsoleObject* FConsoleManager::AddConsoleObject(const TCHAR* Name, IConsoleObj
 				ConsoleObjects.Add(Name, Var);
 				return Var;
 			}
-#if WITH_HOT_RELOAD
-			else if (GIsHotReload)
+#if WITH_RELOAD
+			else if (IsReloadActive())
 			{
 				// Variable is being replaced due to a hot reload - copy state across to new variable, but only if the type hasn't changed
 				{
diff --git a/Engine/Source/Runtime/Core/Private/Modules/ModuleManager.cpp b/Engine/Source/Runtime/Core/Private/Modules/ModuleManager.cpp
--- a/Engine/Source/Runtime/Core/Private/Modules/ModuleManager.cpp
+++ b/Engine/Source/Runtime/Core/Private/Modules/ModuleManager.cpp
@@ -1317,3 +1317,60 @@ int32 FModuleManager::GetModuleCount() const
 	// if it wants to rely on the returned value.
 	return Modules.Num();
 }
+
+namespace
+{
+	EActiveReloadType GActiveReloadType = EActiveReloadType::None;
+	IReload* GActiveReloadInterface = nullptr;
+}
+
+#if WITH_RELOAD
+EActiveReloadType GetActiveReloadType()
+{
+#if WITH_HOT_RELOAD
+	PRAGMA_DISABLE_DEPRECATION_WARNINGS
+		if (GIsHotReload)
+	{
+		check(GActiveReloadInterface);
+		return EActiveReloadType::HotReload;
+	}
+	PRAGMA_ENABLE_DEPRECATION_WARNINGS
+#endif
+
+	return GActiveReloadType;
+}
+
+void BeginReload(EActiveReloadType ActiveReloadType, IReload& Interface)
+{
+	check(GActiveReloadInterface == nullptr);
+#if WITH_HOT_RELOAD
+	PRAGMA_DISABLE_DEPRECATION_WARNINGS
+	GIsHotReload = ActiveReloadType == EActiveReloadType::HotReload;
+	PRAGMA_ENABLE_DEPRECATION_WARNINGS
+#endif
+
+	GActiveReloadType = ActiveReloadType;
+	GActiveReloadInterface = &Interface;
+}
+
+void EndReload()
+{
+#if WITH_HOT_RELOAD
+	PRAGMA_DISABLE_DEPRECATION_WARNINGS
+	GIsHotReload = false;
+	PRAGMA_ENABLE_DEPRECATION_WARNINGS
+#endif
+	GActiveReloadType = EActiveReloadType::None;
+	GActiveReloadInterface = nullptr;
+}
+
+IReload* GetActiveReloadInterface()
+{
+	return GActiveReloadInterface;
+}
+
+bool IsReloadActive()
+{
+	return GetActiveReloadType() != EActiveReloadType::None;
+}
+#endif
diff --git a/Engine/Source/Runtime/Core/Public/Misc/Build.h b/Engine/Source/Runtime/Core/Public/Misc/Build.h
--- a/Engine/Source/Runtime/Core/Public/Misc/Build.h
+++ b/Engine/Source/Runtime/Core/Public/Misc/Build.h
@@ -147,6 +147,18 @@
 	#define WITH_HOT_RELOAD (!IS_MONOLITHIC && !UE_BUILD_SHIPPING && !UE_BUILD_TEST && !UE_GAME && !UE_SERVER)
 #endif
 
+/**
+* Make sure that live coding define is available.  Normally this is supplied by UBT
+*/
+#ifndef WITH_LIVE_CODING
+	#define WITH_LIVE_CODING 0
+#endif
+
+/**
+* Whether we support any type of live reloading
+*/
+#define WITH_RELOAD (WITH_HOT_RELOAD || WITH_LIVE_CODING)
+
 /**
 * Whether we include support for text archive formats. Disabling support allows de-virtualizing archive calls
 * and eliminating string constants for field names.
diff --git a/Engine/Source/Runtime/Core/Public/Modules/ModuleManager.h b/Engine/Source/Runtime/Core/Public/Modules/ModuleManager.h
--- a/Engine/Source/Runtime/Core/Public/Modules/ModuleManager.h
+++ b/Engine/Source/Runtime/Core/Public/Modules/ModuleManager.h
@@ -20,6 +20,7 @@
 
 #if WITH_HOT_RELOAD
 	/** If true, we are reloading a class for HotReload */
+	UE_DEPRECATED(5.0, "GIsHotReload has been deprecated, use IsReloadActive to test to see if a reload is in progress.")
 	extern CORE_API bool GIsHotReload;
 #endif
 
@@ -946,3 +947,54 @@ class FDefaultGameModuleImpl
 #endif	//IS_MONOLITHIC
 
 #endif
+
+/**
+* Enumerates the type of reload in progress
+*/
+enum class EActiveReloadType
+{
+	None,
+	Reinstancing,
+#if WITH_HOT_RELOAD
+	HotReload,
+#endif
+#if WITH_LIVE_CODING
+	LiveCoding,
+#endif
+};
+
+class IReload;
+
+#if WITH_RELOAD
+/**
+* Return the currently active reload.  Check for None to see if reloading is not active.
+* This method respects the GIsHotReload setting.
+*/
+CORE_API EActiveReloadType GetActiveReloadType();
+/**
+* Get the currently active reload interface.
+*/
+CORE_API IReload* GetActiveReloadInterface();
+
+/**
+* Helper method to check to see if reloading is active.
+* This method respects the GIsHotReload setting.
+*/
+CORE_API bool IsReloadActive();
+
+/**
+* Begins the reload process.
+*/
+CORE_API void BeginReload(EActiveReloadType ActiveReloadType, IReload& Interface);
+
+/**
+* Ends the reload process
+*/
+CORE_API void EndReload();
+#else
+inline EActiveReloadType GetActiveReloadType() { return EActiveReloadType::None; }
+
+inline bool IsReloadActive() { return false; }
+
+inline IReload* GetActiveReloadInterface() { return nullptr; }
+#endif
diff --git a/Engine/Source/Runtime/CoreUObject/Private/UObject/Class.cpp b/Engine/Source/Runtime/CoreUObject/Private/UObject/Class.cpp
--- a/Engine/Source/Runtime/CoreUObject/Private/UObject/Class.cpp
+++ b/Engine/Source/Runtime/CoreUObject/Private/UObject/Class.cpp
@@ -28,12 +28,12 @@
 #include "UObject/PropertyTag.h"
 #include "UObject/UnrealType.h"
 #include "UObject/UnrealTypePrivate.h"
+#include "UObject/Reload.h"
 #include "UObject/Stack.h"
 #include "Misc/PackageName.h"
 #include "UObject/ObjectResource.h"
 #include "UObject/LinkerSave.h"
 #include "UObject/Interface.h"
-#include "Misc/HotReloadInterface.h"
 #include "UObject/LinkerPlaceholderClass.h"
 #include "UObject/LinkerPlaceholderFunction.h"
 #include "UObject/StructScriptLoader.h"
@@ -2512,9 +2512,7 @@ void UScriptStruct::DeferCppStructOps(FName Target, ICppStructOps* InCppStructOp
 
 	if (UScriptStruct::ICppStructOps* ExistingOps = DeferredStructOps.FindRef(Target))
 	{
-#if WITH_HOT_RELOAD
-		if (!GIsHotReload) // in hot reload, we will just leak these...they may be in use.
-#endif
+		if (!IsReloadActive()) // in reload, we will just leak these...they may be in use.
 		{
 			check(ExistingOps != InCppStructOps); // if it was equal, then we would be re-adding a now stale pointer to the map
 			delete ExistingOps;
@@ -2549,11 +2547,7 @@ void UScriptStruct::PrepareCppStructOps()
 		
 #if !(UE_BUILD_SHIPPING || UE_BUILD_TEST)
 		// test that the constructor is initializing everything
-		if (!CppStructOps->HasZeroConstructor()
-#if WITH_HOT_RELOAD
-			&& !GIsHotReload // in hot reload, these produce bogus warnings
-#endif
-			)
+		if (!CppStructOps->HasZeroConstructor() && !IsReloadActive()) // when reloading, we get bogus warnings
 		{
 			int32 Size = CppStructOps->GetSize();
 			uint8* TestData00 = (uint8*)FMemory::Malloc(Size);
@@ -5052,12 +5046,12 @@ void UClass::SetSparseClassDataStruct(UScriptStruct* InSparseClassDataStruct)
 	{
 		SparseClassDataStruct = InSparseClassDataStruct;
 
-		// the old type and new type may not match when we do a hot reload so get rid of the old data
+		// the old type and new type may not match when we do a reload so get rid of the old data
 		CleanupSparseClassData();
 	}
 }
 
-#if WITH_HOT_RELOAD
+#if WITH_RELOAD
 
 bool UClass::HotReloadPrivateStaticClass(
 	uint32			InSize,
@@ -5129,7 +5123,7 @@ bool UClass::HotReloadPrivateStaticClass(
 	}
 	else
 	{
-		UE_LOG(LogClass, Warning, TEXT("Hot Reload:  Was not expecting temporary object '%s' for class '%s' to become rooted during construction.  This object cannot be marked pending kill." ), *TempObjectForVTable->GetFName().ToString(), *this->GetName() );
+		UE_LOG(LogClass, Warning, TEXT("Reload:  Was not expecting temporary object '%s' for class '%s' to become rooted during construction.  This object cannot be marked pending kill." ), *TempObjectForVTable->GetFName().ToString(), *this->GetName() );
 	}
 
 	ClassWithin = TClass_WithinClass_StaticClass;
@@ -5171,13 +5165,6 @@ bool UClass::HotReloadPrivateStaticClass(
 
 bool UClass::ReplaceNativeFunction(FName InFName, FNativeFuncPtr InPointer, bool bAddToFunctionRemapTable)
 {
-	IHotReloadInterface* HotReloadSupport = nullptr;
-
-	if(bAddToFunctionRemapTable)
-	{
-		HotReloadSupport = &FModuleManager::LoadModuleChecked<IHotReloadInterface>("HotReload");
-	}
-
 	// Find the function in the class's native function lookup table.
 	for (int32 FunctionIndex = 0; FunctionIndex < NativeFunctionLookupTable.Num(); ++FunctionIndex)
 	{
@@ -5186,7 +5173,7 @@ bool UClass::ReplaceNativeFunction(FName InFName, FNativeFuncPtr InPointer, bool
 		{
 			if (bAddToFunctionRemapTable)
 			{
-				HotReloadSupport->AddHotReloadFunctionRemap(InPointer, NativeFunctionLookup.Pointer);
+				ReloadNotifyFunctionRemap(InPointer, NativeFunctionLookup.Pointer);
 			}
 			NativeFunctionLookup.Pointer = InPointer;
 			return true;
@@ -5199,8 +5186,8 @@ bool UClass::ReplaceNativeFunction(FName InFName, FNativeFuncPtr InPointer, bool
 
 UClass* UClass::GetAuthoritativeClass()
 {
-#if WITH_HOT_RELOAD && WITH_ENGINE
-	if (GIsHotReload)
+#if WITH_RELOAD && WITH_ENGINE
+	if (IsReloadActive())
 	{
 		const TMap<UClass*, UClass*>& ReinstancedClasses = GetClassesToReinstanceForHotReload();
 		if (UClass* const* FoundMapping = ReinstancedClasses.Find(this))
@@ -5216,8 +5203,8 @@ UClass* UClass::GetAuthoritativeClass()
 void UClass::AddNativeFunction(const ANSICHAR* InName, FNativeFuncPtr InPointer)
 {
 	FName InFName(InName);
-#if WITH_HOT_RELOAD
-	if (GIsHotReload)
+#if WITH_RELOAD
+	if (IsReloadActive())
 	{
 		// Find the function in the class's native function lookup table.
 		if (ReplaceNativeFunction(InFName, InPointer, true))
@@ -5227,7 +5214,7 @@ void UClass::AddNativeFunction(const ANSICHAR* InName, FNativeFuncPtr InPointer)
 		else
 		{
 			// function was not found, so it's new
-			UE_LOG(LogClass, Log, TEXT("Function %s is new."), *InFName.ToString());
+			UE_LOG(LogClass, Log, TEXT("Function %s is new or belongs to a modified class."), *InFName.ToString());
 		}
 	}
 #endif
@@ -5237,8 +5224,8 @@ void UClass::AddNativeFunction(const ANSICHAR* InName, FNativeFuncPtr InPointer)
 void UClass::AddNativeFunction(const WIDECHAR* InName, FNativeFuncPtr InPointer)
 {
 	FName InFName(InName);
-#if WITH_HOT_RELOAD
-	if (GIsHotReload)
+#if WITH_RELOAD
+	if (IsReloadActive())
 	{
 		// Find the function in the class's native function lookup table.
 		if (ReplaceNativeFunction(InFName, InPointer, true))
@@ -5248,7 +5235,7 @@ void UClass::AddNativeFunction(const WIDECHAR* InName, FNativeFuncPtr InPointer)
 		else
 		{
 			// function was not found, so it's new
-			UE_LOG(LogClass, Log, TEXT("Function %s is new."), *InFName.ToString());
+			UE_LOG(LogClass, Log, TEXT("Function %s is new or belongs to a modified."), *InFName.ToString());
 		}
 	}
 #endif
@@ -5494,8 +5481,8 @@ void GetPrivateStaticClassBody(
 	UDynamicClass::DynamicClassInitializerType InDynamicClassInitializerFn /*= nullptr*/
 	)
 {
-#if WITH_HOT_RELOAD
-	if (GIsHotReload)
+#if WITH_RELOAD
+	if (IsReloadActive() && GetActiveReloadType() != EActiveReloadType::Reinstancing)
 	{
 		check(!bIsDynamic);
 		UPackage* Package = FindPackage(NULL, PackageName);
@@ -5523,12 +5510,12 @@ void GetPrivateStaticClassBody(
 			}
 			else
 			{
-				UE_LOG(LogClass, Log, TEXT("Could not find existing class %s in package %s for HotReload, assuming new class"), Name, PackageName);
+				UE_LOG(LogClass, Log, TEXT("Could not find existing class %s in package %s for reload, assuming new or modified class"), Name, PackageName);
 			}
 		}
 		else
 		{
-			UE_LOG(LogClass, Log, TEXT("Could not find existing package %s for HotReload of class %s, assuming a new package."), PackageName, Name);
+			UE_LOG(LogClass, Log, TEXT("Could not find existing package %s for reload of class %s, assuming a new package."), PackageName, Name);
 		}
 	}
 #endif
@@ -5600,6 +5587,9 @@ UFunction::UFunction(const FObjectInitializer& ObjectInitializer, UFunction* InS
 , EventGraphFunction(nullptr)
 , EventGraphCallOffset(0)
 #endif
+#if WITH_LIVE_CODING
+, SingletonPtr(nullptr)
+#endif
 {
 }
 
@@ -5609,6 +5599,9 @@ UFunction::UFunction(UFunction* InSuperFunction, EFunctionFlags InFunctionFlags,
 	, RPCId(0)
 	, RPCResponseId(0)
 	, FirstPropertyToInit(NULL)
+#if WITH_LIVE_CODING
+	, SingletonPtr(nullptr)
+#endif
 {
 }
 
@@ -6333,4 +6326,4 @@ IMPLEMENT_CORE_INTRINSIC_CLASS(UDynamicClass, UClass,
 	#endif
 #endif
 
-#include "UObject/DefineUPropertyMacros.h"
+#include "UObject/DefineUPropertyMacros.h"
\ No newline at end of file
diff --git a/Engine/Source/Runtime/CoreUObject/Private/UObject/CoreNative.cpp b/Engine/Source/Runtime/CoreUObject/Private/UObject/CoreNative.cpp
--- a/Engine/Source/Runtime/CoreUObject/Private/UObject/CoreNative.cpp
+++ b/Engine/Source/Runtime/CoreUObject/Private/UObject/CoreNative.cpp
@@ -63,11 +63,10 @@ IMPLEMENT_MODULE( FCoreUObjectModule, CoreUObject );
 
 // if we are not using compiled in natives, we still need this as a base class for intrinsics
 #if !USE_COMPILED_IN_NATIVES
-IMPLEMENT_CLASS(UObject, 0);
 COREUOBJECT_API class UClass* Z_Construct_UClass_UObject();
 UClass* Z_Construct_UClass_UObject()
 {
-	static UClass* OuterClass = nullptr;
+	static UClass* OuterClass = UObject::StaticRegistrationInfo().OuterSingleton;
 	if (!OuterClass)
 	{
 		OuterClass = UObject::StaticClass();
@@ -78,6 +77,7 @@ UClass* Z_Construct_UClass_UObject()
 	check(OuterClass->GetClass());
 	return OuterClass;
 }
+IMPLEMENT_CLASS(UObject, 0);
 #endif
 
 /*-----------------------------------------------------------------------------
diff --git a/Engine/Source/Runtime/CoreUObject/Private/UObject/DeferredRegistry.cpp b/Engine/Source/Runtime/CoreUObject/Private/UObject/DeferredRegistry.cpp
new file mode 100644
--- /dev/null
+++ b/Engine/Source/Runtime/CoreUObject/Private/UObject/DeferredRegistry.cpp
@@ -0,0 +1,67 @@
+// Copyright Epic Games, Inc. All Rights Reserved.
+
+#include "UObject/DeferredRegistry.h"
+#include "UObject/UObjectHash.h"
+#include "Templates/Casts.h"
+
+#if WITH_RELOAD
+
+void ReloadRenameObject(UScriptStruct& Object, const TCHAR* RenamePrefix)
+{
+	// Make sure the old struct is not used by anything
+	Object.ClearFlags(RF_Standalone | RF_Public);
+	Object.RemoveFromRoot();
+	const FName OldRename = MakeUniqueObjectName(GetTransientPackage(), Object.GetClass(), *FString::Printf(TEXT("%s_%s"), RenamePrefix, *Object.GetName()));
+	Object.Rename(*OldRename.ToString(), GetTransientPackage());
+}
+
+void ReloadRenameObject(UEnum& Object, const TCHAR* RenamePrefix)
+{
+	// Make sure the old struct is not used by anything
+	Object.ClearFlags(RF_Standalone | RF_Public);
+	Object.RemoveFromRoot();
+	const FName OldRename = MakeUniqueObjectName(GetTransientPackage(), Object.GetClass(), *FString::Printf(TEXT("%s_%s"), RenamePrefix, *Object.GetName()));
+	Object.Rename(*OldRename.ToString(), GetTransientPackage());
+
+	Object.RemoveNamesFromMasterList();
+}
+
+void ReloadRenameObject(UClass& Class, const TCHAR* RenamePrefix)
+{
+	FString NameWithoutPrefix = UObjectBase::RemoveClassPrefix(*Class.GetName());
+
+	// Rename the old class and move it to transient package
+	Class.RemoveFromRoot();
+	Class.ClearFlags(RF_Standalone | RF_Public);
+	Class.GetDefaultObject()->RemoveFromRoot();
+	Class.GetDefaultObject()->ClearFlags(RF_Standalone | RF_Public);
+	const FName OldClassRename = MakeUniqueObjectName(GetTransientPackage(), Class.GetClass(), *FString::Printf(TEXT("%s_%s"), RenamePrefix, *NameWithoutPrefix));
+	Class.Rename(*OldClassRename.ToString(), GetTransientPackage());
+	Class.SetFlags(RF_Transient);
+	Class.AddToRoot();
+
+	// Make sure enums de-register their names BEFORE we create the new class, otherwise there will be name conflicts
+	TArray<UObject*> ClassSubobjects;
+	GetObjectsWithOuter(&Class, ClassSubobjects);
+	for (auto ClassSubobject : ClassSubobjects)
+	{
+		if (auto Enum = dynamic_cast<UEnum*>(ClassSubobject))
+		{
+			Enum->RemoveNamesFromMasterList();
+		}
+	}
+
+	// Reset singletons for any child fields
+#if WITH_LIVE_CODING
+	for (UField* Field = Class.Children; Field != nullptr; Field = Field->Next)
+	{
+		UFunction* Function = Cast<UFunction>(Field);
+		if (Function != nullptr && Function->SingletonPtr != nullptr)
+		{
+			*Function->SingletonPtr = nullptr;
+		}
+	}
+#endif
+}
+
+#endif // WITH_RELOAD
diff --git a/Engine/Source/Runtime/CoreUObject/Private/UObject/DeferredRegistry.h b/Engine/Source/Runtime/CoreUObject/Private/UObject/DeferredRegistry.h
new file mode 100644
--- /dev/null
+++ b/Engine/Source/Runtime/CoreUObject/Private/UObject/DeferredRegistry.h
@@ -0,0 +1,337 @@
+// Copyright Epic Games, Inc. All Rights Reserved.
+
+#pragma once
+
+#include "Modules/ModuleManager.h"
+#include "UObject/Class.h"
+#include "UObject/Reload.h"
+#include "UObject/UObjectBase.h"
+
+#if WITH_RELOAD
+inline bool operator != (const FClassReloadVersionInfo& lhs, const FClassReloadVersionInfo& rhs)
+{
+	return lhs.Size != rhs.Size || lhs.Hash != rhs.Hash;
+}
+
+inline bool operator != (const FStructReloadVersionInfo& lhs, const FStructReloadVersionInfo& rhs)
+{
+	return lhs.Size != rhs.Size || lhs.Hash != rhs.Hash;
+}
+
+inline bool operator != (const FEnumReloadVersionInfo& lhs, const FEnumReloadVersionInfo& rhs)
+{
+	return lhs.Hash != rhs.Hash;
+}
+
+void ReloadRenameObject(UScriptStruct& Object, const TCHAR* RenamePrefix);
+void ReloadRenameObject(UEnum& Object, const TCHAR* RenamePrefix);
+void ReloadRenameObject(UClass& Object, const TCHAR* RenamePrefix);
+#endif
+
+PRAGMA_DISABLE_DEPRECATION_WARNINGS
+
+template <typename T>
+class TDeferredRegistry
+{
+public:
+	using TInfo = T;
+	using TType = typename T::TType;
+	using TVersion = typename T::TVersion;
+
+private:
+
+	using FPackageAndNameKey = TTuple<FName, FName>;
+
+public:
+
+	/**
+	* Maintains information about a pending registration
+	*/
+	struct FRegistrant
+	{
+		TType* (*OuterRegisterFn)();
+		TType* (*InnerRegisterFn)();
+		const TCHAR* PackageName;
+		TInfo* Info;
+#if WITH_RELOAD
+		TType* OldSingleton;
+#endif
+		FFieldCompiledInInfo* DeprecatedFieldInfo;
+
+#if WITH_RELOAD
+		bool bHasChanged;
+#endif
+	};
+
+public:
+
+	/// <summary>
+	/// Adds the given registration information for the given object.  Objects are either classes, structs, or enumerations.
+	/// </summary>
+	/// <param name="InOuterRegister">Returns a fully initialize instance of the object</param>
+	/// <param name="InInnerRegister">Returns an allocated but uninitialized instance of the object.  This is used only by UClass.</param>
+	/// <param name="InPackageName">Name of the package</param>
+	/// <param name="InName">Name of the object</param>
+	/// <param name="InInfo">Persistent information about the object</param>
+	/// <param name="InVersion">Version information for this incarnation of the object</param>
+	/// <param name="DeprevatedFieldInfo">Deprecated factory object for the inner register for UClass.  Either InInnerRegister or this value must be specified for UClass</param>
+	bool AddRegistration(TType* (*InOuterRegister)(), TType* (*InInnerRegister)(), const TCHAR* InPackageName, const TCHAR* InName, TInfo& InInfo, const TVersion& InVersion, FFieldCompiledInInfo* DeprecatedFieldInfo)
+	{
+#if WITH_RELOAD
+		const FPackageAndNameKey Key = FPackageAndNameKey{ InPackageName, InName };
+		TInfo** ExistingInfo = InfoMap.Find(Key);
+
+		bool bHasChanged = !ExistingInfo || (*ExistingInfo)->ReloadVersionInfo != InVersion;
+		InInfo.ReloadVersionInfo = InVersion;
+		TType* OldSingleton = ExistingInfo ? (*ExistingInfo)->InnerSingleton : nullptr;
+
+		if (ExistingInfo)
+		{
+			if (!bHasChanged && *ExistingInfo != &InInfo)
+			{
+				CopyExistingSingletons(InInfo, **ExistingInfo);
+			}
+			*ExistingInfo = &InInfo;
+		}
+		else
+		{
+			InfoMap.Add(Key, &InInfo);
+		}
+		Registrations.Add(FRegistrant{ InOuterRegister, InInnerRegister, InPackageName, &InInfo, OldSingleton, DeprecatedFieldInfo, bHasChanged });
+		return ExistingInfo != nullptr;
+#else
+		Registrations.Add(FRegistrant{ InOuterRegister, InInnerRegister, InPackageName, &InInfo, DeprecatedFieldInfo });
+		return false;
+#endif
+	}
+
+	/**
+	* Find an existing registration for reload when the version information matches
+	*/
+#if WITH_RELOAD
+	TType* FindMatchingObject(const TCHAR* InPackageName, const TCHAR* InName, const TVersion& InVersion)
+	{
+		const FPackageAndNameKey Key = FPackageAndNameKey{ InPackageName, InName };
+		TInfo** ExistingInfo = InfoMap.Find(Key);
+		if (ExistingInfo == nullptr)
+		{
+			return nullptr;
+		}
+		if ((*ExistingInfo)->ReloadVersionInfo != InVersion)
+		{
+			return nullptr;
+		}
+		return (*ExistingInfo)->InnerSingleton;
+	}
+#endif
+
+	/**
+	* Return the collection of registrations
+	*/
+	TArray<FRegistrant>& GetRegistrations()
+	{
+		return Registrations;
+	}
+
+	/**
+	* Return true if we have registrations
+	*/
+	bool HasPendingRegistrations()
+	{
+		return ProcessedRegistrations < Registrations.Num();
+	}
+
+	/**
+	* Create all the packages for the packages associated with the registrations
+	*/
+	void DoPendingPackageRegistrations()
+	{
+		for (int32 Index = ProcessedRegistrations, Num = Registrations.Num(); Index < Num; ++Index)
+		{
+			CreatePackage(Registrations[Index].PackageName);
+		}
+	}
+
+	/**
+	* Invoke the inner registration method for all the registrations
+	*/
+	void DoPendingInnerRegistrations(bool UpdateCounter)
+	{
+		int32 Num = Registrations.Num();
+		for (int32 Index = ProcessedRegistrations; Index < Num; ++Index)
+		{
+			InnerRegister(Registrations[Index]);
+		}
+
+		if (UpdateCounter)
+		{
+			ProcessedRegistrations = Num;
+		}
+	}
+
+	/**
+	* Invoke the outer registration method for all the registrations
+	*/
+	void DoPendingOuterRegistrations(bool UpdateCounter)
+	{
+		int32 Num = Registrations.Num();
+		for (int32 Index = ProcessedRegistrations; Index < Num; ++Index)
+		{
+			OuterRegister(Registrations[Index]);
+		}
+
+		if (UpdateCounter)
+		{
+			ProcessedRegistrations = Num;
+		}
+	}
+
+	/**
+	* Invoke the outer registration method for all the registrations and invoke the given function with the resulting object
+	*/
+	template <typename FuncType>
+	void DoPendingOuterRegistrations(bool UpdateCounter, FuncType&& InOnRegistration)
+	{
+		int32 Num = Registrations.Num();
+		for (int32 Index = ProcessedRegistrations; Index < Num; ++Index)
+		{
+			TType* Object = OuterRegister(Registrations[Index]);
+			InOnRegistration(*Object);
+		}
+		if (UpdateCounter)
+		{
+			ProcessedRegistrations = Num;
+		}
+	}
+
+	/**
+	* Invoke the register function for an inner registrant
+	*/
+	static TType* InnerRegister(const FRegistrant& Registrant)
+	{
+		return Registrant.InnerRegisterFn();
+	}
+
+	/**
+	* Invoke the register function for an outer registrant
+	*/
+	static TType* OuterRegister(const FRegistrant& Registrant)
+	{
+		return Registrant.OuterRegisterFn();
+	}
+
+	/**
+	* Rename and existing objects that have changed
+	*/
+	void RenameChangedObjects()
+	{
+#if WITH_RELOAD
+		IReload* Reload = GetActiveReloadInterface();
+		if (Reload != nullptr)
+		{
+			const TCHAR* Prefix = Reload->GetPrefix();
+			for (const FRegistrant& Registrant : Registrations)
+			{
+				if (Registrant.bHasChanged && Registrant.OldSingleton != nullptr)
+				{
+					UE_LOG(LogClass, Log, TEXT("%s Reload."), *Registrant.OldSingleton->GetName());
+					ReloadRenameObject(*Registrant.OldSingleton, Prefix);
+
+					// Reset the cached class construct info.
+					Registrant.Info->InnerSingleton = nullptr;
+					Registrant.Info->OuterSingleton = nullptr;
+				}
+			}
+		}
+#endif
+	}
+
+	/**
+	* Send notifications to the reload system
+	*/
+#if WITH_RELOAD
+	void NotifyReload(IReload& Reload)
+	{
+		for (const FRegistrant& Registrant : Registrations)
+		{
+			TType* Old = Registrant.OldSingleton;
+			TType* New = Registrant.bHasChanged ? OuterRegister(Registrant) : Old;
+			Reload.NotifyChange(New, Old);
+		}
+	}
+#endif
+
+	/**
+	* Clear the pending registrations
+	*/
+	void EmptyRegistrations()
+	{
+		TArray<FRegistrant> Local(MoveTemp(Registrations));
+		ProcessedRegistrations = 0;
+	}
+
+	/**
+	* When registering a new info, when the contents haven't changed, we copy singleton pointers 
+	*/
+	void CopyExistingSingletons(TInfo& NewInfo, const TInfo& OldInfo)
+	{
+		NewInfo.InnerSingleton = OldInfo.InnerSingleton;
+		NewInfo.OuterSingleton = OldInfo.OuterSingleton;
+	}
+
+	/**
+	* For the given package and name, return a registration information structure that can be
+	* used for compatibility with older API.  Will be removed after deprecated API removed.
+	*/
+	TInfo& MakeDeprecatedInfo(const TCHAR* InPackageName, const TCHAR* InName)
+	{
+		TUniquePtr<TInfo>& Info = CompatInfoMap.FindOrAdd(FPackageAndNameKey{ InPackageName, InName });
+		if (!Info.IsValid())
+		{
+			Info.Reset(new TInfo);
+		}
+		return *Info.Get();
+	}
+
+	/**
+	* Return the registry singleton
+	*/
+	static TDeferredRegistry& Get()
+	{
+		static TDeferredRegistry Registry;
+		return Registry;
+	}
+
+private:
+	TArray<FRegistrant> Registrations;
+	TMap<FPackageAndNameKey, TUniquePtr<TInfo>> CompatInfoMap; // Can be removed after deprecated API removed
+	int32 ProcessedRegistrations = 0;
+
+#if WITH_RELOAD
+	TMap<FPackageAndNameKey, TInfo*> InfoMap;
+#endif
+};
+
+// Specialization for classes for the LLM_SCOPE
+template <>
+inline typename TDeferredRegistry<FClassRegistrationInfo>::TType* TDeferredRegistry<FClassRegistrationInfo>::InnerRegister(const TDeferredRegistry<FClassRegistrationInfo>::FRegistrant& Registrant)
+{
+	if (Registrant.DeprecatedFieldInfo)
+	{
+		return Registrant.DeprecatedFieldInfo->Register();
+	}
+	LLM_SCOPE(ELLMTag::UObject);
+	return Registrant.InnerRegisterFn();
+}
+
+// Specialization for class to not copy the existing singletons when reloading a class that hasn't changed
+template <>
+inline void TDeferredRegistry<FClassRegistrationInfo>::CopyExistingSingletons(FClassRegistrationInfo& NewInfo, const FClassRegistrationInfo& OldInfo)
+{
+}
+
+using FClassDeferredRegistry = TDeferredRegistry<FClassRegistrationInfo>;
+using FEnumDeferredRegistry = TDeferredRegistry<FEnumRegistrationInfo>;
+using FStructDeferredRegistry = TDeferredRegistry<FStructRegistrationInfo>;
+
+PRAGMA_ENABLE_DEPRECATION_WARNINGS
diff --git a/Engine/Source/Runtime/CoreUObject/Private/UObject/Field.cpp b/Engine/Source/Runtime/CoreUObject/Private/UObject/Field.cpp
--- a/Engine/Source/Runtime/CoreUObject/Private/UObject/Field.cpp
+++ b/Engine/Source/Runtime/CoreUObject/Private/UObject/Field.cpp
@@ -31,7 +31,6 @@ Field.cpp: Defines FField property system fundamentals
 #include "UObject/ObjectResource.h"
 #include "UObject/LinkerSave.h"
 #include "UObject/Interface.h"
-#include "Misc/HotReloadInterface.h"
 #include "UObject/LinkerPlaceholderClass.h"
 #include "UObject/LinkerPlaceholderFunction.h"
 #include "UObject/StructScriptLoader.h"
diff --git a/Engine/Source/Runtime/CoreUObject/Private/UObject/ScriptCore.cpp b/Engine/Source/Runtime/CoreUObject/Private/UObject/ScriptCore.cpp
--- a/Engine/Source/Runtime/CoreUObject/Private/UObject/ScriptCore.cpp
+++ b/Engine/Source/Runtime/CoreUObject/Private/UObject/ScriptCore.cpp
@@ -25,9 +25,9 @@
 #include "UObject/PropertyPortFlags.h"
 #include "UObject/UnrealType.h"
 #include "UObject/Stack.h"
+#include "UObject/Reload.h"
 #include "Blueprint/BlueprintSupport.h"
 #include "UObject/ScriptMacros.h"
-#include "Misc/HotReloadInterface.h"
 #include "UObject/UObjectThreadContext.h"
 #include "HAL/IConsoleManager.h"
 
@@ -711,17 +711,9 @@ COREUOBJECT_API uint8 GRegisterNative( int32 NativeBytecodeIndex, const FNativeF
 	{
 		if( NativeBytecodeIndex<0 || (uint32)NativeBytecodeIndex>UE_ARRAY_COUNT(GNatives) || GNatives[NativeBytecodeIndex]!=&UObject::execUndefined) 
 		{
-#if WITH_HOT_RELOAD
-			if (GIsHotReload)
+			if (!ReloadNotifyFunctionRemap(Func, GNatives[NativeBytecodeIndex]))
 			{
-				IHotReloadInterface& HotReloadSupport = FModuleManager::LoadModuleChecked<IHotReloadInterface>("HotReload");
-				CA_SUPPRESS(6385)				
-				HotReloadSupport.AddHotReloadFunctionRemap(Func, GNatives[NativeBytecodeIndex]);
-			}
-			else
-#endif
-			{
-			GNativeDuplicate = NativeBytecodeIndex;
+				GNativeDuplicate = NativeBytecodeIndex;
 			}
 		}
 		CA_SUPPRESS(6386)
@@ -748,11 +740,7 @@ COREUOBJECT_API uint8 GRegisterCast( int32 CastCode, const FNativeFuncPtr& Func
 
 	if (CastCode != INDEX_NONE)
 	{
-		if(  
-#if WITH_HOT_RELOAD
-			!GIsHotReload && 
-#endif
-			(CastCode<0 || (uint32)CastCode>UE_ARRAY_COUNT(GCasts) || GCasts[CastCode]!=&UObject::execUndefined) ) 
+		if(!IsReloadActive() && (CastCode<0 || (uint32)CastCode>UE_ARRAY_COUNT(GCasts) || GCasts[CastCode]!=&UObject::execUndefined) ) 
 		{
 			GCastDuplicate = CastCode;
 		}
diff --git a/Engine/Source/Runtime/CoreUObject/Private/UObject/UObjectBase.cpp b/Engine/Source/Runtime/CoreUObject/Private/UObject/UObjectBase.cpp
--- a/Engine/Source/Runtime/CoreUObject/Private/UObject/UObjectBase.cpp
+++ b/Engine/Source/Runtime/CoreUObject/Private/UObject/UObjectBase.cpp
@@ -14,11 +14,13 @@
 #include "UObject/UObjectAllocator.h"
 #include "UObject/UObjectHash.h"
 #include "UObject/Class.h"
+#include "UObject/DeferredRegistry.h"
 #include "UObject/UObjectIterator.h"
 #include "UObject/Package.h"
 #include "Templates/Casts.h"
 #include "UObject/GCObject.h"
 #include "UObject/LinkerLoad.h"
+#include "UObject/Reload.h"
 #include "Misc/CommandLine.h"
 #include "Interfaces/IPluginManager.h"
 #include "Serialization/LoadTimeTrace.h"
@@ -541,56 +543,48 @@ void UObjectForceRegistration(UObjectBase* Object, bool bCheckForModuleRelease)
 	}
 }
 
-/**
- * Struct containing the function pointer and package name of a UStruct to be registered with UObject system
- */
-struct FPendingStructRegistrant
-{	
-	class UScriptStruct *(*RegisterFn)();
-	const TCHAR* PackageName;
+// UScriptStruct deferred registration
 
-	FPendingStructRegistrant() {}
-	FPendingStructRegistrant(class UScriptStruct *(*Fn)(), const TCHAR* InPackageName)
-		: RegisterFn(Fn)
-		, PackageName(InPackageName)
-	{
-	}
-	FORCEINLINE bool operator==(const FPendingStructRegistrant& Other) const
-	{
-		return RegisterFn == Other.RegisterFn;
-	}
-};
+TMap<FName, UScriptStruct *(*)()>& GetDynamicStructMap()
+{
+	static TMap<FName, UScriptStruct *(*)()> DynamicStructMap;
+	return DynamicStructMap;
+}
 
-static TArray<FPendingStructRegistrant>& GetDeferredCompiledInStructRegistration()
+void RegisterCompiledInInfo(class UScriptStruct* (*InOuterRegister)(), const TCHAR* InPackageName, const TCHAR* InName, FStructRegistrationInfo& InInfo, const FStructReloadVersionInfo& InVersionInfo)
 {
-	static TArray<FPendingStructRegistrant> DeferredCompiledInRegistration;
-	return DeferredCompiledInRegistration;
+	check(InOuterRegister);
+	FStructDeferredRegistry::Get().AddRegistration(InOuterRegister, nullptr, InPackageName, InName, InInfo, InVersionInfo, nullptr);
+	NotifyRegistrationEvent(InPackageName, InName, ENotifyRegistrationType::NRT_Struct, ENotifyRegistrationPhase::NRP_Added, (UObject * (*)())(InOuterRegister), false);
 }
 
-TMap<FName, UScriptStruct *(*)()>& GetDynamicStructMap()
+void RegisterCompiledInInfo(class UScriptStruct* (*InRegister)(), const TCHAR* PackageName, const TCHAR* ObjectName, const TCHAR* DynamicPackageName, const TCHAR* DynamicPathName)
 {
-	static TMap<FName, UScriptStruct *(*)()> DynamicStructMap;
-	return DynamicStructMap;
+	if (DynamicPackageName != nullptr)
+	{
+		GetConvertedDynamicPackageNameToTypeName().Add(FName(DynamicPackageName), FName(ObjectName));
+	}
+
+	GetDynamicStructMap().Add(DynamicPathName, InRegister);
+	NotifyRegistrationEvent(PackageName, ObjectName, ENotifyRegistrationType::NRT_Struct, ENotifyRegistrationPhase::NRP_Added, (UObject * (*)())(InRegister), true);
 }
 
-void UObjectCompiledInDeferStruct(class UScriptStruct *(*InRegister)(), const TCHAR* PackageName, const TCHAR* ObjectName, bool bDynamic, const TCHAR* DynamicPathName)
+
+void UObjectCompiledInDeferStruct(class UScriptStruct* (*InRegister)(), const TCHAR* PackageName, const TCHAR* ObjectName, bool bDynamic, const TCHAR* DynamicPathName)
 {
 	if (!bDynamic)
 	{
-		// we do reregister StaticStruct in hot reload
-		FPendingStructRegistrant Registrant(InRegister, PackageName);
-		checkSlow(!GetDeferredCompiledInStructRegistration().Contains(Registrant));
-		GetDeferredCompiledInStructRegistration().Add(Registrant);
+		FStructDeferredRegistry& Registry = FStructDeferredRegistry::Get();
+		FStructRegistrationInfo& Info = Registry.MakeDeprecatedInfo(PackageName, ObjectName);
+		Registry.AddRegistration(nullptr, InRegister, PackageName, ObjectName, Info, FStructReloadVersionInfo{}, nullptr);
 	}
 	else
 	{
-		GetDynamicStructMap().Add(DynamicPathName, InRegister);
+		RegisterCompiledInInfo(InRegister, PackageName, ObjectName, nullptr, DynamicPathName);
 	}
-	NotifyRegistrationEvent(PackageName, ObjectName, ENotifyRegistrationType::NRT_Struct, ENotifyRegistrationPhase::NRP_Added, (UObject *(*)())(InRegister), bDynamic);
-
 }
 
-class UScriptStruct *GetStaticStruct(class UScriptStruct *(*InRegister)(), UObject* StructOuter, const TCHAR* StructName, SIZE_T Size, uint32 Crc)
+class UScriptStruct *GetStaticStruct(class UScriptStruct *(*InRegister)(), UObject* StructOuter, const TCHAR* StructName)
 {
 	NotifyRegistrationEvent(*StructOuter->GetOutermost()->GetName(), StructName, ENotifyRegistrationType::NRT_Struct, ENotifyRegistrationPhase::NRP_Started);
 	UScriptStruct *Result = (*InRegister)();
@@ -598,53 +592,44 @@ class UScriptStruct *GetStaticStruct(class UScriptStruct *(*InRegister)(), UObje
 	return Result;
 }
 
-/**
- * Struct containing the function pointer and package name of a UEnum to be registered with UObject system
- */
-struct FPendingEnumRegistrant
-{
-	class UEnum *(*RegisterFn)();
-	const TCHAR* PackageName;
+// UEnum deferred registration
 
-	FPendingEnumRegistrant() {}
-	FPendingEnumRegistrant(class UEnum *(*Fn)(), const TCHAR* InPackageName)
-		: RegisterFn(Fn)
-		, PackageName(InPackageName)
-	{
-	}
-	FORCEINLINE bool operator==(const FPendingEnumRegistrant& Other) const
-	{
-		return RegisterFn == Other.RegisterFn;
-	}
-};
+TMap<FName, UEnum *(*)()>& GetDynamicEnumMap()
+{
+	static TMap<FName, UEnum *(*)()> DynamicEnumMap;
+	return DynamicEnumMap;
+}
 
-// Same thing as GetDeferredCompiledInStructRegistration but for UEnums declared in header files without UClasses.
-static TArray<FPendingEnumRegistrant>& GetDeferredCompiledInEnumRegistration()
+void RegisterCompiledInInfo(class UEnum* (*InOuterRegister)(), const TCHAR* InPackageName, const TCHAR* InName, FEnumRegistrationInfo& InInfo, const FEnumReloadVersionInfo& InVersionInfo)
 {
-	static TArray<FPendingEnumRegistrant> DeferredCompiledInRegistration;
-	return DeferredCompiledInRegistration;
+	check(InOuterRegister);
+	FEnumDeferredRegistry::Get().AddRegistration(InOuterRegister, nullptr, InPackageName, InName, InInfo, InVersionInfo, nullptr);
+	NotifyRegistrationEvent(InPackageName, InName, ENotifyRegistrationType::NRT_Enum, ENotifyRegistrationPhase::NRP_Added, (UObject * (*)())(InOuterRegister), false);
 }
 
-TMap<FName, UEnum *(*)()>& GetDynamicEnumMap()
+void RegisterCompiledInInfo(class UEnum* (*InRegister)(), const TCHAR* PackageName, const TCHAR* ObjectName, const TCHAR* DynamicPackageName, const TCHAR* DynamicPathName)
 {
-	static TMap<FName, UEnum *(*)()> DynamicEnumMap;
-	return DynamicEnumMap;
+	if (DynamicPackageName != nullptr)
+	{
+		GetConvertedDynamicPackageNameToTypeName().Add(FName(DynamicPackageName), FName(ObjectName));
+	}
+
+	GetDynamicEnumMap().Add(DynamicPathName, InRegister);
+	NotifyRegistrationEvent(PackageName, ObjectName, ENotifyRegistrationType::NRT_Enum, ENotifyRegistrationPhase::NRP_Added, (UObject * (*)())(InRegister), true);
 }
 
-void UObjectCompiledInDeferEnum(class UEnum *(*InRegister)(), const TCHAR* PackageName, const TCHAR* ObjectName, bool bDynamic, const TCHAR* DynamicPathName)
+void UObjectCompiledInDeferEnum(class UEnum* (*InRegister)(), const TCHAR* PackageName, const TCHAR* ObjectName, bool bDynamic, const TCHAR* DynamicPathName)
 {
 	if (!bDynamic)
 	{
-		// we do reregister StaticStruct in hot reload
-		FPendingEnumRegistrant Registrant(InRegister, PackageName);
-		checkSlow(!GetDeferredCompiledInEnumRegistration().Contains(Registrant));
-		GetDeferredCompiledInEnumRegistration().Add(Registrant);
+		FEnumDeferredRegistry& Registry = FEnumDeferredRegistry::Get();
+		FEnumRegistrationInfo& Info = Registry.MakeDeprecatedInfo(PackageName, ObjectName);
+		Registry.AddRegistration(InRegister, InRegister, PackageName, ObjectName, Info, FEnumReloadVersionInfo{}, nullptr);
 	}
 	else
 	{
-		GetDynamicEnumMap().Add(DynamicPathName, InRegister);
+		RegisterCompiledInInfo(InRegister, PackageName, ObjectName, nullptr, DynamicPathName);
 	}
-	NotifyRegistrationEvent(PackageName, ObjectName, ENotifyRegistrationType::NRT_Enum, ENotifyRegistrationPhase::NRP_Added, (UObject *(*)())(InRegister), bDynamic);
 }
 
 class UEnum *GetStaticEnum(class UEnum *(*InRegister)(), UObject* EnumOuter, const TCHAR* EnumName)
@@ -655,34 +640,13 @@ class UEnum *GetStaticEnum(class UEnum *(*InRegister)(), UObject* EnumOuter, con
 	return Result;
 }
 
-static TArray<class UClass *(*)()>& GetDeferredCompiledInRegistration()
-{
-	static TArray<class UClass *(*)()> DeferredCompiledInRegistration;
-	return DeferredCompiledInRegistration;
-}
-
-/** Classes loaded with a module, deferred until we register them all in one go */
-static TArray<FFieldCompiledInInfo*>& GetDeferredClassRegistration()
-{
-	static TArray<FFieldCompiledInInfo*> DeferredClassRegistration;
-	return DeferredClassRegistration;
-}
-
-#if WITH_HOT_RELOAD
-/** Map of deferred class registration info (including size and reflection info) */
-static TMap<FName, FFieldCompiledInInfo*>& GetDeferRegisterClassMap()
-{
-	static TMap<FName, FFieldCompiledInInfo*> DeferRegisterClassMap;
-	return DeferRegisterClassMap;
-}
+// UClass deferred registration
 
-/** Classes that changed during hot-reload and need to be re-instanced */
-static TArray<FFieldCompiledInInfo*>& GetHotReloadClasses()
+TMap<FName, FDynamicClassStaticData>& GetDynamicClassMap()
 {
-	static TArray<FFieldCompiledInInfo*> HotReloadClasses;
-	return HotReloadClasses;
+	static TMap<FName, FDynamicClassStaticData> DynamicClassMap;
+	return DynamicClassMap;
 }
-#endif
 
 /** Removes prefix from the native class name */
 FString UObjectBase::RemoveClassPrefix(const TCHAR* ClassName)
@@ -697,162 +661,151 @@ FString UObjectBase::RemoveClassPrefix(const TCHAR* ClassName)
 	return NameWithoutPrefix;
 }
 
-void UClassCompiledInDefer(FFieldCompiledInInfo* ClassInfo, const TCHAR* Name, SIZE_T ClassSize, uint32 Crc)
+void RegisterCompiledInInfo(class UClass* (*InOuterRegister)(), class UClass* (*InInnerRegister)(), const TCHAR* InPackageName, const TCHAR* InName, FClassRegistrationInfo& InInfo, const FClassReloadVersionInfo& InVersionInfo)
 {
-	const FName CPPClassName = Name;
-#if WITH_HOT_RELOAD
-	// Check for existing classes
-	TMap<FName, FFieldCompiledInInfo*>& DeferMap = GetDeferRegisterClassMap();
-	FFieldCompiledInInfo** ExistingClassInfo = DeferMap.Find(CPPClassName);
-	ClassInfo->bHasChanged = !ExistingClassInfo || (*ExistingClassInfo)->Size != ClassInfo->Size || (*ExistingClassInfo)->Crc != ClassInfo->Crc;
-	if (ExistingClassInfo)
-	{
-		// Class exists, this can only happen during hot-reload
-		checkf(GIsHotReload, TEXT("Trying to recreate class '%s' outside of hot reload!"), *CPPClassName.ToString());
-
-		// Get the native name
-		FString NameWithoutPrefix = UObjectBase::RemoveClassPrefix(Name);
-		UClass* ExistingClass = FindObjectChecked<UClass>(ANY_PACKAGE, *NameWithoutPrefix);
-
-		if (ClassInfo->bHasChanged)
-		{
-			// Rename the old class and move it to transient package
-			ExistingClass->RemoveFromRoot();
-			ExistingClass->ClearFlags(RF_Standalone | RF_Public);
-			ExistingClass->GetDefaultObject()->RemoveFromRoot();
-			ExistingClass->GetDefaultObject()->ClearFlags(RF_Standalone | RF_Public);
-			const FName OldClassRename = MakeUniqueObjectName(GetTransientPackage(), ExistingClass->GetClass(), *FString::Printf(TEXT("HOTRELOADED_%s"), *NameWithoutPrefix));
-			ExistingClass->Rename(*OldClassRename.ToString(), GetTransientPackage());
-			ExistingClass->SetFlags(RF_Transient);
-			ExistingClass->AddToRoot();
-
-			// Make sure enums de-register their names BEFORE we create the new class, otherwise there will be name conflicts
-			TArray<UObject*> ClassSubobjects;
-			GetObjectsWithOuter(ExistingClass, ClassSubobjects);
-			for (auto ClassSubobject : ClassSubobjects)
-			{
-				if (auto Enum = dynamic_cast<UEnum*>(ClassSubobject))
-				{
-					Enum->RemoveNamesFromMasterList();
-				}
-			}
-		}
-		ClassInfo->OldClass = ExistingClass;
-		GetHotReloadClasses().Add(ClassInfo);
+	check(InOuterRegister);
+	check(InInnerRegister);
+	bool bExisting = FClassDeferredRegistry::Get().AddRegistration(InOuterRegister, InInnerRegister, InPackageName, InName, InInfo, InVersionInfo, nullptr);
+#if WITH_RELOAD
+	if (bExisting && !IsReloadActive())
+	{
+		// Class exists, this can only happen during hot-reload or live coding
+		UE_LOG(LogUObjectBase, Fatal, TEXT("Trying to recreate class '%s' outside of hot reload and live coding!"), InName);
+	}
+#endif
+	FString NoPrefix(UObjectBase::RemoveClassPrefix(InName));
+	NotifyRegistrationEvent(InPackageName, *NoPrefix, ENotifyRegistrationType::NRT_Class, ENotifyRegistrationPhase::NRP_Added, (UObject * (*)())(InOuterRegister), false);
+	NotifyRegistrationEvent(InPackageName, *(FString(DEFAULT_OBJECT_PREFIX) + NoPrefix), ENotifyRegistrationType::NRT_ClassCDO, ENotifyRegistrationPhase::NRP_Added, (UObject * (*)())(InOuterRegister), false);
+}
 
-		*ExistingClassInfo = ClassInfo;
+void RegisterCompiledInInfo(UClass* (*InOuterRegister)(), UClass* (*InInnerRegister)(), const TCHAR* PackageName, const TCHAR* ObjectName, const TCHAR* DynamicPackageName, const TCHAR* DynamicPathName, void (*InInitSearchableValues)(TMap<FName, FName>&))
+{
+	if (DynamicPackageName != nullptr)
+	{
+		GetConvertedDynamicPackageNameToTypeName().Add(FName(DynamicPackageName), FName(ObjectName));
 	}
-	else
+
+	FDynamicClassStaticData ClassFunctions;
+	ClassFunctions.ZConstructFn = InOuterRegister;
+	ClassFunctions.StaticClassFn = InInnerRegister;
+	if (InInitSearchableValues)
 	{
-		DeferMap.Add(CPPClassName, ClassInfo);
+		InInitSearchableValues(ClassFunctions.SelectedSearchableValues);
 	}
-#endif
-	// We will either create a new class or update the static class pointer of the existing one
-	GetDeferredClassRegistration().Add(ClassInfo);
+	GetDynamicClassMap().Add(FName(DynamicPathName), ClassFunctions);
+
+	FString OriginalPackageName = DynamicPathName;
+	check(OriginalPackageName.EndsWith(ObjectName));
+	OriginalPackageName.RemoveFromEnd(FString(ObjectName));
+	check(OriginalPackageName.EndsWith(TEXT(".")));
+	OriginalPackageName.RemoveFromEnd(FString(TEXT(".")));
+
+	NotifyRegistrationEvent(*OriginalPackageName, ObjectName, ENotifyRegistrationType::NRT_Class, ENotifyRegistrationPhase::NRP_Added, (UObject * (*)())(InOuterRegister), true);
+	NotifyRegistrationEvent(*OriginalPackageName, *(FString(DEFAULT_OBJECT_PREFIX) + ObjectName), ENotifyRegistrationType::NRT_ClassCDO, ENotifyRegistrationPhase::NRP_Added, (UObject * (*)())(InOuterRegister), true);
 }
 
-TMap<FName, FDynamicClassStaticData>& GetDynamicClassMap()
+PRAGMA_DISABLE_DEPRECATION_WARNINGS
+void UClassCompiledInDefer(FFieldCompiledInInfo* ClassInfo, const TCHAR* Name, SIZE_T ClassSize, uint32 Crc)
 {
-	static TMap<FName, FDynamicClassStaticData> DynamicClassMap;
-	return DynamicClassMap;
+	FClassDeferredRegistry& Registry = FClassDeferredRegistry::Get();
+	FClassRegistrationInfo& Info = Registry.MakeDeprecatedInfo(ClassInfo->ClassPackage(), Name);
+	Registry.AddRegistration(nullptr, nullptr, ClassInfo->ClassPackage(), Name, Info, CONSTUCT_RELOAD_VERSION_INFO(FClassReloadVersionInfo, ClassSize, Crc), ClassInfo);
 }
+PRAGMA_ENABLE_DEPRECATION_WARNINGS
 
 void UObjectCompiledInDefer(UClass *(*InRegister)(), UClass *(*InStaticClass)(), const TCHAR* Name, const TCHAR* PackageName, bool bDynamic, const TCHAR* DynamicPathName, void (*InInitSearchableValues)(TMap<FName, FName>&))
 {
 	if (!bDynamic)
 	{
-#if WITH_HOT_RELOAD
 		// Either add all classes if not hot-reloading, or those which have changed
-		TMap<FName, FFieldCompiledInInfo*>& DeferMap = GetDeferRegisterClassMap();
-		if (!GIsHotReload || DeferMap.FindChecked(Name)->bHasChanged)
-#endif
-		{
-			FString NoPrefix(UObjectBase::RemoveClassPrefix(Name));
-			NotifyRegistrationEvent(PackageName, *NoPrefix, ENotifyRegistrationType::NRT_Class, ENotifyRegistrationPhase::NRP_Added, (UObject *(*)())(InRegister), false);
-			NotifyRegistrationEvent(PackageName, *(FString(DEFAULT_OBJECT_PREFIX) + NoPrefix), ENotifyRegistrationType::NRT_ClassCDO, ENotifyRegistrationPhase::NRP_Added, (UObject *(*)())(InRegister), false);
+		FClassDeferredRegistry& ClassRegistry = FClassDeferredRegistry::Get();
+		FClassRegistrationInfo& Info = ClassRegistry.MakeDeprecatedInfo(PackageName, Name);
 
-			TArray<UClass *(*)()>& DeferredCompiledInRegistration = GetDeferredCompiledInRegistration();
-			checkSlow(!DeferredCompiledInRegistration.Contains(InRegister));
-			DeferredCompiledInRegistration.Add(InRegister);
+		// This is slow, but is deprecated.
+		for (FClassDeferredRegistry::FRegistrant& Registrant : ClassRegistry.GetRegistrations())
+		{
+			if (Registrant.Info == &Info)
+			{
+#if WITH_RELOAD
+				if (Registrant.bHasChanged)
+#endif
+				{
+					// This will probably need to be moved
+					FString NoPrefix(UObjectBase::RemoveClassPrefix(Name));
+					NotifyRegistrationEvent(PackageName, *NoPrefix, ENotifyRegistrationType::NRT_Class, ENotifyRegistrationPhase::NRP_Added, (UObject * (*)())(InRegister), false);
+					NotifyRegistrationEvent(PackageName, *(FString(DEFAULT_OBJECT_PREFIX) + NoPrefix), ENotifyRegistrationType::NRT_ClassCDO, ENotifyRegistrationPhase::NRP_Added, (UObject * (*)())(InRegister), false);
+				}
+				Registrant.OuterRegisterFn = InRegister;
+				break;
+			}
 		}
 	}
 	else
 	{
-		FDynamicClassStaticData ClassFunctions;
-		ClassFunctions.ZConstructFn = InRegister;
-		ClassFunctions.StaticClassFn = InStaticClass;
-		if (InInitSearchableValues)
-		{
-			InInitSearchableValues(ClassFunctions.SelectedSearchableValues);
-		}
-		GetDynamicClassMap().Add(FName(DynamicPathName), ClassFunctions);
-
-		FString OriginalPackageName = DynamicPathName;
-		check(OriginalPackageName.EndsWith(Name));
-		OriginalPackageName.RemoveFromEnd(FString(Name));
-		check(OriginalPackageName.EndsWith(TEXT(".")));
-		OriginalPackageName.RemoveFromEnd(FString(TEXT(".")));
-
-		NotifyRegistrationEvent(*OriginalPackageName, Name, ENotifyRegistrationType::NRT_Class, ENotifyRegistrationPhase::NRP_Added, (UObject *(*)())(InRegister), true);
-		NotifyRegistrationEvent(*OriginalPackageName, *(FString(DEFAULT_OBJECT_PREFIX) + Name), ENotifyRegistrationType::NRT_ClassCDO, ENotifyRegistrationPhase::NRP_Added, (UObject *(*)())(InRegister), true);
+		RegisterCompiledInInfo(InRegister, InStaticClass, PackageName, Name, nullptr, DynamicPathName, InInitSearchableValues);
 	}
 }
 
 /** Register all loaded classes */
 void UClassRegisterAllCompiledInClasses()
 {
-#if WITH_HOT_RELOAD
+#if WITH_RELOAD
 	TArray<UClass*> AddedClasses;
 #endif
 	SCOPED_BOOT_TIMING("UClassRegisterAllCompiledInClasses");
 
-	TArray<FFieldCompiledInInfo*>& DeferredClassRegistration = GetDeferredClassRegistration();
-	for (const FFieldCompiledInInfo* Class : DeferredClassRegistration)
+	FClassDeferredRegistry& Registry = FClassDeferredRegistry::Get();
+
+	Registry.RenameChangedObjects();
+
+	for (const FClassDeferredRegistry::FRegistrant& Registrant : Registry.GetRegistrations())
 	{
-		UClass* RegisteredClass = Class->Register();
-#if WITH_HOT_RELOAD
-		if (GIsHotReload && Class->OldClass == nullptr)
+		UClass* RegisteredClass = FClassDeferredRegistry::InnerRegister(Registrant);
+#if WITH_RELOAD
+		if (IsReloadActive() && Registrant.OldSingleton == nullptr)
 		{
 			AddedClasses.Add(RegisteredClass);
 		}
 #endif
 	}
-	DeferredClassRegistration.Empty();
 
-#if WITH_HOT_RELOAD
+#if WITH_RELOAD
 	if (AddedClasses.Num() > 0)
 	{
+		FCoreUObjectDelegates::ReloadAddedClassesDelegate.Broadcast(AddedClasses);
+		PRAGMA_DISABLE_DEPRECATION_WARNINGS
 		FCoreUObjectDelegates::RegisterHotReloadAddedClassesDelegate.Broadcast(AddedClasses);
+		PRAGMA_ENABLE_DEPRECATION_WARNINGS
 	}
 #endif
 }
 
-#if WITH_HOT_RELOAD
-/** Re-instance all existing classes that have changed during hot-reload */
-void UClassReplaceHotReloadClasses()
+#if WITH_RELOAD
+/** Re-instance all existing classes that have changed during reload */
+void UClassReplaceReloadClasses()
 {
-	TArray<FFieldCompiledInInfo*>& HotReloadClasses = GetHotReloadClasses();
-
-	if (FCoreUObjectDelegates::RegisterClassForHotReloadReinstancingDelegate.IsBound())
+	for (const FClassDeferredRegistry::FRegistrant& Registrant : FClassDeferredRegistry::Get().GetRegistrations())
 	{
-		for (const FFieldCompiledInInfo* Class : HotReloadClasses)
+		if (Registrant.OldSingleton == nullptr)
 		{
-			check(Class->OldClass);
-
-			UClass* RegisteredClass = nullptr;
-			if (Class->bHasChanged)
-			{
-				RegisteredClass = Class->Register();
-			}
+			continue;
+		}
 
-			FCoreUObjectDelegates::RegisterClassForHotReloadReinstancingDelegate.Broadcast(Class->OldClass, RegisteredClass, Class->bHasChanged ? EHotReloadedClassFlags::Changed : EHotReloadedClassFlags::None);
+		UClass* RegisteredClass = nullptr;
+		if (Registrant.bHasChanged)
+		{
+			RegisteredClass = FClassDeferredRegistry::InnerRegister(Registrant);
 		}
+
+		PRAGMA_DISABLE_DEPRECATION_WARNINGS
+		FCoreUObjectDelegates::RegisterClassForHotReloadReinstancingDelegate.Broadcast(Registrant.OldSingleton, RegisteredClass, Registrant.bHasChanged ? EHotReloadedClassFlags::Changed : EHotReloadedClassFlags::None);
+		PRAGMA_ENABLE_DEPRECATION_WARNINGS
 	}
 
+	PRAGMA_DISABLE_DEPRECATION_WARNINGS
 	FCoreUObjectDelegates::ReinstanceHotReloadedClassesDelegate.Broadcast();
-	HotReloadClasses.Empty();
+	PRAGMA_ENABLE_DEPRECATION_WARNINGS
 }
-
 #endif
 
 /**
@@ -864,33 +817,30 @@ static void UObjectLoadAllCompiledInDefaultProperties()
 
 	static FName LongEnginePackageName(TEXT("/Script/Engine"));
 
-	TArray<UClass *(*)()>& DeferredCompiledInRegistration = GetDeferredCompiledInRegistration();
+	FClassDeferredRegistry& Registry = FClassDeferredRegistry::Get();
 
-	const bool bHaveRegistrants = DeferredCompiledInRegistration.Num() != 0;
-	if( bHaveRegistrants )
+	if (Registry.HasPendingRegistrations())
 	{
 		SCOPED_BOOT_TIMING("UObjectLoadAllCompiledInDefaultProperties");
 		TArray<UClass*> NewClasses;
 		TArray<UClass*> NewClassesInCoreUObject;
 		TArray<UClass*> NewClassesInEngine;
-		TArray<UClass* (*)()> PendingRegistrants = MoveTemp(DeferredCompiledInRegistration);
-		for (UClass* (*Registrant)() : PendingRegistrants)
-		{
-			UClass* Class = Registrant();
-			UE_LOG(LogUObjectBootstrap, Verbose, TEXT("UObjectLoadAllCompiledInDefaultProperties After Registrant %s %s"), *Class->GetOutermost()->GetName(), *Class->GetName());
-			if (Class->GetOutermost()->GetFName() == GLongCoreUObjectPackageName)
+		Registry.DoPendingOuterRegistrations(true, [&NewClasses, &NewClassesInCoreUObject, &NewClassesInEngine](UClass& Class) -> void
 			{
-				NewClassesInCoreUObject.Add(Class);
-			}
-			else if (Class->GetOutermost()->GetFName() == LongEnginePackageName)
-			{
-				NewClassesInEngine.Add(Class);
-			}
-			else
-			{
-				NewClasses.Add(Class);
-			}
-		}
+				UE_LOG(LogUObjectBootstrap, Verbose, TEXT("UObjectLoadAllCompiledInDefaultProperties After Registrant %s %s"), *Class.GetOutermost()->GetName(), *Class.GetName());
+				if (Class.GetOutermost()->GetFName() == GLongCoreUObjectPackageName)
+				{
+					NewClassesInCoreUObject.Add(&Class);
+				}
+				else if (Class.GetOutermost()->GetFName() == LongEnginePackageName)
+				{
+					NewClassesInEngine.Add(&Class);
+				}
+				else
+				{
+					NewClasses.Add(&Class);
+				}
+			}); 
 		{
 			SCOPED_BOOT_TIMING("CoreUObject Classes");
 			for (UClass* Class : NewClassesInCoreUObject) // we do these first because we assume these never trigger loads
@@ -944,35 +894,18 @@ static void UObjectLoadAllCompiledInStructs()
 {
 	SCOPED_BOOT_TIMING("UObjectLoadAllCompiledInStructs");
 
-	TArray<FPendingEnumRegistrant> PendingEnumRegistrants = MoveTemp(GetDeferredCompiledInEnumRegistration());
-	TArray<FPendingStructRegistrant> PendingStructRegistrants = MoveTemp(GetDeferredCompiledInStructRegistration());
+	FEnumDeferredRegistry& EnumRegistry = FEnumDeferredRegistry::Get();
+	FStructDeferredRegistry& StructRegistry = FStructDeferredRegistry::Get();
 
 	{
 		SCOPED_BOOT_TIMING("UObjectLoadAllCompiledInStructs -  CreatePackages (could be optimized!)");
-		// Load Enums first
-		for (const FPendingEnumRegistrant& EnumRegistrant : PendingEnumRegistrants)
-		{
-			// Make sure the package exists in case it does not contain any UObjects
-			CreatePackage(EnumRegistrant.PackageName);
-		}
-		for (const FPendingStructRegistrant& StructRegistrant : PendingStructRegistrants)
-		{
-			// Make sure the package exists in case it does not contain any UObjects or UEnums
-			CreatePackage(StructRegistrant.PackageName);
-		}
+		EnumRegistry.DoPendingPackageRegistrations();
+		StructRegistry.DoPendingPackageRegistrations();
 	}
 
 	// Load Structs
-
-	for (const FPendingEnumRegistrant& EnumRegistrant : PendingEnumRegistrants)
-	{
-		EnumRegistrant.RegisterFn();
-	}
-
-	for (const FPendingStructRegistrant& StructRegistrant : PendingStructRegistrants)
-	{
-		StructRegistrant.RegisterFn();
-	}
+	EnumRegistry.DoPendingOuterRegistrations(true);
+	StructRegistry.DoPendingOuterRegistrations(true);
 }
 
 void ProcessNewlyLoadedUObjects(FName Package, bool bCanProcessNewlyLoadedObjects)
@@ -991,14 +924,20 @@ void ProcessNewlyLoadedUObjects(FName Package, bool bCanProcessNewlyLoadedObject
 	LLM_SCOPE(ELLMTag::UObject);
 	DECLARE_SCOPE_CYCLE_COUNTER(TEXT("ProcessNewlyLoadedUObjects"), STAT_ProcessNewlyLoadedUObjects, STATGROUP_ObjectVerbose);
 
-	UClassRegisterAllCompiledInClasses();
+	FClassDeferredRegistry& ClassRegistry = FClassDeferredRegistry::Get();
+	FStructDeferredRegistry& StructRegistry = FStructDeferredRegistry::Get();
+	FEnumDeferredRegistry& EnumRegistry = FEnumDeferredRegistry::Get();
 
-	const TArray<UClass* (*)()>& DeferredCompiledInRegistration = GetDeferredCompiledInRegistration();
-	const TArray<FPendingStructRegistrant>& DeferredCompiledInStructRegistration = GetDeferredCompiledInStructRegistration();
-	const TArray<FPendingEnumRegistrant>& DeferredCompiledInEnumRegistration = GetDeferredCompiledInEnumRegistration();
+	StructRegistry.RenameChangedObjects();
+	EnumRegistry.RenameChangedObjects();
+
+	UClassRegisterAllCompiledInClasses();
 
 	bool bNewUObjects = false;
-	while (GFirstPendingRegistrant || DeferredCompiledInRegistration.Num() || DeferredCompiledInStructRegistration.Num() || DeferredCompiledInEnumRegistration.Num())
+	while (GFirstPendingRegistrant ||
+		ClassRegistry.HasPendingRegistrations() ||
+		StructRegistry.HasPendingRegistrations() ||
+		EnumRegistry.HasPendingRegistrations())
 	{
 		bNewUObjects = true;
 		UObjectProcessRegistrants();
@@ -1008,10 +947,23 @@ void ProcessNewlyLoadedUObjects(FName Package, bool bCanProcessNewlyLoadedObject
 
 		UObjectLoadAllCompiledInDefaultProperties();
 	}
-#if WITH_HOT_RELOAD
-	UClassReplaceHotReloadClasses();
+
+#if WITH_RELOAD
+	IReload* Reload = GetActiveReloadInterface();
+	if (Reload != nullptr)
+	{
+		UClassReplaceReloadClasses(); // Legacy
+		EnumRegistry.NotifyReload(*Reload);
+		StructRegistry.NotifyReload(*Reload);
+		ClassRegistry.NotifyReload(*Reload);
+		Reload->Reinstance();
+	}
 #endif
 
+	EnumRegistry.EmptyRegistrations();
+	StructRegistry.EmptyRegistrations();
+	ClassRegistry.EmptyRegistrations();
+
 	if (bNewUObjects && !GIsInitialLoad)
 	{
 		UClass::AssembleReferenceTokenStreams();
@@ -1242,82 +1194,10 @@ const TCHAR* DebugFullName(UObject* Object)
 	}
 }
 
-#if WITH_HOT_RELOAD
-namespace
-{
-	struct FObjectCompiledInfo
-	{
-		/** Registered struct info (including size and reflection info) */
-		static TMap<TTuple<UObject*, FName>, FObjectCompiledInfo>& GetRegisteredInfo()
-		{
-			static TMap<TTuple<UObject*, FName>, FObjectCompiledInfo> StructOrEnumCompiledInfoMap;
-			return StructOrEnumCompiledInfoMap;
-		}
-
-		FObjectCompiledInfo(SIZE_T InClassSize, uint32 InCrc)
-			: Size(InClassSize)
-			, Crc (InCrc)
-		{
-		}
-
-		SIZE_T Size;
-		uint32 Crc;
-	};
-
-	template <typename TType>
-	TType* FindExistingObjectIfHotReload(UObject* Outer, const TCHAR* Name, SIZE_T Size, uint32 Crc)
-	{
-		TTuple<UObject*, FName> Key(Outer, Name);
-
-		bool bChanged = true;
-		if (FObjectCompiledInfo* Info = FObjectCompiledInfo::GetRegisteredInfo().Find(Key))
-		{
-			// Hot-reloaded struct
-			bChanged = Info->Size != Size || Info->Crc != Crc;
-
-			Info->Size = Size;
-			Info->Crc  = Crc;
-		}
-		else
-		{
-			// New struct
-			FObjectCompiledInfo::GetRegisteredInfo().Add(Key, FObjectCompiledInfo(Size, Crc));
-		}
-
-		if (!GIsHotReload)
-		{
-			return nullptr;
-		}
-
-		TType* Existing = FindObject<TType>(Outer, Name);
-		if (!Existing)
-		{
-			// New type added during hot-reload
-			UE_LOG(LogClass, Log, TEXT("Could not find existing type %s for HotReload. Assuming new"), Name);
-			return nullptr;
-		}
-
-		// Existing type, make sure we destroy the old one if it has changed
-		if (bChanged)
-		{
-			// Make sure the old struct is not used by anything
-			Existing->ClearFlags(RF_Standalone | RF_Public);
-			Existing->RemoveFromRoot();
-			const FName OldRename = MakeUniqueObjectName(GetTransientPackage(), Existing->GetClass(), *FString::Printf(TEXT("HOTRELOADED_%s"), Name));
-			Existing->Rename(*OldRename.ToString(), GetTransientPackage());
-			return nullptr;
-		}
-
-		UE_LOG(LogClass, Log, TEXT("%s HotReload."), Name);
-		return Existing;
-	}
-}
-#endif // WITH_HOT_RELOAD
-
 UScriptStruct* FindExistingStructIfHotReloadOrDynamic(UObject* Outer, const TCHAR* StructName, SIZE_T Size, uint32 Crc, bool bIsDynamic)
 {
-#if WITH_HOT_RELOAD
-	UScriptStruct* Result = FindExistingObjectIfHotReload<UScriptStruct>(Outer, StructName, Size, Crc);
+#if WITH_RELOAD
+	UScriptStruct* Result = FStructDeferredRegistry::Get().FindMatchingObject(*Outer->GetName(), StructName, CONSTUCT_RELOAD_VERSION_INFO(FStructReloadVersionInfo, Size, Crc));
 #else
 	UScriptStruct* Result = nullptr;
 #endif
@@ -1330,8 +1210,8 @@ UScriptStruct* FindExistingStructIfHotReloadOrDynamic(UObject* Outer, const TCHA
 
 UEnum* FindExistingEnumIfHotReloadOrDynamic(UObject* Outer, const TCHAR* EnumName, SIZE_T Size, uint32 Crc, bool bIsDynamic)
 {
-#if WITH_HOT_RELOAD
-	UEnum* Result = FindExistingObjectIfHotReload<UEnum>(Outer, EnumName, Size, Crc);
+#if WITH_RELOAD
+	UEnum* Result = FEnumDeferredRegistry::Get().FindMatchingObject(*Outer->GetName(), EnumName, CONSTUCT_RELOAD_VERSION_INFO(FEnumReloadVersionInfo, Crc));
 #else
 	UEnum* Result = nullptr;
 #endif
diff --git a/Engine/Source/Runtime/CoreUObject/Private/UObject/UObjectBaseUtility.cpp b/Engine/Source/Runtime/CoreUObject/Private/UObject/UObjectBaseUtility.cpp
--- a/Engine/Source/Runtime/CoreUObject/Private/UObject/UObjectBaseUtility.cpp
+++ b/Engine/Source/Runtime/CoreUObject/Private/UObject/UObjectBaseUtility.cpp
@@ -215,10 +215,7 @@ bool UObjectBaseUtility::MarkPackageDirty() const
 			// we explicitly disable the ability to dirty a package or map during load.  Commandlets can still
 			// set the dirty state on load.
 			if( IsRunningCommandlet() || 
-				(GIsEditor && !GIsEditorLoadingPackage && !GIsCookerLoadingPackage && !GIsPlayInEditorWorld && !IsInAsyncLoadingThread()
-#if WITH_HOT_RELOAD
-				&& !GIsHotReload
-#endif // WITH_HOT_RELOAD
+				(GIsEditor && !GIsEditorLoadingPackage && !GIsCookerLoadingPackage && !GIsPlayInEditorWorld && !IsInAsyncLoadingThread() && !IsReloadActive()
 #if WITH_EDITORONLY_DATA
 				&& !Package->bIsCookedForEditor // Cooked packages can't be modified nor marked as dirty
 #endif
diff --git a/Engine/Source/Runtime/CoreUObject/Private/UObject/UObjectGlobals.cpp b/Engine/Source/Runtime/CoreUObject/Private/UObject/UObjectGlobals.cpp
--- a/Engine/Source/Runtime/CoreUObject/Private/UObject/UObjectGlobals.cpp
+++ b/Engine/Source/Runtime/CoreUObject/Private/UObject/UObjectGlobals.cpp
@@ -119,9 +119,14 @@ namespace LoadPackageStats
 #endif
 
 /** CoreUObject delegates */
+PRAGMA_DISABLE_DEPRECATION_WARNINGS
 FCoreUObjectDelegates::FRegisterHotReloadAddedClassesDelegate FCoreUObjectDelegates::RegisterHotReloadAddedClassesDelegate;
 FCoreUObjectDelegates::FRegisterClassForHotReloadReinstancingDelegate FCoreUObjectDelegates::RegisterClassForHotReloadReinstancingDelegate;
 FCoreUObjectDelegates::FReinstanceHotReloadedClassesDelegate FCoreUObjectDelegates::ReinstanceHotReloadedClassesDelegate;
+PRAGMA_ENABLE_DEPRECATION_WARNINGS
+FCoreUObjectDelegates::FReloadReinstancingCompleteDelegate FCoreUObjectDelegates::ReloadReinstancingCompleteDelegate;
+FCoreUObjectDelegates::FReloadCompleteDelegate FCoreUObjectDelegates::ReloadCompleteDelegate;
+FCoreUObjectDelegates::FReloadAddedClassesDelegate FCoreUObjectDelegates::ReloadAddedClassesDelegate;
 FCoreUObjectDelegates::FCompiledInUObjectsRegisteredDelegate FCoreUObjectDelegates::CompiledInUObjectsRegisteredDelegate;
 FCoreUObjectDelegates::FIsPackageOKToSaveDelegate FCoreUObjectDelegates::IsPackageOKToSaveDelegate;
 FCoreUObjectDelegates::FOnPackageReloaded FCoreUObjectDelegates::OnPackageReloaded;
@@ -3385,12 +3390,9 @@ UObject* StaticConstructObject_Internal(const FStaticConstructObjectParameters&
 			!InTemplate || 
 			(InName != NAME_None && (Params.bAssumeTemplateIsArchetype || InTemplate == UObject::GetArchetypeFromRequiredInfo(InClass, InOuter, InName, InFlags)))
 			);
-	const bool bCanRecycleSubobjects = bIsNativeFromCDO && (!(InFlags & RF_DefaultSubObject) || !FUObjectThreadContext::Get().IsInConstructor)
-#if WITH_HOT_RELOAD
+
 	// Do not recycle subobjects when performing hot-reload as they may contain old property values.
-	&& !GIsHotReload
-#endif
-		;
+	const bool bCanRecycleSubobjects = bIsNativeFromCDO && (!(InFlags & RF_DefaultSubObject) || !FUObjectThreadContext::Get().IsInConstructor) && !IsReloadActive();
 
 	bool bRecycledSubobject = false;	
 	Result = StaticAllocateObject(InClass, InOuter, InName, InFlags, Params.InternalSetFlags, bCanRecycleSubobjects, &bRecycledSubobject, Params.ExternalPackage);
@@ -4696,7 +4698,7 @@ namespace UE4CodeGen_Private
 	}
 #endif
 
-	void ConstructUFunction(UFunction*& OutFunction, const FFunctionParams& Params)
+	void ConstructUFunction(UFunction*& OutFunction, const FFunctionParams& Params, UFunction** SingletonPtr)
 	{
 		UObject*   (*OuterFunc)() = Params.OuterFunc;
 		UFunction* (*SuperFunc)() = Params.SuperFunc;
@@ -4745,6 +4747,10 @@ namespace UE4CodeGen_Private
 		}
 		OutFunction = NewFunction;
 
+#if WITH_LIVE_CODING
+		NewFunction->SingletonPtr = SingletonPtr;
+#endif
+
 #if WITH_METADATA
 		AddMetaData(NewFunction, Params.MetaDataArray, Params.NumMetaData);
 #endif
diff --git a/Engine/Source/Runtime/CoreUObject/Public/Misc/HotReloadInterface.h b/Engine/Source/Runtime/CoreUObject/Public/Misc/HotReloadInterface.h
--- a/Engine/Source/Runtime/CoreUObject/Public/Misc/HotReloadInterface.h
+++ b/Engine/Source/Runtime/CoreUObject/Public/Misc/HotReloadInterface.h
@@ -82,6 +82,7 @@ public:
 	/**
 	 * Adds a function to re-map after hot-reload.
 	 */
+	UE_DEPRECATED(5.0, "AddHotReloadFunctionRemap has been deprecated, use ReloadNotifyFunctionRemap in Reload.h instead")
 	virtual void AddHotReloadFunctionRemap(FNativeFuncPtr NewFunctionPointer, FNativeFuncPtr OldFunctionPointer) = 0;
 
 	/**
@@ -109,6 +110,7 @@ public:
 	 * @param	bWasTriggeredAutomatically	True if the hot reload was invoked automatically by the hot reload system after detecting a changed DLL
 	 */
 	DECLARE_MULTICAST_DELEGATE_OneParam(FHotReloadEvent, bool /* bWasTriggeredAutomatically */ );
+	UE_DEPRECATED(5.0, "OnHotReload has been deprecated, use FCoreUObjectDelegates::ReloadCompleteDelegate.")
 	virtual FHotReloadEvent& OnHotReload() = 0;
 
 	/**
diff --git a/Engine/Source/Runtime/CoreUObject/Public/UObject/Class.h b/Engine/Source/Runtime/CoreUObject/Public/UObject/Class.h
--- a/Engine/Source/Runtime/CoreUObject/Public/UObject/Class.h
+++ b/Engine/Source/Runtime/CoreUObject/Public/UObject/Class.h
@@ -1826,6 +1826,11 @@ public:
 	int32 EventGraphCallOffset;
 #endif
 
+#if WITH_LIVE_CODING
+	/** Pointer to the cached singleton pointer to this instance */
+	UFunction** SingletonPtr;
+#endif
+
 private:
 	/** C++ function this is bound to */
 	FNativeFuncPtr Func;
@@ -2799,7 +2804,7 @@ public:
 		ClassVTableHelperCtorCallerType InClassVTableHelperCtorCaller,
 		ClassAddReferencedObjectsType InClassAddReferencedObjects);
 
-#if WITH_HOT_RELOAD
+#if WITH_RELOAD
 	/**
 	 * Called when a class is reloading from a DLL...updates various information in-place.
 	 * @param	InSize							sizeof the class
diff --git a/Engine/Source/Runtime/CoreUObject/Public/UObject/ObjectMacros.h b/Engine/Source/Runtime/CoreUObject/Public/UObject/ObjectMacros.h
--- a/Engine/Source/Runtime/CoreUObject/Public/UObject/ObjectMacros.h
+++ b/Engine/Source/Runtime/CoreUObject/Public/UObject/ObjectMacros.h
@@ -11,9 +11,8 @@
 #include "UObject/Script.h"
 
 class FObjectInitializer;
-struct FCompiledInDefer;
 struct FFrame;
-template <typename TClass> struct TClassCompiledInDefer;
+struct FClassReloadVersionInfo;
 
 /** Represents a serializable object pointer in blueprint bytecode. This is always 64-bits, even on 32-bit platforms. */
 typedef	uint64 ScriptPointerType;
@@ -1554,6 +1553,12 @@ namespace UM
 #define DECLARE_FARCHIVE_SERIALIZER( TClass, API ) virtual API void Serialize(FArchive& Ar) override;
 #define DECLARE_FSTRUCTUREDARCHIVE_SERIALIZER( TClass, API ) virtual API void Serialize(FStructuredArchive::FRecord Record) override;
 
+#if WITH_RELOAD
+#define CONSTUCT_RELOAD_VERSION_INFO(VersionInfo, ...) VersionInfo { __VA_ARGS__ }
+#else
+#define CONSTUCT_RELOAD_VERSION_INFO(VersionInfo, ...) VersionInfo()
+#endif
+
 /*-----------------------------------------------------------------------------
 	Class declaration macros.
 -----------------------------------------------------------------------------*/
@@ -1585,6 +1590,8 @@ public: \
 	{ \
 		return TStaticCastFlags; \
 	} \
+	/** Returns the static registration structure for the class */ \
+	static FClassRegistrationInfo& StaticRegistrationInfo(); \
 	/** For internal use only; use StaticConstructObject() to create new objects. */ \
 	inline void* operator new(const size_t InSize, EInternal InInternalOnly, UObject* InOuter = (UObject*)GetTransientPackage(), FName InName = NAME_None, EObjectFlags InSetFlags = RF_NoFlags) \
 	{ \
@@ -1637,16 +1644,16 @@ public: \
 		return nullptr; \
 	}
 
-#if WITH_HOT_RELOAD && !CHECK_PUREVIRTUALS
+#if WITH_RELOAD && !CHECK_PUREVIRTUALS
 	#define DEFINE_VTABLE_PTR_HELPER_CTOR_CALLER(TClass) \
 		static UObject* __VTableCtorCaller(FVTableHelper& Helper) \
 		{ \
 			return new (EC_InternalUseOnlyConstructor, (UObject*)GetTransientPackage(), NAME_None, RF_NeedLoad | RF_ClassDefaultObject | RF_TagGarbageTemp) TClass(Helper); \
 		}
-#else // WITH_HOT_RELOAD && !CHECK_PUREVIRTUALS
+#else // WITH_RELOAD && !CHECK_PUREVIRTUALS
 	#define DEFINE_VTABLE_PTR_HELPER_CTOR_CALLER(TClass) \
 		DEFINE_VTABLE_PTR_HELPER_CTOR_CALLER_DUMMY()
-#endif // WITH_HOT_RELOAD && !CHECK_PUREVIRTUALS
+#endif // WITH_RELOAD && !CHECK_PUREVIRTUALS
 
 #define DECLARE_CLASS_INTRINSIC_NO_CTOR(TClass,TSuperClass,TStaticFlags,TPackage) \
 	DECLARE_CLASS(TClass, TSuperClass, TStaticFlags | CLASS_Intrinsic, CASTCLASS_None, TPackage, NO_API) \
@@ -1727,10 +1734,16 @@ public: \
 
 // Register a class at startup time.
 #define IMPLEMENT_CLASS(TClass, TClassCrc) \
-	static TClassCompiledInDefer<TClass> AutoInitialize##TClass(TEXT(#TClass), sizeof(TClass), TClassCrc); \
+	FClassRegistrationInfo& TClass::StaticRegistrationInfo() \
+	{ \
+		static FClassRegistrationInfo info; \
+		return info; \
+	} \
+	/* Do not change the AutoInitialize_ without changing LC_SymbolPatterns */ \
+	static FRegisterCompiledInInfo AutoInitialize_##TClass(&Z_Construct_UClass_##TClass, TClass::StaticClass, TClass::StaticPackage(), TEXT(#TClass), TClass::StaticRegistrationInfo(), CONSTUCT_RELOAD_VERSION_INFO(FClassReloadVersionInfo, sizeof(TClass), TClassCrc)); \
 	UClass* TClass::GetPrivateStaticClass() \
 	{ \
-		static UClass* PrivateStaticClass = NULL; \
+		static UClass*& PrivateStaticClass = StaticRegistrationInfo().InnerSingleton; \
 		if (!PrivateStaticClass) \
 		{ \
 			/* this could be handled with templates, but we want it external to avoid code bloat */ \
@@ -1756,7 +1769,6 @@ public: \
 
 // Used for intrinsics, this sets up the boiler plate, plus an initialization singleton, which can create properties and GC tokens
 #define IMPLEMENT_INTRINSIC_CLASS(TClass, TRequiredAPI, TSuperClass, TSuperRequiredAPI, TPackage, InitCode) \
-	IMPLEMENT_CLASS(TClass, 0) \
 	TRequiredAPI UClass* Z_Construct_UClass_##TClass(); \
 	struct Z_Construct_UClass_##TClass##_Statics \
 	{ \
@@ -1774,7 +1786,7 @@ public: \
 	}; \
 	UClass* Z_Construct_UClass_##TClass() \
 	{ \
-		static UClass* Class = NULL; \
+		static UClass* Class = TClass::StaticRegistrationInfo().OuterSingleton; \
 		if (!Class) \
 		{ \
 			Class = Z_Construct_UClass_##TClass##_Statics::Construct();\
@@ -1782,7 +1794,7 @@ public: \
 		check(Class->GetClass()); \
 		return Class; \
 	} \
-	static FCompiledInDefer Z_CompiledInDefer_UClass_##TClass(Z_Construct_UClass_##TClass, &TClass::StaticClass, TEXT(TPackage), TEXT(#TClass), false);
+	IMPLEMENT_CLASS(TClass, 0)
 
 #define IMPLEMENT_CORE_INTRINSIC_CLASS(TClass, TSuperClass, InitCode) \
 	IMPLEMENT_INTRINSIC_CLASS(TClass, COREUOBJECT_API, TSuperClass, COREUOBJECT_API, "/Script/CoreUObject" ,InitCode)
diff --git a/Engine/Source/Runtime/CoreUObject/Public/UObject/Reload.h b/Engine/Source/Runtime/CoreUObject/Public/UObject/Reload.h
new file mode 100644
--- /dev/null
+++ b/Engine/Source/Runtime/CoreUObject/Public/UObject/Reload.h
@@ -0,0 +1,70 @@
+// Copyright Epic Games, Inc. All Rights Reserved.
+
+/*=============================================================================
+	Reload.h: Unreal reload support.
+=============================================================================*/
+
+#pragma once
+
+#include "CoreMinimal.h"
+#include "UObject/CoreNative.h"
+#include "Modules/ModuleManager.h"
+
+/**
+ * Systems that implement a reload capability implement this interface and register it with 
+ * the module manager while a reload is in progress.
+ */
+
+class COREUOBJECT_API IReload
+{
+public:
+
+	virtual ~IReload() = default;
+
+	/**
+	 * Returns the type of reload currently in progress.
+	 */
+	virtual EActiveReloadType GetType() const = 0;
+
+	/**
+	 * When classes, structures, and enumerations are renamed, the given prefix is applied.
+	 */
+	virtual const TCHAR* GetPrefix() const = 0;
+
+
+	/**
+	 * Invoke when a duplicate function has been detected.
+	 */
+	virtual void NotifyFunctionRemap(FNativeFuncPtr NewFunctionPointer, FNativeFuncPtr OldFunctionPointer) = 0;
+
+	/**
+	 * Invoke to register a new or changed class, enumeration, or structure for re-instancing.
+	 */
+	virtual void NotifyChange(UClass* New, UClass* Old) = 0;
+	virtual void NotifyChange(UEnum* New, UEnum* Old) = 0;
+	virtual void NotifyChange(UScriptStruct* New, UScriptStruct* Old) = 0;
+
+	/**
+	 * Perform the re-instancing 
+	 */
+	virtual void Reinstance() = 0;
+};
+
+/**
+ * Invoke when a duplicate function has been detected.  If reloading is in progress, it will be notified
+ * of the change.
+ *
+ * @return  Returns true if a reload operation is in progress.  Returns false if the duplication is an error.
+ */
+
+FORCEINLINE bool ReloadNotifyFunctionRemap(FNativeFuncPtr NewFunctionPointer, FNativeFuncPtr OldFunctionPointer)
+{
+#if WITH_RELOAD
+	if (IReload* Reload = GetActiveReloadInterface())
+	{
+		Reload->NotifyFunctionRemap(NewFunctionPointer, OldFunctionPointer);
+		return true;
+	}
+#endif
+	return false;
+}
\ No newline at end of file
diff --git a/Engine/Source/Runtime/CoreUObject/Public/UObject/UObjectBase.h b/Engine/Source/Runtime/CoreUObject/Public/UObject/UObjectBase.h
--- a/Engine/Source/Runtime/CoreUObject/Public/UObject/UObjectBase.h
+++ b/Engine/Source/Runtime/CoreUObject/Public/UObject/UObjectBase.h
@@ -250,10 +250,129 @@ COREUOBJECT_API bool UObjectInitialized();
  */
 COREUOBJECT_API void UObjectForceRegistration(UObjectBase* Object, bool bCheckForModuleRelease = true);
 
-/** 
+/**
+ * Structure that represents the registration information for a given class, structure, or enumeration
+ */
+template <typename T, typename V>
+struct TRegistrationInfo
+{
+	using TType = T;
+	using TVersion = V;
+
+	TType* InnerSingleton = nullptr;
+	TType* OuterSingleton = nullptr;
+	TVersion ReloadVersionInfo;
+};
+
+/**
+ * Helper class to perform registration of object information.  It blindly forwards a call to RegisterCompiledInInfo
+ */
+struct FRegisterCompiledInInfo
+{
+	template <typename ... Args>
+	FRegisterCompiledInInfo(Args&& ... args)
+	{
+		RegisterCompiledInInfo(std::forward<Args>(args)...);
+	}
+};
+
+/**
+ * Reload version information for classes
+ */
+struct FClassReloadVersionInfo
+{
+#if WITH_RELOAD
+	SIZE_T Size = 0;
+	uint32 Hash = 0;
+#endif
+};
+
+/**
+ * Registration information for classes
+ */
+using FClassRegistrationInfo = TRegistrationInfo<UClass, FClassReloadVersionInfo>;
+
+/**
+ * Adds a class registration and version information. The InInfo parameter must be static.
+ */
+COREUOBJECT_API void RegisterCompiledInInfo(class UClass* (*InOuterRegister)(), class UClass* (*InInnerRegister)(), const TCHAR* InPackageName, const TCHAR* InName, FClassRegistrationInfo& InInfo, const FClassReloadVersionInfo& InVersionInfo);
+
+/**
+ * Adds a dynamic class registration.
+ */
+COREUOBJECT_API void RegisterCompiledInInfo(class UClass* (*InOuterRegister)(), class UClass* (*InInnerRegister)(), const TCHAR* Name, const TCHAR* PackageName, const TCHAR* DynamicPackageName, const TCHAR* DynamicPathName, void (*InInitSearchableValues)(TMap<FName, FName>&));
+
+/**
+ * Reload version information for structures
+ */
+struct FStructReloadVersionInfo
+{
+#if WITH_RELOAD
+	SIZE_T Size = 0;
+	uint32 Hash = 0;
+#endif
+};
+
+/**
+ * Registration information for structures
+ */
+using FStructRegistrationInfo = TRegistrationInfo<UScriptStruct, FStructReloadVersionInfo>;
+
+/**
+ * Adds a struct registration and version information. The InInfo parameter must be static.
+ */
+COREUOBJECT_API void RegisterCompiledInInfo(class UScriptStruct* (*InOuterRegister)(), const TCHAR* InPackageName, const TCHAR* InName, FStructRegistrationInfo& InInfo, const FStructReloadVersionInfo& InVersionInfo);
+
+/**
+ * Adds a dynamic struct registration.
+ */
+COREUOBJECT_API void RegisterCompiledInInfo(class UScriptStruct* (*InRegister)(), const TCHAR* PackageName, const TCHAR* ObjectName, const TCHAR* DynamicPackageName, const TCHAR* DynamicPathName);
+
+/**
+ * Invoke the registration method wrapped in notifications.
+ */
+COREUOBJECT_API class UScriptStruct* GetStaticStruct(class UScriptStruct* (*InRegister)(), UObject* StructOuter, const TCHAR* StructName);
+
+UE_DEPRECATED(5.0, "GetStaticStruct with size and hash has been deprecated, use the version without the size and hash.")
+inline class UScriptStruct* GetStaticStruct(class UScriptStruct* (*InRegister)(), UObject* StructOuter, const TCHAR* StructName, SIZE_T Size, uint32 Hash)
+{
+	return GetStaticStruct(InRegister, StructOuter, StructName);
+}
+
+/**
+ * Reload version information for enumerations
+ */
+struct FEnumReloadVersionInfo
+{
+#if WITH_RELOAD
+	uint32 Hash = 0;
+#endif
+};
+
+/**
+ * Registration information for structures
+ */
+using FEnumRegistrationInfo = TRegistrationInfo<UEnum, FEnumReloadVersionInfo>;
+
+/**
+ * Adds a static structure registration and version information. The InInfo parameter must be static.
+ */
+COREUOBJECT_API void RegisterCompiledInInfo(class UEnum* (*InOuterRegister)(), const TCHAR* InPackageName, const TCHAR* InName, FEnumRegistrationInfo& InInfo, const FEnumReloadVersionInfo& InVersionInfo);
+
+/**
+ * Adds a dynamic struct registration.
+ */
+COREUOBJECT_API void RegisterCompiledInInfo(class UEnum* (*InRegister)(), const TCHAR* PackageName, const TCHAR* ObjectName, const TCHAR* DynamicPackageName, const TCHAR* DynamicPathName);
+
+/**
+ * Invoke the registration method wrapped in notifications.
+ */
+COREUOBJECT_API class UEnum* GetStaticEnum(class UEnum* (*InRegister)(), UObject* EnumOuter, const TCHAR* EnumName);
+
+/**
  * Base class for deferred native class registration
  */
-struct FFieldCompiledInInfo
+struct UE_DEPRECATED(5.0, "FFieldCompiledInInfo has been deprecated, use the RegistrationInfo structures below.") FFieldCompiledInInfo
 {
 	FFieldCompiledInInfo(SIZE_T InClassSize, uint32 InCrc)
 		: Size(InClassSize)
@@ -282,13 +401,16 @@ struct FFieldCompiledInInfo
 /**
 * Adds a class to deferred registration queue.
 */
+PRAGMA_DISABLE_DEPRECATION_WARNINGS
+UE_DEPRECATED(5.0, "UClassCompiledInDefer has been deprecated, use RegisterCompiledInInfo.")
 COREUOBJECT_API void UClassCompiledInDefer(FFieldCompiledInInfo* Class, const TCHAR* Name, SIZE_T ClassSize, uint32 Crc);
+PRAGMA_ENABLE_DEPRECATION_WARNINGS
 
 /**
  * Specialized version of the deferred class registration structure.
  */
 template <typename TClass>
-struct TClassCompiledInDefer : public FFieldCompiledInInfo
+struct UE_DEPRECATED(5.0, "TClassCompiledInDefer has been deprecated, use FRegisterCompiledInInfo.") TClassCompiledInDefer : public FFieldCompiledInInfo
 {
 	TClassCompiledInDefer(const TCHAR* InName, SIZE_T InClassSize, uint32 InCrc)
 	: FFieldCompiledInInfo(InClassSize, InCrc)
@@ -309,9 +431,10 @@ struct TClassCompiledInDefer : public FFieldCompiledInInfo
 /**
  * Stashes the singleton function that builds a compiled in class. Later, this is executed.
  */
+UE_DEPRECATED(5.0, "UObjectCompiledInDefer has been deprecated, use RegisterCompiledInInfo.")
 COREUOBJECT_API void UObjectCompiledInDefer(class UClass *(*InRegister)(), class UClass *(*InStaticClass)(), const TCHAR* Name, const TCHAR* PackageName, bool bDynamic, const TCHAR* DynamicPathName, void (*InInitSearchableValues)(TMap<FName, FName>&));
 
-struct FCompiledInDefer
+struct UE_DEPRECATED(5.0, "FCompiledInDefer has been deprecated, use FRegisterCompiledInInfo.") FCompiledInDefer
 {
 	FCompiledInDefer(class UClass *(*InRegister)(), class UClass *(*InStaticClass)(), const TCHAR* PackageName, const TCHAR* Name, bool bDynamic, const TCHAR* DynamicPackageName = nullptr, const TCHAR* DynamicPathName = nullptr, void (*InInitSearchableValues)(TMap<FName, FName>&) = nullptr)
 	{
@@ -319,16 +442,19 @@ struct FCompiledInDefer
 		{
 			GetConvertedDynamicPackageNameToTypeName().Add(FName(DynamicPackageName), FName(Name));
 		}
+		PRAGMA_DISABLE_DEPRECATION_WARNINGS
 		UObjectCompiledInDefer(InRegister, InStaticClass, Name, PackageName, bDynamic, DynamicPathName, InInitSearchableValues);
+		PRAGMA_ENABLE_DEPRECATION_WARNINGS
 	}
 };
 
 /**
  * Stashes the singleton function that builds a compiled in struct (StaticStruct). Later, this is executed.
  */
+UE_DEPRECATED(5.0, "UObjectCompiledInDeferStruct has been deprecated, use RegisterCompiledInInfo.")
 COREUOBJECT_API void UObjectCompiledInDeferStruct(class UScriptStruct *(*InRegister)(), const TCHAR* PackageName, const TCHAR* ObjectName, bool bDynamic, const TCHAR* DynamicPathName);
 
-struct FCompiledInDeferStruct
+struct UE_DEPRECATED(5.0, "FCompiledInDeferStruct has been deprecated, use FRegisterCompiledInInfo.") FCompiledInDeferStruct
 {
 	FCompiledInDeferStruct(class UScriptStruct *(*InRegister)(), const TCHAR* PackageName, const TCHAR* Name, bool bDynamic, const TCHAR* DynamicPackageName, const TCHAR* DynamicPathName)
 	{
@@ -336,21 +462,19 @@ struct FCompiledInDeferStruct
 		{
 			GetConvertedDynamicPackageNameToTypeName().Add(FName(DynamicPackageName), FName(Name));
 		}
+		PRAGMA_DISABLE_DEPRECATION_WARNINGS
 		UObjectCompiledInDeferStruct(InRegister, PackageName, Name, bDynamic, DynamicPathName);
+		PRAGMA_ENABLE_DEPRECATION_WARNINGS
 	}
 };
 
-/**
- * Either call the passed in singleton, or if this is hot reload, find the existing struct
- */
-COREUOBJECT_API class UScriptStruct *GetStaticStruct(class UScriptStruct *(*InRegister)(), UObject* StructOuter, const TCHAR* StructName, SIZE_T Size, uint32 Crc);
-
 /**
  * Stashes the singleton function that builds a compiled in enum. Later, this is executed.
  */
+UE_DEPRECATED(5.0, "UObjectCompiledInDeferEnum has been deprecated, use RegisterCompiledInInfo.")
 COREUOBJECT_API void UObjectCompiledInDeferEnum(class UEnum *(*InRegister)(), const TCHAR* PackageName, const TCHAR* ObjectName, bool bDynamic, const TCHAR* DynamicPathName);
 
-struct FCompiledInDeferEnum
+struct UE_DEPRECATED(5.0, "FCompiledInDeferEnum has been deprecated, use FRegisterCompiledInInfo.") FCompiledInDeferEnum
 {
 	FCompiledInDeferEnum(class UEnum *(*InRegister)(), const TCHAR* PackageName, const TCHAR* Name, bool bDynamic, const TCHAR* DynamicPackageName, const TCHAR* DynamicPathName)
 	{
@@ -358,15 +482,12 @@ struct FCompiledInDeferEnum
 		{
 			GetConvertedDynamicPackageNameToTypeName().Add(FName(DynamicPackageName), FName(Name));
 		}
+		PRAGMA_DISABLE_DEPRECATION_WARNINGS
 		UObjectCompiledInDeferEnum(InRegister, PackageName, Name, bDynamic, DynamicPathName);
+		PRAGMA_ENABLE_DEPRECATION_WARNINGS
 	}
 };
 
-/**
- * Either call the passed in singleton, or if this is hot reload, find the existing enum
- */
-COREUOBJECT_API class UEnum *GetStaticEnum(class UEnum *(*InRegister)(), UObject* EnumOuter, const TCHAR* EnumName);
-
 /** Called during HotReload to hook up an existing structure */
 COREUOBJECT_API class UScriptStruct* FindExistingStructIfHotReloadOrDynamic(UObject* Outer, const TCHAR* StructName, SIZE_T Size, uint32 Crc, bool bIsDynamic);
 
@@ -376,11 +497,6 @@ COREUOBJECT_API class UEnum* FindExistingEnumIfHotReloadOrDynamic(UObject* Outer
 /** Must be called after a module has been loaded that contains UObject classes */
 COREUOBJECT_API void ProcessNewlyLoadedUObjects(FName Package = NAME_None, bool bCanProcessNewlyLoadedObjects = true);
 
-#if WITH_HOT_RELOAD
-/** Map of duplicated CDOs for reinstancing during hot-reload purposes. */
-COREUOBJECT_API TMap<UObject*, UObject*>& GetDuplicatedCDOMap();
-#endif // WITH_HOT_RELOAD
-
 /**
  * Final phase of UObject initialization. all auto register objects are added to the main data structures.
  */
diff --git a/Engine/Source/Runtime/CoreUObject/Public/UObject/UObjectGlobals.h b/Engine/Source/Runtime/CoreUObject/Public/UObject/UObjectGlobals.h
--- a/Engine/Source/Runtime/CoreUObject/Public/UObject/UObjectGlobals.h
+++ b/Engine/Source/Runtime/CoreUObject/Public/UObject/UObjectGlobals.h
@@ -2314,6 +2314,13 @@ enum class EHotReloadedClassFlags
 
 ENUM_CLASS_FLAGS(EHotReloadedClassFlags)
 
+enum class EReloadCompleteReason
+{
+	None,
+	HotReloadAutomatic,
+	HotReloadManual,
+};
+
 /**
  * Global CoreUObject delegates
  */
@@ -2379,16 +2386,31 @@ struct COREUOBJECT_API FCoreUObjectDelegates
 	DECLARE_DELEGATE_RetVal_ThreeParams(bool, FIsPackageOKToSaveDelegate, UPackage*, const FString&, FOutputDevice*);
 	static FIsPackageOKToSaveDelegate IsPackageOKToSaveDelegate;
 
+	/** Delegate for reloaded classes that have been added. */
+	DECLARE_MULTICAST_DELEGATE_OneParam(FReloadAddedClassesDelegate, const TArray<UClass*>&);
+	static FReloadAddedClassesDelegate ReloadAddedClassesDelegate;
+
+	/** Delegate for reload re-instancing complete */
+	DECLARE_MULTICAST_DELEGATE(FReloadReinstancingCompleteDelegate);
+	static FReloadReinstancingCompleteDelegate ReloadReinstancingCompleteDelegate;
+
+	/** Delegate for reload re-instancing complete */
+	DECLARE_MULTICAST_DELEGATE_OneParam(FReloadCompleteDelegate, EReloadCompleteReason);
+	static FReloadCompleteDelegate ReloadCompleteDelegate;
+
 	/** Delegate for registering hot-reloaded classes that have been added  */
 	DECLARE_MULTICAST_DELEGATE_OneParam(FRegisterHotReloadAddedClassesDelegate, const TArray<UClass*>&);
+	UE_DEPRECATED(5.0, "RegisterHotReloadAddedClassesDelegate has been deprecated, use ReloadAddedClassesDelegate.")
 	static FRegisterHotReloadAddedClassesDelegate RegisterHotReloadAddedClassesDelegate;
 
 	/** Delegate for registering hot-reloaded classes that changed after hot-reload for reinstancing */
 	DECLARE_MULTICAST_DELEGATE_ThreeParams(FRegisterClassForHotReloadReinstancingDelegate, UClass*, UClass*, EHotReloadedClassFlags);
+	UE_DEPRECATED(5.0, "RegisterClassForHotReloadReinstancingDelegate has been deprecated, use FReload for class re-instancing.")
 	static FRegisterClassForHotReloadReinstancingDelegate RegisterClassForHotReloadReinstancingDelegate;
 
 	/** Delegate for reinstancing hot-reloaded classes */
 	DECLARE_MULTICAST_DELEGATE(FReinstanceHotReloadedClassesDelegate);
+	UE_DEPRECATED(5.0, "ReinstanceHotReloadedClassesDelegate has been deprecated, use FReload for class re-instancing or ReloadReinstancingCompleteDelegate for notification")
 	static FReinstanceHotReloadedClassesDelegate ReinstanceHotReloadedClassesDelegate;
 
 	/** Delegate for catching when UClasses/UStructs/UEnums would be available via FindObject<>(), but before their CDOs would be constructed. */
@@ -2993,7 +3015,7 @@ namespace UE4CodeGen_Private
 #endif
 	};
 
-	COREUOBJECT_API void ConstructUFunction(UFunction*& OutFunction, const FFunctionParams& Params);
+	COREUOBJECT_API void ConstructUFunction(UFunction*& OutFunction, const FFunctionParams& Params, UFunction** SingletonPtr = nullptr);
 	COREUOBJECT_API void ConstructUEnum(UEnum*& OutEnum, const FEnumParams& Params);
 	COREUOBJECT_API void ConstructUScriptStruct(UScriptStruct*& OutStruct, const FStructParams& Params);
 	COREUOBJECT_API void ConstructUPackage(UPackage*& OutPackage, const FPackageParams& Params);
diff --git a/Engine/Source/Runtime/GameplayTasks/Private/GameplayTaskResource.cpp b/Engine/Source/Runtime/GameplayTasks/Private/GameplayTaskResource.cpp
--- a/Engine/Source/Runtime/GameplayTasks/Private/GameplayTaskResource.cpp
+++ b/Engine/Source/Runtime/GameplayTasks/Private/GameplayTaskResource.cpp
@@ -8,12 +8,12 @@
 TArray<FString> UGameplayTaskResource::ResourceDescriptions;
 #endif // !(UE_BUILD_SHIPPING || UE_BUILD_TEST)
 
-#if WITH_HOT_RELOAD
+#if WITH_RELOAD
 namespace
 {
 	TMap<FName, int8> ClassNameToIDMap;
 }
-#endif // WITH_HOT_RELOAD
+#endif // WITH_RELOAD
 
 UGameplayTaskResource::UGameplayTaskResource(const FObjectInitializer& ObjectInitializer)
 	: Super(ObjectInitializer)
@@ -29,8 +29,8 @@ void UGameplayTaskResource::PostInitProperties()
 
 	if (HasAnyFlags(RF_ClassDefaultObject) && (GetClass()->HasAnyClassFlags(CLASS_Abstract) == false))
 	{
-#if WITH_HOT_RELOAD
-		if (GIsHotReload)
+#if WITH_RELOAD
+		if (IsReloadActive())
 		{
 			if ((bManuallySetID == false || ManualResourceID == INDEX_NONE))
 			{
@@ -43,7 +43,7 @@ void UGameplayTaskResource::PostInitProperties()
 			}
 		}
 		else
-#endif // WITH_HOT_RELOAD
+#endif // WITH_RELOAD
 		{
 			if (bManuallySetID == false || ManualResourceID == INDEX_NONE)
 			{
@@ -56,9 +56,9 @@ void UGameplayTaskResource::PostInitProperties()
 			ResourceDescriptions[DebugId] = GenerateDebugDescription();
 #endif // !(UE_BUILD_SHIPPING || UE_BUILD_TEST)
 
-#if WITH_HOT_RELOAD
+#if WITH_RELOAD
 			ClassNameToIDMap.Add(GetFName(), DebugId);
-#endif // WITH_HOT_RELOAD
+#endif // WITH_RELOAD
 		}
 	}
 }
diff --git a/Engine/Source/Runtime/NavigationSystem/Private/NavAreas/NavArea.cpp b/Engine/Source/Runtime/NavigationSystem/Private/NavAreas/NavArea.cpp
--- a/Engine/Source/Runtime/NavigationSystem/Private/NavAreas/NavArea.cpp
+++ b/Engine/Source/Runtime/NavigationSystem/Private/NavAreas/NavArea.cpp
@@ -18,11 +18,7 @@ UNavArea::UNavArea(const FObjectInitializer& ObjectInitializer)
 
 void UNavArea::FinishDestroy()
 {
-	if (HasAnyFlags(RF_ClassDefaultObject)
-#if WITH_HOT_RELOAD
-		&& !GIsHotReload
-#endif // WITH_HOT_RELOAD
-		)
+	if (HasAnyFlags(RF_ClassDefaultObject) && !IsReloadActive())
 	{
 		UNavigationSystemV1::RequestAreaUnregistering(GetClass());
 	}
@@ -46,9 +42,7 @@ void UNavArea::RegisterArea()
 {
 	if (HasAnyFlags(RF_ClassDefaultObject) && 
 		!HasAnyFlags(RF_NeedInitialization)  // Don't register BP Area that has still not finished loaded their properties, it was also try again to register later via UNavArea::PostLoad()
-#if WITH_HOT_RELOAD
-		&& !GIsHotReload
-#endif // WITH_HOT_RELOAD
+		&& !IsReloadActive()
 		)
 	{
 		UNavigationSystemV1::RequestAreaRegistering(GetClass());
@@ -114,11 +108,7 @@ void UNavArea::PostEditChangeProperty(FPropertyChangedEvent& PropertyChangedEven
 
 	Super::PostEditChangeProperty(PropertyChangedEvent);
 
-	if (HasAnyFlags(RF_ClassDefaultObject)
-#if WITH_HOT_RELOAD
-		&& !GIsHotReload
-#endif // WITH_HOT_RELOAD
-		)
+	if (HasAnyFlags(RF_ClassDefaultObject) && !IsReloadActive())
 	{
 		const FName PropertyName = (PropertyChangedEvent.Property != nullptr) ? PropertyChangedEvent.Property->GetFName() : NAME_None;
 		if (PropertyName == NAME_DefaultCost
diff --git a/Engine/Source/Runtime/NavigationSystem/Private/NavigationSystem.cpp b/Engine/Source/Runtime/NavigationSystem/Private/NavigationSystem.cpp
--- a/Engine/Source/Runtime/NavigationSystem/Private/NavigationSystem.cpp
+++ b/Engine/Source/Runtime/NavigationSystem/Private/NavigationSystem.cpp
@@ -35,10 +35,6 @@
 #include "Editor/LevelEditor/Public/LevelEditor.h"
 #endif
 
-#if WITH_HOT_RELOAD
-#include "Misc/HotReloadInterface.h"
-#endif
-
 #include "NavAreas/NavArea_Null.h"
 #include "NavAreas/NavArea_Obstacle.h"
 #include "NavAreas/NavArea_Default.h"
@@ -700,10 +696,7 @@ void UNavigationSystemV1::PostInitProperties()
 		FCoreUObjectDelegates::PostLoadMapWithWorld.AddUObject(this, &UNavigationSystemV1::OnPostLoadMap);
 		UNavigationSystemV1::NavigationDirtyEvent.AddUObject(this, &UNavigationSystemV1::OnNavigationDirtied);
 
-#if WITH_HOT_RELOAD
-		IHotReloadInterface& HotReloadSupport = FModuleManager::LoadModuleChecked<IHotReloadInterface>("HotReload");
-		HotReloadDelegateHandle = HotReloadSupport.OnHotReload().AddUObject(this, &UNavigationSystemV1::OnHotReload);
-#endif
+		ReloadCompleteDelegateHandle = FCoreUObjectDelegates::ReloadCompleteDelegate.AddUObject(this, &UNavigationSystemV1::OnReloadComplete);
 	}
 }
 
@@ -3943,8 +3936,7 @@ void UNavigationSystemV1::OnNavigationDirtied(const FBox& Bounds)
 	AddDirtyArea(Bounds, ENavigationDirtyFlag::All);
 }
 
-#if WITH_HOT_RELOAD
-void UNavigationSystemV1::OnHotReload(bool bWasTriggeredAutomatically)
+void UNavigationSystemV1::OnReloadComplete(EReloadCompleteReason Reason)
 {
 	if (RequiresNavOctree() && DefaultOctreeController.NavOctree.IsValid() == false)
 	{
@@ -3956,7 +3948,6 @@ void UNavigationSystemV1::OnHotReload(bool bWasTriggeredAutomatically)
 		}
 	}
 }
-#endif // WITH_HOT_RELOAD
 
 void UNavigationSystemV1::CleanUp(FNavigationSystem::ECleanupMode Mode)
 {
@@ -3980,12 +3971,7 @@ void UNavigationSystemV1::CleanUp(FNavigationSystem::ECleanupMode Mode)
 	FWorldDelegates::LevelRemovedFromWorld.RemoveAll(this);
 	FWorldDelegates::OnWorldBeginTearDown.RemoveAll(this);
 
-#if WITH_HOT_RELOAD
-	if (IHotReloadInterface* HotReloadSupport = FModuleManager::GetModulePtr<IHotReloadInterface>("HotReload"))
-	{
-		HotReloadSupport->OnHotReload().Remove(HotReloadDelegateHandle);
-	}
-#endif
+	FCoreUObjectDelegates::ReloadCompleteDelegate.Remove(ReloadCompleteDelegateHandle);
 
 	DestroyNavOctree();
 	
diff --git a/Engine/Source/Runtime/NavigationSystem/Public/NavigationSystem.h b/Engine/Source/Runtime/NavigationSystem/Public/NavigationSystem.h
--- a/Engine/Source/Runtime/NavigationSystem/Public/NavigationSystem.h
+++ b/Engine/Source/Runtime/NavigationSystem/Public/NavigationSystem.h
@@ -1059,12 +1059,10 @@ protected:
 	/** delegate handler called when navigation is dirtied*/
 	void OnNavigationDirtied(const FBox& Bounds);
 
-#if WITH_HOT_RELOAD
-	FDelegateHandle HotReloadDelegateHandle;
+	FDelegateHandle ReloadCompleteDelegateHandle;
 
-	/** called to notify NavigaitonSystem about finished hot reload */
-	virtual void OnHotReload(bool bWasTriggeredAutomatically);
-#endif // WITH_HOT_RELOAD
+	/** called to notify NavigaitonSystem about finished reload */
+	virtual void OnReloadComplete(EReloadCompleteReason Reason);
 
 	/** Registers given navigation data with this Navigation System.
 	 *	@return RegistrationSuccessful if registration was successful, other results mean it failed
