Edigrate cl 15621005:
Fix nullptr crash when caching material parameter data on attribute layers

diff --git a/Engine/Source/Editor/UnrealEd/Private/PreviewMaterial.cpp b/Engine/Source/Editor/UnrealEd/Private/PreviewMaterial.cpp
--- a/Engine/Source/Editor/UnrealEd/Private/PreviewMaterial.cpp
+++ b/Engine/Source/Editor/UnrealEd/Private/PreviewMaterial.cpp
@@ -267,49 +267,52 @@ FMaterialExpressionParameterDataCache CacheMaterialExpressionParameterData(const
 	// but caching all the data in one pass
 	auto CacheMaterialFunctionParameterData = [](UMaterialFunctionInterface* ParameterFunction, TMap<FName, FMaterialParamExpressionData>& ParamDatas)
 	{
-		TArray<UMaterialFunctionInterface*> Functions;
-		ParameterFunction->GetDependentFunctions(Functions);
-		Functions.AddUnique(ParameterFunction);
-
-		for (UMaterialFunctionInterface* Function : Functions)
+		if (ParameterFunction)
 		{
-			for (UMaterialExpression* FunctionExpression : *Function->GetFunctionExpressions())
+			TArray<UMaterialFunctionInterface*> Functions;
+			ParameterFunction->GetDependentFunctions(Functions);
+			Functions.AddUnique(ParameterFunction);
+
+			for (UMaterialFunctionInterface* Function : Functions)
 			{
-				if (const UMaterialExpressionParameter* Parameter = Cast<const UMaterialExpressionParameter>(FunctionExpression))
+				for (UMaterialExpression* FunctionExpression : *Function->GetFunctionExpressions())
 				{
-					FMaterialParamExpressionData ParamData;
-					ParamData.ParamType = UMaterialExpressionParameter::StaticClass();
+					if (const UMaterialExpressionParameter* Parameter = Cast<const UMaterialExpressionParameter>(FunctionExpression))
+					{
+						FMaterialParamExpressionData ParamData;
+						ParamData.ParamType = UMaterialExpressionParameter::StaticClass();
 
-					ParamData.Name = Parameter->ParameterName;
-					ParamData.SortPriority = Parameter->SortPriority;
-					ParamData.Group = Parameter->Group;
+						ParamData.Name = Parameter->ParameterName;
+						ParamData.SortPriority = Parameter->SortPriority;
+						ParamData.Group = Parameter->Group;
 
-					//ensure(!ParamDatas.Contains(ParamData.Name));
-					ParamDatas.Add(ParamData.Name, ParamData);
-				}
-				else if (const UMaterialExpressionTextureSampleParameter* TexParameter = Cast<const UMaterialExpressionTextureSampleParameter>(FunctionExpression))
-				{
-					FMaterialParamExpressionData ParamData;
-					ParamData.ParamType = UMaterialExpressionTextureSampleParameter::StaticClass();
+						//ensure(!ParamDatas.Contains(ParamData.Name));
+						ParamDatas.Add(ParamData.Name, ParamData);
+					}
+					else if (const UMaterialExpressionTextureSampleParameter* TexParameter = Cast<const UMaterialExpressionTextureSampleParameter>(FunctionExpression))
+					{
+						FMaterialParamExpressionData ParamData;
+						ParamData.ParamType = UMaterialExpressionTextureSampleParameter::StaticClass();
 
-					ParamData.Name = TexParameter->ParameterName;
-					ParamData.SortPriority = TexParameter->SortPriority;
-					ParamData.Group = TexParameter->Group;
+						ParamData.Name = TexParameter->ParameterName;
+						ParamData.SortPriority = TexParameter->SortPriority;
+						ParamData.Group = TexParameter->Group;
 
-					//ensure(!ParamDatas.Contains(ParamData.Name));
-					ParamDatas.Add(ParamData.Name, ParamData);
-				}
-				else if (const UMaterialExpressionFontSampleParameter* FontParameter = Cast<const UMaterialExpressionFontSampleParameter>(FunctionExpression))
-				{
-					FMaterialParamExpressionData ParamData;
-					ParamData.ParamType = UMaterialExpressionFontSampleParameter::StaticClass();
+						//ensure(!ParamDatas.Contains(ParamData.Name));
+						ParamDatas.Add(ParamData.Name, ParamData);
+					}
+					else if (const UMaterialExpressionFontSampleParameter* FontParameter = Cast<const UMaterialExpressionFontSampleParameter>(FunctionExpression))
+					{
+						FMaterialParamExpressionData ParamData;
+						ParamData.ParamType = UMaterialExpressionFontSampleParameter::StaticClass();
 
-					ParamData.Name = FontParameter->ParameterName;
-					ParamData.SortPriority = FontParameter->SortPriority;
-					ParamData.Group = FontParameter->Group;
+						ParamData.Name = FontParameter->ParameterName;
+						ParamData.SortPriority = FontParameter->SortPriority;
+						ParamData.Group = FontParameter->Group;
 
-					//ensure(!ParamDatas.Contains(ParamData.Name));
-					ParamDatas.Add(ParamData.Name, ParamData);
+						//ensure(!ParamDatas.Contains(ParamData.Name));
+						ParamDatas.Add(ParamData.Name, ParamData);
+					}
 				}
 			}
 		}
