Fix problem with copy/paste material function nodes
- Doesn't work to associate old/new pins via name, since function input pin names can change
- Instead go back to using indices

diff --git a/Engine/Source/Editor/UnrealEd/Private/MaterialGraphNode_Base.cpp b/Engine/Source/Editor/UnrealEd/Private/MaterialGraphNode_Base.cpp
--- a/Engine/Source/Editor/UnrealEd/Private/MaterialGraphNode_Base.cpp
+++ b/Engine/Source/Editor/UnrealEd/Private/MaterialGraphNode_Base.cpp
@@ -215,12 +215,15 @@ void UMaterialGraphNode_Base::ReconstructNode()
 	// Recreate the new pins
 	AllocateDefaultPins();
 
-	for (UEdGraphPin* OldPin : OldPins)
+	for (int32 PinIndex = 0; PinIndex < OldPins.Num(); ++PinIndex)
 	{
-		// Transfer data to new pin
-		UEdGraphPin* NewPin = FindPin(OldPin->PinName, OldPin->Direction);
-		if (NewPin)
+		UEdGraphPin* OldPin = OldPins[PinIndex];
+		if (PinIndex < Pins.Num())
 		{
+			// Transfer data to new pin
+			UEdGraphPin* NewPin = Pins[PinIndex];
+			ensure(OldPin->Direction == NewPin->Direction);
+			ensure(OldPin->PinType.PinCategory == NewPin->PinType.PinCategory);
 			NewPin->MovePersistentDataFromOldPin(*OldPin);
 		}
 
