Make CanPasteComponents and AdjustComponentDelta take a const component pointer

#rnx

diff --git a/Engine/Source/Editor/UnrealEd/Private/Kismet2/ComponentEditorUtils.cpp b/Engine/Source/Editor/UnrealEd/Private/Kismet2/ComponentEditorUtils.cpp
--- a/Engine/Source/Editor/UnrealEd/Private/Kismet2/ComponentEditorUtils.cpp
+++ b/Engine/Source/Editor/UnrealEd/Private/Kismet2/ComponentEditorUtils.cpp
@@ -43,7 +43,7 @@ struct FComponentObjectTextFactory : public FCustomizableTextObjectFactory
 	TMap<FName, UActorComponent*> NewObjectMap;
 
 	// Determine whether or not scene components in the new object set can be attached to the given scene root component
-	bool CanAttachComponentsTo(USceneComponent* InRootComponent)
+	bool CanAttachComponentsTo(const USceneComponent* InRootComponent)
 	{
 		check(InRootComponent);
 
@@ -52,8 +52,8 @@ struct FComponentObjectTextFactory : public FCustomizableTextObjectFactory
 		for (auto NewComponentIt = NewObjectMap.CreateConstIterator(); NewComponentIt && bCanAttachToRoot; ++NewComponentIt)
 		{
 			// If this is a scene component, and it does not already have a parent within the set
-			USceneComponent* SceneComponent = Cast<USceneComponent>(NewComponentIt->Value);
-			if (SceneComponent != NULL && !ParentMap.Contains(SceneComponent->GetFName()))
+			const USceneComponent* SceneComponent = Cast<USceneComponent>(NewComponentIt->Value);
+			if (SceneComponent && !ParentMap.Contains(SceneComponent->GetFName()))
 			{
 				// Determine if we are allowed to attach the scene component to the given root component
 				bCanAttachToRoot = InRootComponent->CanAttachAsChild(SceneComponent, NAME_None)
@@ -533,7 +533,7 @@ void FComponentEditorUtils::CopyComponents(const TArray<UActorComponent*>& Compo
 	}
 }
 
-bool FComponentEditorUtils::CanPasteComponents(USceneComponent* RootComponent, bool bOverrideCanAttach, bool bPasteAsArchetypes, const FString* SourceData)
+bool FComponentEditorUtils::CanPasteComponents(const USceneComponent* RootComponent, bool bOverrideCanAttach, bool bPasteAsArchetypes, const FString* SourceData)
 {
 	FString ClipboardContent;
 	if (SourceData)
@@ -845,10 +845,9 @@ UActorComponent* FComponentEditorUtils::DuplicateComponent(UActorComponent* Temp
 	return NewCloneComponent;
 }
 
-void FComponentEditorUtils::AdjustComponentDelta(USceneComponent* Component, FVector& Drag, FRotator& Rotation)
+void FComponentEditorUtils::AdjustComponentDelta(const USceneComponent* Component, FVector& Drag, FRotator& Rotation)
 {
-	USceneComponent* ParentSceneComp = Component->GetAttachParent();
-	if (ParentSceneComp)
+	if (const USceneComponent* ParentSceneComp = Component->GetAttachParent())
 	{
 		const FTransform ParentToWorldSpace = ParentSceneComp->GetSocketTransform(Component->GetAttachSocketName());
 
diff --git a/Engine/Source/Editor/UnrealEd/Public/Kismet2/ComponentEditorUtils.h b/Engine/Source/Editor/UnrealEd/Public/Kismet2/ComponentEditorUtils.h
--- a/Engine/Source/Editor/UnrealEd/Public/Kismet2/ComponentEditorUtils.h
+++ b/Engine/Source/Editor/UnrealEd/Public/Kismet2/ComponentEditorUtils.h
@@ -69,7 +69,7 @@ public:
 	 * @param SourceData Component data to paste, or null to use the clipboard
 	 * @return Whether components can be pasted
 	 */
-	static bool CanPasteComponents(USceneComponent* RootComponent, bool bOverrideCanAttach = false, bool bPasteAsArchetypes = false, const FString* SourceData = nullptr);
+	static bool CanPasteComponents(const USceneComponent* RootComponent, bool bOverrideCanAttach = false, bool bPasteAsArchetypes = false, const FString* SourceData = nullptr);
 
 	/**
 	 * Attempts to paste components from the clipboard as siblings of the target component
@@ -139,7 +139,7 @@ public:
 	static bool AttemptApplyMaterialToComponent( USceneComponent* SceneComponent, UMaterialInterface* MaterialToApply, int32 OptionalMaterialSlot = -1 );
 
 	/** Potentially transforms the delta to be applied to a component into the appropriate space */
-	static void AdjustComponentDelta(USceneComponent* Component, FVector& Drag, FRotator& Rotation);
+	static void AdjustComponentDelta(const USceneComponent* Component, FVector& Drag, FRotator& Rotation);
 
 	// Given a template and a property, propagates a default value change to all instances (only if applicable)
 	template<typename T>
