Fixed a crash in the widget SWidgetStackItem. This is a subwidget of the SWidgetStack from the file ReimportFeedbackContext.

The crash was caused by the refactor of the slate attribute. Since this widget call SlatePrepass before setuping a curve sequence on which it call GetLerp in one of its attribute.

diff --git a/Engine/Source/Editor/UnrealEd/Private/AutoReimport/ReimportFeedbackContext.cpp b/Engine/Source/Editor/UnrealEd/Private/AutoReimport/ReimportFeedbackContext.cpp
--- a/Engine/Source/Editor/UnrealEd/Private/AutoReimport/ReimportFeedbackContext.cpp
+++ b/Engine/Source/Editor/UnrealEd/Private/AutoReimport/ReimportFeedbackContext.cpp
@@ -208,7 +208,12 @@ class SWidgetStack : public SVerticalBox
 
 		FLinearColor GetColorAndOpacity() const
 		{
-			return FLinearColor(1.f, 1.f, 1.f, OpacityCurve.GetLerp());
+			if (OpacityCurve.IsPlaying() || OpacityCurve.IsAtEnd())
+			{
+				return FLinearColor(1.f, 1.f, 1.f, OpacityCurve.GetLerp());
+			}
+
+			return FLinearColor(1.f, 1.f, 1.f, 1.f);
 		}
 
 		bool bIsFinished;
