Basic cooked BP support in BP editor (this is a starting point as they are problems remaining specifically around the graph and inherited functions and I have no doubt we will find more things that don't work later).
With these changes however we are now able to derive from a simple cooked BP, with no compile errors and we're able to see inherited components in the SCS editor, add new components to the tree, override inherited node properties and (re)instantiate those child BPs successfully in a level.

diff --git a/Engine/Source/Developer/BlueprintCompilerCppBackend/Private/BlueprintCompilerCppBackendValueHelper.cpp b/Engine/Source/Developer/BlueprintCompilerCppBackend/Private/BlueprintCompilerCppBackendValueHelper.cpp
--- a/Engine/Source/Developer/BlueprintCompilerCppBackend/Private/BlueprintCompilerCppBackendValueHelper.cpp
+++ b/Engine/Source/Developer/BlueprintCompilerCppBackend/Private/BlueprintCompilerCppBackendValueHelper.cpp
@@ -1269,7 +1269,7 @@ FString FEmitDefaultValueHelper::HandleNonNativeComponent(FEmitterLocalContext&
 					const FObjectProperty* ParentVariableProperty = FindFProperty<FObjectProperty>(BPGC, *CleanParentVariableName);
 					ParentVariableName = ParentVariableProperty ? FEmitHelper::GetCppName(ParentVariableProperty) : CleanParentVariableName;
 				}
-				else if (USceneComponent* ParentComponentTemplate = Node->GetParentComponentTemplate(CastChecked<UBlueprint>(BPGC->ClassGeneratedBy)))
+				else if (USceneComponent* ParentComponentTemplate = Node->GetParentComponentTemplate(BPGC))
 				{
 					ParentVariableName = Context.FindGloballyMappedObject(ParentComponentTemplate, USceneComponent::StaticClass());
 				}
diff --git a/Engine/Source/Editor/Kismet/Private/SSCSEditor.cpp b/Engine/Source/Editor/Kismet/Private/SSCSEditor.cpp
--- a/Engine/Source/Editor/Kismet/Private/SSCSEditor.cpp
+++ b/Engine/Source/Editor/Kismet/Private/SSCSEditor.cpp
@@ -4770,17 +4770,17 @@ FSCSEditorTreeNodePtrType SSCSEditor::GetNodeFromActorComponent(const UActorComp
 				// If the given component is one that's created during Blueprint construction
 				if (ActorComponent->IsCreatedByConstructionScript())
 				{
-					TArray<UBlueprint*> ParentBPStack;
-
 					// Check the entire Class hierarchy for the node
+					TArray<UBlueprintGeneratedClass*> ParentBPStack;
 					UBlueprint::GetBlueprintHierarchyFromClass(OwnerClass, ParentBPStack);
 
 					for(int32 StackIndex = ParentBPStack.Num() - 1; StackIndex >= 0; --StackIndex)
 					{
-						if(ParentBPStack[StackIndex]->SimpleConstructionScript)
+						USimpleConstructionScript* ParentSCS = ParentBPStack[StackIndex] ? ParentBPStack[StackIndex]->SimpleConstructionScript : nullptr;
+						if (ParentSCS)
 						{
 							// Attempt to locate an SCS node with a variable name that matches the name of the given component
-							for (USCS_Node* SCS_Node : ParentBPStack[StackIndex]->SimpleConstructionScript->GetAllNodes())
+							for (USCS_Node* SCS_Node : ParentSCS->GetAllNodes())
 							{
 								check(SCS_Node != NULL);
 								if (SCS_Node->GetVariableName() == ActorComponent->GetFName())
@@ -4967,25 +4967,27 @@ void SSCSEditor::BuildSubTreeForActorNode(FSCSEditorActorNodePtrType InActorNode
 		}
 
 		// If it's a Blueprint-generated class, also get the inheritance stack
-		TArray<UBlueprint*> ParentBPStack;
+		TArray<UBlueprintGeneratedClass*> ParentBPStack;
 		UBlueprint::GetBlueprintHierarchyFromClass(Actor->GetClass(), ParentBPStack);
 
+		UBlueprint* ActorBP = (ParentBPStack.Num() > 0 && ParentBPStack[0]) ? Cast<UBlueprint>(ParentBPStack[0]->ClassGeneratedBy) : nullptr;
+		ensure(ActorBP);
+
 		// Add the full SCS tree node hierarchy (including SCS nodes inherited from parent blueprints)
 		for (int32 StackIndex = ParentBPStack.Num() - 1; StackIndex >= 0; --StackIndex)
 		{
-			if (ParentBPStack[StackIndex]->SimpleConstructionScript != nullptr)
+			USimpleConstructionScript* ParentSCS = ParentBPStack[StackIndex] ? ParentBPStack[StackIndex]->SimpleConstructionScript : nullptr;
+			if (ParentSCS)
 			{
-				const TArray<USCS_Node*>& SCS_RootNodes = ParentBPStack[StackIndex]->SimpleConstructionScript->GetRootNodes();
-				for (int32 NodeIndex = 0; NodeIndex < SCS_RootNodes.Num(); ++NodeIndex)
+				for (USCS_Node* SCS_Node : ParentSCS->GetRootNodes())
 				{
-					USCS_Node* SCS_Node = SCS_RootNodes[NodeIndex];
-					check(SCS_Node != nullptr);
+					check(SCS_Node);
 
 					FSCSEditorTreeNodePtrType NewNodePtr;
 					if (SCS_Node->ParentComponentOrVariableName != NAME_None)
 					{
-						USceneComponent* ParentComponent = SCS_Node->GetParentComponentTemplate(ParentBPStack[0]);
-						if (ParentComponent != nullptr)
+						USceneComponent* ParentComponent = SCS_Node->GetParentComponentTemplate(ActorBP);
+						if (ParentComponent)
 						{
 							FSCSEditorTreeNodePtrType ParentNodePtr = FindTreeNode(ParentComponent, InActorNode);
 							if (ParentNodePtr.IsValid())
@@ -5004,12 +5006,12 @@ void SSCSEditor::BuildSubTreeForActorNode(FSCSEditorActorNodePtrType InActorNode
 					{
 						// This call creates ICH override templates for the current Blueprint. Without this, the parent node
 						// search above can fail when attempting to match an inherited node in the tree via component template.
-						NewNodePtr->GetOrCreateEditableComponentTemplate(ParentBPStack[0]);
+						NewNodePtr->GetOrCreateEditableComponentTemplate(ActorBP);
 						for (FSCSEditorTreeNodePtrType ChildNodePtr : NewNodePtr->GetChildren())
 						{
 							if (ensure(ChildNodePtr.IsValid()))
 							{
-								ChildNodePtr->GetOrCreateEditableComponentTemplate(ParentBPStack[0]);
+								ChildNodePtr->GetOrCreateEditableComponentTemplate(ActorBP);
 							}
 						}
 					}
diff --git a/Engine/Source/Editor/UnrealEd/Private/Kismet2/BlueprintEditorUtils.cpp b/Engine/Source/Editor/UnrealEd/Private/Kismet2/BlueprintEditorUtils.cpp
--- a/Engine/Source/Editor/UnrealEd/Private/Kismet2/BlueprintEditorUtils.cpp
+++ b/Engine/Source/Editor/UnrealEd/Private/Kismet2/BlueprintEditorUtils.cpp
@@ -1768,14 +1768,14 @@ void FBlueprintEditorUtils::PropagateParentBlueprintDefaults(UClass* ClassToProp
 	check(NewCDO);
 
 	// Get the blueprint's BP derived lineage
-	TArray<UBlueprint*> ParentBP;
-	UBlueprint::GetBlueprintHierarchyFromClass(ClassToPropagate, ParentBP);
+	TArray<UBlueprintGeneratedClass*> ParentBPStack;
+	UBlueprint::GetBlueprintHierarchyFromClass(ClassToPropagate, ParentBPStack);
 
 	// Starting from the least derived BP class, copy the properties into the new CDO
-	for(int32 i = ParentBP.Num() - 1; i > 0; i--)
+	for(int32 Index = ParentBPStack.Num() - 1; Index > 0; --Index)
 	{
-		checkf(ParentBP[i]->GeneratedClass != nullptr, TEXT("Parent classes for class %s have not yet been generated.  Compile-on-load must be processed for the parent class first."), *ClassToPropagate->GetName());
-		UObject* LayerCDO = ParentBP[i]->GeneratedClass->GetDefaultObject();
+		checkf(ParentBPStack[Index], TEXT("Parent classes for class %s have not yet been generated.  Compile-on-load must be processed for the parent class first."), *ClassToPropagate->GetName());
+		UObject* LayerCDO = ParentBPStack[Index]->GetDefaultObject();
 
 		UEditorEngine::FCopyPropertiesForUnrelatedObjectsParams CopyDetails;
 		CopyDetails.bReplaceObjectClassReferences = false;
@@ -3669,14 +3669,30 @@ int32 FBlueprintEditorUtils::FindTimelineIndex(const UBlueprint* Blueprint, cons
 
 void FBlueprintEditorUtils::GetSCSVariableNameList(const UBlueprint* Blueprint, TSet<FName>& VariableNames)
 {
-	if(Blueprint != nullptr && Blueprint->SimpleConstructionScript != nullptr)
+	if (Blueprint)
 	{
-		for (USCS_Node* SCS_Node : Blueprint->SimpleConstructionScript->GetAllNodes())
+		GetSCSVariableNameList(Blueprint->SimpleConstructionScript, VariableNames);
+	}
+}
+
+void FBlueprintEditorUtils::GetSCSVariableNameList(const UBlueprintGeneratedClass* BPGC, TSet<FName>& VariableNames)
+{
+	if (BPGC)
+	{
+		GetSCSVariableNameList(BPGC->SimpleConstructionScript, VariableNames);
+	}
+}
+
+void FBlueprintEditorUtils::GetSCSVariableNameList(const USimpleConstructionScript* SCS, TSet<FName>& VariableNames)
+{
+	if (SCS)
+	{
+		for (USCS_Node* SCS_Node : SCS->GetAllNodes())
 		{
-			if (SCS_Node != nullptr)
+			if (SCS_Node)
 			{
 				const FName VariableName = SCS_Node->GetVariableName();
-				if(VariableName != NAME_None)
+				if (VariableName != NAME_None)
 				{
 					VariableNames.Add(VariableName);
 				}
@@ -4593,21 +4609,28 @@ void FBlueprintEditorUtils::GetClassVariableList(const UBlueprint* Blueprint, TS
 		if (bIncludePrivateVars)
 		{
 			// Include SCS node variable names, timelines, and other member variables that may be pending compilation. Consider them to be "private" as they're not technically accessible for editing just yet.
-			TArray<UBlueprint*> ParentBPStack;
+			TArray<UBlueprintGeneratedClass*> ParentBPStack;
 			UBlueprint::GetBlueprintHierarchyFromClass(Blueprint->SkeletonGeneratedClass, ParentBPStack);
 			for (int32 StackIndex = ParentBPStack.Num() - 1; StackIndex >= 0; --StackIndex)
 			{
-				UBlueprint* ParentBP = ParentBPStack[StackIndex];
-				check(ParentBP != nullptr);
+				UBlueprintGeneratedClass* ParentPBGC = ParentBPStack[StackIndex];
+				check(ParentPBGC);
+				UBlueprint* ParentBP = Cast<UBlueprint>(ParentPBGC->ClassGeneratedBy);
 
-				GetSCSVariableNameList(ParentBP, VisibleVariables);
+				GetSCSVariableNameList(ParentPBGC, VisibleVariables);
 
-				for (int32 VariableIndex = 0; VariableIndex < ParentBP->NewVariables.Num(); ++VariableIndex)
+				if (ParentBP)
 				{
-					VisibleVariables.Add(ParentBP->NewVariables[VariableIndex].VarName);
+					for (const FBPVariableDescription& Variable : ParentBP->NewVariables)
+					{
+						VisibleVariables.Add(Variable.VarName);
+					}
 				}
 
-				for (UTimelineTemplate* Timeline : ParentBP->Timelines)
+				// Since we defer copying the timeline templates to the BPGC until compile time, 
+				// we consider the BP (when present) to be authoritative.
+				const TArray<TObjectPtr<UTimelineTemplate>>& Timelines = ParentBP ? ParentBP->Timelines : ParentPBGC->Timelines;
+				for (UTimelineTemplate* Timeline : Timelines)
 				{
 					if (Timeline)
 					{
diff --git a/Engine/Source/Editor/UnrealEd/Private/Kismet2/Kismet2.cpp b/Engine/Source/Editor/UnrealEd/Private/Kismet2/Kismet2.cpp
--- a/Engine/Source/Editor/UnrealEd/Private/Kismet2/Kismet2.cpp
+++ b/Engine/Source/Editor/UnrealEd/Private/Kismet2/Kismet2.cpp
@@ -455,7 +455,7 @@ UBlueprint* FKismetEditorUtilities::CreateBlueprint(UClass* ParentClass, UObject
 		FBlueprintEditorUtils::AddFunctionGraph(NewBP, UCSGraph, /*bIsUserCreated=*/ false, AActor::StaticClass());
 
 		// If the blueprint is derived from another blueprint, add in a super-call automatically
-		if( NewBP->ParentClass && NewBP->ParentClass->ClassGeneratedBy )
+		if (NewBP->ParentClass && NewBP->ParentClass->HasAnyClassFlags(CLASS_CompiledFromBlueprint))
 		{
 			check( UCSGraph->Nodes.Num() > 0 );
 			UK2Node_FunctionEntry* UCSEntry = CastChecked<UK2Node_FunctionEntry>(UCSGraph->Nodes[0]);
diff --git a/Engine/Source/Editor/UnrealEd/Public/Kismet2/BlueprintEditorUtils.h b/Engine/Source/Editor/UnrealEd/Public/Kismet2/BlueprintEditorUtils.h
--- a/Engine/Source/Editor/UnrealEd/Public/Kismet2/BlueprintEditorUtils.h
+++ b/Engine/Source/Editor/UnrealEd/Public/Kismet2/BlueprintEditorUtils.h
@@ -1279,6 +1279,20 @@ public:
 	 */
 	static void GetSCSVariableNameList(const UBlueprint* Blueprint, TSet<FName>& VariableNames);
 
+	/**
+	 * Gets a list of SCS node variable names for the given BPGC.
+	 *
+	 * @param [in,out]	VariableNames		The list of variable names for the SCS node array.
+	 */
+	static void GetSCSVariableNameList(const UBlueprintGeneratedClass* BPGC, TSet<FName>& VariableNames);
+
+	/**
+	 * Gets a list of SCS node variable names for the given SCS.
+	 *
+	 * @param [in,out]	VariableNames		The list of variable names for the SCS node array.
+	 */
+	static void GetSCSVariableNameList(const USimpleConstructionScript* SCS, TSet<FName>& VariableNames);
+
 	/** 
 	 * Gets a list of function names in blueprints that implement the interface defined by the given blueprint.
 	 *
diff --git a/Engine/Source/Runtime/Engine/Classes/Engine/Blueprint.h b/Engine/Source/Runtime/Engine/Classes/Engine/Blueprint.h
--- a/Engine/Source/Runtime/Engine/Classes/Engine/Blueprint.h
+++ b/Engine/Source/Runtime/Engine/Classes/Engine/Blueprint.h
@@ -928,6 +928,15 @@ public:
 	 * @return						true if there were no status errors in any of the parent blueprints, otherwise false
 	 */
 	static bool GetBlueprintHierarchyFromClass(const UClass* InClass, TArray<UBlueprint*>& OutBlueprintParents);
+
+	/**
+	 * Gets an array of all BPGCs used to generate this class and its parents.  0th elements is the BPGC used to generate InClass
+	 *
+	 * @param InClass				The class to get the blueprint lineage for
+	 * @param OutBlueprintParents	Array of BPGCs used to generate this class and its parents.  0th = this, Nth = least derived BP-based parent
+	 * @return						true if there were no status errors in any of the parent blueprints, otherwise false
+	 */
+	static bool GetBlueprintHierarchyFromClass(const UClass* InClass, TArray<UBlueprintGeneratedClass*>& OutBlueprintParents);
 	
 #if WITH_EDITOR
 	/** returns true if the class hierarchy is error free */
diff --git a/Engine/Source/Runtime/Engine/Classes/Engine/SCS_Node.h b/Engine/Source/Runtime/Engine/Classes/Engine/SCS_Node.h
--- a/Engine/Source/Runtime/Engine/Classes/Engine/SCS_Node.h
+++ b/Engine/Source/Runtime/Engine/Classes/Engine/SCS_Node.h
@@ -191,6 +191,9 @@ class USCS_Node : public UObject
 	/** Finds and returns the parent component template through the given Blueprint */
 	ENGINE_API USceneComponent* GetParentComponentTemplate(UBlueprint* InBlueprint) const;
 
+	/** Finds and returns the parent component template through the given BPGC */
+	ENGINE_API USceneComponent* GetParentComponentTemplate(UBlueprintGeneratedClass* BPGC) const;
+
 	/** Saves the current state of SimpleConstructionScript and its nodes to the transaction buffer. */
 	ENGINE_API void SaveToTransactionBuffer();
 #endif
diff --git a/Engine/Source/Runtime/Engine/Private/Blueprint.cpp b/Engine/Source/Runtime/Engine/Private/Blueprint.cpp
--- a/Engine/Source/Runtime/Engine/Private/Blueprint.cpp
+++ b/Engine/Source/Runtime/Engine/Private/Blueprint.cpp
@@ -1387,7 +1387,7 @@ UBlueprint* UBlueprint::GetBlueprintFromClass(const UClass* InClass)
 
 bool UBlueprint::GetBlueprintHierarchyFromClass(const UClass* InClass, TArray<UBlueprint*>& OutBlueprintParents)
 {
-	OutBlueprintParents.Empty();
+	OutBlueprintParents.Reset();
 
 	bool bNoErrors = true;
 	const UClass* CurrentClass = InClass;
@@ -1400,9 +1400,9 @@ bool UBlueprint::GetBlueprintHierarchyFromClass(const UClass* InClass, TArray<UB
 #endif // #if WITH_EDITORONLY_DATA
 
 		// If valid, use stored ParentClass rather than the actual UClass::GetSuperClass(); handles the case when the class has not been recompiled yet after a reparent operation.
-		if(const UClass* ParentClass = BP->ParentClass)
+		if (BP->ParentClass)
 		{
-			CurrentClass = ParentClass;
+			CurrentClass = BP->ParentClass;
 		}
 		else
 		{
@@ -1414,6 +1414,40 @@ bool UBlueprint::GetBlueprintHierarchyFromClass(const UClass* InClass, TArray<UB
 	return bNoErrors;
 }
 
+bool UBlueprint::GetBlueprintHierarchyFromClass(const UClass* InClass, TArray<UBlueprintGeneratedClass*>& OutBlueprintParents)
+{
+	OutBlueprintParents.Reset();
+
+	bool bNoErrors = true;
+	UBlueprintGeneratedClass* CurrentClass = Cast<UBlueprintGeneratedClass>(const_cast<UClass*>(InClass));
+	while (CurrentClass)
+	{
+		OutBlueprintParents.Add(CurrentClass);
+
+		UBlueprint* BP = UBlueprint::GetBlueprintFromClass(CurrentClass);
+
+#if WITH_EDITORONLY_DATA
+		if (BP)
+		{
+			bNoErrors &= (BP->Status != BS_Error);
+		}
+#endif // #if WITH_EDITORONLY_DATA
+
+		// If valid, use stored ParentClass rather than the actual UClass::GetSuperClass(); handles the case when the class has not been recompiled yet after a reparent operation.
+		if (BP && BP->ParentClass)
+		{
+			CurrentClass = Cast<UBlueprintGeneratedClass>(BP->ParentClass);
+		}
+		else
+		{
+			check(CurrentClass);
+			CurrentClass = Cast<UBlueprintGeneratedClass>(CurrentClass->GetSuperClass());
+		}
+	}
+
+	return bNoErrors;
+}
+
 #if WITH_EDITOR
 bool UBlueprint::IsBlueprintHierarchyErrorFree(const UClass* InClass)
 {
diff --git a/Engine/Source/Runtime/Engine/Private/SCS_Node.cpp b/Engine/Source/Runtime/Engine/Private/SCS_Node.cpp
--- a/Engine/Source/Runtime/Engine/Private/SCS_Node.cpp
+++ b/Engine/Source/Runtime/Engine/Private/SCS_Node.cpp
@@ -485,9 +485,7 @@ void USCS_Node::SetParent(USCS_Node* InParentNode)
 	ensure(InParentNode);
 	USimpleConstructionScript* ParentSCS = InParentNode ? InParentNode->GetSCS() : nullptr;
 	ensure(ParentSCS);
-	UBlueprint* ParentBlueprint = ParentSCS ? ParentSCS->GetBlueprint() : nullptr;
-	ensure(ParentBlueprint);
-	UClass* ParentBlueprintGeneratedClass = ParentBlueprint ? ParentBlueprint->GeneratedClass : nullptr;
+	UClass* ParentBlueprintGeneratedClass = ParentSCS ? ParentSCS->GetOwnerClass() : nullptr;
 
 	if (ParentBlueprintGeneratedClass && InParentNode)
 	{
@@ -529,17 +527,23 @@ void USCS_Node::SetParent(const USceneComponent* InParentComponent)
 }
 
 USceneComponent* USCS_Node::GetParentComponentTemplate(UBlueprint* InBlueprint) const
+{
+	check(InBlueprint && InBlueprint->GeneratedClass);
+	return GetParentComponentTemplate(CastChecked<UBlueprintGeneratedClass>(InBlueprint->GeneratedClass));
+}
+
+USceneComponent* USCS_Node::GetParentComponentTemplate(UBlueprintGeneratedClass* BPGC) const
 {
 	USceneComponent* ParentComponentTemplate = nullptr;
 	if(ParentComponentOrVariableName != NAME_None)
 	{
-		check(InBlueprint != nullptr && InBlueprint->GeneratedClass != nullptr);
+		check(BPGC);
 
 		// If the parent component template is found in the 'Components' array of the CDO (i.e. native)
 		if(bIsParentComponentNative)
 		{
 			// Access the Blueprint CDO
-			AActor* CDO = InBlueprint->GeneratedClass->GetDefaultObject<AActor>();
+			AActor* CDO = BPGC->GetDefaultObject<AActor>();
 			if(CDO != nullptr)
 			{
 				// Find the component template in the CDO that matches the specified name
@@ -559,26 +563,25 @@ USceneComponent* USCS_Node::GetParentComponentTemplate(UBlueprint* InBlueprint)
 		else
 		{
 			// Get the Blueprint hierarchy
-			TArray<UBlueprint*> ParentBPStack;
-			UBlueprint::GetBlueprintHierarchyFromClass(InBlueprint->GeneratedClass, ParentBPStack);
+			TArray<UBlueprintGeneratedClass*> ParentBPStack;
+			UBlueprint::GetBlueprintHierarchyFromClass(BPGC, ParentBPStack);
 
 			// Find the parent Blueprint in the hierarchy
-			for(int32 StackIndex = ParentBPStack.Num() - 1; StackIndex > 0 && !ParentComponentTemplate; --StackIndex)
+			for (int32 StackIndex = ParentBPStack.Num() - 1; StackIndex > 0 && !ParentComponentTemplate; --StackIndex)
 			{
-				UBlueprint* ParentBlueprint = ParentBPStack[StackIndex];
-				if(ParentBlueprint != nullptr
-					&& ParentBlueprint->SimpleConstructionScript != nullptr
-					&& ParentBlueprint->GeneratedClass->GetFName() == ParentComponentOwnerClassName)
+				UBlueprintGeneratedClass* ParentBPGC = ParentBPStack[StackIndex];
+				if (ParentBPGC
+					&& ParentBPGC->SimpleConstructionScript
+					&& ParentBPGC->GetFName() == ParentComponentOwnerClassName)
 				{
 					// Find the SCS node with a variable name that matches the specified name
-					TArray<USCS_Node*> ParentSCSNodes = ParentBlueprint->SimpleConstructionScript->GetAllNodes();
-					for(int32 ParentNodeIndex = 0; ParentNodeIndex < ParentSCSNodes.Num(); ++ParentNodeIndex)
+					for (USCS_Node* Node : ParentBPGC->SimpleConstructionScript->GetAllNodes())
 					{
-						USceneComponent* CompTemplate = Cast<USceneComponent>(ParentSCSNodes[ParentNodeIndex]->ComponentTemplate);
-						if(CompTemplate != nullptr && ParentSCSNodes[ParentNodeIndex]->GetVariableName() == ParentComponentOrVariableName)
+						USceneComponent* CompTemplate = Cast<USceneComponent>(Node->ComponentTemplate);
+						if (CompTemplate && Node->GetVariableName() == ParentComponentOrVariableName)
 						{
 							// Found a match; this is our parent, we're done
-							ParentComponentTemplate = Cast<USceneComponent>(ParentSCSNodes[ParentNodeIndex]->GetActualComponentTemplate(Cast<UBlueprintGeneratedClass>(InBlueprint->GeneratedClass)));
+							ParentComponentTemplate = Cast<USceneComponent>(Node->GetActualComponentTemplate(Cast<UBlueprintGeneratedClass>(BPGC)));
 							break;
 						}
 					}
diff --git a/Engine/Source/Runtime/Engine/Private/SimpleConstructionScript.cpp b/Engine/Source/Runtime/Engine/Private/SimpleConstructionScript.cpp
--- a/Engine/Source/Runtime/Engine/Private/SimpleConstructionScript.cpp
+++ b/Engine/Source/Runtime/Engine/Private/SimpleConstructionScript.cpp
@@ -1047,12 +1047,12 @@ USceneComponent* USimpleConstructionScript::GetSceneRootComponentTemplate(bool b
 	if(!RootComponentTemplate)
 	{
 		// Get the Blueprint hierarchy
-		TArray<UBlueprint*> BPStack;
-		if(GeneratedClass != nullptr)
+		TArray<UBlueprintGeneratedClass*> BPStack;
+		if (GeneratedClass)
 		{
 			UBlueprint::GetBlueprintHierarchyFromClass(GeneratedClass, BPStack);
 		}
-		else if(ParentClass != nullptr)
+		else if (ParentClass)
 		{
 			UBlueprint::GetBlueprintHierarchyFromClass(ParentClass, BPStack);
 		}
@@ -1060,18 +1060,20 @@ USceneComponent* USimpleConstructionScript::GetSceneRootComponentTemplate(bool b
 		// Note: Normally if the Blueprint has a parent, we can assume that the parent already has a scene root component set,
 		// ...but we'll run through the hierarchy just in case there are legacy BPs out there that might not adhere to this assumption.
 		TArray<const USimpleConstructionScript*> SCSStack;
+		SCSStack.Reserve(BPStack.Num() + 1);
+
 		SCSStack.Add(this);
 
-		for(int32 StackIndex = 0; StackIndex < BPStack.Num(); ++StackIndex)
+		for(UBlueprintGeneratedClass* BPGC : BPStack)
 		{
-			if(BPStack[StackIndex] && BPStack[StackIndex]->SimpleConstructionScript && !SCSStack.Contains(BPStack[StackIndex]->SimpleConstructionScript))
+			if (BPGC && BPGC->SimpleConstructionScript)
 			{
-				// UBlueprint::GetBlueprintHierarchyFromClass returns first children then parents. So we need to revert the order.
-				SCSStack.Insert(BPStack[StackIndex]->SimpleConstructionScript, 0);
+				SCSStack.AddUnique(BPGC->SimpleConstructionScript);
 			}
 		}
 
-		for(int32 StackIndex = 0; StackIndex < SCSStack.Num() && !RootComponentTemplate; ++StackIndex)
+		// UBlueprint::GetBlueprintHierarchyFromClass returns first children then parents. So we need to revert the order.
+		for (int32 StackIndex = SCSStack.Num() - 1; StackIndex >= 0 && !RootComponentTemplate; --StackIndex)
 		{
 			const TArray<USCS_Node*>& SCSRootNodes = SCSStack[StackIndex]->GetRootNodes();
 
