Don't run construction scripts for template actors, such as sequencer spawn templates.

diff --git a/Engine/Source/Editor/UnrealEd/Private/Kismet2/KismetReinstanceUtilities.cpp b/Engine/Source/Editor/UnrealEd/Private/Kismet2/KismetReinstanceUtilities.cpp
--- a/Engine/Source/Editor/UnrealEd/Private/Kismet2/KismetReinstanceUtilities.cpp
+++ b/Engine/Source/Editor/UnrealEd/Private/Kismet2/KismetReinstanceUtilities.cpp
@@ -2435,7 +2435,7 @@ void FBlueprintCompileReinstancer::ReplaceInstancesOfClass_Inner(TMap<UClass*, U
 		}
 
 		// Skipping CDOs as CSs are not allowed for them.
-		if (!ActorInstance->HasAnyFlags(RF_ClassDefaultObject))
+		if (!ActorInstance->HasAnyFlags(RF_ClassDefaultObject) && ActorInstance->GetLevel())
 		{
 			ActorInstance->RerunConstructionScripts();
 		}
