Avoid creating FName only for comparison (heavily used during Houdini Import)

diff --git a/Engine/Source/Editor/UnrealEd/Private/Factories/ActorFactory.cpp b/Engine/Source/Editor/UnrealEd/Private/Factories/ActorFactory.cpp
--- a/Engine/Source/Editor/UnrealEd/Private/Factories/ActorFactory.cpp
+++ b/Engine/Source/Editor/UnrealEd/Private/Factories/ActorFactory.cpp
@@ -223,10 +223,12 @@ UActorFactory::UActorFactory(const FObjectInitializer& ObjectInitializer)
 
 bool UActorFactory::CanCreateActorFrom( const FAssetData& AssetData, FText& OutErrorMsg )
 {
+	TRACE_CPUPROFILER_EVENT_SCOPE(UActorFactory::CanCreateActorFrom);
+
 	// By Default we assume the factory can't work with existing asset data
 	return !AssetData.IsValid() || 
-		AssetData.ObjectPath == FName(*GetDefaultActor( AssetData )->GetPathName()) || 
-		AssetData.ObjectPath == FName(*GetDefaultActor( AssetData )->GetClass()->GetPathName());
+		AssetData.ObjectPath == *GetDefaultActor( AssetData )->GetPathName() || 
+		AssetData.ObjectPath == *GetDefaultActor( AssetData )->GetClass()->GetPathName();
 }
 
 AActor* UActorFactory::GetDefaultActor( const FAssetData& AssetData )
