Adds DeployFolder to CookGlobalShaderCommandlet.cpp

diff --git a/Engine/Source/Editor/UnrealEd/Private/Commandlets/CookGlobalShaderCommandlet.cpp b/Engine/Source/Editor/UnrealEd/Private/Commandlets/CookGlobalShaderCommandlet.cpp
--- a/Engine/Source/Editor/UnrealEd/Private/Commandlets/CookGlobalShaderCommandlet.cpp
+++ b/Engine/Source/Editor/UnrealEd/Private/Commandlets/CookGlobalShaderCommandlet.cpp
@@ -43,18 +43,24 @@ int32 UCookGlobalShadersCommandlet::Main(const FString& Params)
 		UE_LOG(LogCookGlobalShaders, Log, TEXT("CookGlobalShaders"));
 		UE_LOG(LogCookGlobalShaders, Log, TEXT("This commandlet will allow you to generate the global shaders file which can be used to override what is used in a cooked build by deploying the loose file."));
 		UE_LOG(LogCookGlobalShaders, Log, TEXT("Options:"));
-		UE_LOG(LogCookGlobalShaders, Log, TEXT(" Required: -platform=<platform>     (Which platform you want to cook for, i.e. windows)"));
-		UE_LOG(LogCookGlobalShaders, Log, TEXT(" Optional: -device=<name>           (Set which device to use, when enabled the reload command will be sent to the device once the shaders are cooked)"));
-		UE_LOG(LogCookGlobalShaders, Log, TEXT(" Optional: -deploy                  (Must be used with -device and will deploy the shader file onto the device rather than in the staged builds folder)"));
-		UE_LOG(LogCookGlobalShaders, Log, TEXT(" Optional: -stage=<optional path>   (Moved the shader file into the staged builds folder, destination can be overriden)"));
-		UE_LOG(LogCookGlobalShaders, Log, TEXT(" Optional: -reload                  (Execute a shader reload on the device, only works if the device is valid or a default one was found"));
-		UE_LOG(LogCookGlobalShaders, Log, TEXT(" Optional: -shaderpdb=<path>        (Sets the shader pdb root)"));
+		UE_LOG(LogCookGlobalShaders, Log, TEXT(" Required: -platform=<platform>             (Which platform you want to cook for, i.e. windows)"));
+		UE_LOG(LogCookGlobalShaders, Log, TEXT(" Optional: -device=<name>                   (Set which device to use, when enabled the reload command will be sent to the device once the shaders are cooked)"));
+		UE_LOG(LogCookGlobalShaders, Log, TEXT(" Optional: -deploy=<optional deploy folder> (Must be used with -device and will deploy the shader file onto the device rather than in the staged builds folder)"));
+		UE_LOG(LogCookGlobalShaders, Log, TEXT(" Optional: -stage=<optional path>           (Moved the shader file into the staged builds folder, destination can be overriden)"));
+		UE_LOG(LogCookGlobalShaders, Log, TEXT(" Optional: -reload                          (Execute a shader reload on the device, only works if the device is valid or a default one was found"));
+		UE_LOG(LogCookGlobalShaders, Log, TEXT(" Optional: -shaderpdb=<path>                (Sets the shader pdb root)"));
 		return 0;
 	}
 
-	const bool bDeployToDevice = Switches.Contains("deploy");
-	const bool bCopyToStaged = Switches.Contains("stage");
-	const bool bExecuteReload = Switches.Contains("reload");
+	const bool bDeployToDevice = Switches.Contains(TEXT("deploy")) || ParamVals.Contains(TEXT("deploy"));
+	const bool bCopyToStaged = Switches.Contains(TEXT("stage"));
+	const bool bExecuteReload = Switches.Contains(TEXT("reload"));
+
+	FString DeployFolder;
+	if (ParamVals.Contains(TEXT("deploy")))
+	{
+		DeployFolder = ParamVals[TEXT("deploy")];
+	}
 
 	// Parse platform
 	FString PlatformName;
@@ -96,6 +102,11 @@ int32 UCookGlobalShadersCommandlet::Main(const FString& Params)
 		if ( !TargetDevice.IsValid() )
 		{
 			UE_LOG(LogCookGlobalShaders, Warning, TEXT("Failed to find target device '%s', reload / deploy will not be valid"), *TargetDeviceName);
+
+			for (int i = 0; i < TargetDevices.Num(); ++i)
+			{
+				UE_LOG(LogCookGlobalShaders, Warning, TEXT("	%s"), *TargetDevices[i]->GetName());
+			}
 		}
 	}
 	else
@@ -172,16 +183,19 @@ int32 UCookGlobalShadersCommandlet::Main(const FString& Params)
 		{
 			// Execute Copy
 			UE_LOG(LogCookGlobalShaders, Log, TEXT("Copying Cooked Files..."));
-			bCopySucceeded =  DeviceHelper->CopyFilesToDevice(TargetDevice.Get(), FilesToCopy);
+			bCopySucceeded = DeviceHelper->CopyFilesToDevice(TargetDevice.Get(), FilesToCopy);
 		}
 		// if no helper, but we want to deploy, use the TargetPlatform
 		else if (bDeployToDevice && TargetDevice != nullptr)
 		{
 			bCopySucceeded = true;
 
+			TMap<FString,FString> CustomPlatformData;
+			CustomPlatformData.Add(TEXT("DeployFolder"), DeployFolder);
+
 			for (auto It : FilesToCopy)
 			{
-				bCopySucceeded = bCopySucceeded && TargetPlatform->CopyFileToTarget(TargetDevice->GetId().GetDeviceName(), It.Key, It.Value, {});
+				bCopySucceeded = bCopySucceeded && TargetPlatform->CopyFileToTarget(TargetDevice->GetId().GetDeviceName(), It.Key, It.Value, CustomPlatformData);
 			}
 		}
 
