Process FPendingCleanupObjects periodically when in the cooking server similar to how the editor processes these every tick in FEngineLoop::Tick()

- Adding a FlushRenderingCommands() call to the end of COTF and CBTB.  This cleans up pending delete objects, and will also flush Release() render command calls on FRenderResources so we don't pending delete objects before they are done on the rendering thread.
- Adding a stat to see how this impacts cooks.
- I'm intentially not caching a previous frame's pending objects to delete (like FEngineLoop::Tick) because we shouldn't be rendering in the cooker so this _should_ be OK.
- FMaterialShaderMaps were not being deleted and accumulating in a global list.  This resulted in stale maps hanging around (ones that failed compilation) and a memory leak.
- ODSC code would attempt to use some of the stale maps in the list and crash.
- Anything that inherited from FDeferredCleanupInterface would exhibit the same problems as FMaterialShaderMaps and be leaking in the cooking server.  This includes FLightMap and FShadowMap which were being GC'ed during cooks.

diff --git a/Engine/Source/Editor/UnrealEd/Private/Commandlets/CookCommandlet.cpp b/Engine/Source/Editor/UnrealEd/Private/Commandlets/CookCommandlet.cpp
--- a/Engine/Source/Editor/UnrealEd/Private/Commandlets/CookCommandlet.cpp
+++ b/Engine/Source/Editor/UnrealEd/Private/Commandlets/CookCommandlet.cpp
@@ -75,6 +75,7 @@ namespace DetailedCookStats
 	double TickLoopShaderProcessAsyncResultsTimeSec = 0.0;
 	double TickLoopProcessDeferredCommandsTimeSec = 0.0;
 	double TickLoopTickCommandletStatsTimeSec = 0.0;
+	double TickLoopFlushRenderingCommandsTimeSec = 0.0;
 
 	FCookStatsManager::FAutoRegisterCallback RegisterCookStats([](FCookStatsManager::AddStatFuncRef AddStat)
 	{
@@ -96,9 +97,10 @@ namespace DetailedCookStats
 		ADD_COOK_STAT_FLT(" 0. 1. 4", TickLoopShaderProcessAsyncResultsTimeSec);
 		ADD_COOK_STAT_FLT(" 0. 1. 5", TickLoopProcessDeferredCommandsTimeSec);
 		ADD_COOK_STAT_FLT(" 0. 1. 6", TickLoopTickCommandletStatsTimeSec);
-		ADD_COOK_STAT_FLT(" 0. 1. 7", TargetPlatforms);
-		ADD_COOK_STAT_FLT(" 0. 1. 8", CookProject);
-		ADD_COOK_STAT_FLT(" 0. 1. 9", CookCultures);
+		ADD_COOK_STAT_FLT(" 0. 1. 7", TickLoopFlushRenderingCommandsTimeSec);
+		ADD_COOK_STAT_FLT(" 0. 1. 8", TargetPlatforms);
+		ADD_COOK_STAT_FLT(" 0. 1. 9", CookProject);
+		ADD_COOK_STAT_FLT(" 0. 1. 10", CookCultures);
 		#undef ADD_COOK_STAT_FLT
 	});
 
@@ -583,6 +585,12 @@ bool UCookCommandlet::CookOnTheFly( FGuid InstanceId, int32 Timeout, bool bForce
 
 			CookOnTheFlyGCController.ConditionallyCollectGarbage(CookOnTheFlyServer);
 
+			{
+				// Flush rendering commands to release any RHI resources (shaders and shader maps).
+				// Delete any FPendingCleanupObjects (shader maps).
+				FlushRenderingCommands();
+			}
+
 			CookOnTheFlyServer->WaitForRequests(100 /* timeoutMs */);
 		}
 	}
@@ -1138,6 +1146,14 @@ bool UCookCommandlet::CookByTheBook( const TArray<ITargetPlatform*>& Platforms)
 					UE_SCOPED_COOKTIMER_AND_DURATION(CookByTheBook_ProcessDeferredCommands, DetailedCookStats::TickLoopProcessDeferredCommandsTimeSec);
 					ProcessDeferredCommands();
 				}
+
+				{
+					UE_SCOPED_COOKTIMER_AND_DURATION(CookByTheBook_TickCommandletStats, DetailedCookStats::TickLoopFlushRenderingCommandsTimeSec);
+
+					// Flush rendering commands to release any RHI resources (shaders and shader maps).
+					// Delete any FPendingCleanupObjects (shader maps).
+					FlushRenderingCommands();
+				}
 			}
 
 			{
