Adding mesh shader support to pipeline caches.
- Pipeline caches had serialization code for the new types but wasn't creating the shaders on preload
- Fixing the cache commandlet to be flexible with the supported shader types.

diff --git a/Engine/Source/Editor/UnrealEd/Private/Commandlets/ShaderPipelineCacheToolsCommandlet.cpp b/Engine/Source/Editor/UnrealEd/Private/Commandlets/ShaderPipelineCacheToolsCommandlet.cpp
--- a/Engine/Source/Editor/UnrealEd/Private/Commandlets/ShaderPipelineCacheToolsCommandlet.cpp
+++ b/Engine/Source/Editor/UnrealEd/Private/Commandlets/ShaderPipelineCacheToolsCommandlet.cpp
@@ -392,7 +392,7 @@ int32 DumpPSOSC(FString& Token)
 		}
 		else if (Item.Type == FPipelineCacheFileFormatPSO::DescriptorType::Graphics)
 		{
-			check(!(Item.GraphicsDesc.VertexShader == FSHAHash()));
+			check(!(Item.GraphicsDesc.VertexShader == FSHAHash() || Item.GraphicsDesc.MeshShader == FSHAHash()));
 			StringRep = Item.GraphicsDesc.ToString();
 		}
 		else if (Item.Type == FPipelineCacheFileFormatPSO::DescriptorType::RayTracing)
@@ -751,6 +751,10 @@ int32 ExpandPSOSC(const TArray<FString>& Tokens)
 			case FPipelineCacheFileFormatPSO::DescriptorType::Graphics:
 				UE_LOG(LogShaderPipelineCacheTools, Verbose, TEXT("VertexShader"));
 				PrintShaders(InverseMap, Item.GraphicsDesc.VertexShader);
+				UE_LOG(LogShaderPipelineCacheTools, Verbose, TEXT("MeshShader"));
+				PrintShaders(InverseMap, Item.GraphicsDesc.MeshShader);
+				UE_LOG(LogShaderPipelineCacheTools, Verbose, TEXT("AmplificationShader"));
+				PrintShaders(InverseMap, Item.GraphicsDesc.AmplificationShader);
 				UE_LOG(LogShaderPipelineCacheTools, Verbose, TEXT("FragmentShader"));
 				PrintShaders(InverseMap, Item.GraphicsDesc.FragmentShader);
 				UE_LOG(LogShaderPipelineCacheTools, Verbose, TEXT("GeometryShader"));
@@ -804,7 +808,7 @@ int32 ExpandPSOSC(const TArray<FString>& Tokens)
 	{ 
 		NumExamined++;
 		
-		check(SF_Vertex == 0 && SF_Compute == 5);
+		static_assert(SF_Vertex == 0 && SF_Compute == 7, "Shader Frequencies have changed, please update");
 		TArray<int32> StableShadersPerSlot[SF_NumFrequencies];
 		bool ActivePerSlot[SF_NumFrequencies] = { false };
 
@@ -817,6 +821,8 @@ int32 ExpandPSOSC(const TArray<FString>& Tokens)
 		else if (Item.Type == FPipelineCacheFileFormatPSO::DescriptorType::Graphics)
 		{
 			ActivePerSlot[SF_Vertex] = GetStableShaders(InverseMap, StableShaderKeyIndexTable, Item.GraphicsDesc.VertexShader, StableShadersPerSlot[SF_Vertex], OutAnyActiveButMissing);
+			ActivePerSlot[SF_Mesh] = GetStableShaders(InverseMap, StableShaderKeyIndexTable, Item.GraphicsDesc.MeshShader, StableShadersPerSlot[SF_Mesh], OutAnyActiveButMissing);
+			ActivePerSlot[SF_Amplification] = GetStableShaders(InverseMap, StableShaderKeyIndexTable, Item.GraphicsDesc.AmplificationShader, StableShadersPerSlot[SF_Amplification], OutAnyActiveButMissing);
 			ActivePerSlot[SF_Pixel] = GetStableShaders(InverseMap, StableShaderKeyIndexTable, Item.GraphicsDesc.FragmentShader, StableShadersPerSlot[SF_Pixel], OutAnyActiveButMissing);
 			ActivePerSlot[SF_Geometry] = GetStableShaders(InverseMap, StableShaderKeyIndexTable, Item.GraphicsDesc.GeometryShader, StableShadersPerSlot[SF_Geometry], OutAnyActiveButMissing);
 			ActivePerSlot[SF_Hull] = GetStableShaders(InverseMap, StableShaderKeyIndexTable, Item.GraphicsDesc.HullShader, StableShadersPerSlot[SF_Hull], OutAnyActiveButMissing);
@@ -843,6 +849,8 @@ int32 ExpandPSOSC(const TArray<FString>& Tokens)
 			{
 				UE_LOG(LogShaderPipelineCacheTools, Display, TEXT("   %s"), *Item.GraphicsDesc.StateToString());
 				PrintShaders(InverseMap, StableShaderKeyIndexTable, Item.GraphicsDesc.VertexShader, TEXT("VertexShader"));
+				PrintShaders(InverseMap, StableShaderKeyIndexTable, Item.GraphicsDesc.MeshShader, TEXT("MeshShader"));
+				PrintShaders(InverseMap, StableShaderKeyIndexTable, Item.GraphicsDesc.AmplificationShader, TEXT("AmplificationShader"));
 				PrintShaders(InverseMap, StableShaderKeyIndexTable, Item.GraphicsDesc.FragmentShader, TEXT("FragmentShader"));
 				PrintShaders(InverseMap, StableShaderKeyIndexTable, Item.GraphicsDesc.GeometryShader, TEXT("GeometryShader"));
 				PrintShaders(InverseMap, StableShaderKeyIndexTable, Item.GraphicsDesc.HullShader, TEXT("HullShader"));
@@ -920,6 +928,8 @@ int32 ExpandPSOSC(const TArray<FString>& Tokens)
 				UE_LOG(LogShaderPipelineCacheTools, Warning, TEXT("   %s"), *Item.GraphicsDesc.StateToString());
 
 				PrintShaders(InverseMap, StableShaderKeyIndexTable, Item.GraphicsDesc.VertexShader, TEXT("VertexShader"));
+				PrintShaders(InverseMap, StableShaderKeyIndexTable, Item.GraphicsDesc.MeshShader, TEXT("MeshShader"));
+				PrintShaders(InverseMap, StableShaderKeyIndexTable, Item.GraphicsDesc.AmplificationShader, TEXT("AmplificationShader"));
 				PrintShaders(InverseMap, StableShaderKeyIndexTable, Item.GraphicsDesc.FragmentShader, TEXT("FragmentShader"));
 				PrintShaders(InverseMap, StableShaderKeyIndexTable, Item.GraphicsDesc.GeometryShader, TEXT("GeometryShader"));
 				PrintShaders(InverseMap, StableShaderKeyIndexTable, Item.GraphicsDesc.HullShader, TEXT("HullShader"));
@@ -1199,17 +1209,9 @@ static TSet<FPipelineCacheFileFormatPSO> ParseStableCSV(const FString& FileName,
 			return;
 		}
 
-		TArray<FStringView, TInlineAllocator<2 + SF_Compute>> Parts;
+		TArray<FStringView, TInlineAllocator<2 + SF_NumFrequencies>> Parts;
 		ParseQuoteComma(Line, Parts);
 
-		if (Parts.Num() != 2 + SF_Compute) // SF_Compute here because the stablepc.csv file format does not have a compute slot
-		{
-			// Assume the rest of the file csv lines are are bad or are in an out of date format. If one is, they probably all are.
-			UE_LOG(LogShaderPipelineCacheTools, Warning, TEXT("File %s is not in the correct format ignoring the rest of its contents."), *FileName);
-			bParsed = false;
-			return;
-		}
-
 		FPipelineCacheFileFormatPSO PSO;
 		FMemory::Memzero(PSO);
 		PSO.Type = FPipelineCacheFileFormatPSO::DescriptorType::Graphics; // we will change this to compute later if needed
@@ -1224,24 +1226,64 @@ static TSet<FPipelineCacheFileFormatPSO> ParseStableCSV(const FString& FileName,
 		}
 
 		// For backward compatibility, compute shaders are stored as a zeroed graphics desc with the shader in the hull shader slot.
+		static FName NAME_SF_Vertex("SF_Vertex");
+		static FName NAME_SF_Mesh("SF_Mesh");
+		static FName NAME_SF_Amplification("SF_Amplification");
+		static FName NAME_SF_Hull("SF_Hull");
+		static FName NAME_SF_Domain("SF_Domain");
+		static FName NAME_SF_Pixel("SF_Pixel");
+		static FName NAME_SF_Geometry("SF_Geometry");
 		static FName NAME_SF_Compute("SF_Compute");
 		static FName NAME_SF_RayGen("SF_RayGen");
 		static FName NAME_SF_RayMiss("SF_RayMiss");
 		static FName NAME_SF_RayHitGroup("SF_RayHitGroup");
 		static FName NAME_SF_RayCallable("SF_RayCallable");
-		for (int32 SlotIndex = 0; SlotIndex < SF_Compute; ++SlotIndex)
+		for (int32 PartIndex = 2; PartIndex < Parts.Num(); ++PartIndex)
 		{
-			if (Parts[SlotIndex + 2].IsEmpty())
+			if (Parts[PartIndex].IsEmpty())
 			{
 				continue;
 			}
 
 			FStableShaderKeyAndValue Shader;
-			Shader.ParseFromString(Parts[SlotIndex + 2]);
+			Shader.ParseFromString(Parts[PartIndex]);
 
-			int32 AdjustedSlotIndex = SlotIndex;
+			int32 AdjustedSlotIndex = SF_Vertex;
 
-			if (Shader.TargetFrequency == NAME_SF_RayGen)
+			if (Shader.TargetFrequency == NAME_SF_Vertex)
+			{
+				AdjustedSlotIndex = SF_Vertex;
+			}
+			else if (Shader.TargetFrequency == NAME_SF_Mesh)
+			{
+				AdjustedSlotIndex = SF_Mesh;
+			}
+			else if (Shader.TargetFrequency == NAME_SF_Amplification)
+			{
+				AdjustedSlotIndex = SF_Amplification;
+			}
+			else if (Shader.TargetFrequency == NAME_SF_Hull)
+			{
+				AdjustedSlotIndex = SF_Hull;
+			}
+			else if (Shader.TargetFrequency == NAME_SF_Domain)
+			{
+				AdjustedSlotIndex = SF_Domain;
+			}
+			else if (Shader.TargetFrequency == NAME_SF_Pixel)
+			{
+				AdjustedSlotIndex = SF_Pixel;
+			}
+			else if (Shader.TargetFrequency == NAME_SF_Geometry)
+			{
+				AdjustedSlotIndex = SF_Geometry;
+			}
+			else if (Shader.TargetFrequency == NAME_SF_Compute)
+			{
+				PSO.Type = FPipelineCacheFileFormatPSO::DescriptorType::Compute;
+				AdjustedSlotIndex = SF_Compute;
+			}
+			else if (Shader.TargetFrequency == NAME_SF_RayGen)
 			{
 				PSO.Type = FPipelineCacheFileFormatPSO::DescriptorType::RayTracing;
 				AdjustedSlotIndex = SF_RayGen;
@@ -1263,20 +1305,9 @@ static TSet<FPipelineCacheFileFormatPSO> ParseStableCSV(const FString& FileName,
 			}
 			else
 			{
-				// Graphics and compute
-
-				if (SlotIndex == SF_Hull)
-				{
-					if (Shader.TargetFrequency == NAME_SF_Compute)
-					{
-						PSO.Type = FPipelineCacheFileFormatPSO::DescriptorType::Compute;
-						AdjustedSlotIndex = SF_Compute;
-					}
-				}
-				else
-				{
-					check(Shader.TargetFrequency != NAME_SF_Compute);
-				}
+				UE_LOG(LogShaderPipelineCacheTools, Warning, TEXT("File %s is not in the correct format (GraphicsDesc) ignoring the rest of its contents."), *FileName);
+				bParsed = false;
+				return;
 			}
 
 			FSHAHash Match;
@@ -1312,6 +1343,12 @@ static TSet<FPipelineCacheFileFormatPSO> ParseStableCSV(const FString& FileName,
 			case SF_Vertex:
 				PSO.GraphicsDesc.VertexShader = Match;
 				break;
+			case SF_Mesh:
+				PSO.GraphicsDesc.MeshShader = Match;
+				break;
+			case SF_Amplification:
+				PSO.GraphicsDesc.AmplificationShader = Match;
+				break;
 			case SF_Pixel:
 				PSO.GraphicsDesc.FragmentShader = Match;
 				break;
@@ -1346,6 +1383,8 @@ static TSet<FPipelineCacheFileFormatPSO> ParseStableCSV(const FString& FileName,
 		{
 			check(PSO.ComputeDesc.ComputeShader != FSHAHash() &&
 				PSO.GraphicsDesc.VertexShader == FSHAHash() &&
+				PSO.GraphicsDesc.MeshShader == FSHAHash() &&
+				PSO.GraphicsDesc.AmplificationShader == FSHAHash() &&
 				PSO.GraphicsDesc.FragmentShader == FSHAHash() &&
 				PSO.GraphicsDesc.GeometryShader == FSHAHash() &&
 				PSO.GraphicsDesc.HullShader == FSHAHash() &&
@@ -1997,7 +2036,7 @@ int32 BuildPSOSC(const TArray<FString>& Tokens)
 			}
 			else if (Item.Type == FPipelineCacheFileFormatPSO::DescriptorType::Graphics)
 			{
-				check(!(Item.GraphicsDesc.VertexShader == FSHAHash()));
+				check(!(Item.GraphicsDesc.VertexShader == FSHAHash() || Item.GraphicsDesc.MeshShader == FSHAHash()));
 				StringRep = Item.GraphicsDesc.ToString();
 			}
 			else if (Item.Type == FPipelineCacheFileFormatPSO::DescriptorType::RayTracing)
@@ -2035,6 +2074,8 @@ int32 BuildPSOSC(const TArray<FString>& Tokens)
 					{
 						if (
 							TestItem.GraphicsDesc.VertexShader == Item.GraphicsDesc.VertexShader &&
+							TestItem.GraphicsDesc.MeshShader == Item.GraphicsDesc.MeshShader &&
+							TestItem.GraphicsDesc.AmplificationShader == Item.GraphicsDesc.AmplificationShader &&
 							TestItem.GraphicsDesc.FragmentShader == Item.GraphicsDesc.FragmentShader &&
 							TestItem.GraphicsDesc.GeometryShader == Item.GraphicsDesc.GeometryShader &&
 							TestItem.GraphicsDesc.HullShader == Item.GraphicsDesc.HullShader &&
diff --git a/Engine/Source/Runtime/RenderCore/Private/ShaderCodeLibrary.cpp b/Engine/Source/Runtime/RenderCore/Private/ShaderCodeLibrary.cpp
--- a/Engine/Source/Runtime/RenderCore/Private/ShaderCodeLibrary.cpp
+++ b/Engine/Source/Runtime/RenderCore/Private/ShaderCodeLibrary.cpp
@@ -2656,6 +2656,24 @@ FComputeShaderRHIRef FShaderCodeLibrary::CreateComputeShader(EShaderPlatform Pla
 	return nullptr;
 }
 
+FMeshShaderRHIRef FShaderCodeLibrary::CreateMeshShader(EShaderPlatform Platform, const FSHAHash& Hash)
+{
+	if (FShaderLibrariesCollection::Impl)
+	{
+		return FMeshShaderRHIRef(FShaderLibrariesCollection::Impl->CreateShader(SF_Mesh, Hash));
+	}
+	return nullptr;
+}
+
+FAmplificationShaderRHIRef FShaderCodeLibrary::CreateAmplificationShader(EShaderPlatform Platform, const FSHAHash& Hash)
+{
+	if (FShaderLibrariesCollection::Impl)
+	{
+		return FAmplificationShaderRHIRef(FShaderLibrariesCollection::Impl->CreateShader(SF_Amplification, Hash));
+	}
+	return nullptr;
+}
+
 FRayTracingShaderRHIRef FShaderCodeLibrary::CreateRayTracingShader(EShaderPlatform Platform, const FSHAHash& Hash, EShaderFrequency Frequency)
 {
 	if (FShaderLibrariesCollection::Impl)
diff --git a/Engine/Source/Runtime/RenderCore/Private/ShaderPipelineCache.cpp b/Engine/Source/Runtime/RenderCore/Private/ShaderPipelineCache.cpp
--- a/Engine/Source/Runtime/RenderCore/Private/ShaderPipelineCache.cpp
+++ b/Engine/Source/Runtime/RenderCore/Private/ShaderPipelineCache.cpp
@@ -617,32 +617,56 @@ bool FShaderPipelineCache::Precompile(FRHICommandListImmediate& RHICmdList, ESha
 		if(FPipelineCacheFileFormatPSO::DescriptorType::Graphics == PSO.Type)
 		{
 			FGraphicsPipelineStateInitializer GraphicsInitializer;
-			
-			FRHIVertexDeclaration* VertexDesc = PipelineStateCache::GetOrCreateVertexDeclaration(PSO.GraphicsDesc.VertexDescriptor);
-			GraphicsInitializer.BoundShaderState.VertexDeclarationRHI = VertexDesc;
-			
-			FVertexShaderRHIRef VertexShader;
-			if (PSO.GraphicsDesc.VertexShader != FSHAHash())
-			{
-				VertexShader = FShaderCodeLibrary::CreateVertexShader(Platform, PSO.GraphicsDesc.VertexShader);
-				GraphicsInitializer.BoundShaderState.VertexShaderRHI = VertexShader;
-			}
 
-	#if PLATFORM_SUPPORTS_TESSELLATION_SHADERS
-			FHullShaderRHIRef HullShader;
-			if (PSO.GraphicsDesc.HullShader != FSHAHash())
+			if (PSO.GraphicsDesc.MeshShader != FSHAHash())
 			{
-				HullShader = FShaderCodeLibrary::CreateHullShader(Platform, PSO.GraphicsDesc.HullShader);
-				GraphicsInitializer.BoundShaderState.HullShaderRHI = HullShader;
-			}
+				FMeshShaderRHIRef MeshShader = FShaderCodeLibrary::CreateMeshShader(Platform, PSO.GraphicsDesc.MeshShader);
+				GraphicsInitializer.BoundShaderState.MeshShaderRHI = MeshShader;
 
-			FDomainShaderRHIRef DomainShader;
-			if (PSO.GraphicsDesc.DomainShader != FSHAHash())
+				if (PSO.GraphicsDesc.AmplificationShader != FSHAHash())
+				{
+					FAmplificationShaderRHIRef AmplificationShader = FShaderCodeLibrary::CreateAmplificationShader(Platform, PSO.GraphicsDesc.AmplificationShader);
+					GraphicsInitializer.BoundShaderState.AmplificationShaderRHI = AmplificationShader;
+				}
+			}
+			else
 			{
-				DomainShader = FShaderCodeLibrary::CreateDomainShader(Platform, PSO.GraphicsDesc.DomainShader);
-				GraphicsInitializer.BoundShaderState.DomainShaderRHI = DomainShader;
+				FRHIVertexDeclaration* VertexDesc = PipelineStateCache::GetOrCreateVertexDeclaration(PSO.GraphicsDesc.VertexDescriptor);
+				GraphicsInitializer.BoundShaderState.VertexDeclarationRHI = VertexDesc;
+			
+				FVertexShaderRHIRef VertexShader;
+				if (PSO.GraphicsDesc.VertexShader != FSHAHash())
+				{
+					VertexShader = FShaderCodeLibrary::CreateVertexShader(Platform, PSO.GraphicsDesc.VertexShader);
+					GraphicsInitializer.BoundShaderState.VertexShaderRHI = VertexShader;
+				}
+
+#if PLATFORM_SUPPORTS_TESSELLATION_SHADERS
+				FHullShaderRHIRef HullShader;
+				if (PSO.GraphicsDesc.HullShader != FSHAHash())
+				{
+					HullShader = FShaderCodeLibrary::CreateHullShader(Platform, PSO.GraphicsDesc.HullShader);
+					GraphicsInitializer.BoundShaderState.HullShaderRHI = HullShader;
+				}
+
+				FDomainShaderRHIRef DomainShader;
+				if (PSO.GraphicsDesc.DomainShader != FSHAHash())
+				{
+					DomainShader = FShaderCodeLibrary::CreateDomainShader(Platform, PSO.GraphicsDesc.DomainShader);
+					GraphicsInitializer.BoundShaderState.DomainShaderRHI = DomainShader;
+				}
+#endif
+
+#if PLATFORM_SUPPORTS_GEOMETRY_SHADERS
+				FGeometryShaderRHIRef GeometryShader;
+				if (PSO.GraphicsDesc.GeometryShader != FSHAHash())
+				{
+					GeometryShader = FShaderCodeLibrary::CreateGeometryShader(Platform, PSO.GraphicsDesc.GeometryShader);
+					GraphicsInitializer.BoundShaderState.GeometryShaderRHI = GeometryShader;
+				}
+#endif
 			}
-	#endif
+
 			FPixelShaderRHIRef FragmentShader;
 			if (PSO.GraphicsDesc.FragmentShader != FSHAHash())
 			{
@@ -650,14 +674,6 @@ bool FShaderPipelineCache::Precompile(FRHICommandListImmediate& RHICmdList, ESha
 				GraphicsInitializer.BoundShaderState.PixelShaderRHI = FragmentShader;
 			}
 
-	#if PLATFORM_SUPPORTS_GEOMETRY_SHADERS
-			FGeometryShaderRHIRef GeometryShader;
-			if (PSO.GraphicsDesc.GeometryShader != FSHAHash())
-			{
-				GeometryShader = FShaderCodeLibrary::CreateGeometryShader(Platform, PSO.GraphicsDesc.GeometryShader);
-				GraphicsInitializer.BoundShaderState.GeometryShaderRHI = GeometryShader;
-			}
-	#endif
 			auto BlendState = GetOrCreateBlendState(PSO.GraphicsDesc.BlendState);
 			GraphicsInitializer.BlendState = BlendState;
 			
@@ -808,7 +824,20 @@ void FShaderPipelineCache::PreparePipelineBatch(TDoubleLinkedList<FPipelineCache
             if (PSO.Type == FPipelineCacheFileFormatPSO::DescriptorType::Graphics)
 			{
                 // See if the shaders exist in the current code libraries, before trying to load the shader data
-                if (PSO.GraphicsDesc.VertexShader != EmptySHA)
+				if (PSO.GraphicsDesc.MeshShader != EmptySHA)
+				{
+					RequiredShaders.Add(PSO.GraphicsDesc.MeshShader);
+					bOK &= FShaderCodeLibrary::ContainsShaderCode(PSO.GraphicsDesc.MeshShader);
+					UE_CLOG(!bOK, LogRHI, Verbose, TEXT("Failed to find MeshShader shader: %s"), *(PSO.GraphicsDesc.MeshShader.ToString()));
+
+					if (PSO.GraphicsDesc.AmplificationShader != EmptySHA)
+					{
+						RequiredShaders.Add(PSO.GraphicsDesc.AmplificationShader);
+						bOK &= FShaderCodeLibrary::ContainsShaderCode(PSO.GraphicsDesc.AmplificationShader);
+						UE_CLOG(!bOK, LogRHI, Verbose, TEXT("Failed to find AmplificationShader shader: %s"), *(PSO.GraphicsDesc.AmplificationShader.ToString()));
+					}
+				}
+                else if (PSO.GraphicsDesc.VertexShader != EmptySHA)
                 {
                     RequiredShaders.Add(PSO.GraphicsDesc.VertexShader);
                     bOK &= FShaderCodeLibrary::ContainsShaderCode(PSO.GraphicsDesc.VertexShader);
@@ -826,12 +855,6 @@ void FShaderPipelineCache::PreparePipelineBatch(TDoubleLinkedList<FPipelineCache
 						bOK &= FShaderCodeLibrary::ContainsShaderCode(PSO.GraphicsDesc.DomainShader);
 						UE_CLOG(!bOK, LogRHI, Verbose, TEXT("Failed to find DomainShader shader: %s"), *(PSO.GraphicsDesc.DomainShader.ToString()));
 					}
-					if (PSO.GraphicsDesc.FragmentShader != EmptySHA)
-					{
-						RequiredShaders.Add(PSO.GraphicsDesc.FragmentShader);
-						bOK &= FShaderCodeLibrary::ContainsShaderCode(PSO.GraphicsDesc.FragmentShader);
-						UE_CLOG(!bOK, LogRHI, Verbose, TEXT("Failed to find FragmentShader shader: %s"), *(PSO.GraphicsDesc.FragmentShader.ToString()));
-					}
 					if (PSO.GraphicsDesc.GeometryShader != EmptySHA)
 					{
 						RequiredShaders.Add(PSO.GraphicsDesc.GeometryShader);
@@ -846,13 +869,34 @@ void FShaderPipelineCache::PreparePipelineBatch(TDoubleLinkedList<FPipelineCache
 					UE_LOG(LogRHI, Error, TEXT("PSO Entry has no vertex shader: %u this is an invalid entry!"), PSORead->Hash);
 					bOK = false;
 				}
-                
+
+				if (PSO.GraphicsDesc.FragmentShader != EmptySHA)
+				{
+					RequiredShaders.Add(PSO.GraphicsDesc.FragmentShader);
+					bOK &= FShaderCodeLibrary::ContainsShaderCode(PSO.GraphicsDesc.FragmentShader);
+					UE_CLOG(!bOK, LogRHI, Verbose, TEXT("Failed to find FragmentShader shader: %s"), *(PSO.GraphicsDesc.FragmentShader.ToString()));
+				}
+
                 // If everything is OK then we can issue reads of the actual shader code
 				if (bOK && PSO.GraphicsDesc.VertexShader != FSHAHash())
 				{
 					bOK &= FShaderCodeLibrary::PreloadShader(PSO.GraphicsDesc.VertexShader, AsyncJob.ReadRequests);
                     UE_CLOG(!bOK, LogRHI, Verbose, TEXT("Failed to read VertexShader shader: %s"), *(PSO.GraphicsDesc.VertexShader.ToString()));
 				}
+
+				// If everything is OK then we can issue reads of the actual shader code
+				if (bOK && PSO.GraphicsDesc.MeshShader != FSHAHash())
+				{
+					bOK &= FShaderCodeLibrary::PreloadShader(PSO.GraphicsDesc.MeshShader, AsyncJob.ReadRequests);
+					UE_CLOG(!bOK, LogRHI, Verbose, TEXT("Failed to read MeshShader shader: %s"), *(PSO.GraphicsDesc.MeshShader.ToString()));
+				}
+
+				// If everything is OK then we can issue reads of the actual shader code
+				if (bOK && PSO.GraphicsDesc.AmplificationShader != FSHAHash())
+				{
+					bOK &= FShaderCodeLibrary::PreloadShader(PSO.GraphicsDesc.AmplificationShader, AsyncJob.ReadRequests);
+					UE_CLOG(!bOK, LogRHI, Verbose, TEXT("Failed to read AmplificationShader shader: %s"), *(PSO.GraphicsDesc.AmplificationShader.ToString()));
+				}
 		
 				if (bOK && PSO.GraphicsDesc.HullShader != EmptySHA)
 				{
diff --git a/Engine/Source/Runtime/RenderCore/Public/ShaderCodeLibrary.h b/Engine/Source/Runtime/RenderCore/Public/ShaderCodeLibrary.h
--- a/Engine/Source/Runtime/RenderCore/Public/ShaderCodeLibrary.h
+++ b/Engine/Source/Runtime/RenderCore/Public/ShaderCodeLibrary.h
@@ -194,6 +194,8 @@ struct RENDERCORE_API FShaderCodeLibrary
 	static FDomainShaderRHIRef CreateDomainShader(EShaderPlatform Platform, const FSHAHash& Hash);
 	static FGeometryShaderRHIRef CreateGeometryShader(EShaderPlatform Platform, const FSHAHash& Hash);
 	static FComputeShaderRHIRef CreateComputeShader(EShaderPlatform Platform, const FSHAHash& Hash);
+	static FMeshShaderRHIRef CreateMeshShader(EShaderPlatform Platform, const FSHAHash& Hash);
+	static FAmplificationShaderRHIRef CreateAmplificationShader(EShaderPlatform Platform, const FSHAHash& Hash);
 	static FRayTracingShaderRHIRef CreateRayTracingShader(EShaderPlatform Platform, const FSHAHash& Hash, EShaderFrequency Frequency);
 
 	// Total number of shader entries in the library
