Don't include deleted external actors from autosave, since it will result in no package being saved anyway and will make that actor transient (which will be fixed by [REDACTED_PERSON] in Save2).

diff --git a/Engine/Source/Editor/UnrealEd/Private/FileHelpers.cpp b/Engine/Source/Editor/UnrealEd/Private/FileHelpers.cpp
--- a/Engine/Source/Editor/UnrealEd/Private/FileHelpers.cpp
+++ b/Engine/Source/Editor/UnrealEd/Private/FileHelpers.cpp
@@ -2861,15 +2861,33 @@ EAutosaveContentPackagesResult::Type FEditorFileUtils::AutosaveMapEx(const FStri
 					if (ExternalPackage->IsDirty() && (bForceIfNotInList || DirtyPackagesForAutoSave.Contains(ExternalPackage))
 						&& FPackageName::IsValidLongPackageName(ExternalPackage->GetName(), /*bIncludeReadOnlyRoots=*/false))
 					{
-						const FString AutosaveFilename = GetAutoSaveFilename(ExternalPackage, AbsoluteAutosaveDir, AutosaveIndex, FPackageName::GetAssetPackageExtension());
-						if (!GEditor->Exec(nullptr, *FString::Printf(TEXT("OBJ SAVEPACKAGE PACKAGE=\"%s\" FILE=\"%s\" SILENT=false AUTOSAVING=true"), *ExternalPackage->GetName(), *AutosaveFilename)))
+						// Don't try to save external packages that will get deleted
+						bool bActorPackageNeedsToSave = false;
+						ForEachObjectWithPackage(ExternalPackage, [&bActorPackageNeedsToSave](UObject* Object)
 						{
-							return EAutosaveContentPackagesResult::Failure;
+							if (Cast<AActor>(Object))
+							{
+								if (!Object->IsPendingKill())
+								{
+									bActorPackageNeedsToSave = true;
+									return false;
+								}
+							}
+							return true;
+						}, false);
+
+						if (bActorPackageNeedsToSave)
+						{
+							const FString AutosaveFilename = GetAutoSaveFilename(ExternalPackage, AbsoluteAutosaveDir, AutosaveIndex, FPackageName::GetAssetPackageExtension());
+							if (!GEditor->Exec(nullptr, *FString::Printf(TEXT("OBJ SAVEPACKAGE PACKAGE=\"%s\" FILE=\"%s\" SILENT=false AUTOSAVING=true"), *ExternalPackage->GetName(), *AutosaveFilename)))
+							{
+								return EAutosaveContentPackagesResult::Failure;
+							}
+							// We saved an actor
+							bResult = true;
+							// Re-mark the package as dirty, because autosaving it will have cleared the dirty flag
+							ExternalPackage->MarkPackageDirty();
 						}
-						// We saved an actor
-						bResult = true;
-						// Re-mark the package as dirty, because autosaving it will have cleared the dirty flag
-						ExternalPackage->MarkPackageDirty();
 					}
 				}
 			}
