Mode Manager InputKey should still send input to the modes, even when the tools context handles the press.

#rnx

diff --git a/Engine/Source/Editor/UnrealEd/Private/EditorModeManager.cpp b/Engine/Source/Editor/UnrealEd/Private/EditorModeManager.cpp
--- a/Engine/Source/Editor/UnrealEd/Private/EditorModeManager.cpp
+++ b/Engine/Source/Editor/UnrealEd/Private/EditorModeManager.cpp
@@ -1380,13 +1380,13 @@ bool FEditorModeTools::ProcessCapturedMouseMoves( FEditorViewportClient* InViewp
 /** Notifies all active modes of keyboard input */
 bool FEditorModeTools::InputKey(FEditorViewportClient* InViewportClient, FViewport* Viewport, FKey Key, EInputEvent Event, bool bRouteToToolsContext)
 {
-	if (bRouteToToolsContext && InteractiveToolsContext->InputKey(InViewportClient, Viewport, Key, Event))
+	bool bHandled = false;
+	if (bRouteToToolsContext)
 	{
-		return true;
+		bHandled |= InteractiveToolsContext->InputKey(InViewportClient, Viewport, Key, Event);
 	}
-	bool bHandled = false;
 	
-	// Copy the modes and iterate of that since a key may remove the edit mode and change CopyActiveScriptableModes
+	// Copy the modes and iterate of that since a key may remove the edit mode and change ActiveScriptableModes
 	TArray<UEdMode*> CopyActiveScriptableModes(ActiveScriptableModes);
 	for (UEdMode* Mode : CopyActiveScriptableModes)
 	{
