Updated async task notifications to the new progress notification api and reskined a few misc notifications

diff --git a/Engine/Source/Developer/SourceControl/Private/SSourceControlLogin.cpp b/Engine/Source/Developer/SourceControl/Private/SSourceControlLogin.cpp
--- a/Engine/Source/Developer/SourceControl/Private/SSourceControlLogin.cpp
+++ b/Engine/Source/Developer/SourceControl/Private/SSourceControlLogin.cpp
@@ -251,8 +251,9 @@ void SSourceControlLogin::DisplayConnectionSuccess() const
 	FNotificationInfo Info( LOCTEXT("ConnectionSuccessful", "Connection to source control was successful!") );
 	Info.bFireAndForget = true;
 	Info.bUseSuccessFailIcons = true;
-	Info.Image = FEditorStyle::GetBrush(TEXT("NotificationList.SuccessImage"));
-	FSlateNotificationManager::Get().AddNotification(Info);
+	TSharedPtr<SNotificationItem> Notification = FSlateNotificationManager::Get().AddNotification(Info);
+
+	Notification->SetCompletionState(SNotificationItem::CS_Success);
 }
 
 EVisibility SSourceControlLogin::GetThrobberVisibility() const
diff --git a/Engine/Source/Editor/UnrealEd/Private/Animation/AsyncAnimCompressionUI.cpp b/Engine/Source/Editor/UnrealEd/Private/Animation/AsyncAnimCompressionUI.cpp
--- a/Engine/Source/Editor/UnrealEd/Private/Animation/AsyncAnimCompressionUI.cpp
+++ b/Engine/Source/Editor/UnrealEd/Private/Animation/AsyncAnimCompressionUI.cpp
@@ -6,38 +6,30 @@
 #include "Animation/AnimCompressionDerivedDataPublic.h"
 #include "Animation/AnimSequence.h"
 
-/** Notification class for asynchronous shader compiling. */
-class FAnimCompressionNotificationImpl : public FGlobalEditorNotification
-{
-protected:
-	/** FGlobalEditorNotification interface */
-	virtual bool ShouldShowNotification(const bool bIsNotificationAlreadyActive) const override;
-	virtual void SetNotificationText(const TSharedPtr<SNotificationItem>& InNotificationItem) const override;
-};
-
 /** Global notification object. */
-FAnimCompressionNotificationImpl GAnimCompressionNotification;
-
-bool FAnimCompressionNotificationImpl::ShouldShowNotification(const bool bIsNotificationAlreadyActive) const
-{
-	const uint32 RemainingJobsThreshold = 0;
-	const uint32 ActiveJobs = GAsyncCompressedAnimationsTracker ? GAsyncCompressedAnimationsTracker->GetNumRemainingJobs() : 0;
 
-	return ActiveJobs > RemainingJobsThreshold;
-}
-
-void FAnimCompressionNotificationImpl::SetNotificationText(const TSharedPtr<SNotificationItem>& InNotificationItem) const
+/** Notification class for asynchronous shader compiling. */
+class FAnimCompressionNotificationImpl : public FGlobalEditorProgressNotification
 {
-	if(GAsyncCompressedAnimationsTracker)
+public:
+	FAnimCompressionNotificationImpl()
+		: FGlobalEditorProgressNotification(NSLOCTEXT("AsyncAnimCompression", "AnimCompressionInProgress", "Compressing Animations"))
+	{}
+protected:
+	/** FGlobalEditorProgressNotification interface */
+	virtual int32 UpdateProgress()
 	{
-		const int32 RemainingJobs = GAsyncCompressedAnimationsTracker->GetNumRemainingJobs();
+		const int32 RemainingJobs = GAsyncCompressedAnimationsTracker ? GAsyncCompressedAnimationsTracker->GetNumRemainingJobs() : 0;
 		if (RemainingJobs > 0)
 		{
 			FFormatNamedArguments Args;
 			Args.Add(TEXT("AnimsToCompress"), FText::AsNumber(RemainingJobs));
-			const FText ProgressMessage = FText::Format(NSLOCTEXT("AsyncAnimCompression", "AnimCompressionInProgressFormat", "Compressing Animations ({AnimsToCompress})"), Args);
+			UpdateProgressMessage(FText::Format(NSLOCTEXT("AsyncAnimCompression", "AnimCompressionInProgressFormat", "Compressing Animations ({AnimsToCompress})"), Args));
 
-			InNotificationItem->SetText(ProgressMessage);
 		}
+
+		return RemainingJobs;
 	}
-}
\ No newline at end of file
+};
+
+FAnimCompressionNotificationImpl GAnimCompressionNotification;
diff --git a/Engine/Source/Editor/UnrealEd/Private/AssetGuideline.cpp b/Engine/Source/Editor/UnrealEd/Private/AssetGuideline.cpp
--- a/Engine/Source/Editor/UnrealEd/Private/AssetGuideline.cpp
+++ b/Engine/Source/Editor/UnrealEd/Private/AssetGuideline.cpp
@@ -56,15 +56,11 @@ public:
 		ChildSlot
 		[
 			SNew(SBorder)
-			.BorderImage(FCoreStyle::Get().GetBrush("NotificationList.ItemBackground"))
+			.Padding( FMargin(5) )
+			.BorderImage(FCoreStyle::Get().GetBrush("NotificationList.ItemBackground_Border"))
+			.BorderBackgroundColor(FColor(0,0,0,1))
 			[
-				SNew(SBorder)
-				.Padding( FMargin(5) )
-				.BorderImage(FCoreStyle::Get().GetBrush("NotificationList.ItemBackground_Border"))
-				.BorderBackgroundColor(FColor(0,0,0,1))
-				[
-					ConstructInternals(InArgs)
-				]
+				ConstructInternals(InArgs)
 			]
 		];
 	}
diff --git a/Engine/Source/Editor/UnrealEd/Private/CardRepresentationBuildNotification.cpp b/Engine/Source/Editor/UnrealEd/Private/CardRepresentationBuildNotification.cpp
--- a/Engine/Source/Editor/UnrealEd/Private/CardRepresentationBuildNotification.cpp
+++ b/Engine/Source/Editor/UnrealEd/Private/CardRepresentationBuildNotification.cpp
@@ -5,44 +5,38 @@
 #include "Stats/Stats.h"
 #include "TickableEditorObject.h"
 #include "MeshCardRepresentation.h"
-#include "Framework/Notifications/NotificationManager.h"
-#include "Widgets/Notifications/SNotificationList.h"
 #include "Templates/UniquePtr.h"
+#include "GlobalEditorNotification.h"
 
 /** Notification class for asynchronous card representation building. */
-class FCardRepresentationBuildNotificationImpl
-	: public FTickableEditorObject
+class FCardRepresentationBuildNotificationImpl : public FGlobalEditorProgressNotification
 {
 
 public:
-
 	FCardRepresentationBuildNotificationImpl()
+		: FGlobalEditorProgressNotification(NSLOCTEXT("CardRepresentationBuild", "CardRepresentationBuildInProgress", "Building Cards"))
 	{
-		LastEnableTime = 0;
 	}
 
-	/** Starts the notification. */
-	void BuildStarted();
-
-	/** Ends the notification. */
-	void BuildFinished();
-
-protected:
-	/** FTickableEditorObject interface */
-	virtual void Tick(float DeltaTime) override;
-	virtual ETickableTickType GetTickableTickType() const override
+private:
+	virtual bool AllowedToStartNotification() const override
 	{
-		return ETickableTickType::Always;
+		return GCardRepresentationAsyncQueue ? GCardRepresentationAsyncQueue->GetNumOutstandingTasks() >= 10 : false;
 	}
-	virtual TStatId GetStatId() const override;
 
-private:
+	virtual int32 UpdateProgress()
+	{
+		const int32 RemainingJobs = GCardRepresentationAsyncQueue ? GCardRepresentationAsyncQueue->GetNumOutstandingTasks() : 0;
+		if (RemainingJobs > 0)
+		{
+			FFormatNamedArguments Args;
+			Args.Add(TEXT("BuildTasks"), FText::AsNumber(GCardRepresentationAsyncQueue->GetNumOutstandingTasks()));
+			UpdateProgressMessage(FText::Format(NSLOCTEXT("CardRepresentationBuild", "CardRepresentationBuildInProgressFormat", "Building Cards ({BuildTasks})"), Args));
 
-	/** Tracks the last time the notification was started, used to avoid spamming. */
-	double LastEnableTime;
+		}
 
-	/** In progress message */
-	TWeakPtr<SNotificationItem> NotificationPtr;
+		return RemainingJobs;
+	}
 };
 
 /** Global notification object. */
@@ -58,80 +52,3 @@ void TearDownCardRepresentationBuildNotification()
 {
 	GCardRepresentationBuildNotification = nullptr;
 }
-
-void FCardRepresentationBuildNotificationImpl::BuildStarted()
-{
-	LastEnableTime = FPlatformTime::Seconds();
-
-	// Starting a new request! Notify the UI.
-	if (NotificationPtr.IsValid())
-	{
-		NotificationPtr.Pin()->ExpireAndFadeout();
-	}
-	
-	FNotificationInfo Info( NSLOCTEXT("CardRepresentationBuild", "CardRepresentationBuildInProgress", "Building Card Representations") );
-	Info.bFireAndForget = false;
-	
-	// Setting fade out and expire time to 0 as the expire message is currently very obnoxious
-	Info.FadeOutDuration = 0.0f;
-	Info.ExpireDuration = 0.0f;
-
-	NotificationPtr = FSlateNotificationManager::Get().AddNotification(Info);
-
-	if (NotificationPtr.IsValid())
-	{
-		NotificationPtr.Pin()->SetCompletionState(SNotificationItem::CS_Pending);
-	}
-}
-
-void FCardRepresentationBuildNotificationImpl::BuildFinished()
-{
-	// Finished all requests! Notify the UI.
-	TSharedPtr<SNotificationItem> NotificationItem = NotificationPtr.Pin();
-
-	if (NotificationItem.IsValid())
-	{
-		NotificationItem->SetText( NSLOCTEXT("CardRepresentationBuild", "BuildFinished", "Finished building Card Representations") );
-		NotificationItem->SetCompletionState(SNotificationItem::CS_Success);
-		NotificationItem->ExpireAndFadeout();
-
-		NotificationPtr.Reset();
-	}
-}
-
-void FCardRepresentationBuildNotificationImpl::Tick(float DeltaTime) 
-{
-	if (GCardRepresentationAsyncQueue)
-	{
-		// Trigger a new notification if we are doing an async build, and we haven't displayed the notification recently
-		if (GCardRepresentationAsyncQueue->GetNumOutstandingTasks() > 0
-			&& !NotificationPtr.IsValid()
-			&& (FPlatformTime::Seconds() - LastEnableTime) > 5)
-		{
-			BuildStarted();
-		}
-		// Disable the notification when we are no longer doing an async compile
-		else if (GCardRepresentationAsyncQueue->GetNumOutstandingTasks() == 0 && NotificationPtr.IsValid())
-		{
-			BuildFinished();
-		}
-		else if (GCardRepresentationAsyncQueue->GetNumOutstandingTasks() > 0 && NotificationPtr.IsValid())
-		{
-			TSharedPtr<SNotificationItem> NotificationItem = NotificationPtr.Pin();
-
-			if (NotificationItem.IsValid())
-			{
-				FFormatNamedArguments Args;
-				Args.Add( TEXT("BuildTasks"), FText::AsNumber( GCardRepresentationAsyncQueue->GetNumOutstandingTasks() ) );
-				FText ProgressMessage = FText::Format(NSLOCTEXT("CardRepresentationBuild", "CardRepresentationBuildInProgressFormat", "Building Card Representations ({BuildTasks})"), Args);
-
-				NotificationItem->SetText( ProgressMessage );
-			}
-		}
-	}
-}
-
-TStatId FCardRepresentationBuildNotificationImpl::GetStatId() const
-{
-	RETURN_QUICK_DECLARE_CYCLE_STAT(FCardRepresentationBuildNotificationImpl, STATGROUP_Tickables);
-}
diff --git a/Engine/Source/Editor/UnrealEd/Private/DistanceFieldBuildNotification.cpp b/Engine/Source/Editor/UnrealEd/Private/DistanceFieldBuildNotification.cpp
--- a/Engine/Source/Editor/UnrealEd/Private/DistanceFieldBuildNotification.cpp
+++ b/Engine/Source/Editor/UnrealEd/Private/DistanceFieldBuildNotification.cpp
@@ -5,44 +5,33 @@
 #include "Stats/Stats.h"
 #include "TickableEditorObject.h"
 #include "DistanceFieldAtlas.h"
-#include "Framework/Notifications/NotificationManager.h"
-#include "Widgets/Notifications/SNotificationList.h"
 #include "Templates/UniquePtr.h"
+#include "GlobalEditorNotification.h"
 
 /** Notification class for asynchronous distance field building. */
-class FDistanceFieldBuildNotificationImpl
-	: public FTickableEditorObject
+class FDistanceFieldBuildNotificationImpl : public FGlobalEditorProgressNotification
 {
-
 public:
 
 	FDistanceFieldBuildNotificationImpl()
+		: FGlobalEditorProgressNotification(NSLOCTEXT("DistanceFieldBuild", "DistanceFieldBuildInProgress", "Building Mesh Distance Fields"))
 	{
-		LastEnableTime = 0;
 	}
 
-	/** Starts the notification. */
-	void DistanceFieldBuildStarted();
-
-	/** Ends the notification. */
-	void DistanceFieldBuildFinished();
-
 protected:
-	/** FTickableEditorObject interface */
-	virtual void Tick(float DeltaTime) override;
-	virtual ETickableTickType GetTickableTickType() const override
+	/** FGlobalEditorProgressNotification interface */
+	virtual int32 UpdateProgress()
 	{
-		return ETickableTickType::Always;
-	}
-	virtual TStatId GetStatId() const override;
-
-private:
-
-	/** Tracks the last time the notification was started, used to avoid spamming. */
-	double LastEnableTime;
+		const int32 RemainingJobs = GDistanceFieldAsyncQueue ? GDistanceFieldAsyncQueue->GetNumOutstandingTasks() : 0;
+		if (RemainingJobs > 0)
+		{
+			FFormatNamedArguments Args;
+			Args.Add(TEXT("BuildTasks"), FText::AsNumber(RemainingJobs));
+			UpdateProgressMessage(FText::Format(NSLOCTEXT("DistanceFieldBuild", "DistanceFieldBuildInProgressFormat", "Building Mesh Distance Fields ({BuildTasks})"), Args));
+		}
 
-	/** In progress message */
-	TWeakPtr<SNotificationItem> DistanceFieldNotificationPtr;
+		return RemainingJobs;
+	}
 };
 
 /** Global notification object. */
@@ -58,80 +47,3 @@ void TearDownDistanceFieldBuildNotification()
 {
 	GDistanceFieldBuildNotification = nullptr;
 }
-
-void FDistanceFieldBuildNotificationImpl::DistanceFieldBuildStarted()
-{
-	LastEnableTime = FPlatformTime::Seconds();
-
-	// Starting a new request! Notify the UI.
-	if (DistanceFieldNotificationPtr.IsValid())
-	{
-		DistanceFieldNotificationPtr.Pin()->ExpireAndFadeout();
-	}
-	
-	FNotificationInfo Info( NSLOCTEXT("DistanceFieldBuild", "DistanceFieldBuildInProgress", "Building Mesh Distance Fields") );
-	Info.bFireAndForget = false;
-	
-	// Setting fade out and expire time to 0 as the expire message is currently very obnoxious
-	Info.FadeOutDuration = 0.0f;
-	Info.ExpireDuration = 0.0f;
-
-	DistanceFieldNotificationPtr = FSlateNotificationManager::Get().AddNotification(Info);
-
-	if (DistanceFieldNotificationPtr.IsValid())
-	{
-		DistanceFieldNotificationPtr.Pin()->SetCompletionState(SNotificationItem::CS_Pending);
-	}
-}
-
-void FDistanceFieldBuildNotificationImpl::DistanceFieldBuildFinished()
-{
-	// Finished all requests! Notify the UI.
-	TSharedPtr<SNotificationItem> NotificationItem = DistanceFieldNotificationPtr.Pin();
-
-	if (NotificationItem.IsValid())
-	{
-		NotificationItem->SetText( NSLOCTEXT("DistanceFieldBuild", "DistanceFieldBuildFinished", "Finished building Distance Fields!") );
-		NotificationItem->SetCompletionState(SNotificationItem::CS_Success);
-		NotificationItem->ExpireAndFadeout();
-
-		DistanceFieldNotificationPtr.Reset();
-	}
-}
-
-void FDistanceFieldBuildNotificationImpl::Tick(float DeltaTime) 
-{
-	if (GDistanceFieldAsyncQueue)
-	{
-		// Trigger a new notification if we are doing an async build, and we haven't displayed the notification recently
-		if (GDistanceFieldAsyncQueue->GetNumOutstandingTasks() > 0
-			&& !DistanceFieldNotificationPtr.IsValid()
-			&& (FPlatformTime::Seconds() - LastEnableTime) > 5)
-		{
-			DistanceFieldBuildStarted();
-		}
-		// Disable the notification when we are no longer doing an async compile
-		else if (GDistanceFieldAsyncQueue->GetNumOutstandingTasks() == 0 && DistanceFieldNotificationPtr.IsValid())
-		{
-			DistanceFieldBuildFinished();
-		}
-		else if (GDistanceFieldAsyncQueue->GetNumOutstandingTasks() > 0 && DistanceFieldNotificationPtr.IsValid())
-		{
-			TSharedPtr<SNotificationItem> NotificationItem = DistanceFieldNotificationPtr.Pin();
-
-			if (NotificationItem.IsValid())
-			{
-				FFormatNamedArguments Args;
-				Args.Add( TEXT("BuildTasks"), FText::AsNumber( GDistanceFieldAsyncQueue->GetNumOutstandingTasks() ) );
-				FText ProgressMessage = FText::Format(NSLOCTEXT("DistanceFieldBuild", "DistanceFieldBuildInProgressFormat", "Building Mesh Distance Fields ({BuildTasks})"), Args);
-
-				NotificationItem->SetText( ProgressMessage );
-			}
-		}
-	}
-}
-
-TStatId FDistanceFieldBuildNotificationImpl::GetStatId() const
-{
-	RETURN_QUICK_DECLARE_CYCLE_STAT(FDistanceFieldBuildNotificationImpl, STATGROUP_Tickables);
-}
diff --git a/Engine/Source/Editor/UnrealEd/Private/GlobalEditorNotification.cpp b/Engine/Source/Editor/UnrealEd/Private/GlobalEditorNotification.cpp
--- a/Engine/Source/Editor/UnrealEd/Private/GlobalEditorNotification.cpp
+++ b/Engine/Source/Editor/UnrealEd/Private/GlobalEditorNotification.cpp
@@ -2,6 +2,8 @@
 
 #include "GlobalEditorNotification.h"
 #include "UnrealClient.h"
+#include "Engine/EngineTypes.h"
+#include "Editor.h"
 
 void FGlobalEditorNotification::Tick(float DeltaTime)
 {
@@ -12,3 +14,156 @@ TStatId FGlobalEditorNotification::GetStatId() const
 {
 	RETURN_QUICK_DECLARE_CYCLE_STAT(FGlobalEditorNotification, STATGROUP_Tickables);
 }
+
+void FGlobalEditorProgressNotification::UpdateProgressMessage(const FText NewMessage)
+{
+	ProgressMessage = NewMessage;
+}
+
+void FGlobalEditorProgressNotification::CancelNotification()
+{
+	FSlateNotificationManager::Get().CancelProgressNotification(NotificationHandle);
+	NotificationHandle.Reset();
+	TotalWorkNeeded = 0;
+	CurrentWorkCompleted = 0;
+}
+
+void FGlobalEditorProgressNotification::StartNotification(int32 InTotalWorkNeeded, FText ProgressText)
+{
+	CurrentWorkCompleted = 0;
+	TotalWorkNeeded = InTotalWorkNeeded;
+	NotificationHandle = FSlateNotificationManager::Get().StartProgressNotification(ProgressText, TotalWorkNeeded);
+}
+
+void FGlobalEditorProgressNotification::UpdateNotification(int32 InTotalWorkComplete, FText ProgressText)
+{
+	FSlateNotificationManager::Get().UpdateProgressNotification(NotificationHandle, InTotalWorkComplete, TotalWorkNeeded, ProgressText);
+}
+
+void FGlobalEditorProgressNotification::Tick(float DeltaTime)
+{
+	if(AllowedToStartNotification() || NotificationHandle.IsValid())
+	{
+		const int32 RemainingJobs = UpdateProgress();
+
+		if (RemainingJobs != 0 && !NotificationHandle.IsValid())
+		{
+			StartNotification(RemainingJobs, ProgressMessage);
+		}
+		else if (NotificationHandle.IsValid())
+		{
+			const int32 DeltaWorkDone = (TotalWorkNeeded - RemainingJobs) - CurrentWorkCompleted;
+			// More tasks were added, update the total work
+			if (DeltaWorkDone < 0)
+			{
+				TotalWorkNeeded += -DeltaWorkDone;
+			}
+			else if (DeltaWorkDone > 0)
+			{
+				CurrentWorkCompleted += DeltaWorkDone;
+			}
+
+			UpdateNotification(CurrentWorkCompleted, ProgressMessage);
+
+			
+		}
+
+		if (RemainingJobs == 0 && NotificationHandle.IsValid())
+		{
+			NotificationHandle.Reset();
+			TotalWorkNeeded = 0;
+			CurrentWorkCompleted = 0;
+		}
+	}
+}
+
+TStatId FGlobalEditorProgressNotification::GetStatId() const
+{
+	RETURN_QUICK_DECLARE_CYCLE_STAT(FGlobalEditorProgressNotification, STATGROUP_Tickables);
+}
+
+
+
+/** Notification class for asynchronous card representation building. */
+class FTestProgressNotification : public FGlobalEditorProgressNotification
+{
+
+public:
+	FTestProgressNotification()
+		: FGlobalEditorProgressNotification(FText::FromString(TEXT("Test Async Progres")))
+	{
+		GEditor->GetTimerManager()->SetTimer(
+			Handle,
+			FTimerDelegate::CreateLambda
+			(
+				[this]()
+				{
+					--RemainingJobs;
+				}
+			),
+			1.0f,
+			true
+		);
+	}
+
+	~FTestProgressNotification()
+	{
+		GEditor->GetTimerManager()->ClearTimer(Handle);
+	}
+private:
+	virtual int32 UpdateProgress()
+	{
+		if (RemainingJobs <= 0)
+		{
+			GEditor->GetTimerManager()->ClearTimer(Handle);
+			Handle.Invalidate();
+		}
+		return RemainingJobs;
+	}
+public:
+	FTimerHandle Handle;
+	int32 RemainingJobs = 0;
+};
+
+static TArray<TUniquePtr<FTestProgressNotification>> TestProgressNotificationTasks;
+
+
+static void StartWorkTest()
+{
+	TUniquePtr<FTestProgressNotification> TestProgressNotificationTask = MakeUnique<FTestProgressNotification>();
+	TestProgressNotificationTask->RemainingJobs = 100;
+
+	TestProgressNotificationTasks.Add(MoveTemp(TestProgressNotificationTask));
+}
+
+static void StopWorkTest()
+{
+	FTimerHandle Handle;
+	GEditor->GetTimerManager()->SetTimer(
+		Handle,
+		FTimerDelegate::CreateLambda
+		(
+			[]()
+			{
+				int32 TaskIndex = FMath::Rand() % TestProgressNotificationTasks.Num();
+				TUniquePtr<FTestProgressNotification>& TestProgressNotificationTask = TestProgressNotificationTasks[TaskIndex];
+				TestProgressNotificationTask->CancelNotification();
+				TestProgressNotificationTask->RemainingJobs = 0;
+			}
+		),
+		5.0f,
+		false);
+}
+
+static void AddWork()
+{
+	int32 TaskIndex = FMath::Rand() % TestProgressNotificationTasks.Num();
+	TUniquePtr<FTestProgressNotification>& TestProgressNotificationTask = TestProgressNotificationTasks[TaskIndex];
+	//if (TestProgressNotificationTask)
+	{
+		TestProgressNotificationTask->RemainingJobs += 100;
+	}
+}
+FAutoConsoleCommand StartWorkTestCommand(TEXT("StartWorkTest"), TEXT(""), FConsoleCommandDelegate::CreateStatic(&StartWorkTest));
+FAutoConsoleCommand AddWorkTestCommand(TEXT("AddWork"), TEXT(""), FConsoleCommandDelegate::CreateStatic(&AddWork));
+FAutoConsoleCommand StopWorkTestCommand(TEXT("StopWorkTest"), TEXT(""), FConsoleCommandDelegate::CreateStatic(&StopWorkTest));
diff --git a/Engine/Source/Editor/UnrealEd/Private/GlobalEditorNotification.h b/Engine/Source/Editor/UnrealEd/Private/GlobalEditorNotification.h
--- a/Engine/Source/Editor/UnrealEd/Private/GlobalEditorNotification.h
+++ b/Engine/Source/Editor/UnrealEd/Private/GlobalEditorNotification.h
@@ -6,6 +6,7 @@
 #include "Widgets/Notifications/GlobalNotification.h"
 #include "Stats/Stats.h"
 #include "TickableEditorObject.h"
+#include "Framework/Notifications/NotificationManager.h"
 
 /**
  * Class used to provide simple global editor notifications (for things like shader compilation and texture streaming) 
@@ -28,4 +29,47 @@ private:
 	virtual void Tick(float DeltaTime) override;
 	virtual ETickableTickType GetTickableTickType() const override { return ETickableTickType::Always; }
 	virtual TStatId GetStatId() const override;
+
+};
+
+/**
+ * A special kind of editor notification for showing async or background tasks with measurable progress
+ */
+class FGlobalEditorProgressNotification : public FTickableEditorObject
+{
+public:
+	FGlobalEditorProgressNotification(const FText InProgressMessage)
+		: ProgressMessage(InProgressMessage)
+	{}
+
+public:
+	/** Whether or not the notification is allowed to start. */
+	virtual bool AllowedToStartNotification() const { return true; }
+
+	/**
+	 * Called every frame to update progress. The first time a non-zero value is returned, the notification starts and will end and reach 100% completion when
+	 * the number of jobs reaches zero again
+	 * 
+	 * @return Return the number of remaining work for this notification. (E.g number of threaded jobs left)
+ 	 */
+	virtual int32 UpdateProgress() = 0;
+
+	/** Updates the display text of this notification */
+	void UpdateProgressMessage(const FText NewMessage);
+
+	/** Cancels this notification */
+	void CancelNotification();
+private:
+	void StartNotification(int32 InStartWorkNum, FText ProgressText);
+	void UpdateNotification(int32 InTotalWorkComplete, FText ProgressText);
+
+	/** FTickableEditorObject interface */
+	virtual void Tick(float DeltaTime) override;
+	virtual ETickableTickType GetTickableTickType() const override { return ETickableTickType::Always; }
+	virtual TStatId GetStatId() const override;
+private:
+	FProgressNotificationHandle NotificationHandle;
+	FText ProgressMessage;
+	int32 TotalWorkNeeded = 0;
+	int32 CurrentWorkCompleted = 0;
 };
diff --git a/Engine/Source/Editor/UnrealEd/Private/ShaderCompilingNotification.cpp b/Engine/Source/Editor/UnrealEd/Private/ShaderCompilingNotification.cpp
--- a/Engine/Source/Editor/UnrealEd/Private/ShaderCompilingNotification.cpp
+++ b/Engine/Source/Editor/UnrealEd/Private/ShaderCompilingNotification.cpp
@@ -7,32 +7,35 @@
 #include "Widgets/Notifications/SNotificationList.h"
 
 /** Notification class for asynchronous shader compiling. */
-class FShaderCompilingNotificationImpl : public FGlobalEditorNotification
+class FShaderCompilingNotificationImpl : public FGlobalEditorProgressNotification
 {
+public:
+	FShaderCompilingNotificationImpl()
+		: FGlobalEditorProgressNotification(NSLOCTEXT("ShaderCompile", "ShaderCompileInProgress", "Compiling Shaders"))
+	{
+
+	}
+
 protected:
-	/** FGlobalEditorNotification interface */
-	virtual bool ShouldShowNotification(const bool bIsNotificationAlreadyActive) const override;
-	virtual void SetNotificationText(const TSharedPtr<SNotificationItem>& InNotificationItem) const override;
+	virtual bool AllowedToStartNotification() const override
+	{
+		return GShaderCompilingManager->ShouldDisplayCompilingNotification();
+	}
+	virtual int32 UpdateProgress()
+	{
+		const int32 RemainingJobs = GShaderCompilingManager->IsCompiling() ? GShaderCompilingManager->GetNumRemainingJobs() : 0;
+		if (RemainingJobs > 0)
+		{
+			FFormatNamedArguments Args;
+			Args.Add(TEXT("ShaderJobs"), FText::AsNumber(GShaderCompilingManager->GetNumRemainingJobs()));
+			UpdateProgressMessage(FText::Format(NSLOCTEXT("ShaderCompile", "ShaderCompileInProgressFormat", "Compiling Shaders ({ShaderJobs})"), Args));
+
+		}
+
+		return RemainingJobs;
+	}
 };
 
 /** Global notification object. */
 FShaderCompilingNotificationImpl GShaderCompilingNotification;
 
-bool FShaderCompilingNotificationImpl::ShouldShowNotification(const bool bIsNotificationAlreadyActive) const
-{
-	// ShouldDisplayCompilingNotification is more a hint, and may start returning false while there's still work being done
-	// If we're already showing the notification, we should continue to do so until all the jobs have finished
-	return GShaderCompilingManager->ShouldDisplayCompilingNotification() || (bIsNotificationAlreadyActive && GShaderCompilingManager->IsCompiling());
-}
-
-void FShaderCompilingNotificationImpl::SetNotificationText(const TSharedPtr<SNotificationItem>& InNotificationItem) const
-{
-	if (GShaderCompilingManager->IsCompiling())
-	{
-		FFormatNamedArguments Args;
-		Args.Add(TEXT("ShaderJobs"), FText::AsNumber(GShaderCompilingManager->GetNumRemainingJobs()));
-		const FText ProgressMessage = FText::Format(NSLOCTEXT("ShaderCompile", "ShaderCompileInProgressFormat", "Compiling Shaders ({ShaderJobs})"), Args);
-
-		InNotificationItem->SetText(ProgressMessage);
-	}
-}
diff --git a/Engine/Source/Editor/UnrealEd/Private/Subsystems/AssetEditorSubsystem.cpp b/Engine/Source/Editor/UnrealEd/Private/Subsystems/AssetEditorSubsystem.cpp
--- a/Engine/Source/Editor/UnrealEd/Private/Subsystems/AssetEditorSubsystem.cpp
+++ b/Engine/Source/Editor/UnrealEd/Private/Subsystems/AssetEditorSubsystem.cpp
@@ -696,11 +696,12 @@ void UAssetEditorSubsystem::RestorePreviouslyOpenAssets()
 	TArray<FString> OpenAssets;
 	GConfig->GetArray(TEXT("AssetEditorSubsystem"), TEXT("OpenAssetsAtExit"), OpenAssets, GEditorPerProjectIni);
 
-	bool bCleanShutdown = false;
+	bool bCleanShutdown = true;
 	GConfig->GetBool(TEXT("AssetEditorSubsystem"), TEXT("CleanShutdown"), bCleanShutdown, GEditorPerProjectIni);
 
 	SaveOpenAssetEditors(false);
 
+/*
 	if (OpenAssets.Num() > 0)
 	{
 		if (bCleanShutdown)
@@ -720,10 +721,10 @@ void UAssetEditorSubsystem::RestorePreviouslyOpenAssets()
 				GConfig->GetBool(TEXT("AssetEditorSubsystem"), TEXT("SuppressRestorePreviouslyOpenAssetsNotification"), bSuppressNotification, GEditorPerProjectIni);
 
 				if (!bSuppressNotification)
-				{
+				{*/
 					// Ask the user; this doesn't block so will reopen the assets later
 					SpawnRestorePreviouslyOpenAssetsNotification(bCleanShutdown, OpenAssets);
-				}
+	/*			}
 			}
 		}
 		else
@@ -731,7 +732,7 @@ void UAssetEditorSubsystem::RestorePreviouslyOpenAssets()
 			// If we crashed, we always ask regardless of what the user previously said
 			SpawnRestorePreviouslyOpenAssetsNotification(bCleanShutdown, OpenAssets);
 		}
-	}
+	}*/
 }
 
 void UAssetEditorSubsystem::SetAutoRestoreAndDisableSaving(const bool bInAutoRestoreAndDisableSaving)
@@ -783,8 +784,8 @@ void UAssetEditorSubsystem::SpawnRestorePreviouslyOpenAssetsNotification(const b
 	));
 
 	// We will let the notification expire automatically after 10 seconds
-	Info.bFireAndForget = false;
-	Info.ExpireDuration = 10.0f;
+	Info.bFireAndForget = true;
+	Info.ExpireDuration = 20.0f;
 
 	// We want the auto-save to be subtle
 	Info.bUseLargeFont = false;
diff --git a/Engine/Source/Editor/UnrealEd/Private/TextureStreamingNotification.cpp b/Engine/Source/Editor/UnrealEd/Private/TextureStreamingNotification.cpp
--- a/Engine/Source/Editor/UnrealEd/Private/TextureStreamingNotification.cpp
+++ b/Engine/Source/Editor/UnrealEd/Private/TextureStreamingNotification.cpp
@@ -5,15 +5,37 @@
 #include "Editor.h"
 #include "GlobalEditorNotification.h"
 #include "ContentStreaming.h"
-#include "Widgets/Notifications/SNotificationList.h"
 
 /** Notification class for texture streaming. */
-class FTextureStreamingNotificationImpl : public FGlobalEditorNotification
+class FTextureStreamingNotificationImpl : public FGlobalEditorProgressNotification
 {
+public:
+	FTextureStreamingNotificationImpl()
+		: FGlobalEditorProgressNotification(NSLOCTEXT("StreamingTextures", "StreamingTexturesInProgress", "Streaming Textures"))
+	{}
+
 protected:
 	/** FGlobalEditorNotification interface */
-	virtual bool ShouldShowNotification(const bool bIsNotificationAlreadyActive) const override;
-	virtual void SetNotificationText(const TSharedPtr<SNotificationItem>& InNotificationItem) const override;
+	virtual bool AllowedToStartNotification() const override
+	{
+		// We only want to show the notification initially if there's still enough work left to do to warrant it
+		// Never show these notifications during PIE
+		const int32 NumStreamingTextures = GetNumStreamingTextures();
+		return !GEditor->PlayWorld && NumStreamingTextures > 300;
+	}
+
+	virtual int32 UpdateProgress() override
+	{
+		const int32 NumStreamingTextures = GetNumStreamingTextures();
+		if (NumStreamingTextures > 0)
+		{
+			FFormatNamedArguments Args;
+			Args.Add(TEXT("NumTextures"), FText::AsNumber(NumStreamingTextures));
+			UpdateProgressMessage(FText::Format(NSLOCTEXT("StreamingTextures", "StreamingTexturesInProgressFormat", "Streaming Textures ({NumTextures})"), Args));
+		}
+
+		return NumStreamingTextures;
+	}
 
 private:
 	static int32 GetNumStreamingTextures();
@@ -22,28 +44,6 @@ private:
 /** Global notification object. */
 FTextureStreamingNotificationImpl GTextureStreamingNotification;
 
-bool FTextureStreamingNotificationImpl::ShouldShowNotification(const bool bIsNotificationAlreadyActive) const
-{
-	// We only want to show the notification initially if there's still enough work left to do to warrant it
-	// However, if we're already showing the notification, we should continue to do so until all the streaming has finished
-	// Never show these notifications during PIE
-	const int32 NumStreamingTextures = GetNumStreamingTextures();
-	return !GEditor->PlayWorld && ((NumStreamingTextures > 300) || (bIsNotificationAlreadyActive && NumStreamingTextures > 0));
-}
-
-void FTextureStreamingNotificationImpl::SetNotificationText(const TSharedPtr<SNotificationItem>& InNotificationItem) const
-{
-	const int32 NumStreamingTextures = GetNumStreamingTextures();
-	if (NumStreamingTextures > 0)
-	{
-		FFormatNamedArguments Args;
-		Args.Add(TEXT("NumTextures"), FText::AsNumber(NumStreamingTextures));
-		const FText ProgressMessage = FText::Format(NSLOCTEXT("StreamingTextures", "StreamingTexturesInProgressFormat", "Streaming Textures ({NumTextures})"), Args);
-
-		InNotificationItem->SetText(ProgressMessage);
-	}
-}
-
 int32 FTextureStreamingNotificationImpl::GetNumStreamingTextures()
 {
 	FStreamingManagerCollection& StreamingManagerCollection = IStreamingManager::Get();
diff --git a/Engine/Source/Runtime/Engine/Private/AsyncCompilationHelpers.cpp b/Engine/Source/Runtime/Engine/Private/AsyncCompilationHelpers.cpp
--- a/Engine/Source/Runtime/Engine/Private/AsyncCompilationHelpers.cpp
+++ b/Engine/Source/Runtime/Engine/Private/AsyncCompilationHelpers.cpp
@@ -39,45 +39,31 @@ static FAutoConsoleCommand CVarAsyncAssetDumpStallStacks(
 void FAsyncCompilationNotification::Update(int32 NumJobs)
 {
 	check(IsInGameThread());
-	TSharedPtr<SNotificationItem> NotificationItem = NotificationItemPtr.Pin();
 
 	FFormatNamedArguments Args;
 	Args.Add(TEXT("AssetType"), AssetType);
+	Args.Add(TEXT("NumJobs"), FText::AsNumber(NumJobs));
+	FText ProgressMessage = FText::Format(LOCTEXT("AsyncCompilationProgress", "Preparing {AssetType} ({NumJobs})"), Args);
 
 	if (NumJobs == 0)
 	{
-		if (NotificationItem.IsValid())
+		if (NotificationHandle.IsValid())
 		{
-			NotificationItem->SetText(FText::Format(LOCTEXT("AsyncCompilationFinished", "Finished preparing {AssetType}!"), Args));
-			NotificationItem->SetCompletionState(SNotificationItem::CS_Success);
-			NotificationItem->ExpireAndFadeout();
-
-			NotificationItemPtr.Reset();
+			FSlateNotificationManager::Get().UpdateProgressNotification(NotificationHandle, StartNumJobs, StartNumJobs, ProgressMessage);
 		}
+		StartNumJobs = 0;
+		NotificationHandle = FProgressNotificationHandle();
 	}
-	else
+	else 
 	{
-		Args.Add(TEXT("NumJobs"), FText::AsNumber(NumJobs));
-		FText ProgressMessage = FText::Format(LOCTEXT("AsyncCompilationProgress", "Preparing {AssetType} ({NumJobs})"), Args);
-
-		if (!NotificationItem.IsValid())
+		if(!NotificationHandle.IsValid())
 		{
-			FNotificationInfo Info(ProgressMessage);
-			Info.bFireAndForget = false;
-
-			// Setting fade out and expire time to 0 as the expire message is currently very obnoxious
-			Info.FadeOutDuration = 0.0f;
-			Info.ExpireDuration = 0.0f;
-
-			NotificationItem = FSlateNotificationManager::Get().AddNotification(Info);
-			NotificationItemPtr = NotificationItem;
+			StartNumJobs = NumJobs;
+			NotificationHandle = FSlateNotificationManager::Get().StartProgressNotification(ProgressMessage, StartNumJobs);
 		}
-
-		if (NotificationItem.IsValid())
+		else
 		{
-			NotificationItem->SetCompletionState(SNotificationItem::CS_Pending);
-			NotificationItem->SetVisibility(EVisibility::HitTestInvisible);
-			NotificationItem->SetText(ProgressMessage);
+			FSlateNotificationManager::Get().UpdateProgressNotification(NotificationHandle, StartNumJobs - NumJobs, StartNumJobs, ProgressMessage);
 		}
 	}
 };
diff --git a/Engine/Source/Runtime/Engine/Public/AsyncCompilationHelpers.h b/Engine/Source/Runtime/Engine/Public/AsyncCompilationHelpers.h
--- a/Engine/Source/Runtime/Engine/Public/AsyncCompilationHelpers.h
+++ b/Engine/Source/Runtime/Engine/Public/AsyncCompilationHelpers.h
@@ -7,10 +7,10 @@
 
 #include "HAL/IConsoleManager.h"
 #include "Logging/LogMacros.h"
+#include "Framework/Notifications/NotificationManager.h"
 
 DECLARE_LOG_CATEGORY_EXTERN(LogAsyncCompilation, Log, All);
 
-class SNotificationItem;
 class FQueuedThreadPoolWrapper;
 
 class FAsyncCompilationNotification
@@ -23,7 +23,9 @@ public:
 
 	void Update(int32 NumJobs);
 private:
-	TWeakPtr<SNotificationItem> NotificationItemPtr;
+	FProgressNotificationHandle NotificationHandle;
+	/** The number of jobs when the notification began */
+	int32 StartNumJobs;
 	FText AssetType;
 };
 
