Remove commented out code used for testing

diff --git a/Engine/Source/Editor/UnrealEd/Private/Subsystems/AssetEditorSubsystem.cpp b/Engine/Source/Editor/UnrealEd/Private/Subsystems/AssetEditorSubsystem.cpp
--- a/Engine/Source/Editor/UnrealEd/Private/Subsystems/AssetEditorSubsystem.cpp
+++ b/Engine/Source/Editor/UnrealEd/Private/Subsystems/AssetEditorSubsystem.cpp
@@ -701,7 +701,6 @@ void UAssetEditorSubsystem::RestorePreviouslyOpenAssets()
 
 	SaveOpenAssetEditors(false);
 
-/*
 	if (OpenAssets.Num() > 0)
 	{
 		if (bCleanShutdown)
@@ -721,10 +720,10 @@ void UAssetEditorSubsystem::RestorePreviouslyOpenAssets()
 				GConfig->GetBool(TEXT("AssetEditorSubsystem"), TEXT("SuppressRestorePreviouslyOpenAssetsNotification"), bSuppressNotification, GEditorPerProjectIni);
 
 				if (!bSuppressNotification)
-				{*/
+				{
 					// Ask the user; this doesn't block so will reopen the assets later
 					SpawnRestorePreviouslyOpenAssetsNotification(bCleanShutdown, OpenAssets);
-	/*			}
+				}
 			}
 		}
 		else
@@ -732,7 +731,7 @@ void UAssetEditorSubsystem::RestorePreviouslyOpenAssets()
 			// If we crashed, we always ask regardless of what the user previously said
 			SpawnRestorePreviouslyOpenAssetsNotification(bCleanShutdown, OpenAssets);
 		}
-	}*/
+	}
 }
 
 void UAssetEditorSubsystem::SetAutoRestoreAndDisableSaving(const bool bInAutoRestoreAndDisableSaving)
