Made UEdMode disallow tool startup in PIE and SIE by default, and made that behavior subclass-configurable.

#rnx

diff --git a/Engine/Source/Editor/UnrealEd/Private/Tools/UEdMode.cpp b/Engine/Source/Editor/UnrealEd/Private/Tools/UEdMode.cpp
--- a/Engine/Source/Editor/UnrealEd/Private/Tools/UEdMode.cpp
+++ b/Engine/Source/Editor/UnrealEd/Private/Tools/UEdMode.cpp
@@ -228,13 +228,22 @@ void UEdMode::RegisterTool(TSharedPtr<FUICommandInfo> UICommand, FString ToolIde
 	ToolsContext->ToolManager->RegisterToolType(ToolIdentifier, Builder);
 	CommandList->MapAction(UICommand,
 		FExecuteAction::CreateUObject(ToolsContext.Get(), &UEdModeInteractiveToolsContext::StartTool, ToolIdentifier),
-		FCanExecuteAction::CreateUObject(ToRawPtr(ToolsContext->ToolManager), &UInteractiveToolManager::CanActivateTool, EToolSide::Mouse, ToolIdentifier),
+		FCanExecuteAction::CreateWeakLambda(ToolsContext.Get(), [this, ToolIdentifier]() {
+			return ShouldToolStartBeAllowed(ToolIdentifier) && 
+				 ToolsContext->ToolManager->CanActivateTool(EToolSide::Mouse, ToolIdentifier);
+			}),
 		FIsActionChecked::CreateUObject(ToolsContext.Get(), &UEdModeInteractiveToolsContext::IsToolActive, EToolSide::Mouse, ToolIdentifier),
 		EUIActionRepeatMode::RepeatDisabled);
 
 	RegisteredTools.Emplace(UICommand, ToolIdentifier);
 }
 
+bool UEdMode::ShouldToolStartBeAllowed(const FString& ToolIdentifier) const
+{
+	// Disallow starting tools when playing in editor or simulating.
+	return !GEditor->PlayWorld && !GIsPlayInEditorWorld;
+}
+
 void UEdMode::Exit()
 {
 	if (SettingsObject)
diff --git a/Engine/Source/Editor/UnrealEd/Public/Tools/UEdMode.h b/Engine/Source/Editor/UnrealEd/Public/Tools/UEdMode.h
--- a/Engine/Source/Editor/UnrealEd/Public/Tools/UEdMode.h
+++ b/Engine/Source/Editor/UnrealEd/Public/Tools/UEdMode.h
@@ -163,6 +163,11 @@ public:
 
 	virtual void RegisterTool(TSharedPtr<FUICommandInfo> UICommand, FString ToolIdentifier, UInteractiveToolBuilder* Builder);
 
+	/** 
+	 * Subclasses can override this to add additional checks on whether a tool should be allowed to start.
+	 * By default the check disallows starting tools during play/simulate in editor.
+	 */
+	virtual bool ShouldToolStartBeAllowed(const FString& ToolIdentifier) const;
 
 	virtual void Exit();
 	virtual UTexture2D* GetVertexTexture();
