Source-control & World Partition changes needed for Mantle support

diff --git a/Engine/Source/Editor/UnrealEd/Private/PackageSourceControlHelper.cpp b/Engine/Source/Editor/UnrealEd/Private/PackageSourceControlHelper.cpp
--- a/Engine/Source/Editor/UnrealEd/Private/PackageSourceControlHelper.cpp
+++ b/Engine/Source/Editor/UnrealEd/Private/PackageSourceControlHelper.cpp
@@ -30,81 +30,144 @@ ISourceControlProvider& FPackageSourceControlHelper::GetSourceControlProvider()
 
 bool FPackageSourceControlHelper::Delete(const FString& PackageName) const
 {
-	FString Filename = SourceControlHelpers::PackageFilename(PackageName);
+	TArray<FString> PackageNames = { PackageName };
+	return Delete(PackageNames);
+}
 
-	UE_LOG(LogCommandletPackageHelper, Verbose, TEXT("Deleting %s"), *Filename);
-	
+bool FPackageSourceControlHelper::Delete(const TArray<FString>& PackageNames) const
+{
+	// Early out when not using source control
 	if (!UseSourceControl())
 	{
-		if (!IPlatformFile::GetPlatformPhysical().SetReadOnly(*Filename, false) ||
-			!IPlatformFile::GetPlatformPhysical().DeleteFile(*Filename))
+		for (const FString& PackageName : PackageNames)
 		{
-			UE_LOG(LogCommandletPackageHelper, Error, TEXT("Error deleting %s"), *Filename);
-			return false;
+			FString Filename = SourceControlHelpers::PackageFilename(PackageName);
+
+			if (!IPlatformFile::GetPlatformPhysical().SetReadOnly(*Filename, false) ||
+				!IPlatformFile::GetPlatformPhysical().DeleteFile(*Filename))
+			{
+				UE_LOG(LogCommandletPackageHelper, Error, TEXT("Error deleting %s"), *Filename);
+				return false;
+			}
 		}
+
+		return true;
 	}
-	else
+
+	TArray<FString> FilesToRevert;
+	TArray<FString> FilesToDeleteFromDisk;
+	TArray<FString> FilesToDeleteFromSCC;
+
+	bool bSCCErrorsFound = false;
+
+	// First: get latest state from source control
+	TArray<FString> Filenames;
+	TArray<FSourceControlStateRef> SourceControlStates;
+	Filenames.Reset(PackageNames.Num());
+
+	for (const FString& PackageName : PackageNames)
+	{
+		Filenames.Emplace(SourceControlHelpers::PackageFilename(PackageName));
+	}
+
+	ECommandResult::Type UpdateState = GetSourceControlProvider().GetState(Filenames, SourceControlStates, EStateCacheUsage::ForceUpdate);
+
+	if (UpdateState != ECommandResult::Succeeded)
+	{
+		UE_LOG(LogCommandletPackageHelper, Error, TEXT("Could not get source control state for packages"));
+		return false;
+	}
+
+	for(FSourceControlStateRef& SourceControlState : SourceControlStates)
 	{
-		FSourceControlStatePtr SourceControlState = GetSourceControlProvider().GetState(Filename, EStateCacheUsage::ForceUpdate);
+		const FString& Filename = SourceControlState->GetFilename();
 
-		if (SourceControlState.IsValid() && SourceControlState->IsSourceControlled())
+		UE_LOG(LogCommandletPackageHelper, Verbose, TEXT("Deleting %s"), *Filename);
+
+		if (SourceControlState->IsSourceControlled())
 		{
 			FString OtherCheckedOutUser;
 			if (SourceControlState->IsCheckedOutOther(&OtherCheckedOutUser))
 			{
 				UE_LOG(LogCommandletPackageHelper, Error, TEXT("Overwriting package %s already checked out by %s, will not submit"), *Filename, *OtherCheckedOutUser);
-				return false;
+				bSCCErrorsFound = true;
 			}
 			else if (!SourceControlState->IsCurrent())
 			{
 				UE_LOG(LogCommandletPackageHelper, Error, TEXT("Overwriting package %s (not at head revision), will not submit"), *Filename);
-				return false;
+				bSCCErrorsFound = true;
 			}
 			else if (SourceControlState->IsAdded())
 			{
-				if (GetSourceControlProvider().Execute(ISourceControlOperation::Create<FRevert>(), Filename) != ECommandResult::Succeeded)
-				{
-					UE_LOG(LogCommandletPackageHelper, Error, TEXT("Error reverting package %s from source control"), *Filename);
-					return false;
-				}
-
-				if (!IFileManager::Get().Delete(*Filename, false, true))
-				{
-					UE_LOG(LogCommandletPackageHelper, Error, TEXT("Error deleting package %s locally"), *Filename);
-					return false;
-				}
+				FilesToRevert.Add(Filename);
+				FilesToDeleteFromDisk.Add(Filename);
 			}
 			else
 			{
-				UE_LOG(LogCommandletPackageHelper, Log, TEXT("Deleting package %s from source control"), *Filename);
-
 				if (SourceControlState->IsCheckedOut())
 				{
-					if (GetSourceControlProvider().Execute(ISourceControlOperation::Create<FRevert>(), Filename) != ECommandResult::Succeeded)
-					{
-						UE_LOG(LogCommandletPackageHelper, Error, TEXT("Error reverting package %s from source control"), *Filename);
-						return false;
-					}
+					FilesToRevert.Add(Filename);
 				}
 
-				if (GetSourceControlProvider().Execute(ISourceControlOperation::Create<FDelete>(), Filename) != ECommandResult::Succeeded)
-				{
-					UE_LOG(LogCommandletPackageHelper, Error, TEXT("Error deleting package %s from source control"), *Filename);
-					return false;
-				}
+				FilesToDeleteFromSCC.Add(Filename);
 			}
 		}
 		else
 		{
-			if (!IFileManager::Get().Delete(*Filename, false, true))
-			{
-				UE_LOG(LogCommandletPackageHelper, Error, TEXT("Error deleting package %s locally"), *Filename);
-				return false;
-			}
+			FilesToDeleteFromDisk.Add(Filename);
 		}
 	}
 
-	return true;
+	if (bSCCErrorsFound)
+	{
+		// Errors were found, we'll cancel everything
+		return false;
+	}
+
+	// It's possible that not all files were in the source control cache, in which case we should still add them to the
+	// files to delete on disk.
+	if (Filenames.Num() != SourceControlStates.Num())
+	{
+		for (FSourceControlStateRef& SourceControlState : SourceControlStates)
+		{
+			Filenames.Remove(SourceControlState->GetFilename());
+		}
+
+		FilesToDeleteFromDisk.Append(Filenames);
+	}
+
+	// First, revert files from SCC
+	if (FilesToRevert.Num() > 0)
+	{
+		if (GetSourceControlProvider().Execute(ISourceControlOperation::Create<FRevert>(), FilesToRevert) != ECommandResult::Succeeded)
+		{
+			UE_LOG(LogCommandletPackageHelper, Error, TEXT("Error reverting packages from source control"));
+			return false;
+		}
+	}
+
+	// Then delete files from SCC
+	if (FilesToDeleteFromSCC.Num() > 0)
+	{
+		if (GetSourceControlProvider().Execute(ISourceControlOperation::Create<FDelete>(), FilesToDeleteFromSCC) != ECommandResult::Succeeded)
+		{
+			UE_LOG(LogCommandletPackageHelper, Error, TEXT("Error deleting packages from source control"));
+			return false;
+		}
+	}
+
+	// Then delete files on disk
+	bool bDeleteOnDiskOk = true;
+	for (const FString& Filename : FilesToDeleteFromDisk)
+	{
+		if (!IFileManager::Get().Delete(*Filename, false, true))
+		{
+			UE_LOG(LogCommandletPackageHelper, Error, TEXT("Error deleting package %s locally"), *Filename);
+			bDeleteOnDiskOk = false;
+		}
+	}
+
+	return bDeleteOnDiskOk;
 }
 
 bool FPackageSourceControlHelper::Delete(UPackage* Package) const
diff --git a/Engine/Source/Editor/UnrealEd/Public/PackageSourceControlHelper.h b/Engine/Source/Editor/UnrealEd/Public/PackageSourceControlHelper.h
--- a/Engine/Source/Editor/UnrealEd/Public/PackageSourceControlHelper.h
+++ b/Engine/Source/Editor/UnrealEd/Public/PackageSourceControlHelper.h
@@ -11,6 +11,7 @@ class UNREALED_API FPackageSourceControlHelper
 public:
 	bool UseSourceControl() const;
 	bool Delete(const FString& PackageName) const;
+	bool Delete(const TArray<FString>& PackageNames) const;
 	bool Delete(UPackage* Package) const;
 	bool Delete(const TArray<UPackage*>& Packages) const;
 	bool AddToSourceControl(UPackage* Package) const;
diff --git a/Engine/Source/Runtime/Engine/Private/WorldPartition/ActorDescContainer.cpp b/Engine/Source/Runtime/Engine/Private/WorldPartition/ActorDescContainer.cpp
--- a/Engine/Source/Runtime/Engine/Private/WorldPartition/ActorDescContainer.cpp
+++ b/Engine/Source/Runtime/Engine/Private/WorldPartition/ActorDescContainer.cpp
@@ -215,6 +215,18 @@ void UActorDescContainer::OnPackageDeleted(UPackage* Package)
 	}
 }
 
+void UActorDescContainer::RemoveActor(const FGuid& ActorGuid)
+{
+	if (TUniquePtr<FWorldPartitionActorDesc>* ExistingActorDesc = Actors.FindRef(ActorGuid))
+	{
+		// This should never be called on already loaded actors
+		check(!(*ExistingActorDesc)->GetActor());
+		OnActorDescRemoved(*ExistingActorDesc);
+		Actors.Remove((*ExistingActorDesc)->GetGuid());
+		ExistingActorDesc->Release();
+	}
+}
+
 void UActorDescContainer::RegisterDelegates()
 {
 	if (GEditor && !IsTemplate())
diff --git a/Engine/Source/Runtime/Engine/Private/WorldPartition/ActorDescList.cpp b/Engine/Source/Runtime/Engine/Private/WorldPartition/ActorDescList.cpp
--- a/Engine/Source/Runtime/Engine/Private/WorldPartition/ActorDescList.cpp
+++ b/Engine/Source/Runtime/Engine/Private/WorldPartition/ActorDescList.cpp
@@ -43,7 +43,21 @@ const FWorldPartitionActorDesc* FActorDescList::GetActorDesc(const FString& Pack
 	const FName PackageFName(*PackageName);
 	for (const TUniquePtr<FWorldPartitionActorDesc>& ActorDescPtr : ActorDescList)
 	{
-		if (ActorDescPtr->GetActorPackage() == PackageFName)
+		if (ActorDescPtr && ActorDescPtr->GetActorPackage() == PackageFName)
+		{
+			return ActorDescPtr.Get();
+		}
+	}
+
+	return nullptr;
+}
+
+const FWorldPartitionActorDesc* FActorDescList::GetActorDesc(const FSoftObjectPath& InActorPath) const
+{
+	const FName ActorPath(*InActorPath.ToString());
+	for (const TUniquePtr<FWorldPartitionActorDesc>& ActorDescPtr : ActorDescList)
+	{
+		if (ActorDescPtr && ActorDescPtr->GetActorPath() == ActorPath)
 		{
 			return ActorDescPtr.Get();
 		}
diff --git a/Engine/Source/Runtime/Engine/Public/WorldPartition/ActorDescContainer.h b/Engine/Source/Runtime/Engine/Public/WorldPartition/ActorDescContainer.h
--- a/Engine/Source/Runtime/Engine/Public/WorldPartition/ActorDescContainer.h
+++ b/Engine/Source/Runtime/Engine/Public/WorldPartition/ActorDescContainer.h
@@ -28,6 +28,9 @@ public:
 	virtual void OnPackageDeleted(UPackage* Package);
 
 	FName GetContainerPackage() const { return ContainerPackageName; }
+
+	/** Removes an actor desc without the need to load a package */
+	virtual void RemoveActor(const FGuid& ActorGuid);
 #endif
 
 	UPROPERTY(Transient)
diff --git a/Engine/Source/Runtime/Engine/Public/WorldPartition/ActorDescList.h b/Engine/Source/Runtime/Engine/Public/WorldPartition/ActorDescList.h
--- a/Engine/Source/Runtime/Engine/Public/WorldPartition/ActorDescList.h
+++ b/Engine/Source/Runtime/Engine/Public/WorldPartition/ActorDescList.h
@@ -25,6 +25,7 @@ public:
 	const FWorldPartitionActorDesc& GetActorDescChecked(const FGuid& Guid) const;
 
 	const FWorldPartitionActorDesc* GetActorDesc(const FString& PackageName) const;
+	const FWorldPartitionActorDesc* GetActorDesc(const FSoftObjectPath& InActorPath) const;
 
 	int32 GetActorDescCount() const { return Actors.Num(); }
 
