WorldPartition HLOD Builder - Output actor name when building

diff --git a/Engine/Source/Editor/UnrealEd/Private/WorldPartition/WorldPartitionHLODsBuilder.cpp b/Engine/Source/Editor/UnrealEd/Private/WorldPartition/WorldPartitionHLODsBuilder.cpp
--- a/Engine/Source/Editor/UnrealEd/Private/WorldPartition/WorldPartitionHLODsBuilder.cpp
+++ b/Engine/Source/Editor/UnrealEd/Private/WorldPartition/WorldPartitionHLODsBuilder.cpp
@@ -391,7 +391,7 @@ bool UWorldPartitionHLODsBuilder::BuildHLODActors()
 		AWorldPartitionHLOD* HLODActor = CastChecked<AWorldPartitionHLOD>(ActorDesc->GetActor());
 
 		CurrentActor++;
-		UE_LOG(LogWorldPartitionHLODsBuilder, Display, TEXT("    [%d/%d] Building HLOD actor..."), CurrentActor, HLODActorsToBuild.Num());
+		UE_LOG(LogWorldPartitionHLODsBuilder, Display, TEXT("    [%d/%d] Building HLOD actor %s..."), CurrentActor, HLODActorsToBuild.Num(), *HLODActor->GetActorLabel());
 
 		HLODActor->BuildHLOD();
 
