Cleaned-up questionable recursion in AGroupActor

Group actors would recurse Modify, PostEditMove, and InvalidateLightingCacheDetailed into the actors within the group, which would result in these functions being called twice during drag operations. It doesn't really make sense for these to recurse, as their location is not directly tied to the group actor iself which is why each actor within the group needs to be moved separately during a drag.

diff --git a/Engine/Source/Editor/UnrealEd/Classes/Editor/GroupActor.h b/Engine/Source/Editor/UnrealEd/Classes/Editor/GroupActor.h
--- a/Engine/Source/Editor/UnrealEd/Classes/Editor/GroupActor.h
+++ b/Engine/Source/Editor/UnrealEd/Classes/Editor/GroupActor.h
@@ -37,14 +37,11 @@ class AGroupActor : public AActor
 	virtual void PostLoad() override;
 	virtual void PostEditChangeProperty(FPropertyChangedEvent& PropertyChangedEvent) override;
 	virtual void PostEditUndo() override;
-	virtual bool Modify( bool bAlwaysMarkDirty=true ) override;
 	virtual bool IsEditorOnly() const override { return true; }
 	//~ End UObject Interface
 
 	//~ Begin AActor Interface
 	virtual void PostActorCreated() override;
-	virtual void InvalidateLightingCacheDetailed(bool bTranslationOnly) override;
-	virtual void PostEditMove(bool bFinished) override;
 	virtual void SetIsTemporarilyHiddenInEditor( bool bIsHidden ) override;
 	virtual void GetActorBounds(bool bOnlyCollidingComponents, FVector& Origin, FVector& BoxExtent, bool bIncludeFromChildActors) const override;
 	virtual void GetActorLocationBounds(bool bOnlyCollidingComponents, FVector& Origin, FVector& BoxExtent, bool bIncludeFromChildActors) const override;
@@ -54,6 +51,12 @@ class AGroupActor : public AActor
 	// @todo document
 	virtual bool IsSelected() const;
 
+	/**
+	 * Enumerate all actors and subgroups for this group (includes this).
+	 * @param InCallback		The function to call for each actor
+	 */
+	UNREALED_API void ForEachActorInGroup(TFunctionRef<void(AActor*)> InCallback);
+
 	/**
 	 * Enumerate all actors and subgroups for this group that should be moved alongside this actor (includes this).
 	 * @param InSelectionSet	The selection set to use, or null to use the global selection set
diff --git a/Engine/Source/Editor/UnrealEd/Private/GroupActor.cpp b/Engine/Source/Editor/UnrealEd/Private/GroupActor.cpp
--- a/Engine/Source/Editor/UnrealEd/Private/GroupActor.cpp
+++ b/Engine/Source/Editor/UnrealEd/Private/GroupActor.cpp
@@ -103,6 +103,25 @@ bool AGroupActor::IsSelected() const
 	return (IsLocked() && HasSelectedActors()) || Super::IsSelected();
 }
 
+void AGroupActor::ForEachActorInGroup(TFunctionRef<void(AActor*)> InCallback)
+{
+	for (AActor* Actor : GroupActors)
+	{
+		if (Actor)
+		{
+			InCallback(Actor);
+		}
+	}
+	for (AGroupActor* SubGroup : SubGroups)
+	{
+		if (SubGroup)
+		{
+			SubGroup->ForEachActorInGroup(InCallback);
+		}
+	}
+	InCallback(this);
+}
+
 namespace GroupActorHelpers
 {
 
@@ -180,65 +199,6 @@ void AGroupActor::GroupApplyDelta(const FVector& InDrag, const FRotator& InRot,
 	});
 }
 
-bool AGroupActor::Modify(bool bAlwaysMarkDirty/*=true*/)
-{
-	bool bSavedToTransactionBuffer = false;
-	for(int32 ActorIndex=0; ActorIndex<GroupActors.Num(); ++ActorIndex)
-	{
-		if( GroupActors[ActorIndex] != NULL )
-		{
-			bSavedToTransactionBuffer = GroupActors[ActorIndex]->Modify(bAlwaysMarkDirty) || bSavedToTransactionBuffer;
-		}
-	}
-	for(int32 SubGroupIndex=0; SubGroupIndex<SubGroups.Num(); ++SubGroupIndex)
-	{
-		if( SubGroups[SubGroupIndex] != NULL )
-		{
-			bSavedToTransactionBuffer = SubGroups[SubGroupIndex]->Modify(bAlwaysMarkDirty) || bSavedToTransactionBuffer;
-		}
-	}
-	bSavedToTransactionBuffer = Super::Modify(bAlwaysMarkDirty) || bSavedToTransactionBuffer;
-	return  bSavedToTransactionBuffer;
-}
-
-void AGroupActor::InvalidateLightingCacheDetailed(bool bTranslationOnly)
-{
-	for(int32 ActorIndex=0; ActorIndex<GroupActors.Num(); ++ActorIndex)
-	{
-		if( GroupActors[ActorIndex] != NULL )
-		{
-			GroupActors[ActorIndex]->InvalidateLightingCacheDetailed(bTranslationOnly);
-		}
-	}
-	for(int32 SubGroupIndex=0; SubGroupIndex<SubGroups.Num(); ++SubGroupIndex)
-	{
-		if( SubGroups[SubGroupIndex] != NULL )
-		{
-			SubGroups[SubGroupIndex]->InvalidateLightingCacheDetailed(bTranslationOnly);
-		}
-	}
-	Super::InvalidateLightingCacheDetailed(bTranslationOnly);
-}
-
-void AGroupActor::PostEditMove(bool bFinished)
-{
-	for(int32 ActorIndex=0; ActorIndex<GroupActors.Num(); ++ActorIndex)
-	{
-		if( GroupActors[ActorIndex] != NULL )
-		{
-			GroupActors[ActorIndex]->PostEditMove(bFinished);
-		}
-	}
-	for(int32 SubGroupIndex=0; SubGroupIndex<SubGroups.Num(); ++SubGroupIndex)
-	{
-		if( SubGroups[SubGroupIndex] != NULL )
-		{
-			SubGroups[SubGroupIndex]->PostEditMove(bFinished);
-		}
-	}
-	Super::PostEditMove(bFinished);
-}
-
 void AGroupActor::SetIsTemporarilyHiddenInEditor( bool bIsHidden )
 {
 	Super::SetIsTemporarilyHiddenInEditor( bIsHidden );
diff --git a/Engine/Source/Editor/UnrealEd/Private/LevelEditorViewport.cpp b/Engine/Source/Editor/UnrealEd/Private/LevelEditorViewport.cpp
--- a/Engine/Source/Editor/UnrealEd/Private/LevelEditorViewport.cpp
+++ b/Engine/Source/Editor/UnrealEd/Private/LevelEditorViewport.cpp
@@ -1547,11 +1547,7 @@ bool FLevelEditorViewportClient::CanApplyMaterialToHitProxy( const HHitProxy* Hi
 }
 
 FTrackingTransaction::FTrackingTransaction()
-	: TransCount(0)
-	, ScopedTransaction(NULL)
-	, TrackingTransactionState(ETransactionState::Inactive)
 {
-	USelection::SelectionChangedEvent.AddRaw(this, &FTrackingTransaction::OnEditorSelectionChanged);
 }
 
 FTrackingTransaction::~FTrackingTransaction()
@@ -1568,53 +1564,52 @@ void FTrackingTransaction::Begin(const FText& Description, AActor* AdditionalAct
 
 	TrackingTransactionState = ETransactionState::Active;
 
-	TSet<AGroupActor*> GroupActors;
-
-	auto ProcessActorModify = [&](AActor* Actor)
+	if (UTypedElementSelectionSet* SelectionSet = GetMutableSelectionSet())
 	{
-		// Some tracking transactions, such as a duplication operation into a sublevel do not modify the selected Actors
-		// Store the initial package dirty state so we can restore if necessary
-		UPackage* Package = Actor->GetOutermost();
-		if (!InitialPackageDirtyStates.Contains(Package))
-		{
-			InitialPackageDirtyStates.Add(Package, Package->IsDirty());
-		}
-		
-		Actor->Modify();
+		// TODO: Ideally "Begin" would *not* speculatively call Modify at all (nor need to track/restore dirty state), 
+		//       as Modify should be called on-demand prior to changes within the transaction, however various existing 
+		//       code makes assumptions that tracking transactions will call Modify on the active selection.
 
-		if (UActorGroupingUtils::IsGroupingActive())
+		// Call modify the currently selected objects, gathering up any groups for processing later
+		TSet<AGroupActor*> GroupsPendingModify;
+		auto ModifyObject = [this, &GroupsPendingModify](UObject* InObject)
 		{
-			// if this actor is in a group, add the GroupActor into a list to be modified shortly
-			AGroupActor* ActorLockedRootGroup = AGroupActor::GetRootForActor(Actor, true);
-			if (ActorLockedRootGroup != nullptr)
+			// Track the dirty state of the packages so that it can be restored if the selection changes mid-transaction (eg, drag duplication)
+			UPackage* Package = InObject->GetOutermost();
+			if (!InitialPackageDirtyStates.Contains(Package))
 			{
-				GroupActors.Add(ActorLockedRootGroup);
+				InitialPackageDirtyStates.Add(Package, Package->IsDirty());
 			}
-		}
-	};
 
-	// Modify selected actors to record their state at the start of the transaction
-	for (FSelectionIterator It(GEditor->GetSelectedActorIterator()); It; ++It)
-	{
-		AActor* Actor = CastChecked<AActor>(*It);
-		ProcessActorModify(Actor);
-	}
-	// And the additional optional actor provided
-	if (AdditionalActor)
-	{
-		ProcessActorModify(AdditionalActor);
-	}
+			AActor* Actor = Cast<AActor>(InObject);
+			if (Actor && UActorGroupingUtils::IsGroupingActive())
+			{
+				if (AGroupActor* GroupActor = AGroupActor::GetRootForActor(Actor, true))
+				{
+					GroupsPendingModify.Add(GroupActor);
+				}
+			}
 
-	// Modify unique group actors
-	for (AGroupActor* GroupActor : GroupActors)
-	{		
-		GroupActor->Modify();
-	}
+			InObject->Modify();
 
-	// Modify selected components
-	for (FSelectionIterator It(GEditor->GetSelectedComponentIterator()); It; ++It)
-	{
-		CastChecked<UActorComponent>(*It)->Modify();
+			return true;
+		};
+		SelectionSet->ForEachSelectedObject(ModifyObject);
+		if (AdditionalActor)
+		{
+			ModifyObject(AdditionalActor);
+		}
+
+		// Recursively call modify on any groups, as viewport manipulation of an actor within a group tends to affect the entire group
+		for (AGroupActor* GroupPendingModify : GroupsPendingModify)
+		{
+			GroupPendingModify->ForEachActorInGroup([](AActor* InActor)
+			{
+				InActor->Modify();
+			});
+		}
+
+		SelectionSet->OnChanged().AddRaw(this, &FTrackingTransaction::OnEditorSelectionChanged);
 	}
 }
 
@@ -1627,6 +1622,10 @@ void FTrackingTransaction::End()
 	}
 	TrackingTransactionState = ETransactionState::Inactive;
 	InitialPackageDirtyStates.Empty();
+	if (UTypedElementSelectionSet* SelectionSet = GetMutableSelectionSet())
+	{
+		SelectionSet->OnChanged().RemoveAll(this);
+	}
 }
 
 void FTrackingTransaction::Cancel()
@@ -1655,7 +1654,19 @@ void FTrackingTransaction::PromotePendingToActive()
 	}
 }
 
-void FTrackingTransaction::OnEditorSelectionChanged(UObject* NewSelection)
+// TODO: Ideally this should be given the selection set for the level editor, rather than rely on the global state
+const UTypedElementSelectionSet* FTrackingTransaction::GetSelectionSet() const
+{
+	return GetMutableSelectionSet();
+}
+UTypedElementSelectionSet* FTrackingTransaction::GetMutableSelectionSet() const
+{
+	return (GEditor && GEditor->GetSelectedActors())
+		? GEditor->GetSelectedActors()->GetElementSelectionSet()
+		: nullptr;
+}
+
+void FTrackingTransaction::OnEditorSelectionChanged(const UTypedElementSelectionSet* InSelectionSet)
 {
 	if (!InitialPackageDirtyStates.Num())
 	{
@@ -1664,28 +1675,26 @@ void FTrackingTransaction::OnEditorSelectionChanged(UObject* NewSelection)
 
 	checkf(TrackingTransactionState == ETransactionState::Active, TEXT("Inactive tracking transaction contains package states"));
 
-	// The selection has changed due to a duplication or other operation
-	// So, restore dirty state on any package which is no longer selected
-	TArray<UPackage*> SelectedPackages;
-	for (FSelectionIterator It(GEditor->GetSelectedActorIterator()); It; ++It)
+	// The selection state has changed mid-transaction (eg, drag duplication)
+	// Restore the dirty state on any packages which no longer have selected objects
+	TSet<UPackage*> SelectedPackages;
+	InSelectionSet->ForEachSelectedObject([&SelectedPackages](UObject* InObject)
 	{
-		AActor* Actor = static_cast<AActor*>(*It);
-		checkSlow(Actor->IsA(AActor::StaticClass()));
-		SelectedPackages.AddUnique(Actor->GetOutermost());
-	}
+		SelectedPackages.Add(InObject->GetOutermost());
+		return true;
+	});
 
-	for (const auto& Entry : InitialPackageDirtyStates)
+	for (const TTuple<UPackage*, bool>& InitialPackageDirtyState : InitialPackageDirtyStates)
 	{
-		// If no Actors in the package are currently selected, restore unmodified package state
-		UPackage* Package = Entry.Key;
-		if (!SelectedPackages.Contains(Package) && !Entry.Value && Package->IsDirty())
+		UPackage* Package = InitialPackageDirtyState.Key;
+		const bool bInitialDirtyState = InitialPackageDirtyState.Value;
+		if (!SelectedPackages.Contains(Package) && !bInitialDirtyState && Package->IsDirty())
 		{
 			Package->SetDirtyFlag(false);
 		}
 	}
 
 	InitialPackageDirtyStates.Empty();
-
 }
 
 
diff --git a/Engine/Source/Editor/UnrealEd/Public/LevelEditorViewport.h b/Engine/Source/Editor/UnrealEd/Public/LevelEditorViewport.h
--- a/Engine/Source/Editor/UnrealEd/Public/LevelEditorViewport.h
+++ b/Engine/Source/Editor/UnrealEd/Public/LevelEditorViewport.h
@@ -104,18 +104,21 @@ struct UNREALED_API FTrackingTransaction
 
 	bool IsPending() const { return TrackingTransactionState == ETransactionState::Pending; }
 	
-	int32 TransCount;
+	int32 TransCount = 0;
 
 private:
 
+	const UTypedElementSelectionSet* GetSelectionSet() const;
+	UTypedElementSelectionSet* GetMutableSelectionSet() const;
+
 	/** Editor selection changed delegate handler */	
-	void OnEditorSelectionChanged(UObject* NewSelection);
+	void OnEditorSelectionChanged(const UTypedElementSelectionSet* InSelectionSet);
 
 	/** The current transaction. */
-	class FScopedTransaction*	ScopedTransaction;
+	class FScopedTransaction* ScopedTransaction = nullptr;
 
 	/** This is set to Active if TrackingStarted() has initiated a transaction, Pending if a transaction will begin before the next delta change */
-	ETransactionState::Enum TrackingTransactionState;
+	ETransactionState::Enum TrackingTransactionState = ETransactionState::Inactive;
 
 	/** The description to use if a pending transaction turns into a real transaction */
 	FText PendingDescription;
