Optimized sub-object look-up within transactions

FPersistentObjectRef would perform a StaticFindObject for each sub-object in its object chain, for each request to resolve the object pointer. This could be quite costly when finding the record for a large transaction (eg, when moving thousands of objects).

This change stores a cache of the resolved object chain, and will use this cache if every object in the chain still passes its validation check. In order to do this, the guts of FPersistentObjectRef have been made private, and the -> operator was also removed as it could promote writing code that would repeatedly resolve the object pointer.

diff --git a/Engine/Source/Editor/UnrealEd/Classes/Editor/Transactor.h b/Engine/Source/Editor/UnrealEd/Classes/Editor/Transactor.h
--- a/Engine/Source/Editor/UnrealEd/Classes/Editor/Transactor.h
+++ b/Engine/Source/Editor/UnrealEd/Classes/Editor/Transactor.h
@@ -12,9 +12,10 @@
 #include "UObject/Object.h"
 #include "UObject/Package.h"
 #include "Serialization/ArchiveUObject.h"
-#include "Misc/ITransaction.h"
 #include "Serialization/ArchiveSerializedPropertyChain.h"
+#include "Misc/ITransaction.h"
 #include "Algo/Find.h"
+#include <type_traits>
 #include "Transactor.generated.h"
 
 /*-----------------------------------------------------------------------------
@@ -54,44 +55,53 @@ protected:
 		*/
 		struct FPersistentObjectRef
 		{
-			// This enum represents all of the different special cases
-			// we are handling with this type:
-			enum class EReferenceType : uint8
-			{
-				SubObject,
-				RootObject,
-
-				Unknown
-			};
-
-			EReferenceType ReferenceType;
-			UObject* Object;
-			TArray<FName> SubObjectHierarchyID;
-			FName ComponentName;
-
+		public:
 			friend FArchive& operator<<(FArchive& Ar, FPersistentObjectRef& ReferencedObject)
 			{
-				Ar << (uint8&)ReferencedObject.ReferenceType;
-				Ar << ReferencedObject.Object;
-				Ar << ReferencedObject.SubObjectHierarchyID;
+				Ar << (std::underlying_type_t<EReferenceType>&)ReferencedObject.ReferenceType;
+				Ar << ReferencedObject.RootObject;
+				Ar << ReferencedObject.SubObjectHierarchyIDs;
 				return Ar;
 			}
 
-			FPersistentObjectRef()
-				: ReferenceType(EReferenceType::Unknown)
-				, Object(nullptr)
-			{}
+			FPersistentObjectRef() = default;
 
 			explicit FPersistentObjectRef(UObject* InObject);
 
-			UObject* Get() const;
+			bool IsRootObjectReference() const
+			{
+				return ReferenceType == EReferenceType::RootObject;
+			}
 
-			UObject* operator->() const
+			bool IsSubObjectReference() const
 			{
-				UObject* Obj = Get();
-				check(Obj);
-				return Obj;
+				return ReferenceType == EReferenceType::SubObject;
 			}
+
+			UObject* Get() const;
+
+			void AddReferencedObjects(FReferenceCollector& Collector);
+
+		private:
+			/** This enum represents all of the different special cases we are handling with this type */
+			enum class EReferenceType : uint8
+			{
+				Unknown,
+				RootObject,
+				SubObject,
+			};
+
+			/** The reference type we're handling */
+			EReferenceType ReferenceType = EReferenceType::Unknown;
+			/** Stores the object pointer when ReferenceType==RootObject, and the outermost pointer of the sub-object chain when when ReferenceType==SubObject */
+			UObject* RootObject = nullptr;
+			/** Stores the sub-object name chain when ReferenceType==SubObject */
+			TArray<FName, TInlineAllocator<4>> SubObjectHierarchyIDs;
+
+			/** Cached pointers corresponding to RootObject when ReferenceType==SubObject (@note cache needs testing on access as it may have become stale) */
+			mutable TWeakObjectPtr<UObject> CachedRootObject;
+			/** Cache of pointers corresponding to the items within SubObjectHierarchyIDs when ReferenceType==SubObject (@note cache needs testing on access as it may have become stale) */
+			mutable TArray<TWeakObjectPtr<UObject>, TInlineAllocator<4>> CachedSubObjectHierarchy;
 		};
 
 		struct FSerializedProperty
diff --git a/Engine/Source/Editor/UnrealEd/Private/EditorTransaction.cpp b/Engine/Source/Editor/UnrealEd/Private/EditorTransaction.cpp
--- a/Engine/Source/Editor/UnrealEd/Private/EditorTransaction.cpp
+++ b/Engine/Source/Editor/UnrealEd/Private/EditorTransaction.cpp
@@ -6,6 +6,7 @@
 #include "UObject/Object.h"
 #include "UObject/Package.h"
 #include "Algo/Find.h"
+#include "Algo/Reverse.h"
 #include "Engine/Level.h"
 #include "Components/ActorComponent.h"
 #include "Model.h"
@@ -18,43 +19,6 @@
 
 DEFINE_LOG_CATEGORY_STATIC(LogEditorTransaction, Log, All);
 
-inline UObject* BuildSubobjectKey(UObject* InObj, TArray<FName>& OutHierarchyNames)
-{
-	auto UseOuter = [](const UObject* Obj)
-	{
-		if (Obj == nullptr)
-		{
-			return false;
-		}
-
-		const bool bIsCDO = Obj->HasAllFlags(RF_ClassDefaultObject);
-		const UObject* CDO = bIsCDO ? Obj : nullptr;
-		const bool bIsClassCDO = (CDO != nullptr) ? (CDO->GetClass()->ClassDefaultObject == CDO) : false;
-		if(!bIsClassCDO && CDO)
-		{
-			// Likely a trashed CDO, try to recover. Only known cause of this is
-			// ambiguous use of DSOs:
-			CDO = CDO->GetClass()->ClassDefaultObject;
-		}
-		const UActorComponent* AsComponent = Cast<UActorComponent>(Obj);
-		const bool bIsDSO = Obj->HasAnyFlags(RF_DefaultSubObject);
-		const bool bIsSCSComponent = AsComponent && AsComponent->IsCreatedByConstructionScript();
-		return (bIsCDO && bIsClassCDO) || bIsDSO || bIsSCSComponent;
-	};
-	
-	UObject* Outermost = nullptr;
-
-	UObject* Iter = InObj;
-	while (UseOuter(Iter))
-	{
-		OutHierarchyNames.Add(Iter->GetFName());
-		Iter = Iter->GetOuter();
-		Outermost = Iter;
-	}
-
-	return Outermost;
-}
-
 /*-----------------------------------------------------------------------------
 	A single transaction.
 -----------------------------------------------------------------------------*/
@@ -107,15 +71,17 @@ void FTransaction::FObjectRecord::SerializeContents( FArchive& Ar, int32 InOper
 	if( Array )
 	{
 		const bool bWasArIgnoreOuterRef = Ar.ArIgnoreOuterRef;
-		if (Object.SubObjectHierarchyID.Num() != 0)
+		if (Object.IsSubObjectReference())
 		{
 			Ar.ArIgnoreOuterRef = true;
 		}
 
-		//UE_LOG( LogEditorTransaction, Log, TEXT("Array %s %i*%i: %i"), Object ? *Object->GetFullName() : TEXT("Invalid Object"), Index, ElementSize, InOper);
+		UObject* CurrentObject = Object.Get();
+
+		//UE_LOG( LogEditorTransaction, Log, TEXT("Array %s %i*%i: %i"), CurrentObject ? *CurrentObject->GetFullName() : TEXT("Invalid Object"), Index, ElementSize, InOper);
 
-		check((SIZE_T)Array >= (SIZE_T)Object.Get() + sizeof(UObject));
-		check((SIZE_T)Array + sizeof(FScriptArray) <= (SIZE_T)Object.Get() + Object->GetClass()->GetPropertiesSize());
+		check((SIZE_T)Array >= (SIZE_T)CurrentObject + sizeof(UObject));
+		check((SIZE_T)Array + sizeof(FScriptArray) <= (SIZE_T)CurrentObject + CurrentObject->GetClass()->GetPropertiesSize());
 		check(ElementSize!=0);
 		check(DefaultConstructor!=NULL);
 		check(Serializer!=NULL);
@@ -158,7 +124,7 @@ void FTransaction::FObjectRecord::SerializeContents( FArchive& Ar, int32 InOper
 	}
 	else
 	{
-		//UE_LOG(LogEditorTransaction, Log,  TEXT("Object %s"), *Object->GetFullName());
+		//UE_LOG(LogEditorTransaction, Log,  TEXT("Object %s"), *Object.Get()->GetFullName());
 		check(Index==0);
 		check(ElementSize==0);
 		check(DefaultConstructor==NULL);
@@ -175,7 +141,7 @@ void FTransaction::FObjectRecord::SerializeObject( FArchive& Ar )
 	if (CurrentObject)
 	{
 		const bool bWasArIgnoreOuterRef = Ar.ArIgnoreOuterRef;
-		if (Object.SubObjectHierarchyID.Num() != 0)
+		if (Object.IsSubObjectReference())
 		{
 			Ar.ArIgnoreOuterRef = true;
 		}
@@ -646,24 +612,47 @@ FArchive& operator<<( FArchive& Ar, FTransaction::FObjectRecord& R )
 }
 
 FTransaction::FObjectRecord::FPersistentObjectRef::FPersistentObjectRef(UObject* InObject)
-	: ReferenceType(EReferenceType::Unknown)
-	, Object(nullptr)
 {
-	check(InObject);
-	UObject* Outermost = BuildSubobjectKey(InObject, SubObjectHierarchyID);
+	RootObject = InObject;
+	{
+		auto UseOuter = [](const UObject* Obj)
+		{
+			if (Obj == nullptr)
+			{
+				return false;
+			}
+
+			const bool bIsCDO = Obj->HasAllFlags(RF_ClassDefaultObject);
+			const UObject* CDO = bIsCDO ? Obj : nullptr;
+			const bool bIsClassCDO = (CDO != nullptr) ? (CDO->GetClass()->ClassDefaultObject == CDO) : false;
+			if (!bIsClassCDO && CDO)
+			{
+				// Likely a trashed CDO, try to recover
+				// Only known cause of this is ambiguous use of DSOs
+				CDO = CDO->GetClass()->ClassDefaultObject;
+			}
+			const UActorComponent* AsComponent = Cast<UActorComponent>(Obj);
+			const bool bIsDSO = Obj->HasAnyFlags(RF_DefaultSubObject);
+			const bool bIsSCSComponent = AsComponent && AsComponent->IsCreatedByConstructionScript();
+			return (bIsCDO && bIsClassCDO) || bIsDSO || bIsSCSComponent;
+		};
 
-	if (SubObjectHierarchyID.Num()>0)
+		while (UseOuter(RootObject))
+		{
+			SubObjectHierarchyIDs.Add(RootObject->GetFName());
+			RootObject = RootObject->GetOuter();
+		}
+	}
+	check(RootObject);
+
+	if (SubObjectHierarchyIDs.Num() > 0)
 	{
-		check(Outermost);
-		//check(Outermost != GetTransientPackage());
 		ReferenceType = EReferenceType::SubObject;
-		Object = Outermost;
+		Algo::Reverse(SubObjectHierarchyIDs);
 	}
 	else
 	{
-		SubObjectHierarchyID.Empty();
 		ReferenceType = EReferenceType::RootObject;
-		Object = InObject;
 	}
 
 	// Make sure that when we look up the object we find the same thing:
@@ -674,61 +663,97 @@ UObject* FTransaction::FObjectRecord::FPersistentObjectRef::Get() const
 {
 	if (ReferenceType == EReferenceType::SubObject)
 	{
-		check (SubObjectHierarchyID.Num() > 0)
-		// find the subobject:
-		UObject* CurrentObject = Object;
-		bool bFoundTargetSubObject = (SubObjectHierarchyID.Num() == 0);
-		if (!bFoundTargetSubObject)
+		check(SubObjectHierarchyIDs.Num() > 0);
+
+		UObject* CurrentObject = nullptr;
+
+		// Do we have a valid cached pointer?
+		if (!CachedRootObject.IsExplicitlyNull() && SubObjectHierarchyIDs.Num() == CachedSubObjectHierarchy.Num())
 		{
-			// Current increasing depth into sub-objects, starts at 1 to avoid the sub-object found and placed in NextObject.
-			int SubObjectDepth = SubObjectHierarchyID.Num() - 1;
-			UObject* NextObject = CurrentObject;
-			while (NextObject != nullptr && !bFoundTargetSubObject)
+			// Root object is a pointer test
+			{
+				CurrentObject = CachedRootObject.GetEvenIfUnreachable();
+				if (CurrentObject != RootObject)
+				{
+					CurrentObject = nullptr;
+				}
+			}
+
+			// All other sub-objects are a name test
+			for (int32 SubObjectIndex = 0; CurrentObject && SubObjectIndex < SubObjectHierarchyIDs.Num(); ++SubObjectIndex)
 			{
-				// Look for any UObject with the CurrentObject's outer to find the next sub-object:
-				NextObject = StaticFindObjectFast(UObject::StaticClass(), CurrentObject, SubObjectHierarchyID[SubObjectDepth]);
-				bFoundTargetSubObject = SubObjectDepth == 0;
-				--SubObjectDepth;
-				CurrentObject = NextObject;
+				CurrentObject = CachedSubObjectHierarchy[SubObjectIndex].GetEvenIfUnreachable();
+				if (CurrentObject && CurrentObject->GetFName() != SubObjectHierarchyIDs[SubObjectIndex])
+				{
+					CurrentObject = nullptr;
+				}
 			}
 		}
+		if (CurrentObject)
+		{
+			return CurrentObject;
+		}
+
+		// Cached pointer is invalid
+		CachedRootObject.Reset();
+		CachedSubObjectHierarchy.Reset();
 
-		return bFoundTargetSubObject ? CurrentObject : nullptr;
+		// Try to find and cache the subobject
+		CachedRootObject = RootObject;
+		CurrentObject = RootObject;
+		for (int32 SubObjectIndex = 0; CurrentObject && SubObjectIndex < SubObjectHierarchyIDs.Num(); ++SubObjectIndex)
+		{
+			CurrentObject = StaticFindObjectFast(UObject::StaticClass(), CurrentObject, SubObjectHierarchyIDs[SubObjectIndex]);
+			CachedSubObjectHierarchy.Add(CurrentObject);
+		}
+		if (CurrentObject)
+		{
+			check(!CachedRootObject.IsExplicitlyNull() && SubObjectHierarchyIDs.Num() == CachedSubObjectHierarchy.Num());
+			return CurrentObject;
+		}
+
+		// Cached pointer is invalid
+		CachedRootObject.Reset();
+		CachedSubObjectHierarchy.Reset();
+
+		return nullptr;
 	}
 
-	return Object;
+	return RootObject;
 }
 
-void FTransaction::FObjectRecord::AddReferencedObjects( FReferenceCollector& Collector )
+void FTransaction::FObjectRecord::FPersistentObjectRef::AddReferencedObjects(FReferenceCollector& Collector)
 {
-	UObject* Obj = Object.Object;
+	UObject* Obj = RootObject;
 	Collector.AddReferencedObject(Obj);
-	Object.Object = Obj;
+	RootObject = Obj;
+}
+
+void FTransaction::FObjectRecord::AddReferencedObjects( FReferenceCollector& Collector )
+{
+	Object.AddReferencedObjects(Collector);
 
-	auto AddSerializedObjectReferences = [](FSerializedObject& InSerializedObject, FReferenceCollector& InCollector)
+	auto AddSerializedObjectReferences = [&Collector](FSerializedObject& InSerializedObject)
 	{
 		for (FPersistentObjectRef& ReferencedObject : InSerializedObject.ReferencedObjects)
 		{
-			UObject* RefObj = ReferencedObject.Object;
-			InCollector.AddReferencedObject(RefObj);
-			ReferencedObject.Object = RefObj;
+			ReferencedObject.AddReferencedObjects(Collector);
 		}
 
 		if (InSerializedObject.ObjectAnnotation.IsValid())
 		{
-			InSerializedObject.ObjectAnnotation->AddReferencedObjects(InCollector);
+			InSerializedObject.ObjectAnnotation->AddReferencedObjects(Collector);
 		}
 	};
-	AddSerializedObjectReferences(SerializedObject, Collector);
-	AddSerializedObjectReferences(SerializedObjectFlip, Collector);
-	AddSerializedObjectReferences(SerializedObjectSnapshot, Collector);
+	AddSerializedObjectReferences(SerializedObject);
+	AddSerializedObjectReferences(SerializedObjectFlip);
+	AddSerializedObjectReferences(SerializedObjectSnapshot);
 }
 
 bool FTransaction::FObjectRecord::ContainsPieObject() const
 {
 	{
-		UObject* Obj = Object.Object;
-
+		const UObject* Obj = Object.Get();
 		if(Obj && Obj->GetOutermost()->HasAnyPackageFlags(PKG_PlayInEditor))
 		{
 			return true;
@@ -739,7 +764,7 @@ bool FTransaction::FObjectRecord::ContainsPieObject() const
 	{
 		for (const FPersistentObjectRef& ReferencedObject : InSerializedObject.ReferencedObjects)
 		{
-			const UObject* Obj = ReferencedObject.Object;
+			const UObject* Obj = ReferencedObject.Get();
 			if (Obj && Obj->GetOutermost()->HasAnyPackageFlags(PKG_PlayInEditor))
 			{
 				return true;
