Change presave prototype to take an FObjectSaveContextRef to provide more information about the save parameters.
#rn Minor Cooking

diff --git a/Engine/Plugins/Developer/Concert/ConcertSync/ConcertSyncClient/Source/ConcertSyncClient/Private/ConcertClientPackageBridge.cpp b/Engine/Plugins/Developer/Concert/ConcertSync/ConcertSyncClient/Source/ConcertSyncClient/Private/ConcertClientPackageBridge.cpp
--- a/Engine/Plugins/Developer/Concert/ConcertSync/ConcertSyncClient/Source/ConcertSyncClient/Private/ConcertClientPackageBridge.cpp
+++ b/Engine/Plugins/Developer/Concert/ConcertSync/ConcertSyncClient/Source/ConcertSyncClient/Private/ConcertClientPackageBridge.cpp
@@ -1,12 +1,14 @@
 // Copyright Epic Games, Inc. All Rights Reserved.
 
 #include "ConcertClientPackageBridge.h"
+
 #include "ConcertLogGlobal.h"
 #include "ConcertWorkspaceData.h"
 #include "ConcertSyncClientUtil.h"
 
 #include "Engine/World.h"
 #include "Engine/Engine.h"
+#include "UObject/ObjectSaveContext.h"
 #include "UObject/Package.h"
 #include "UObject/PackageReload.h"
 #include "Misc/Paths.h"
@@ -56,8 +58,8 @@ FConcertClientPackageBridge::FConcertClientPackageBridge()
 	// correctly outside of take recorder.  There is a lot of code in MultiUser that assumes that in -game mode is
 	// "receive" only.
 	//
-	UPackage::PreSavePackageEvent.AddRaw(this, &FConcertClientPackageBridge::HandlePackagePreSave);
-	UPackage::PackageSavedEvent.AddRaw(this, &FConcertClientPackageBridge::HandlePackageSaved);
+	UPackage::PreSavePackageWithContextEvent.AddRaw(this, &FConcertClientPackageBridge::HandlePackagePreSave);
+	UPackage::PackageSavedWithContextEvent.AddRaw(this, &FConcertClientPackageBridge::HandlePackageSaved);
 	FCoreUObjectDelegates::OnPackageReloaded.AddRaw(this, &FConcertClientPackageBridge::HandleAssetReload);
 
 	if (GIsEditor)
@@ -79,8 +81,8 @@ FConcertClientPackageBridge::~FConcertClientPackageBridge()
 {
 #if WITH_EDITOR
 	// Unregister Package Events
-	UPackage::PreSavePackageEvent.RemoveAll(this);
-	UPackage::PackageSavedEvent.RemoveAll(this);
+	UPackage::PreSavePackageWithContextEvent.RemoveAll(this);
+	UPackage::PackageSavedWithContextEvent.RemoveAll(this);
 	FCoreUObjectDelegates::OnPackageReloaded.RemoveAll(this);
 
 	if (GIsEditor)
@@ -122,7 +124,7 @@ bool& FConcertClientPackageBridge::GetIgnoreLocalDiscardRef()
 	return bIgnoreLocalDiscard;
 }
 
-void FConcertClientPackageBridge::HandlePackagePreSave(UPackage* Package)
+void FConcertClientPackageBridge::HandlePackagePreSave(UPackage* Package, FObjectPreSaveContext ObjectSaveContext)
 {
 	// Ignore package operations fired by the cooker (cook on the fly).
 	if (GIsCookerLoadingPackage)
@@ -162,10 +164,8 @@ void FConcertClientPackageBridge::HandlePackagePreSave(UPackage* Package)
 	UE_LOG(LogConcert, Verbose, TEXT("Asset Pre-Saved: %s"), *Package->GetName());
 }
 
-void FConcertClientPackageBridge::HandlePackageSaved(const FString& PackageFilename, UObject* Outer)
+void FConcertClientPackageBridge::HandlePackageSaved(const FString& PackageFilename, UPackage* Package, FObjectPostSaveContext ObjectSaveContext)
 {
-	UPackage* Package = CastChecked<UPackage>(Outer);
-
 	// Ignore package operations fired by the cooker (cook on the fly).
 	if (GIsCookerLoadingPackage)
 	{
diff --git a/Engine/Plugins/Developer/Concert/ConcertSync/ConcertSyncClient/Source/ConcertSyncClient/Private/ConcertClientPackageBridge.h b/Engine/Plugins/Developer/Concert/ConcertSync/ConcertSyncClient/Source/ConcertSyncClient/Private/ConcertClientPackageBridge.h
--- a/Engine/Plugins/Developer/Concert/ConcertSync/ConcertSyncClient/Source/ConcertSyncClient/Private/ConcertClientPackageBridge.h
+++ b/Engine/Plugins/Developer/Concert/ConcertSync/ConcertSyncClient/Source/ConcertSyncClient/Private/ConcertClientPackageBridge.h
@@ -4,6 +4,8 @@
 
 #include "IConcertClientPackageBridge.h"
 
+class FObjectPostSaveContext;
+class FObjectPreSaveContext;
 class UPackage;
 class FPackageReloadedEvent;
 
@@ -26,10 +28,10 @@ public:
 
 private:
 	/** Called prior to a package being saved to disk */
-	void HandlePackagePreSave(UPackage* Package);
+	void HandlePackagePreSave(UPackage* Package, FObjectPreSaveContext ObjectSavecontext);
 
 	/** Called after a package has been saved to disk */
-	void HandlePackageSaved(const FString& PackageFilename, UObject* Outer);
+	void HandlePackageSaved(const FString& PackageFilename, UPackage* Package, FObjectPostSaveContext ObjectSaveContext);
 
 	/** Called when a new asset is added */
 	void HandleAssetAdded(UObject *Object);
diff --git a/Engine/Plugins/Editor/AssetSearch/Source/Private/AssetSearchManager.cpp b/Engine/Plugins/Editor/AssetSearch/Source/Private/AssetSearchManager.cpp
--- a/Engine/Plugins/Editor/AssetSearch/Source/Private/AssetSearchManager.cpp
+++ b/Engine/Plugins/Editor/AssetSearch/Source/Private/AssetSearchManager.cpp
@@ -1,6 +1,7 @@
 // Copyright Epic Games, Inc. All Rights Reserved.
 
 #include "AssetSearchManager.h"
+
 #include "IAssetRegistry.h"
 #include "AssetRegistryModule.h"
 #include "AssetSearchDatabase.h"
@@ -29,6 +30,7 @@
 #include "PackageTools.h"
 #include "UObject/ObjectMacros.h"
 #include "UObject/ObjectKey.h"
+#include "UObject/ObjectSaveContext.h"
 #include "UObject/WeakObjectPtrTemplates.h"
 #include "FileHelpers.h"
 #include "Misc/MessageDialog.h"
@@ -202,7 +204,7 @@ FAssetSearchManager::~FAssetSearchManager()
 
 	StopScanningAssets();
 
-	UPackage::PackageSavedEvent.RemoveAll(this);
+	UPackage::PackageSavedWithContextEvent.RemoveAll(this);
 	FCoreUObjectDelegates::OnAssetLoaded.RemoveAll(this);
 	UObject::FAssetRegistryTag::OnGetExtraObjectTags.RemoveAll(this);
 
@@ -227,7 +229,7 @@ void FAssetSearchManager::Start()
 
 	RegisterSearchProvider(TEXT("AssetRegistry"), MakeUnique<FAssetRegistrySearchProvider>());
 
-	UPackage::PackageSavedEvent.AddRaw(this, &FAssetSearchManager::HandlePackageSaved);
+	UPackage::PackageSavedWithContextEvent.AddRaw(this, &FAssetSearchManager::HandlePackageSaved);
 	FCoreUObjectDelegates::OnAssetLoaded.AddRaw(this, &FAssetSearchManager::OnAssetLoaded);
 
 	TickerHandle = FTicker::GetCoreTicker().AddTicker(FTickerDelegate::CreateRaw(this, &FAssetSearchManager::Tick_GameThread), 0);
@@ -446,7 +448,7 @@ void FAssetSearchManager::OnAssetScanFinished()
 	});
 }
 
-void FAssetSearchManager::HandlePackageSaved(const FString& PackageFilename, UObject* Outer)
+void FAssetSearchManager::HandlePackageSaved(const FString& PackageFilename, UPackage* Package, FObjectPostSaveContext ObjectSaveContext)
 {
 	check(IsInGameThread());
 
@@ -456,8 +458,6 @@ void FAssetSearchManager::HandlePackageSaved(const FString& PackageFilename, UOb
 		return;
 	}
 
-	UPackage* Package = CastChecked<UPackage>(Outer);
-
 	if (GIsEditor && !IsRunningCommandlet())
 	{
 		TArray<UObject*> Objects;
diff --git a/Engine/Plugins/Editor/AssetSearch/Source/Private/AssetSearchManager.h b/Engine/Plugins/Editor/AssetSearch/Source/Private/AssetSearchManager.h
--- a/Engine/Plugins/Editor/AssetSearch/Source/Private/AssetSearchManager.h
+++ b/Engine/Plugins/Editor/AssetSearch/Source/Private/AssetSearchManager.h
@@ -9,6 +9,7 @@
 #include "Containers/Queue.h"
 #include "HAL/Runnable.h"
 
+class FObjectPostSaveContext;
 class FRunnableThread;
 class UClass;
 class ISearchProvider;
@@ -46,7 +47,7 @@ private:
 	void OnAssetRemoved(const FAssetData& InAssetData);
 	void OnAssetScanFinished();
 
-	void HandlePackageSaved(const FString& PackageFilename, UObject* Outer);
+	void HandlePackageSaved(const FString& PackageFilename, UPackage* Package, FObjectPostSaveContext ObjectSaveContext);
 	void OnAssetLoaded(UObject* InObject);
 
 	void AddOrUpdateAsset(const FAssetData& InAsset, const FString& IndexedJson, const FString& DerivedDataKey);
diff --git a/Engine/Plugins/Editor/DataValidation/Source/DataValidation/Private/DataValidationModule.cpp b/Engine/Plugins/Editor/DataValidation/Source/DataValidation/Private/DataValidationModule.cpp
--- a/Engine/Plugins/Editor/DataValidation/Source/DataValidation/Private/DataValidationModule.cpp
+++ b/Engine/Plugins/Editor/DataValidation/Source/DataValidation/Private/DataValidationModule.cpp
@@ -1,7 +1,9 @@
 // Copyright Epic Games, Inc. All Rights Reserved.
 
 #include "DataValidationModule.h"
+
 #include "UObject/Object.h"
+#include "UObject/ObjectSaveContext.h"
 #include "UObject/SoftObjectPath.h"
 #include "GameFramework/HUD.h"
 
@@ -43,7 +45,7 @@ private:
 	TSharedRef<FExtender> OnExtendContentBrowserPathSelectionMenu(const TArray<FString>& SelectedPaths);
 	void CreateDataValidationContentBrowserAssetMenu(FMenuBuilder& MenuBuilder, TArray<FAssetData> SelectedAssets);
 	void CreateDataValidationContentBrowserPathMenu(FMenuBuilder& MenuBuilder, TArray<FString> SelectedPaths);
-	void OnPackageSaved(const FString& PackageFileName, UObject* PackageObj);
+	void OnPackageSaved(const FString& PackageFileName, UPackage* Package, FObjectPostSaveContext ObjectSaveContext);
 
 	// Adds Asset and any assets it depends on to the set DependentAssets
 	void FindAssetDependencies(const FAssetRegistryModule& AssetRegistryModule, const FAssetData& Asset, TSet<FAssetData>& DependentAssets);
@@ -78,7 +80,7 @@ void FDataValidationModule::StartupModule()
 		FCoreDelegates::OnPostEngineInit.AddRaw(this, &FDataValidationModule::RegisterMenus);
 
 		// Add save callback
-		UPackage::PackageSavedEvent.AddRaw(this, &FDataValidationModule::OnPackageSaved);
+		UPackage::PackageSavedWithContextEvent.AddRaw(this, &FDataValidationModule::OnPackageSaved);
 
 		ISettingsModule& SettingsModule = FModuleManager::LoadModuleChecked<ISettingsModule>("Settings");
 		SettingsModule.RegisterSettings("Editor", "Advanced", "DataValidation",
@@ -105,7 +107,7 @@ void FDataValidationModule::ShutdownModule()
 		UToolMenus::UnregisterOwner(this);
 		FCoreDelegates::OnPostEngineInit.RemoveAll(this);
 
-		UPackage::PackageSavedEvent.RemoveAll(this);
+		UPackage::PackageSavedWithContextEvent.RemoveAll(this);
 	}
 }
 
@@ -306,12 +308,12 @@ void FDataValidationModule::CreateDataValidationContentBrowserPathMenu(FMenuBuil
 	);
 }
 
-void FDataValidationModule::OnPackageSaved(const FString& PackageFileName, UObject* PackageObj)
+void FDataValidationModule::OnPackageSaved(const FString& PackageFileName, UPackage* Package, FObjectPostSaveContext ObjectSaveContext)
 {
 	UEditorValidatorSubsystem* EditorValidationSubsystem = GEditor->GetEditorSubsystem<UEditorValidatorSubsystem>();
-	if (EditorValidationSubsystem && PackageObj)
+	if (EditorValidationSubsystem && Package)
 	{
-		EditorValidationSubsystem->ValidateSavedPackage(PackageObj->GetFName());
+		EditorValidationSubsystem->ValidateSavedPackage(Package->GetFName());
 	}
 }
 
diff --git a/Engine/Plugins/Editor/GameplayTagsEditor/Source/GameplayTagsEditor/Private/GameplayTagsEditorModule.cpp b/Engine/Plugins/Editor/GameplayTagsEditor/Source/GameplayTagsEditor/Private/GameplayTagsEditorModule.cpp
--- a/Engine/Plugins/Editor/GameplayTagsEditor/Source/GameplayTagsEditor/Private/GameplayTagsEditorModule.cpp
+++ b/Engine/Plugins/Editor/GameplayTagsEditor/Source/GameplayTagsEditor/Private/GameplayTagsEditorModule.cpp
@@ -1,6 +1,7 @@
 // Copyright Epic Games, Inc. All Rights Reserved.
 
 #include "GameplayTagsEditorModule.h"
+
 #include "Misc/Paths.h"
 #include "PropertyEditorModule.h"
 #include "Factories/Factory.h"
@@ -33,6 +34,7 @@
 #include "UObject/UObjectHash.h"
 #include "HAL/IConsoleManager.h"
 #include "Misc/FileHelper.h"
+#include "UObject/ObjectSaveContext.h"
 
 #define LOCTEXT_NAMESPACE "GameplayTagEditor"
 
@@ -90,7 +92,7 @@ public:
 			GEditor->GetEditorSubsystem<UImportSubsystem>()->OnAssetPostImport.AddRaw(this, &FGameplayTagsEditorModule::OnObjectReimported);
 			FEditorDelegates::OnEditAssetIdentifiers.AddRaw(this, &FGameplayTagsEditorModule::OnEditGameplayTag);
 			IGameplayTagsModule::OnTagSettingsChanged.AddRaw(this, &FGameplayTagsEditorModule::OnEditorSettingsChanged);
-			UPackage::PackageSavedEvent.AddRaw(this, &FGameplayTagsEditorModule::OnPackageSaved);
+			UPackage::PackageSavedWithContextEvent.AddRaw(this, &FGameplayTagsEditorModule::OnPackageSaved);
 		}
 	}
 
@@ -124,7 +126,7 @@ public:
 		}
 		FEditorDelegates::OnEditAssetIdentifiers.RemoveAll(this);
 		IGameplayTagsModule::OnTagSettingsChanged.RemoveAll(this);
-		UPackage::PackageSavedEvent.RemoveAll(this);
+		UPackage::PackageSavedWithContextEvent.RemoveAll(this);
 	}
 
 	void OnEditorSettingsChanged()
@@ -145,7 +147,7 @@ public:
 		}
 	}
 
-	void OnPackageSaved(const FString& PackageFileName, UObject* PackageObj)
+	void OnPackageSaved(const FString& PackageFileName, UPackage* Package, FObjectPostSaveContext ObjectSaveContext)
 	{
 		if (GIsEditor && !IsRunningCommandlet())
 		{
@@ -155,7 +157,7 @@ public:
 
 			TArray<UObject*> Objects;
 			const bool bIncludeNestedObjects = false;
-			GetObjectsWithPackage(Cast<UPackage>(PackageObj), Objects, bIncludeNestedObjects);
+			GetObjectsWithPackage(Package, Objects, bIncludeNestedObjects);
 			for (UObject* Entry : Objects)
 			{
 				if (UDataTable* DataTable = Cast<UDataTable>(Entry))
diff --git a/Engine/Plugins/Enterprise/LidarPointCloud/Source/LidarPointCloudRuntime/Private/LidarPointCloud.cpp b/Engine/Plugins/Enterprise/LidarPointCloud/Source/LidarPointCloudRuntime/Private/LidarPointCloud.cpp
--- a/Engine/Plugins/Enterprise/LidarPointCloud/Source/LidarPointCloudRuntime/Private/LidarPointCloud.cpp
+++ b/Engine/Plugins/Enterprise/LidarPointCloud/Source/LidarPointCloudRuntime/Private/LidarPointCloud.cpp
@@ -1,6 +1,7 @@
 // Copyright Epic Games, Inc. All Rights Reserved.
 
 #include "LidarPointCloud.h"
+
 #include "LidarPointCloudShared.h"
 #include "LidarPointCloudActor.h"
 #include "LidarPointCloudComponent.h"
@@ -26,6 +27,7 @@
 #include "Editor.h"
 #include "Styling/SlateStyleRegistry.h"
 #include "Misc/MessageDialog.h"
+#include "UObject/ObjectSaveContext.h"
 #include "UObject/UObjectGlobals.h"
 #endif
 
@@ -46,8 +48,8 @@ class FLidarPackageReloader : public FTickableGameObject
 	TArray<TWeakObjectPtr<UPackage>> PackagesToReload;
 
 public:
-	FLidarPackageReloader() { UPackage::PackageSavedEvent.AddRaw(this, &FLidarPackageReloader::OnPackageSaved); }
-	virtual ~FLidarPackageReloader() { UPackage::PackageSavedEvent.RemoveAll(this); }
+	FLidarPackageReloader() { UPackage::PackageSavedWithContextEvent.AddRaw(this, &FLidarPackageReloader::OnPackageSaved); }
+	virtual ~FLidarPackageReloader() { UPackage::PackageSavedWithContextEvent.RemoveAll(this); }
 	virtual TStatId GetStatId() const override { RETURN_QUICK_DECLARE_CYCLE_STAT(LidarPackageReloader, STATGROUP_Tickables); }
 	virtual ETickableTickType GetTickableTickType() const override { return ETickableTickType::Always; }
 	virtual bool IsTickableInEditor() const override { return true; }
@@ -66,11 +68,10 @@ public:
 	}
 
 private:
-	void OnPackageSaved(const FString& Filename, UObject* Obj)
+	void OnPackageSaved(const FString& Filename, UPackage* Package, FObjectPostSaveContext ObjectSaveContext)
 	{
 		if (GetDefault<ULidarPointCloudSettings>()->bReleaseAssetAfterSaving)
 		{
-			UPackage* Package = Cast<UPackage>(Obj);
 			if (Cast<ULidarPointCloud>(Package->FindAssetInPackage()))
 			{
 				PackagesToReload.Add(Package);
@@ -360,7 +361,14 @@ void ULidarPointCloud::BeginDestroy()
 
 void ULidarPointCloud::PreSave(const class ITargetPlatform* TargetPlatform)
 {
+	PRAGMA_DISABLE_DEPRECATION_WARNINGS;
 	Super::PreSave(TargetPlatform);
+	PRAGMA_ENABLE_DEPRECATION_WARNINGS;
+}
+
+void ULidarPointCloud::PreSave(FObjectPreSaveContext ObjectSaveContext)
+{
+	Super::PreSave(ObjectSaveContext);
 	OnPreSaveCleanupEvent.Broadcast();
 }
 
diff --git a/Engine/Plugins/Enterprise/LidarPointCloud/Source/LidarPointCloudRuntime/Public/LidarPointCloud.h b/Engine/Plugins/Enterprise/LidarPointCloud/Source/LidarPointCloudRuntime/Public/LidarPointCloud.h
--- a/Engine/Plugins/Enterprise/LidarPointCloud/Source/LidarPointCloudRuntime/Public/LidarPointCloud.h
+++ b/Engine/Plugins/Enterprise/LidarPointCloud/Source/LidarPointCloudRuntime/Public/LidarPointCloud.h
@@ -160,7 +160,9 @@ public:
 	virtual void PostLoad() override;
 	virtual void GetAssetRegistryTags(TArray<FAssetRegistryTag>& OutTags) const override;
 	virtual void BeginDestroy() override;
+	UE_DEPRECATED(5.0, "Use version that takes FObjectPreSaveContext instead.")
 	virtual void PreSave(const class ITargetPlatform* TargetPlatform) override;
+	virtual void PreSave(FObjectPreSaveContext ObjectSaveContext) override;
 #if WITH_EDITOR
 	virtual void ClearAllCachedCookedPlatformData() override;
 	virtual void PostEditChangeProperty(FPropertyChangedEvent & PropertyChangedEvent) override;
diff --git a/Engine/Plugins/Experimental/Animation/ContextualAnimation/Source/ContextualAnimation/Private/ContextualAnimCompositeSceneAsset.cpp b/Engine/Plugins/Experimental/Animation/ContextualAnimation/Source/ContextualAnimation/Private/ContextualAnimCompositeSceneAsset.cpp
--- a/Engine/Plugins/Experimental/Animation/ContextualAnimation/Source/ContextualAnimation/Private/ContextualAnimCompositeSceneAsset.cpp
+++ b/Engine/Plugins/Experimental/Animation/ContextualAnimation/Source/ContextualAnimation/Private/ContextualAnimCompositeSceneAsset.cpp
@@ -1,8 +1,10 @@
 // Copyright Epic Games, Inc. All Rights Reserved.
 
 #include "ContextualAnimCompositeSceneAsset.h"
+
 #include "ContextualAnimMetadata.h"
 #include "Animation/AnimMontage.h"
+#include "UObject/ObjectSaveContext.h"
 
 const FName UContextualAnimCompositeSceneAsset::InteractorRoleName = FName(TEXT("interactor"));
 const FName UContextualAnimCompositeSceneAsset::InteractableRoleName = FName(TEXT("interactable"));
@@ -24,6 +26,13 @@ EContextualAnimJoinRule UContextualAnimCompositeSceneAsset::GetJoinRuleForRole(c
 }
 
 void UContextualAnimCompositeSceneAsset::PreSave(const class ITargetPlatform* TargetPlatform)
+{
+	PRAGMA_DISABLE_DEPRECATION_WARNINGS;
+	Super::PreSave(TargetPlatform);
+	PRAGMA_ENABLE_DEPRECATION_WARNINGS;
+}
+
+void UContextualAnimCompositeSceneAsset::PreSave(FObjectPreSaveContext ObjectSaveContext)
 {
 	// Necessary for FCompactPose that uses a FAnimStackAllocator (TMemStackAllocator) which allocates from FMemStack.
 	// When allocating memory from FMemStack we need to explicitly use FMemMark to ensure items are freed when the scope exits. 
@@ -32,7 +41,7 @@ void UContextualAnimCompositeSceneAsset::PreSave(const class ITargetPlatform* Ta
 	FMemMark Mark(FMemStack::Get());
 
 	// Generate scene pivot for each alignment section
-	Super::PreSave(TargetPlatform);
+	Super::PreSave(ObjectSaveContext);
 
 	// Generate alignment tracks relative to scene pivot
 	for (FContextualAnimData& Data : InteractorTrack.AnimDataContainer)
diff --git a/Engine/Plugins/Experimental/Animation/ContextualAnimation/Source/ContextualAnimation/Private/ContextualAnimSceneAsset.cpp b/Engine/Plugins/Experimental/Animation/ContextualAnimation/Source/ContextualAnimation/Private/ContextualAnimSceneAsset.cpp
--- a/Engine/Plugins/Experimental/Animation/ContextualAnimation/Source/ContextualAnimation/Private/ContextualAnimSceneAsset.cpp
+++ b/Engine/Plugins/Experimental/Animation/ContextualAnimation/Source/ContextualAnimation/Private/ContextualAnimSceneAsset.cpp
@@ -1,11 +1,13 @@
 // Copyright Epic Games, Inc. All Rights Reserved.
 
 #include "ContextualAnimSceneAsset.h"
+
 #include "AnimationRuntime.h"
 #include "Animation/AnimMontage.h"
 #include "Animation/AnimationPoseData.h"
 #include "Animation/AnimTypes.h"
 #include "ContextualAnimUtilities.h"
+#include "UObject/ObjectSaveContext.h"
 
 UContextualAnimSceneAsset::UContextualAnimSceneAsset(const FObjectInitializer& ObjectInitializer)
 	: Super(ObjectInitializer)
@@ -13,6 +15,13 @@ UContextualAnimSceneAsset::UContextualAnimSceneAsset(const FObjectInitializer& O
 }
 
 void UContextualAnimSceneAsset::PreSave(const class ITargetPlatform* TargetPlatform)
+{
+	PRAGMA_DISABLE_DEPRECATION_WARNINGS;
+	Super::PreSave(TargetPlatform);
+	PRAGMA_ENABLE_DEPRECATION_WARNINGS;
+}
+
+void UContextualAnimSceneAsset::PreSave(FObjectPreSaveContext ObjectSaveContext)
 {
 	// Necessary for FCompactPose that uses a FAnimStackAllocator (TMemStackAllocator) which allocates from FMemStack.
 	// When allocating memory from FMemStack we need to explicitly use FMemMark to ensure items are freed when the scope exits. 
@@ -21,7 +30,7 @@ void UContextualAnimSceneAsset::PreSave(const class ITargetPlatform* TargetPlatf
 	FMemMark Mark(FMemStack::Get());
 
 	// Generate scene pivot for each alignment section
-	Super::PreSave(TargetPlatform);
+	Super::PreSave(ObjectSaveContext);
 	
 	const FTransform MeshToComponentInverse = MeshToComponent.Inverse();
 	const float SampleInterval = 1.f / SampleRate;
diff --git a/Engine/Plugins/Experimental/Animation/ContextualAnimation/Source/ContextualAnimation/Private/ContextualAnimSceneAssetBase.cpp b/Engine/Plugins/Experimental/Animation/ContextualAnimation/Source/ContextualAnimation/Private/ContextualAnimSceneAssetBase.cpp
--- a/Engine/Plugins/Experimental/Animation/ContextualAnimation/Source/ContextualAnimation/Private/ContextualAnimSceneAssetBase.cpp
+++ b/Engine/Plugins/Experimental/Animation/ContextualAnimation/Source/ContextualAnimation/Private/ContextualAnimSceneAssetBase.cpp
@@ -1,6 +1,7 @@
 // Copyright Epic Games, Inc. All Rights Reserved.
 
 #include "ContextualAnimSceneAssetBase.h"
+
 #include "AnimationRuntime.h"
 #include "Animation/AnimMontage.h"
 #include "Animation/AnimationPoseData.h"
@@ -9,6 +10,7 @@
 #include "ContextualAnimScenePivotProvider.h"
 #include "ContextualAnimUtilities.h"
 #include "ContextualAnimMetadata.h"
+#include "UObject/ObjectSaveContext.h"
 
 UContextualAnimSceneAssetBase::UContextualAnimSceneAssetBase(const FObjectInitializer& ObjectInitializer)
 	: Super(ObjectInitializer)
@@ -96,7 +98,14 @@ void UContextualAnimSceneAssetBase::GenerateAlignmentTracksRelativeToScenePivot(
 
 void UContextualAnimSceneAssetBase::PreSave(const class ITargetPlatform* TargetPlatform)
 {
+	PRAGMA_DISABLE_DEPRECATION_WARNINGS;
 	Super::PreSave(TargetPlatform);
+	PRAGMA_ENABLE_DEPRECATION_WARNINGS;
+}
+
+void UContextualAnimSceneAssetBase::PreSave(FObjectPreSaveContext ObjectSaveContext)
+{
+	Super::PreSave(ObjectSaveContext);
 
 	// Generate scene pivot for each alignment section
 	for (int32 Idx = 0; Idx < AlignmentSections.Num(); Idx++)
diff --git a/Engine/Plugins/Experimental/Animation/ContextualAnimation/Source/ContextualAnimation/Public/ContextualAnimCompositeSceneAsset.h b/Engine/Plugins/Experimental/Animation/ContextualAnimation/Source/ContextualAnimation/Public/ContextualAnimCompositeSceneAsset.h
--- a/Engine/Plugins/Experimental/Animation/ContextualAnimation/Source/ContextualAnimation/Public/ContextualAnimCompositeSceneAsset.h
+++ b/Engine/Plugins/Experimental/Animation/ContextualAnimation/Source/ContextualAnimation/Public/ContextualAnimCompositeSceneAsset.h
@@ -82,7 +82,9 @@ public:
 
 	UContextualAnimCompositeSceneAsset(const FObjectInitializer& ObjectInitializer);
 
+	UE_DEPRECATED(5.0, "Use version that takes FObjectPreSaveContext instead.")
 	virtual void PreSave(const class ITargetPlatform* TargetPlatform) override;
+	virtual void PreSave(FObjectPreSaveContext ObjectSaveContext) override;
 
 	bool QueryData(FContextualAnimQueryResult& OutResult, const FContextualAnimQueryParams& QueryParams, const FTransform& ToWorldTransform) const;
 
diff --git a/Engine/Plugins/Experimental/Animation/PoseSearch/Source/Runtime/Private/PoseSearch.cpp b/Engine/Plugins/Experimental/Animation/PoseSearch/Source/Runtime/Private/PoseSearch.cpp
--- a/Engine/Plugins/Experimental/Animation/PoseSearch/Source/Runtime/Private/PoseSearch.cpp
+++ b/Engine/Plugins/Experimental/Animation/PoseSearch/Source/Runtime/Private/PoseSearch.cpp
@@ -20,6 +20,7 @@
 #include "Animation/AnimationPoseData.h"
 #include "AnimationRuntime.h"
 #include "BonePose.h"
+#include "UObject/ObjectSaveContext.h"
 
 IMPLEMENT_ANIMGRAPH_MESSAGE(UE::PoseSearch::IPoseHistoryProvider);
 
@@ -221,6 +222,13 @@ bool FPoseSearchFeatureVectorLayout::IsValid(int32 MaxNumBones) const
 // UPoseSearchSchema
 
 void UPoseSearchSchema::PreSave(const class ITargetPlatform* TargetPlatform)
+{
+	PRAGMA_DISABLE_DEPRECATION_WARNINGS;
+	Super::PreSave(TargetPlatform);
+	PRAGMA_ENABLE_DEPRECATION_WARNINGS;
+}
+
+void UPoseSearchSchema::PreSave(FObjectPreSaveContext ObjectSaveContext)
 {
 	SampleRate = FMath::Clamp(SampleRate, 1, 60);
 	SamplingInterval = 1.0f / SampleRate;
@@ -235,7 +243,7 @@ void UPoseSearchSchema::PreSave(const class ITargetPlatform* TargetPlatform)
 	GenerateLayout();
 	ResolveBoneReferences();
 
-	Super::PreSave(TargetPlatform);
+	Super::PreSave(ObjectSaveContext);
 }
 
 void UPoseSearchSchema::PostLoad()
@@ -396,6 +404,13 @@ void FPoseSearchIndex::Reset()
 // UPoseSearchSequenceMetaData
 
 void UPoseSearchSequenceMetaData::PreSave(const class ITargetPlatform* TargetPlatform)
+{
+	PRAGMA_DISABLE_DEPRECATION_WARNINGS;
+	Super::PreSave(TargetPlatform);
+	PRAGMA_ENABLE_DEPRECATION_WARNINGS;
+}
+
+void UPoseSearchSequenceMetaData::PreSave(FObjectPreSaveContext ObjectSaveContext)
 {
 	SearchIndex.Reset();
 
@@ -408,7 +423,7 @@ void UPoseSearchSequenceMetaData::PreSave(const class ITargetPlatform* TargetPla
 		}
 	}
 
-	Super::PreSave(TargetPlatform);
+	Super::PreSave(ObjectSaveContext);
 }
 
 bool UPoseSearchSequenceMetaData::IsValidForIndexing() const
@@ -472,6 +487,13 @@ bool UPoseSearchDatabase::IsValidForSearch() const
 }
 
 void UPoseSearchDatabase::PreSave(const class ITargetPlatform* TargetPlatform)
+{
+	PRAGMA_DISABLE_DEPRECATION_WARNINGS;
+	Super::PreSave(TargetPlatform);
+	PRAGMA_ENABLE_DEPRECATION_WARNINGS;
+}
+
+void UPoseSearchDatabase::PreSave(FObjectPreSaveContext ObjectSaveContext)
 {
 	SearchIndex.Reset();
 
@@ -480,7 +502,7 @@ void UPoseSearchDatabase::PreSave(const class ITargetPlatform* TargetPlatform)
 		UE::PoseSearch::BuildIndex(this);
 	}
 
-	Super::PreSave(TargetPlatform);
+	Super::PreSave(ObjectSaveContext);
 }
 
 
diff --git a/Engine/Plugins/Experimental/CommonConversation/Source/CommonConversationGraph/Private/CommonConversationGraphModule.cpp b/Engine/Plugins/Experimental/CommonConversation/Source/CommonConversationGraph/Private/CommonConversationGraphModule.cpp
--- a/Engine/Plugins/Experimental/CommonConversation/Source/CommonConversationGraph/Private/CommonConversationGraphModule.cpp
+++ b/Engine/Plugins/Experimental/CommonConversation/Source/CommonConversationGraph/Private/CommonConversationGraphModule.cpp
@@ -1,9 +1,11 @@
 // Copyright Epic Games, Inc. All Rights Reserved.
 
 #include "CommonConversationGraphModule.h"
+
 #include "ConversationDatabase.h"
 #include "ConversationCompiler.h"
 #include "Engine/AssetManager.h"
+#include "UObject/ObjectSaveContext.h"
 #include "UObject/UObjectHash.h"
 #include "UObject/CoreRedirects.h"
 
@@ -17,7 +19,7 @@ DEFINE_LOG_CATEGORY(LogCommonConversationGraph);
 
 void FCommonConversationGraphModule::StartupModule()
 {
-	UPackage::PreSavePackageEvent.AddRaw(this, &FCommonConversationGraphModule::HandlePreSavePackage);
+	UPackage::PreSavePackageWithContextEvent.AddRaw(this, &FCommonConversationGraphModule::HandlePreSavePackage);
 
 #if WITH_EDITOR
 	FEditorDelegates::BeginPIE.AddRaw(this, &FCommonConversationGraphModule::HandleBeginPIE);
@@ -33,14 +35,14 @@ void FCommonConversationGraphModule::ShutdownModule()
 	// This function may be called during shutdown to clean up your module.  For modules that support dynamic reloading,
 	// we call this function before unloading the module.
 
-	UPackage::PreSavePackageEvent.RemoveAll(this);
+	UPackage::PreSavePackageWithContextEvent.RemoveAll(this);
 
 #if WITH_EDITOR
 	FEditorDelegates::BeginPIE.RemoveAll(this);
 #endif
 }
 
-void FCommonConversationGraphModule::HandlePreSavePackage(UPackage* Package)
+void FCommonConversationGraphModule::HandlePreSavePackage(UPackage* Package, FObjectPreSaveContext ObjectSaveContext)
 {
 	TArray<UObject*> Objects;
 	const bool bIncludeNestedObjects = false;
diff --git a/Engine/Plugins/Experimental/CommonConversation/Source/CommonConversationRuntime/Private/ConversationDatabase.cpp b/Engine/Plugins/Experimental/CommonConversation/Source/CommonConversationRuntime/Private/ConversationDatabase.cpp
--- a/Engine/Plugins/Experimental/CommonConversation/Source/CommonConversationRuntime/Private/ConversationDatabase.cpp
+++ b/Engine/Plugins/Experimental/CommonConversation/Source/CommonConversationRuntime/Private/ConversationDatabase.cpp
@@ -4,6 +4,7 @@
 
 #if WITH_EDITOR
 #include "EdGraph/EdGraph.h"
+#include "UObject/ObjectSaveContext.h"
 #endif
 #include "AssetRegistry/IAssetRegistry.h"
 #include "AssetRegistry/AssetRegistryModule.h"
@@ -57,9 +58,14 @@ FLinearColor UConversationDatabase::GetDebugParticipantColor(FGameplayTag Partic
 #if WITH_EDITOR
 void UConversationDatabase::PreSave(const class ITargetPlatform* TargetPlatform)
 {
+	PRAGMA_DISABLE_DEPRECATION_WARNINGS;
 	Super::PreSave(TargetPlatform);
+	PRAGMA_ENABLE_DEPRECATION_WARNINGS;
+}
 
-
+void UConversationDatabase::PreSave(FObjectPreSaveContext ObjectSaveContext)
+{
+	Super::PreSave(ObjectSaveContext);
 }
 
 EDataValidationResult UConversationDatabase::IsDataValid(TArray<FText>& ValidationErrors)
diff --git a/Engine/Plugins/Experimental/ControlRig/Source/ControlRigDeveloper/Private/ControlRigBlueprint.cpp b/Engine/Plugins/Experimental/ControlRig/Source/ControlRigDeveloper/Private/ControlRigBlueprint.cpp
--- a/Engine/Plugins/Experimental/ControlRig/Source/ControlRigDeveloper/Private/ControlRigBlueprint.cpp
+++ b/Engine/Plugins/Experimental/ControlRig/Source/ControlRigDeveloper/Private/ControlRigBlueprint.cpp
@@ -1,6 +1,7 @@
 // Copyright Epic Games, Inc. All Rights Reserved.
 
 #include "ControlRigBlueprint.h"
+
 #include "ControlRigBlueprintGeneratedClass.h"
 #include "EdGraph/EdGraph.h"
 #include "EdGraphNode_Comment.h"
@@ -11,6 +12,7 @@
 #include "Graph/ControlRigGraph.h"
 #include "Graph/ControlRigGraphNode.h"
 #include "Graph/ControlRigGraphSchema.h"
+#include "UObject/ObjectSaveContext.h"
 #include "UObject/UObjectGlobals.h"
 #include "ControlRigObjectVersion.h"
 #include "ControlRigDeveloper.h"
@@ -293,7 +295,14 @@ void UControlRigBlueprint::Serialize(FArchive& Ar)
 
 void UControlRigBlueprint::PreSave(const class ITargetPlatform* TargetPlatform)
 {
+	PRAGMA_DISABLE_DEPRECATION_WARNINGS;
 	Super::PreSave(TargetPlatform);
+	PRAGMA_ENABLE_DEPRECATION_WARNINGS;
+}
+
+void UControlRigBlueprint::PreSave(FObjectPreSaveContext ObjectSaveContext)
+{
+	Super::PreSave(ObjectSaveContext);
 
 	SupportedEventNames.Reset();
 	if (UControlRigBlueprintGeneratedClass* RigClass = GetControlRigBlueprintGeneratedClass())
diff --git a/Engine/Plugins/Experimental/VirtualProductionUtilities/Source/VPUtilitiesEditor/Private/VPCameraBase.cpp b/Engine/Plugins/Experimental/VirtualProductionUtilities/Source/VPUtilitiesEditor/Private/VPCameraBase.cpp
--- a/Engine/Plugins/Experimental/VirtualProductionUtilities/Source/VPUtilitiesEditor/Private/VPCameraBase.cpp
+++ b/Engine/Plugins/Experimental/VirtualProductionUtilities/Source/VPUtilitiesEditor/Private/VPCameraBase.cpp
@@ -12,6 +12,7 @@
 #include "ScopedTransaction.h"
 #include "SLevelViewport.h"
 #include "VREditorMode.h"
+#include "UObject/ObjectSaveContext.h"
 #endif
 
 #define LOCTEXT_NAMESPACE "VPCameraUIBase"
@@ -89,8 +90,15 @@ void AVPCameraBase::BeginDestroy()
 
 void AVPCameraBase::PreSave(const class ITargetPlatform* TargetPlatform)
 {
-	RemoveUser(true);
+	PRAGMA_DISABLE_DEPRECATION_WARNINGS;
 	Super::PreSave(TargetPlatform);
+	PRAGMA_ENABLE_DEPRECATION_WARNINGS;
+}
+
+void AVPCameraBase::PreSave(FObjectPreSaveContext ObjectSaveContext)
+{
+	RemoveUser(true);
+	Super::PreSave(ObjectSaveContext);
 }
 
 void AVPCameraBase::AddPreview()
diff --git a/Engine/Plugins/FX/Niagara/Source/Niagara/Classes/NiagaraScript.h b/Engine/Plugins/FX/Niagara/Source/Niagara/Classes/NiagaraScript.h
--- a/Engine/Plugins/FX/Niagara/Source/Niagara/Classes/NiagaraScript.h
+++ b/Engine/Plugins/FX/Niagara/Source/Niagara/Classes/NiagaraScript.h
@@ -610,8 +610,10 @@ public:
 #endif
 
 	//~ Begin UObject interface
-	void PreSave(const class ITargetPlatform* TargetPlatform) override;
-	void Serialize(FArchive& Ar)override;
+	UE_DEPRECATED(5.0, "Use version that takes FObjectPreSaveContext instead.")
+	virtual void PreSave(const class ITargetPlatform* TargetPlatform) override;
+	virtual void PreSave(FObjectPreSaveContext ObjectSaveContext) override;
+	virtual void Serialize(FArchive& Ar)override;
 	virtual void PostLoad() override;
 #if WITH_EDITOR
 	virtual void PostEditChangeProperty(FPropertyChangedEvent& PropertyChangedEvent) override;
diff --git a/Engine/Plugins/FX/Niagara/Source/Niagara/Classes/NiagaraSystem.h b/Engine/Plugins/FX/Niagara/Source/Niagara/Classes/NiagaraSystem.h
--- a/Engine/Plugins/FX/Niagara/Source/Niagara/Classes/NiagaraSystem.h
+++ b/Engine/Plugins/FX/Niagara/Source/Niagara/Classes/NiagaraSystem.h
@@ -214,10 +214,12 @@ public:
 
 	//~ UObject interface
 	void PostInitProperties();
-	void Serialize(FArchive& Ar)override;
+	virtual void Serialize(FArchive& Ar) override;
 	virtual void PostLoad() override; 
 	virtual void BeginDestroy() override;
+	UE_DEPRECATED(5.0, "Use version that takes FObjectPreSaveContext instead.")
 	virtual void PreSave(const class ITargetPlatform* TargetPlatform) override;
+	virtual void PreSave(FObjectPreSaveContext ObjectSaveContext) override;
 #if WITH_EDITOR
 	virtual void PreEditChange(FProperty* PropertyThatWillChange)override;
 	virtual void PostEditChangeProperty(struct FPropertyChangedEvent& PropertyChangedEvent) override; 
diff --git a/Engine/Plugins/FX/Niagara/Source/Niagara/Private/NiagaraScript.cpp b/Engine/Plugins/FX/Niagara/Source/Niagara/Private/NiagaraScript.cpp
--- a/Engine/Plugins/FX/Niagara/Source/Niagara/Private/NiagaraScript.cpp
+++ b/Engine/Plugins/FX/Niagara/Source/Niagara/Private/NiagaraScript.cpp
@@ -1,6 +1,7 @@
 // Copyright Epic Games, Inc. All Rights Reserved.
 
 #include "NiagaraScript.h"
+
 #include "Modules/ModuleManager.h"
 #include "NiagaraScriptSourceBase.h"
 #include "NiagaraComponent.h"
@@ -19,6 +20,7 @@
 #include "HAL/PlatformFileManager.h"
 #include "Misc/FileHelper.h"
 #include "UObject/EditorObjectVersion.h"
+#include "UObject/ObjectSaveContext.h"
 #include "UObject/ReleaseObjectVersion.h"
 #include "NiagaraDataInterfaceSkeletalMesh.h"
 #include "NiagaraDataInterfaceStaticMesh.h"
@@ -865,7 +867,14 @@ void UNiagaraScript::GenerateDefaultFunctionBindings()
 
 void UNiagaraScript::PreSave(const class ITargetPlatform* TargetPlatform)
 {
+	PRAGMA_DISABLE_DEPRECATION_WARNINGS;
 	Super::PreSave(TargetPlatform);
+	PRAGMA_ENABLE_DEPRECATION_WARNINGS;
+}
+
+void UNiagaraScript::PreSave(FObjectPreSaveContext ObjectSaveContext)
+{
+	Super::PreSave(ObjectSaveContext);
 
 #if WITH_EDITORONLY_DATA
 	// Pre-save can happen in any order for objects in the package and since this is now used to cache data for execution we need to make sure that the system compilation
@@ -882,7 +891,7 @@ void UNiagaraScript::PreSave(const class ITargetPlatform* TargetPlatform)
 	// Make sure the data interfaces are consistent to prevent crashes in later caching operations.
 	if (CachedScriptVM.DataInterfaceInfo.Num() != CachedDefaultDataInterfaces.Num())
 	{
-		UE_LOG(LogNiagara, Warning, TEXT("Data interface count mistmatch during script presave. Invaliding compile results (see full log for details).  Script: %s"), *GetPathName());
+		UE_LOG(LogNiagara, Warning, TEXT("Data interface count mismatch during script presave. Invaliding compile results (see full log for details).  Script: %s"), *GetPathName());
 		UE_LOG(LogNiagara, Log, TEXT("Compiled DataInterfaceInfos:"));
 		for (const FNiagaraScriptDataInterfaceCompileInfo& DataInterfaceCompileInfo : CachedScriptVM.DataInterfaceInfo)
 		{
@@ -902,6 +911,7 @@ void UNiagaraScript::PreSave(const class ITargetPlatform* TargetPlatform)
 		return;
 	}
 
+	const ITargetPlatform* TargetPlatform = ObjectSaveContext.GetTargetPlatform();
 	if (TargetPlatform && TargetPlatform->RequiresCookedData())
 	{
 		TOptional<ENiagaraSimTarget> SimTarget = GetSimTarget();
diff --git a/Engine/Plugins/FX/Niagara/Source/Niagara/Private/NiagaraSystem.cpp b/Engine/Plugins/FX/Niagara/Source/Niagara/Private/NiagaraSystem.cpp
--- a/Engine/Plugins/FX/Niagara/Source/Niagara/Private/NiagaraSystem.cpp
+++ b/Engine/Plugins/FX/Niagara/Source/Niagara/Private/NiagaraSystem.cpp
@@ -2,6 +2,7 @@
 
 
 #include "NiagaraSystem.h"
+
 #include "NiagaraRendererProperties.h"
 #include "NiagaraRenderer.h"
 #include "NiagaraConstants.h"
@@ -27,6 +28,7 @@
 #include "Algo/RemoveIf.h"
 #include "Async/Async.h"
 #include "Misc/ScopedSlowTask.h"
+#include "UObject/ObjectSaveContext.h"
 
 #define LOCTEXT_NAMESPACE "NiagaraSystem"
 
@@ -132,9 +134,16 @@ void UNiagaraSystem::BeginDestroy()
 	//FNiagaraWorldManager::DestroyAllSystemSimulations(this);
 }
 
-void UNiagaraSystem::PreSave(const class ITargetPlatform * TargetPlatform)
+void UNiagaraSystem::PreSave(const class ITargetPlatform* TargetPlatform)
 {
+	PRAGMA_DISABLE_DEPRECATION_WARNINGS;
 	Super::PreSave(TargetPlatform);
+	PRAGMA_ENABLE_DEPRECATION_WARNINGS;
+}
+
+void UNiagaraSystem::PreSave(FObjectPreSaveContext ObjectSaveContext)
+{
+	Super::PreSave(ObjectSaveContext);
 #if WITH_EDITORONLY_DATA
 	WaitForCompilationComplete();
 #endif
diff --git a/Engine/Plugins/Runtime/DataRegistry/Source/DataRegistry/Private/DataRegistrySource_CurveTable.cpp b/Engine/Plugins/Runtime/DataRegistry/Source/DataRegistry/Private/DataRegistrySource_CurveTable.cpp
--- a/Engine/Plugins/Runtime/DataRegistry/Source/DataRegistry/Private/DataRegistrySource_CurveTable.cpp
+++ b/Engine/Plugins/Runtime/DataRegistry/Source/DataRegistry/Private/DataRegistrySource_CurveTable.cpp
@@ -1,10 +1,12 @@
 // Copyright Epic Games, Inc. All Rights Reserved.
 
 #include "DataRegistrySource_CurveTable.h"
+
 #include "DataRegistryTypesPrivate.h"
 #include "DataRegistrySettings.h"
 #include "Interfaces/ITargetPlatform.h"
 #include "Engine/AssetManager.h"
+#include "UObject/ObjectSaveContext.h"
 
 void UDataRegistrySource_CurveTable::SetSourceTable(const TSoftObjectPtr<UCurveTable>& InSourceTable, const FDataRegistrySource_DataTableRules& InTableRules)
 {
@@ -295,7 +297,14 @@ void UDataRegistrySource_CurveTable::EditorRefreshSource()
 
 void UDataRegistrySource_CurveTable::PreSave(const ITargetPlatform* TargetPlatform)
 {
+	PRAGMA_DISABLE_DEPRECATION_WARNINGS;
 	Super::PreSave(TargetPlatform);
+	PRAGMA_ENABLE_DEPRECATION_WARNINGS;
+}
+
+void UDataRegistrySource_CurveTable::PreSave(FObjectPreSaveContext ObjectSaveContext)
+{
+	Super::PreSave(ObjectSaveContext);
 
 	// Force load it to validate type on save
 	SetCachedTable(true);
diff --git a/Engine/Plugins/Runtime/DataRegistry/Source/DataRegistry/Private/DataRegistrySource_DataTable.cpp b/Engine/Plugins/Runtime/DataRegistry/Source/DataRegistry/Private/DataRegistrySource_DataTable.cpp
--- a/Engine/Plugins/Runtime/DataRegistry/Source/DataRegistry/Private/DataRegistrySource_DataTable.cpp
+++ b/Engine/Plugins/Runtime/DataRegistry/Source/DataRegistry/Private/DataRegistrySource_DataTable.cpp
@@ -1,10 +1,12 @@
 // Copyright Epic Games, Inc. All Rights Reserved.
 
 #include "DataRegistrySource_DataTable.h"
+
 #include "DataRegistryTypesPrivate.h"
 #include "DataRegistrySettings.h"
 #include "Interfaces/ITargetPlatform.h"
 #include "Engine/AssetManager.h"
+#include "UObject/ObjectSaveContext.h"
 
 void UDataRegistrySource_DataTable::SetSourceTable(const TSoftObjectPtr<UDataTable>& InSourceTable, const FDataRegistrySource_DataTableRules& InTableRules)
 {
@@ -292,7 +294,14 @@ void UDataRegistrySource_DataTable::EditorRefreshSource()
 
 void UDataRegistrySource_DataTable::PreSave(const ITargetPlatform* TargetPlatform)
 {
+	PRAGMA_DISABLE_DEPRECATION_WARNINGS;
 	Super::PreSave(TargetPlatform);
+	PRAGMA_ENABLE_DEPRECATION_WARNINGS;
+}
+
+void UDataRegistrySource_DataTable::PreSave(FObjectPreSaveContext ObjectSaveContext)
+{
+	Super::PreSave(ObjectSaveContext);
 
 	// Force load it to validate type on save
 	SetCachedTable(true);
diff --git a/Engine/Plugins/Runtime/HairStrands/Source/HairStrandsCore/Private/GroomAsset.cpp b/Engine/Plugins/Runtime/HairStrands/Source/HairStrandsCore/Private/GroomAsset.cpp
--- a/Engine/Plugins/Runtime/HairStrands/Source/HairStrandsCore/Private/GroomAsset.cpp
+++ b/Engine/Plugins/Runtime/HairStrands/Source/HairStrandsCore/Private/GroomAsset.cpp
@@ -1,6 +1,7 @@
 // Copyright Epic Games, Inc. All Rights Reserved.
 
 #include "GroomAsset.h"
+
 #include "Async/Async.h"
 #include "EngineUtils.h"
 #include "GroomAssetImportData.h"
@@ -27,6 +28,7 @@
 #include "Engine/StaticMesh.h"
 #include "Logging/LogMacros.h"
 #include "HairStrandsCore.h"
+#include "UObject/ObjectSaveContext.h"
 
 #if WITH_EDITOR
 #include "Interfaces/ITargetPlatform.h"
@@ -732,6 +734,13 @@ void UGroomAsset::PostLoad()
 }
 
 void UGroomAsset::PreSave(const class ITargetPlatform* TargetPlatform)
+{
+	PRAGMA_DISABLE_DEPRECATION_WARNINGS;
+	Super::PreSave(TargetPlatform);
+	PRAGMA_ENABLE_DEPRECATION_WARNINGS;
+}
+
+void UGroomAsset::PreSave(FObjectPreSaveContext ObjectSaveContext)
 {
 #if WITH_EDITORONLY_DATA
 	check(AreGroupsValid());
@@ -792,7 +801,9 @@ void UGroomAsset::PreSave(const class ITargetPlatform* TargetPlatform)
 	UpdateHairGroupsInfo();
 	UpdateCachedSettings();
 #endif
+	Super::PreSave(ObjectSaveContext);
 }
+
 void UGroomAsset::BeginDestroy()
 {
 	ReleaseResource();
diff --git a/Engine/Plugins/Runtime/HairStrands/Source/HairStrandsCore/Private/GroomBindingAsset.cpp b/Engine/Plugins/Runtime/HairStrands/Source/HairStrandsCore/Private/GroomBindingAsset.cpp
--- a/Engine/Plugins/Runtime/HairStrands/Source/HairStrandsCore/Private/GroomBindingAsset.cpp
+++ b/Engine/Plugins/Runtime/HairStrands/Source/HairStrandsCore/Private/GroomBindingAsset.cpp
@@ -1,11 +1,13 @@
 // Copyright Epic Games, Inc. All Rights Reserved.
 
 #include "GroomBindingAsset.h"
+
 #include "EngineUtils.h"
 #include "GeometryCache.h"
 #include "GeometryCacheMeshData.h"
 #include "GroomAsset.h"
 #include "GroomBindingBuilder.h"
+#include "UObject/ObjectSaveContext.h"
 
 #if WITH_EDITORONLY_DATA
 #include "DerivedDataCacheInterface.h"
@@ -284,6 +286,13 @@ void UGroomBindingAsset::PostLoad()
 }
 
 void UGroomBindingAsset::PreSave(const class ITargetPlatform* TargetPlatform)
+{
+	PRAGMA_DISABLE_DEPRECATION_WARNINGS;
+	Super::PreSave(TargetPlatform);
+	PRAGMA_ENABLE_DEPRECATION_WARNINGS;
+}
+
+void UGroomBindingAsset::PreSave(FObjectPreSaveContext ObjectSaveContext)
 {
 #if WITH_EDITOR
 	while (QueryStatus == EQueryStatus::Submitted)
@@ -291,7 +300,7 @@ void UGroomBindingAsset::PreSave(const class ITargetPlatform* TargetPlatform)
 		FPlatformProcess::Sleep(1);
 	}
 #endif
-	Super::PreSave(TargetPlatform);
+	Super::PreSave(ObjectSaveContext);
 #if WITH_EDITOR
 	OnGroomBindingAssetChanged.Broadcast();
 #endif
@@ -299,7 +308,14 @@ void UGroomBindingAsset::PreSave(const class ITargetPlatform* TargetPlatform)
 
 void UGroomBindingAsset::PostSaveRoot(bool bCleanupIsRequired)
 {
+	PRAGMA_DISABLE_DEPRECATION_WARNINGS;
 	Super::PostSaveRoot(bCleanupIsRequired);
+	PRAGMA_ENABLE_DEPRECATION_WARNINGS;
+}
+
+void UGroomBindingAsset::PostSaveRoot(FObjectPostSaveRootContext ObjectSaveContext)
+{
+	Super::PostSaveRoot(ObjectSaveContext);
 #if WITH_EDITOR
 	OnGroomBindingAssetChanged.Broadcast();
 #endif
diff --git a/Engine/Plugins/Runtime/HairStrands/Source/HairStrandsCore/Public/GroomBindingAsset.h b/Engine/Plugins/Runtime/HairStrands/Source/HairStrandsCore/Public/GroomBindingAsset.h
--- a/Engine/Plugins/Runtime/HairStrands/Source/HairStrandsCore/Public/GroomBindingAsset.h
+++ b/Engine/Plugins/Runtime/HairStrands/Source/HairStrandsCore/Public/GroomBindingAsset.h
@@ -121,8 +121,12 @@ public:
 	//~ Begin UObject Interface.
 	virtual void GetResourceSizeEx(FResourceSizeEx& CumulativeResourceSize) override;
 	virtual void PostLoad() override;
+	UE_DEPRECATED(5.0, "Use version that takes FObjectPreSaveContext instead.")
 	virtual void PreSave(const class ITargetPlatform* TargetPlatform) override;
+	virtual void PreSave(FObjectPreSaveContext ObjectSaveContext) override;
+	UE_DEPRECATED(5.0, "Use version that takes FObjectPostSaveContext instead.")
 	virtual void PostSaveRoot(bool bCleanupIsRequired) override;
+	virtual void PostSaveRoot(FObjectPostSaveRootContext ObjectSaveContext) override;
 	virtual void BeginDestroy() override;
 	virtual void Serialize(FArchive& Ar) override;
 
diff --git a/Engine/Source/Developer/BlueprintCompilerCppBackend/Private/BlueprintCompilerCppBackendUMG.cpp b/Engine/Source/Developer/BlueprintCompilerCppBackend/Private/BlueprintCompilerCppBackendUMG.cpp
--- a/Engine/Source/Developer/BlueprintCompilerCppBackend/Private/BlueprintCompilerCppBackendUMG.cpp
+++ b/Engine/Source/Developer/BlueprintCompilerCppBackend/Private/BlueprintCompilerCppBackendUMG.cpp
@@ -18,7 +18,7 @@ void FBackendHelperUMG::WidgetFunctionsInHeader(FEmitterLocalContext& Context)
 	if (Cast<UWidgetBlueprintGeneratedClass>(Context.GetCurrentlyGeneratedClass()))
 	{
 		Context.Header.AddLine(FString::Printf(TEXT("virtual void %s(TArray<FName>& SlotNames) const override;"), GET_FUNCTION_NAME_STRING_CHECKED(UUserWidget, GetSlotNames)));
-		Context.Header.AddLine(FString::Printf(TEXT("virtual void %s(const class ITargetPlatform* TargetPlatform) override;"), GET_FUNCTION_NAME_STRING_CHECKED(UUserWidget, PreSave)));
+		Context.Header.AddLine(FString::Printf(TEXT("virtual void %s(FObjectPreSaveContext ObjectSaveContext) override;"), GET_FUNCTION_NAME_STRING_CHECKED_OneParam(UUserWidget, PreSave, FObjectPreSaveContext)));
 		Context.Header.AddLine(TEXT("virtual void InitializeNativeClassData() override;"));
 	}
 }
@@ -145,10 +145,10 @@ void FBackendHelperUMG::EmitWidgetInitializationFunctions(FEmitterLocalContext&
 		}
 
 		// PreSave
-		Context.AddLine(FString::Printf(TEXT("void %s::%s(const class ITargetPlatform* TargetPlatform)"), *CppClassName, GET_FUNCTION_NAME_STRING_CHECKED(UUserWidget, PreSave)));
+		Context.AddLine(FString::Printf(TEXT("void %s::%s(FObjectPreSaveContext ObjectSaveContext)"), *CppClassName, GET_FUNCTION_NAME_STRING_CHECKED_OneParam(UUserWidget, PreSave, FObjectPreSaveContext)));
 		Context.AddLine(TEXT("{"));
 		Context.IncreaseIndent();
-		Context.AddLine(FString::Printf(TEXT("Super::%s(TargetPlatform);"), GET_FUNCTION_NAME_STRING_CHECKED(UObject, PreSave)));
+		Context.AddLine(FString::Printf(TEXT("Super::%s(ObjectSaveContext);"), GET_FUNCTION_NAME_STRING_CHECKED_OneParam(UObject, PreSave, FObjectPreSaveContext)));
 		Context.AddLine(TEXT("TArray<FName> LocalNamedSlots;"));
 		Context.AddLine(FString::Printf(TEXT("%s(LocalNamedSlots);"), GET_FUNCTION_NAME_STRING_CHECKED(UUserWidget, GetSlotNames)));
 		Context.AddLine(TEXT("RemoveObsoleteBindings(LocalNamedSlots);")); //RemoveObsoleteBindings is protected - no check
diff --git a/Engine/Source/Editor/BehaviorTreeEditor/Private/BehaviorTreeEditor.cpp b/Engine/Source/Editor/BehaviorTreeEditor/Private/BehaviorTreeEditor.cpp
--- a/Engine/Source/Editor/BehaviorTreeEditor/Private/BehaviorTreeEditor.cpp
+++ b/Engine/Source/Editor/BehaviorTreeEditor/Private/BehaviorTreeEditor.cpp
@@ -1,10 +1,12 @@
 // Copyright Epic Games, Inc. All Rights Reserved.
 
 #include "BehaviorTreeEditor.h"
+
 #include "Widgets/Text/STextBlock.h"
 #include "Framework/MultiBox/MultiBoxBuilder.h"
 #include "Engine/Blueprint.h"
 #include "Widgets/Layout/SBorder.h"
+#include "UObject/ObjectSaveContext.h"
 #include "UObject/Package.h"
 #include "BehaviorTree/BTDecorator.h"
 #include "BehaviorTree/BTCompositeNode.h"
@@ -72,7 +74,7 @@ FBehaviorTreeEditor::FBehaviorTreeEditor()
 	: IBehaviorTreeEditor()
 {
 	// listen for package change events to update injected nodes
-	OnPackageSavedDelegateHandle     = UPackage::PackageSavedEvent.AddRaw(this, &FBehaviorTreeEditor::OnPackageSaved);
+	OnPackageSavedDelegateHandle     = UPackage::PackageSavedWithContextEvent.AddRaw(this, &FBehaviorTreeEditor::OnPackageSaved);
 
 	bShowDecoratorRangeLower = false;
 	bShowDecoratorRangeSelf = false;
@@ -91,7 +93,7 @@ FBehaviorTreeEditor::FBehaviorTreeEditor()
 
 FBehaviorTreeEditor::~FBehaviorTreeEditor()
 {
-	UPackage::PackageSavedEvent.Remove(OnPackageSavedDelegateHandle);
+	UPackage::PackageSavedWithContextEvent.Remove(OnPackageSavedDelegateHandle);
 
 	Debugger.Reset();
 }
@@ -1107,7 +1109,7 @@ void FBehaviorTreeEditor::OnFinishedChangingProperties(const FPropertyChangedEve
 	BehaviorTree->BTGraph->GetSchema()->ForceVisualizationCacheClear();
 }
 
-void FBehaviorTreeEditor::OnPackageSaved(const FString& PackageFileName, UObject* Outer)
+void FBehaviorTreeEditor::OnPackageSaved(const FString& PackageFileName, UPackage* Package, FObjectPostSaveContext ObjectSaveContext)
 {
 	UBehaviorTreeGraph* MyGraph = BehaviorTree ? Cast<UBehaviorTreeGraph>(BehaviorTree->BTGraph) : NULL;
 	if (MyGraph)
diff --git a/Engine/Source/Editor/BehaviorTreeEditor/Private/BehaviorTreeEditor.h b/Engine/Source/Editor/BehaviorTreeEditor/Private/BehaviorTreeEditor.h
--- a/Engine/Source/Editor/BehaviorTreeEditor/Private/BehaviorTreeEditor.h
+++ b/Engine/Source/Editor/BehaviorTreeEditor/Private/BehaviorTreeEditor.h
@@ -16,6 +16,7 @@ class FBehaviorTreeDebugger;
 class FBehaviorTreeEditorToolbar;
 class FDocumentTabFactory;
 class FDocumentTracker;
+class FObjectPostSaveContext;
 class IDetailsView;
 class SBehaviorTreeBlackboardEditor;
 class SBehaviorTreeBlackboardView;
@@ -90,7 +91,7 @@ public:
 	void JumpToNode(const UEdGraphNode* Node);
 
 	bool IsPropertyEditable() const;
-	void OnPackageSaved(const FString& PackageFileName, UObject* Outer);
+	void OnPackageSaved(const FString& PackageFileName, UPackage* Package, FObjectPostSaveContext ObjectSaveContext);
 	void OnFinishedChangingProperties(const FPropertyChangedEvent& PropertyChangedEvent);
 	virtual void OnClassListUpdated() override;
 
diff --git a/Engine/Source/Editor/BlueprintGraph/Private/K2Node_FunctionEntry.cpp b/Engine/Source/Editor/BlueprintGraph/Private/K2Node_FunctionEntry.cpp
--- a/Engine/Source/Editor/BlueprintGraph/Private/K2Node_FunctionEntry.cpp
+++ b/Engine/Source/Editor/BlueprintGraph/Private/K2Node_FunctionEntry.cpp
@@ -2,11 +2,13 @@
 
 
 #include "K2Node_FunctionEntry.h"
+
 #include "Engine/Blueprint.h"
 #include "Animation/AnimBlueprint.h"
 #include "UObject/UnrealType.h"
 #include "UObject/BlueprintsObjectVersion.h"
 #include "UObject/FrameworkObjectVersion.h"
+#include "UObject/ObjectSaveContext.h"
 #include "UObject/StructOnScope.h"
 #include "Engine/UserDefinedStruct.h"
 #include "EdGraph/EdGraphSchema.h"
@@ -203,8 +205,15 @@ UK2Node_FunctionEntry::UK2Node_FunctionEntry(const FObjectInitializer& ObjectIni
 
 void UK2Node_FunctionEntry::PreSave(const class ITargetPlatform* TargetPlatform)
 {
+	PRAGMA_DISABLE_DEPRECATION_WARNINGS;
 	Super::PreSave(TargetPlatform);
-	
+	PRAGMA_ENABLE_DEPRECATION_WARNINGS;
+}
+
+void UK2Node_FunctionEntry::PreSave(FObjectPreSaveContext ObjectSaveContext)
+{
+	Super::PreSave(ObjectSaveContext);
+
 	const UBlueprint* Blueprint = HasValidBlueprint() ? GetBlueprint() : nullptr;
 	if (Blueprint && LocalVariables.Num() > 0)
 	{
diff --git a/Engine/Source/Editor/Experimental/EditorInteractiveToolsFramework/Private/EdModeInteractiveToolsContext.cpp b/Engine/Source/Editor/Experimental/EditorInteractiveToolsFramework/Private/EdModeInteractiveToolsContext.cpp
--- a/Engine/Source/Editor/Experimental/EditorInteractiveToolsFramework/Private/EdModeInteractiveToolsContext.cpp
+++ b/Engine/Source/Editor/Experimental/EditorInteractiveToolsFramework/Private/EdModeInteractiveToolsContext.cpp
@@ -2,6 +2,7 @@
 
 
 #include "EdModeInteractiveToolsContext.h"
+
 #include "Editor.h"
 #include "EditorViewportClient.h"
 #include "EditorModeManager.h"
@@ -32,6 +33,7 @@
 
 #include "BaseGizmos/GizmoRenderingUtil.h"
 #include "UnrealClient.h"
+#include "UObject/ObjectSaveContext.h"
 
 //#define ENABLE_DEBUG_PRINTING
 
@@ -546,7 +548,7 @@ void UEdModeInteractiveToolsContext::Initialize(IToolsContextQueriesAPI* Queries
 	{
 		TerminateActiveToolsOnPIEStart();
 	});
-	PreSaveWorldDelegateHandle = FEditorDelegates::PreSaveWorld.AddLambda([this](uint32 SaveFlags, UWorld* World)
+	PreSaveWorldDelegateHandle = FEditorDelegates::PreSaveWorldWithContext.AddLambda([this](UWorld* World, FObjectPreSaveContext ObjectSaveContext)
 	{
 		TerminateActiveToolsOnSaveWorld();
 	});
@@ -580,7 +582,7 @@ void UEdModeInteractiveToolsContext::Shutdown()
 	{
 		LevelEditor->OnMapChanged().Remove(WorldTearDownDelegateHandle);
 		FEditorDelegates::BeginPIE.Remove(BeginPIEDelegateHandle);
-		FEditorDelegates::PreSaveWorld.Remove(PreSaveWorldDelegateHandle);
+		FEditorDelegates::PreSaveWorldWithContext.Remove(PreSaveWorldDelegateHandle);
 		GEditor->OnViewportClientListChanged().Remove(ViewportClientListChangedHandle);
 
 		// auto-accept any in-progress tools
diff --git a/Engine/Source/Editor/Kismet/Classes/BlueprintPaletteFavorites.h b/Engine/Source/Editor/Kismet/Classes/BlueprintPaletteFavorites.h
--- a/Engine/Source/Editor/Kismet/Classes/BlueprintPaletteFavorites.h
+++ b/Engine/Source/Editor/Kismet/Classes/BlueprintPaletteFavorites.h
@@ -237,7 +237,7 @@ public:
 	/** 
 	 * A list of strings that are used to identify specific palette actions. 
 	 * This is what gets saved out when the user has customized their own set, 
-	 * and is not updated until PreSave(const class ITargetPlatform* TargetPlatform).
+	 * and is not updated until PreSave.
 	 */
 	UPROPERTY(config)
 	TArray<FString> CustomFavorites;
diff --git a/Engine/Source/Editor/LandscapeEditor/Private/LandscapeEdMode.cpp b/Engine/Source/Editor/LandscapeEditor/Private/LandscapeEdMode.cpp
--- a/Engine/Source/Editor/LandscapeEditor/Private/LandscapeEdMode.cpp
+++ b/Engine/Source/Editor/LandscapeEditor/Private/LandscapeEdMode.cpp
@@ -1,6 +1,7 @@
 // Copyright Epic Games, Inc. All Rights Reserved.
 
 #include "LandscapeEdMode.h"
+
 #include "SceneView.h"
 #include "Engine/Texture2D.h"
 #include "EditorViewportClient.h"
@@ -32,6 +33,7 @@
 #include "LandscapeInfoMap.h"
 #include "LandscapeImportHelper.h"
 #include "LandscapeConfigHelper.h"
+#include "UObject/ObjectSaveContext.h"
 
 //Slate dependencies
 #include "Misc/FeedbackContext.h"
@@ -668,7 +670,7 @@ void FEdModeLandscape::Exit()
 	FEdMode::Exit();
 }
 
-void FEdModeLandscape::OnPreSaveWorld(uint32 InSaveFlags, const class UWorld* InWorld)
+void FEdModeLandscape::OnPreSaveWorld(const class UWorld* InWorld, FObjectPreSaveContext ObjectSaveContext)
 {
 	// Avoid doing this during cooking to keep determinism and we don't want to do this on GameWorlds.
 	if (!InWorld->IsGameWorld() && !GIsCookerLoadingPackage)
diff --git a/Engine/Source/Editor/LandscapeEditor/Private/LandscapeEdMode.h b/Engine/Source/Editor/LandscapeEditor/Private/LandscapeEdMode.h
--- a/Engine/Source/Editor/LandscapeEditor/Private/LandscapeEdMode.h
+++ b/Engine/Source/Editor/LandscapeEditor/Private/LandscapeEdMode.h
@@ -585,7 +585,7 @@ public:
 	DECLARE_EVENT(FEdModeLandscape, FTargetsListUpdated);
 	static FTargetsListUpdated TargetsListUpdated;
 
-	void OnPreSaveWorld(uint32 InSaveFlags, const class UWorld* InWorld);
+	void OnPreSaveWorld(const class UWorld* InWorld, FObjectPreSaveContext ObjectSaveContext);
 
 	/** Handle notification that visible levels may have changed and we should update the editable landscapes list */
 	void HandleLevelsChanged(bool ShouldExitMode);
diff --git a/Engine/Source/Editor/LandscapeEditor/Private/LandscapeEditorModule.cpp b/Engine/Source/Editor/LandscapeEditor/Private/LandscapeEditorModule.cpp
--- a/Engine/Source/Editor/LandscapeEditor/Private/LandscapeEditorModule.cpp
+++ b/Engine/Source/Editor/LandscapeEditor/Private/LandscapeEditorModule.cpp
@@ -1,6 +1,7 @@
 // Copyright Epic Games, Inc. All Rights Reserved.
 
 #include "LandscapeEditorModule.h"
+
 #include "Modules/ModuleManager.h"
 #include "Textures/SlateIcon.h"
 #include "Framework/Commands/UICommandList.h"
@@ -33,6 +34,7 @@
 #include "LandscapeSubsystem.h"
 
 #include "LandscapeRender.h"
+#include "UObject/ObjectSaveContext.h"
 
 #define LOCTEXT_NAMESPACE "LandscapeEditor"
 
@@ -65,7 +67,7 @@ public:
 	{
 		FLandscapeEditorCommands::Register();
 
-		PreSaveWorldHandle = FEditorDelegates::PreSaveWorld.AddRaw(this, &FLandscapeEditorModule::OnPreSaveWorld);
+		PreSaveWorldHandle = FEditorDelegates::PreSaveWorldWithContext.AddRaw(this, &FLandscapeEditorModule::OnPreSaveWorld);
 
 		// register the editor mode
 		FEditorModeRegistry::Get().RegisterMode<FEdModeLandscape>(
@@ -141,7 +143,7 @@ public:
 	{
 		FLandscapeEditorCommands::Unregister();
 
-		FEditorDelegates::PreSaveWorld.Remove(PreSaveWorldHandle);
+		FEditorDelegates::PreSaveWorldWithContext.Remove(PreSaveWorldHandle);
 
 		// unregister the editor mode
 		FEditorModeRegistry::Get().UnregisterMode(FBuiltinEditorModes::EM_Landscape);
@@ -292,12 +294,12 @@ public:
 		}
 	}
 
-	void OnPreSaveWorld(uint32 SaveFlags, class UWorld* World)
+	void OnPreSaveWorld( class UWorld* World, FObjectPreSaveContext ObjectSaveContext)
 	{
 		FEdModeLandscape* EdMode = (FEdModeLandscape*)GLevelEditorModeTools().GetActiveMode(FBuiltinEditorModes::EM_Landscape);
 		if (EdMode)
 		{
-			EdMode->OnPreSaveWorld(SaveFlags, World);
+			EdMode->OnPreSaveWorld(World, ObjectSaveContext);
 		}
 	}
 
diff --git a/Engine/Source/Editor/MaterialEditor/Private/MaterialInstanceEditor.cpp b/Engine/Source/Editor/MaterialEditor/Private/MaterialInstanceEditor.cpp
--- a/Engine/Source/Editor/MaterialEditor/Private/MaterialInstanceEditor.cpp
+++ b/Engine/Source/Editor/MaterialEditor/Private/MaterialInstanceEditor.cpp
@@ -1,11 +1,13 @@
 // Copyright Epic Games, Inc. All Rights Reserved.
 
 #include "MaterialInstanceEditor.h"
+
 #include "Widgets/Text/STextBlock.h"
 #include "EngineGlobals.h"
 #include "Misc/ConfigCacheIni.h"
 #include "Modules/ModuleManager.h"
 #include "Widgets/Views/SListView.h"
+#include "UObject/ObjectSaveContext.h"
 #include "UObject/Package.h"
 #include "Editor.h"
 #include "EditorStyleSet.h"
@@ -610,7 +612,7 @@ FMaterialInstanceEditor::FMaterialInstanceEditor()
 , FunctionMaterialProxy(nullptr)
 , FunctionInstanceProxy(nullptr)
 {
-	UPackage::PreSavePackageEvent.AddRaw(this, &FMaterialInstanceEditor::PreSavePackage);
+	UPackage::PreSavePackageWithContextEvent.AddRaw(this, &FMaterialInstanceEditor::PreSavePackage);
 }
 
 FMaterialInstanceEditor::~FMaterialInstanceEditor()
@@ -620,7 +622,7 @@ FMaterialInstanceEditor::~FMaterialInstanceEditor()
 
 	GEditor->UnregisterForUndo( this );
 
-	UPackage::PreSavePackageEvent.RemoveAll(this);
+	UPackage::PreSavePackageWithContextEvent.RemoveAll(this);
 
 	// The streaming data will be null if there were any edits
 	if (MaterialEditorInstance && MaterialEditorInstance->SourceInstance && !MaterialEditorInstance->SourceInstance->HasTextureStreamingData())
@@ -1118,7 +1120,7 @@ void FMaterialInstanceEditor::RefreshPreviewAsset()
 	PreviewVC->SetPreviewAsset(PreviewAsset);
 }
 
-void FMaterialInstanceEditor::PreSavePackage(UPackage* Package)
+void FMaterialInstanceEditor::PreSavePackage(UPackage* Package, FObjectPreSaveContext ObjectSaveContext)
 {
 	// The streaming data will be null if there were any edits
 	if (MaterialEditorInstance && 
diff --git a/Engine/Source/Editor/MaterialEditor/Private/MaterialInstanceEditor.h b/Engine/Source/Editor/MaterialEditor/Private/MaterialInstanceEditor.h
--- a/Engine/Source/Editor/MaterialEditor/Private/MaterialInstanceEditor.h
+++ b/Engine/Source/Editor/MaterialEditor/Private/MaterialInstanceEditor.h
@@ -15,6 +15,7 @@
 #include "TickableEditorObject.h"
 
 class FCanvas;
+class FObjectPreSaveContext;
 class UToolMenu;
 class UMaterialEditorInstanceConstant;
 class UMaterialInterface;
@@ -75,7 +76,7 @@ public:
 	/** Post edit change notify for properties. */
 	virtual void NotifyPostChange( const FPropertyChangedEvent& PropertyChangedEvent, FProperty* PropertyThatChanged ) override;
 
-	void PreSavePackage(UPackage* Obj);
+	void PreSavePackage(UPackage* Obj, FObjectPreSaveContext ObjectSaveContext);
 
 	/** Rebuilds the inheritance list for this material instance. */
 	void RebuildInheritanceList();
diff --git a/Engine/Source/Editor/MeshPaint/Private/IMeshPaintMode.cpp b/Engine/Source/Editor/MeshPaint/Private/IMeshPaintMode.cpp
--- a/Engine/Source/Editor/MeshPaint/Private/IMeshPaintMode.cpp
+++ b/Engine/Source/Editor/MeshPaint/Private/IMeshPaintMode.cpp
@@ -1,6 +1,7 @@
 // Copyright Epic Games, Inc. All Rights Reserved.
 
 #include "IMeshPaintMode.h"
+
 #include "SceneView.h"
 #include "EditorViewportClient.h"
 #include "Modules/ModuleManager.h"
@@ -17,6 +18,7 @@
 #include "EditorSupportDelegates.h"
 
 #include "MeshPaintHelpers.h"
+#include "UObject/ObjectSaveContext.h"
 
 //Slate dependencies
 #include "LevelEditor.h"
@@ -71,8 +73,8 @@ void IMeshPaintEdMode::Enter()
 	FCoreUObjectDelegates::OnObjectsReplaced.AddSP(this, &IMeshPaintEdMode::OnObjectsReplaced);
 
 	// Hook into pre/post world save, so that the original collision volumes can be temporarily reinstated
-	FEditorDelegates::PreSaveWorld.AddSP(this, &IMeshPaintEdMode::OnPreSaveWorld);
-	FEditorDelegates::PostSaveWorld.AddSP(this, &IMeshPaintEdMode::OnPostSaveWorld);
+	FEditorDelegates::PreSaveWorldWithContext.AddSP(this, &IMeshPaintEdMode::OnPreSaveWorld);
+	FEditorDelegates::PostSaveWorldWithContext.AddSP(this, &IMeshPaintEdMode::OnPostSaveWorld);
 
 	// Catch assets if they are about to be (re)imported
 	GEditor->GetEditorSubsystem<UImportSubsystem>()->OnAssetPostImport.AddSP(this, &IMeshPaintEdMode::OnPostImportAsset);
@@ -139,8 +141,8 @@ void IMeshPaintEdMode::Exit()
 	AssetRegistryModule.Get().OnAssetRemoved().RemoveAll(this);
 	FReimportManager::Instance()->OnPostReimport().RemoveAll(this);
 	GEditor->GetEditorSubsystem<UImportSubsystem>()->OnAssetPostImport.RemoveAll(this);
-	FEditorDelegates::PreSaveWorld.RemoveAll(this);
-	FEditorDelegates::PostSaveWorld.RemoveAll(this);
+	FEditorDelegates::PreSaveWorldWithContext.RemoveAll(this);
+	FEditorDelegates::PostSaveWorldWithContext.RemoveAll(this);
 	FCoreUObjectDelegates::OnObjectsReplaced.RemoveAll(this);
 	USelection::SelectionChangedEvent.Remove(SelectionChangedHandle);
 
@@ -293,12 +295,12 @@ bool IMeshPaintEdMode::InputKey( FEditorViewportClient* InViewportClient, FViewp
 	return bHandled;
 }
 
-void IMeshPaintEdMode::OnPreSaveWorld(uint32 SaveFlags, UWorld* World)
+void IMeshPaintEdMode::OnPreSaveWorld(UWorld* World, FObjectPreSaveContext ObjectSaveContext)
 {
 	MeshPainter->Refresh();
 }
 
-void IMeshPaintEdMode::OnPostSaveWorld(uint32 SaveFlags, UWorld* World, bool bSuccess)
+void IMeshPaintEdMode::OnPostSaveWorld(UWorld* World, FObjectPostSaveContext ObjectSaveContext)
 {
 	MeshPainter->Refresh();
 }
diff --git a/Engine/Source/Editor/MeshPaint/Public/IMeshPaintMode.h b/Engine/Source/Editor/MeshPaint/Public/IMeshPaintMode.h
--- a/Engine/Source/Editor/MeshPaint/Public/IMeshPaintMode.h
+++ b/Engine/Source/Editor/MeshPaint/Public/IMeshPaintMode.h
@@ -8,6 +8,8 @@
 
 class FEditorViewportClient;
 class FPrimitiveDrawInterface;
+class FObjectPostSaveContext;
+class FObjectPreSaveContext;
 class FSceneView;
 class FViewport;
 class UFactory;
@@ -60,10 +62,10 @@ public:
 	IMeshPainter* GetMeshPainter();
 private:
 	/** Called prior to saving a level */
-	void OnPreSaveWorld(uint32 SaveFlags, UWorld* World);
+	void OnPreSaveWorld(UWorld* World, FObjectPreSaveContext ObjectSaveContext);
 
 	/** Called after saving a level */
-	void OnPostSaveWorld(uint32 SaveFlags, UWorld* World, bool bSuccess);
+	void OnPostSaveWorld(UWorld* World, FObjectPostSaveContext ObjectSaveContext);
 
 	/** Called when an asset has just been imported */
 	void OnPostImportAsset(UFactory* Factory, UObject* Object);
diff --git a/Engine/Source/Editor/Sequencer/Private/LevelEditorSequencerIntegration.cpp b/Engine/Source/Editor/Sequencer/Private/LevelEditorSequencerIntegration.cpp
--- a/Engine/Source/Editor/Sequencer/Private/LevelEditorSequencerIntegration.cpp
+++ b/Engine/Source/Editor/Sequencer/Private/LevelEditorSequencerIntegration.cpp
@@ -1,6 +1,7 @@
 // Copyright Epic Games, Inc. All Rights Reserved.
 
 #include "LevelEditorSequencerIntegration.h"
+
 #include "SequencerEdMode.h"
 #include "Styling/SlateIconFinder.h"
 #include "PropertyHandle.h"
@@ -43,6 +44,7 @@
 #include "Widgets/Layout/SBox.h"
 #include "Widgets/Text/STextBlock.h"
 #include "UObject/ObjectKey.h"
+#include "UObject/ObjectSaveContext.h"
 #include "UnrealEdGlobals.h"
 #include "UnrealEdMisc.h"
 #include "Editor/UnrealEdEngine.h"
@@ -169,12 +171,12 @@ void FLevelEditorSequencerIntegration::Initialize(const FLevelEditorSequencerInt
 
 	// Register for saving the level so that the state of the scene can be restored before saving and updated after saving.
 	{
-		FDelegateHandle Handle = FEditorDelegates::PreSaveWorld.AddRaw(this, &FLevelEditorSequencerIntegration::OnPreSaveWorld);
-		AcquiredResources.Add([=]{ FEditorDelegates::PreSaveWorld.Remove(Handle); });
+		FDelegateHandle Handle = FEditorDelegates::PreSaveWorldWithContext.AddRaw(this, &FLevelEditorSequencerIntegration::OnPreSaveWorld);
+		AcquiredResources.Add([=]{ FEditorDelegates::PreSaveWorldWithContext.Remove(Handle); });
 	}
 	{
-		FDelegateHandle Handle = FEditorDelegates::PostSaveWorld.AddRaw(this, &FLevelEditorSequencerIntegration::OnPostSaveWorld);
-		AcquiredResources.Add([=]{ FEditorDelegates::PostSaveWorld.Remove(Handle); });
+		FDelegateHandle Handle = FEditorDelegates::PostSaveWorldWithContext.AddRaw(this, &FLevelEditorSequencerIntegration::OnPostSaveWorld);
+		AcquiredResources.Add([=]{ FEditorDelegates::PostSaveWorldWithContext.Remove(Handle); });
 	}
 	{
 		FDelegateHandle Handle = FEditorDelegates::PreBeginPIE.AddRaw(this, &FLevelEditorSequencerIntegration::OnPreBeginPIE);
@@ -308,7 +310,7 @@ void FLevelEditorSequencerIntegration::OnActorLabelChanged(AActor* ChangedActor)
 	}
 }
 
-void FLevelEditorSequencerIntegration::OnPreSaveWorld(uint32 SaveFlags, class UWorld* World)
+void FLevelEditorSequencerIntegration::OnPreSaveWorld(class UWorld* World, FObjectPreSaveContext ObjectSaveContext)
 {
 	// Restore the saved state so that the level save can save that instead of the animated state.
 	IterateAllSequencers(
@@ -322,7 +324,7 @@ void FLevelEditorSequencerIntegration::OnPreSaveWorld(uint32 SaveFlags, class UW
 	);
 }
 
-void FLevelEditorSequencerIntegration::OnPostSaveWorld(uint32 SaveFlags, class UWorld* World, bool bSuccess)
+void FLevelEditorSequencerIntegration::OnPostSaveWorld(class UWorld* World, FObjectPostSaveContext ObjectSaveContext)
 {
 	// Reset the time after saving so that an update will be triggered to put objects back to their animated state.
 	IterateAllSequencers(
diff --git a/Engine/Source/Editor/Sequencer/Public/LevelEditorSequencerIntegration.h b/Engine/Source/Editor/Sequencer/Public/LevelEditorSequencerIntegration.h
--- a/Engine/Source/Editor/Sequencer/Public/LevelEditorSequencerIntegration.h
+++ b/Engine/Source/Editor/Sequencer/Public/LevelEditorSequencerIntegration.h
@@ -14,6 +14,8 @@ class AActor;
 class FExtender;
 class FMenuBuilder;
 class FSequencer;
+class FObjectPostSaveContext;
+class FObjectPreSaveContext;
 class FUICommandList;
 class IAssetViewport;
 class ISequencer;
@@ -90,10 +92,10 @@ public:
 private:
 
 	/** Called before the world is going to be saved. The sequencer puts everything back to its initial state. */
-	void OnPreSaveWorld(uint32 SaveFlags, UWorld* World);
+	void OnPreSaveWorld(UWorld* World, FObjectPreSaveContext ObjectSaveContext);
 
 	/** Called after the world has been saved. The sequencer updates to the animated state. */
-	void OnPostSaveWorld(uint32 SaveFlags, UWorld* World, bool bSuccess);
+	void OnPostSaveWorld(UWorld* World, FObjectPostSaveContext ObjectSaveContext);
 
 	/** Called after a level has been added */
 	void OnLevelAdded(ULevel* InLevel, UWorld* InWorld);
diff --git a/Engine/Source/Editor/UMGEditor/Private/WidgetBlueprint.cpp b/Engine/Source/Editor/UMGEditor/Private/WidgetBlueprint.cpp
--- a/Engine/Source/Editor/UMGEditor/Private/WidgetBlueprint.cpp
+++ b/Engine/Source/Editor/UMGEditor/Private/WidgetBlueprint.cpp
@@ -1,6 +1,7 @@
 // Copyright Epic Games, Inc. All Rights Reserved.
 
 #include "WidgetBlueprint.h"
+
 #include "Components/Widget.h"
 #include "Blueprint/UserWidget.h"
 #include "MovieScene.h"
@@ -19,6 +20,7 @@
 #include "WidgetBlueprintCompiler.h"
 #include "UObject/EditorObjectVersion.h"
 #include "UObject/FortniteMainBranchObjectVersion.h"
+#include "UObject/ObjectSaveContext.h"
 #include "WidgetGraphSchema.h"
 #include "UMGEditorProjectSettings.h"
 
@@ -586,7 +588,14 @@ void UWidgetBlueprint::ReplaceDeprecatedNodes()
 #if WITH_EDITORONLY_DATA
 void UWidgetBlueprint::PreSave(const class ITargetPlatform* TargetPlatform)
 {
+	PRAGMA_DISABLE_DEPRECATION_WARNINGS;
 	Super::PreSave(TargetPlatform);
+	PRAGMA_ENABLE_DEPRECATION_WARNINGS;
+}
+
+void UWidgetBlueprint::PreSave(FObjectPreSaveContext ObjectSaveContext)
+{
+	Super::PreSave(ObjectSaveContext);
 }
 #endif // WITH_EDITORONLY_DATA
 
diff --git a/Engine/Source/Editor/UnrealEd/Classes/Editor/EditorEngine.h b/Engine/Source/Editor/UnrealEd/Classes/Editor/EditorEngine.h
--- a/Engine/Source/Editor/UnrealEd/Classes/Editor/EditorEngine.h
+++ b/Engine/Source/Editor/UnrealEd/Classes/Editor/EditorEngine.h
@@ -2201,11 +2201,16 @@ public:
 	virtual bool InitializePhysicsSceneForSaveIfNecessary(UWorld* World, bool &bOutForceInitialized);
 	void CleanupPhysicsSceneThatWasInitializedForSave(UWorld* World, bool bForceInitialized);
 
-	/** Invoked before a UWorld is saved to update editor systems */
+	UE_DEPRECATED(5.0, "Use version that takes FObjectPreSaveContext instead.")
 	virtual void OnPreSaveWorld(uint32 SaveFlags, UWorld* World);
+	UE_DEPRECATED(5.0, "Use version that takes FObjectPostSaveContext instead.")
+	virtual void OnPostSaveWorld(uint32 SaveFlags, UWorld* World, uint32 OriginalPackageFlags, bool bSuccess);
+
+	/** Invoked before a UWorld is saved to update editor systems */
+	virtual void OnPreSaveWorld(UWorld* World, FObjectPreSaveContext ObjectSaveContext);
 
 	/** Invoked after a UWorld is saved to update editor systems */
-	virtual void OnPostSaveWorld(uint32 SaveFlags, UWorld* World, uint32 OriginalPackageFlags, bool bSuccess);
+	virtual void OnPostSaveWorld(UWorld* World, FObjectPostSaveContext ObjectSaveContext);
 
 	/**
 	 * Adds provided package to a default changelist
diff --git a/Engine/Source/Editor/UnrealEd/Private/CookOnTheFlyServer.cpp b/Engine/Source/Editor/UnrealEd/Private/CookOnTheFlyServer.cpp
--- a/Engine/Source/Editor/UnrealEd/Private/CookOnTheFlyServer.cpp
+++ b/Engine/Source/Editor/UnrealEd/Private/CookOnTheFlyServer.cpp
@@ -44,6 +44,7 @@
 #include "UObject/SavePackage.h"
 #include "UObject/MetaData.h"
 #include "UObject/ConstructorHelpers.h"
+#include "UObject/ObjectSaveContext.h"
 #include "UObject/UObjectArray.h"
 #include "Misc/PackageName.h"
 #include "Misc/RedirectCollector.h"
@@ -8142,13 +8143,16 @@ uint32 UCookOnTheFlyServer::FullLoadAndSave(uint32& CookedPackageCount)
 	TSet<UPackage*> ProcessedPackages;
 	ProcessedPackages.Reserve(65536);
 
-	TMap<UWorld*, bool> WorldsToPostSaveRoot;
+	TMap<UWorld*, TArray<bool>> WorldsToPostSaveRoot;
 	WorldsToPostSaveRoot.Reserve(1024);
 
 	TArray<UObject*> ObjectsToWaitForCookedPlatformData;
 	ObjectsToWaitForCookedPlatformData.Reserve(65536);
 
 	TArray<FString> PackagesToLoad;
+
+	FObjectSaveContextData ObjectSaveContext(nullptr, nullptr, TEXT(""), SaveFlags);
+
 	do
 	{
 		PackagesToLoad.Reset();
@@ -8236,35 +8240,42 @@ uint32 UCookOnTheFlyServer::FullLoadAndSave(uint32& CookedPackageCount)
 							UWorld* World = Cast<UWorld>(Obj);
 							bool bInitializedPhysicsSceneForSave = false;
 							bool bForceInitializedWorld = false;
-							if (World && bSaveConcurrent)
-							{
-								UE_SCOPED_HIERARCHICAL_COOKTIMER(FullLoadAndSave_SettingUpWorlds);
-								// We need a physics scene at save time in case code does traces during onsave events.
-								bInitializedPhysicsSceneForSave = GEditor->InitializePhysicsSceneForSaveIfNecessary(World, bForceInitializedWorld);
-
-								GIsCookerLoadingPackage = true;
-								{
-									UE_SCOPED_HIERARCHICAL_COOKTIMER(FullLoadAndSave_PreSaveWorld);
-									GEditor->OnPreSaveWorld(SaveFlags, World);
-								}
-								{
-									UE_SCOPED_HIERARCHICAL_COOKTIMER(FullLoadAndSave_PreSaveRoot);
-									bool bCleanupIsRequired = World->PreSaveRoot(TEXT(""));
-									WorldsToPostSaveRoot.Add(World, bCleanupIsRequired);
-								}
-								GIsCookerLoadingPackage = false;
-							}
 
 							bool bAllPlatformDataLoaded = true;
 							bool bIsTexture = Obj->IsA(UTexture::StaticClass());
+							bool bFirstPlatform = true;
 							for (const ITargetPlatform* TargetPlatform : TargetPlatforms)
 							{
+								ObjectSaveContext.TargetPlatform = TargetPlatform;
+								ObjectSaveContext.bOuterConcurrentSave = bFirstPlatform;
+								bFirstPlatform = false;
+
 								if (bSaveConcurrent)
 								{
 									GIsCookerLoadingPackage = true;
+									if (World)
+									{
+										UE_SCOPED_HIERARCHICAL_COOKTIMER(FullLoadAndSave_SettingUpWorlds);
+										// We need a physics scene at save time in case code does traces during onsave events.
+										if (ObjectSaveContext.bOuterConcurrentSave)
+										{
+											bInitializedPhysicsSceneForSave = GEditor->InitializePhysicsSceneForSaveIfNecessary(World, bForceInitializedWorld);
+										}
+
+										{
+											UE_SCOPED_HIERARCHICAL_COOKTIMER(FullLoadAndSave_PreSaveWorld);
+											GEditor->OnPreSaveWorld(World, FObjectPreSaveContext(ObjectSaveContext));
+										}
+										{
+											UE_SCOPED_HIERARCHICAL_COOKTIMER(FullLoadAndSave_PreSaveRoot);
+											UE::SavePackageUtilities::CallPreSaveRoot(World, ObjectSaveContext);
+											WorldsToPostSaveRoot.FindOrAdd(World).Add(ObjectSaveContext.bCleanupRequired);
+										}
+									}
+
 									{
 										UE_SCOPED_HIERARCHICAL_COOKTIMER(FullLoadAndSave_PreSave);
-										Obj->PreSave(TargetPlatform);
+										UE::SavePackageUtilities::CallPreSave(Obj, ObjectSaveContext);
 									}
 									GIsCookerLoadingPackage = false;
 								}
@@ -8610,11 +8621,18 @@ uint32 UCookOnTheFlyServer::FullLoadAndSave(uint32& CookedPackageCount)
 	{
 		UE_LOG(LogCook, Display, TEXT("Calling PostSaveRoot on worlds..."));
 		UE_SCOPED_HIERARCHICAL_COOKTIMER(FullLoadAndSave_PostSaveRoot);
-		for (auto WorldIt = WorldsToPostSaveRoot.CreateConstIterator(); WorldIt; ++WorldIt)
+		for (const TPair<UWorld*,TArray<bool>>& WorldIt : WorldsToPostSaveRoot)
 		{
-			UWorld* World = WorldIt.Key();
+			UWorld* World = WorldIt.Key;
+			const TArray<bool>& PlatformNeedsCleanup = WorldIt.Value;
 			check(World);
-			World->PostSaveRoot(WorldIt.Value());
+			check(PlatformNeedsCleanup.Num() == TargetPlatforms.Num());
+			for (int PlatformIndex = TargetPlatforms.Num() - 1; PlatformIndex >= 0; --PlatformIndex)
+			{
+				ObjectSaveContext.TargetPlatform = TargetPlatforms[PlatformIndex];
+				ObjectSaveContext.bOuterConcurrentSave = PlatformIndex == 0;
+				UE::SavePackageUtilities::CallPostSaveRoot(World, ObjectSaveContext, PlatformNeedsCleanup[PlatformIndex]);
+			}
 		}
 	}
 
diff --git a/Engine/Source/Editor/UnrealEd/Private/Editor.cpp b/Engine/Source/Editor/UnrealEd/Private/Editor.cpp
--- a/Engine/Source/Editor/UnrealEd/Private/Editor.cpp
+++ b/Engine/Source/Editor/UnrealEd/Private/Editor.cpp
@@ -114,6 +114,8 @@ FSimpleMulticastDelegate								FEditorDelegates::PropertySelectionChange;
 FSimpleMulticastDelegate								FEditorDelegates::PostLandscapeLayerUpdated;
 FEditorDelegates::FOnPreSaveWorld						FEditorDelegates::PreSaveWorld;
 FEditorDelegates::FOnPostSaveWorld						FEditorDelegates::PostSaveWorld;
+FEditorDelegates::FOnPreSaveWorldWithContext			FEditorDelegates::PreSaveWorldWithContext;
+FEditorDelegates::FOnPostSaveWorldWithContext			FEditorDelegates::PostSaveWorldWithContext;
 FEditorDelegates::FOnFinishPickingBlueprintClass		FEditorDelegates::OnFinishPickingBlueprintClass;
 FEditorDelegates::FOnNewAssetCreation					FEditorDelegates::OnConfigureNewAssetProperties;
 FEditorDelegates::FOnNewAssetCreation					FEditorDelegates::OnNewAssetCreated;
diff --git a/Engine/Source/Editor/UnrealEd/Private/EditorActorFolders.cpp b/Engine/Source/Editor/UnrealEd/Private/EditorActorFolders.cpp
--- a/Engine/Source/Editor/UnrealEd/Private/EditorActorFolders.cpp
+++ b/Engine/Source/Editor/UnrealEd/Private/EditorActorFolders.cpp
@@ -1,6 +1,7 @@
 // Copyright Epic Games, Inc. All Rights Reserved.
 
 #include "EditorActorFolders.h"
+
 #include "GameFramework/Actor.h"
 #include "HAL/FileManager.h"
 #include "Misc/Paths.h"
@@ -13,6 +14,7 @@
 #include "Editor.h"
 #include "EditorFolderUtils.h"
 #include "ScopedTransaction.h"
+#include "UObject/ObjectSaveContext.h"
 
 #define LOCTEXT_NAMESPACE "FActorFolders"
 
@@ -47,7 +49,7 @@ FActorFolders::FActorFolders()
 	GEngine->OnLevelActorListChanged().AddRaw(this, &FActorFolders::OnLevelActorListChanged);
 
 	FEditorDelegates::MapChange.AddRaw(this, &FActorFolders::OnMapChange);
-	FEditorDelegates::PostSaveWorld.AddRaw(this, &FActorFolders::OnWorldSaved);
+	FEditorDelegates::PostSaveWorldWithContext.AddRaw(this, &FActorFolders::OnWorldSaved);
 }
 
 FActorFolders::~FActorFolders()
@@ -57,7 +59,7 @@ FActorFolders::~FActorFolders()
 	GEngine->OnLevelActorListChanged().RemoveAll(this);
 
 	FEditorDelegates::MapChange.RemoveAll(this);
-	FEditorDelegates::PostSaveWorld.RemoveAll(this);
+	FEditorDelegates::PostSaveWorldWithContext.RemoveAll(this);
 }
 
 void FActorFolders::AddReferencedObjects(FReferenceCollector& Collector)
@@ -130,7 +132,7 @@ void FActorFolders::OnMapChange(uint32 MapChangeFlags)
 	OnLevelActorListChanged();
 }
 
-void FActorFolders::OnWorldSaved(uint32 SaveFlags, UWorld* World, bool bSuccess)
+void FActorFolders::OnWorldSaved(UWorld* World, FObjectPostSaveContext ObjectSaveContext)
 {
 	// Attempt to save the folder state
 	const UEditorActorFolders* const* ExisingFolders = TemporaryWorldFolders.Find(World);
diff --git a/Engine/Source/Editor/UnrealEd/Private/EditorEngine.cpp b/Engine/Source/Editor/UnrealEd/Private/EditorEngine.cpp
--- a/Engine/Source/Editor/UnrealEd/Private/EditorEngine.cpp
+++ b/Engine/Source/Editor/UnrealEd/Private/EditorEngine.cpp
@@ -1,6 +1,7 @@
 // Copyright Epic Games, Inc. All Rights Reserved.
 
 #include "Editor/EditorEngine.h"
+
 #include "Misc/MessageDialog.h"
 #include "HAL/FileManager.h"
 #include "Misc/CommandLine.h"
@@ -11,6 +12,7 @@
 #include "Misc/AssetRegistryInterface.h"
 #include "Modules/ModuleManager.h"
 #include "UObject/MetaData.h"
+#include "UObject/ObjectSaveContext.h"
 #include "UObject/ConstructorHelpers.h"
 #include "Application/ThrottleManager.h"
 #include "Framework/Application/SlateApplication.h"
@@ -4288,16 +4290,14 @@ FSavePackageResultStruct UEditorEngine::Save( UPackage* InOuter, UObject* InBase
 		}
 	}
 
-	// Record the package flags before OnPreSaveWorld. They will be used in OnPostSaveWorld.
-	const uint32 OriginalPackageFlags = (InOuter ? InOuter->GetPackageFlags() : 0);
-
 	SlowTask.EnterProgressFrame(10);
 
 	UWorld* World = Cast<UWorld>(Base);
 	bool bInitializedPhysicsSceneForSave = false;
 	bool bForceInitializedWorld = false;
 	const bool bSavingConcurrent = !!(SaveFlags & ESaveFlags::SAVE_Concurrent);
-	
+
+	FObjectSaveContextData ObjectSaveContext(InOuter, TargetPlatform, Filename, SaveFlags);
 	UWorld *OriginalOwningWorld = nullptr;
 	if ( World )
 	{
@@ -4305,7 +4305,10 @@ FSavePackageResultStruct UEditorEngine::Save( UPackage* InOuter, UObject* InBase
 		{
 			bInitializedPhysicsSceneForSave = InitializePhysicsSceneForSaveIfNecessary(World, bForceInitializedWorld);
 
-			OnPreSaveWorld(SaveFlags, World);
+			PRAGMA_DISABLE_DEPRECATION_WARNINGS;
+			OnPreSaveWorld(ObjectSaveContext.SaveFlags, World);
+			PRAGMA_ENABLE_DEPRECATION_WARNINGS;
+			OnPreSaveWorld(World, FObjectPreSaveContext(ObjectSaveContext));
 		}
 
 		OriginalOwningWorld = World->PersistentLevel->OwningWorld;
@@ -4322,10 +4325,14 @@ FSavePackageResultStruct UEditorEngine::Save( UPackage* InOuter, UObject* InBase
 
 	SlowTask.EnterProgressFrame(70);
 
+	PRAGMA_DISABLE_DEPRECATION_WARNINGS;
 	UPackage::PreSavePackageEvent.Broadcast(InOuter);
+	PRAGMA_ENABLE_DEPRECATION_WARNINGS;
+	UPackage::PreSavePackageWithContextEvent.Broadcast(InOuter, FObjectPreSaveContext(ObjectSaveContext));
 	FSavePackageResultStruct Result = UPackage::Save(InOuter, Base, TopLevelFlags, Filename, Error, Conform, bForceByteSwapping, bWarnOfLongFilename, SaveFlags, TargetPlatform, FinalTimeStamp, bSlowTask, InOutDiffMap, SavePackageContext);
 
 	SlowTask.EnterProgressFrame(10);
+	ObjectSaveContext.bSaveSucceeded = Result == ESavePackageResult::Success;
 
 	// If the package is a valid candidate for being automatically-added to source control, go ahead and add it
 	// to the default changelist
@@ -4361,7 +4368,10 @@ FSavePackageResultStruct UEditorEngine::Save( UPackage* InOuter, UObject* InBase
 
 		if (!bSavingConcurrent)
 		{
-			OnPostSaveWorld(SaveFlags, World, OriginalPackageFlags, Result == ESavePackageResult::Success);
+			PRAGMA_DISABLE_DEPRECATION_WARNINGS;
+			OnPostSaveWorld(ObjectSaveContext.SaveFlags, World, ObjectSaveContext.OriginalPackageFlags, ObjectSaveContext.bSaveSucceeded);
+			PRAGMA_ENABLE_DEPRECATION_WARNINGS;
+			OnPostSaveWorld(World, FObjectPostSaveContext(ObjectSaveContext));
 
 			if (bInitializedPhysicsSceneForSave)
 			{
@@ -4393,6 +4403,14 @@ bool UEditorEngine::SavePackage(UPackage* InOuter, UObject* InBase, EObjectFlags
 
 void UEditorEngine::OnPreSaveWorld(uint32 SaveFlags, UWorld* World)
 {
+}
+
+void UEditorEngine::OnPreSaveWorld(UWorld* World, FObjectPreSaveContext ObjectSaveContext)
+{
+	if (!ObjectSaveContext.IsFirstConcurrentSave())
+	{
+		return;
+	}
 	if ( !ensure(World) )
 	{
 		return;
@@ -4403,14 +4421,17 @@ void UEditorEngine::OnPreSaveWorld(uint32 SaveFlags, UWorld* World)
 	check(World->PersistentLevel);
 
 	// Pre save world event
-	FEditorDelegates::PreSaveWorld.Broadcast(SaveFlags, World);
+	PRAGMA_DISABLE_DEPRECATION_WARNINGS;
+	FEditorDelegates::PreSaveWorld.Broadcast(ObjectSaveContext.GetSaveFlags(), World);
+	PRAGMA_ENABLE_DEPRECATION_WARNINGS;
+	FEditorDelegates::PreSaveWorldWithContext.Broadcast(World, ObjectSaveContext);
 
 	// Update cull distance volumes (and associated primitives).
 	World->UpdateCullDistanceVolumes();
 
 	if ( !IsRunningCommandlet() )
 	{
-		const bool bAutosave = (SaveFlags & SAVE_FromAutosave) != 0;
+		const bool bAutosave = (ObjectSaveContext.GetSaveFlags() & SAVE_FromAutosave) != 0;
 		if ( bAutosave )
 		{
 			// Temporarily flag packages saved under a PIE filename as PKG_PlayInEditor for serialization so loading
@@ -4475,15 +4496,25 @@ void UEditorEngine::OnPreSaveWorld(uint32 SaveFlags, UWorld* World)
 
 void UEditorEngine::OnPostSaveWorld(uint32 SaveFlags, UWorld* World, uint32 OriginalPackageFlags, bool bSuccess)
 {
+}
+
+void UEditorEngine::OnPostSaveWorld(UWorld* World, FObjectPostSaveContext ObjectSaveContext)
+{
+	if (!ObjectSaveContext.IsLastConcurrentSave())
+	{
+		return;
+	}
 	if ( !ensure(World) )
 	{
 		return;
 	}
+	uint32 OriginalPackageFlags = ObjectSaveContext.GetOriginalPackageFlags();
+	bool bSuccess = ObjectSaveContext.SaveSucceeded();
 
 	if ( !IsRunningCommandlet() )
 	{
 		UPackage* WorldPackage = World->GetOutermost();
-		const bool bAutosave = (SaveFlags & SAVE_FromAutosave) != 0;
+		const bool bAutosave = (ObjectSaveContext.GetSaveFlags() & SAVE_FromAutosave) != 0;
 		if ( bAutosave )
 		{
 			// Restore original value of PKG_PlayInEditor if we changed it during PIE saving
@@ -4552,7 +4583,10 @@ void UEditorEngine::OnPostSaveWorld(uint32 SaveFlags, UWorld* World, uint32 Orig
 	}
 
 	// Post save world event
-	FEditorDelegates::PostSaveWorld.Broadcast(SaveFlags, World, bSuccess);
+	PRAGMA_DISABLE_DEPRECATION_WARNINGS;
+	FEditorDelegates::PostSaveWorld.Broadcast(ObjectSaveContext.GetSaveFlags(), World, bSuccess);
+	PRAGMA_ENABLE_DEPRECATION_WARNINGS;
+	FEditorDelegates::PostSaveWorldWithContext.Broadcast(World, ObjectSaveContext);
 }
 
 APlayerStart* UEditorEngine::CheckForPlayerStart()
diff --git a/Engine/Source/Editor/UnrealEd/Private/PackageAutoSaver.cpp b/Engine/Source/Editor/UnrealEd/Private/PackageAutoSaver.cpp
--- a/Engine/Source/Editor/UnrealEd/Private/PackageAutoSaver.cpp
+++ b/Engine/Source/Editor/UnrealEd/Private/PackageAutoSaver.cpp
@@ -1,6 +1,8 @@
 // Copyright Epic Games, Inc. All Rights Reserved.
 
 #include "PackageAutoSaver.h"
+
+#include "UObject/ObjectSaveContext.h"
 #include "UObject/Package.h"
 #include "HAL/FileManager.h"
 #include "Misc/Paths.h"
@@ -100,14 +102,14 @@ FPackageAutoSaver::FPackageAutoSaver()
 	UPackage::PackageMarkedDirtyEvent.AddRaw(this, &FPackageAutoSaver::OnMarkPackageDirty);
 
 	// Register for the package modified callback to catch packages that have been saved
-	UPackage::PackageSavedEvent.AddRaw(this, &FPackageAutoSaver::OnPackageSaved);
+	UPackage::PackageSavedWithContextEvent.AddRaw(this, &FPackageAutoSaver::OnPackageSaved);
 }
 
 FPackageAutoSaver::~FPackageAutoSaver()
 {
 	UPackage::PackageDirtyStateChangedEvent.RemoveAll(this);
 	UPackage::PackageMarkedDirtyEvent.RemoveAll(this);
-	UPackage::PackageSavedEvent.RemoveAll(this);
+	UPackage::PackageSavedWithContextEvent.RemoveAll(this);
 }
 
 void FPackageAutoSaver::UpdateAutoSaveCount(const float DeltaSeconds)
@@ -338,10 +340,8 @@ void FPackageAutoSaver::OnMarkPackageDirty(UPackage* Pkg, bool bWasDirty)
 	UpdateDirtyListsForPackage(Pkg);
 }
 
-void FPackageAutoSaver::OnPackageSaved(const FString& Filename, UObject* Obj)
+void FPackageAutoSaver::OnPackageSaved(const FString& Filename, UPackage* Pkg, FObjectPostSaveContext ObjectSaveContext)
 {
-	UPackage* const Pkg = Cast<UPackage>(Obj);
-
 	// If this has come from an auto-save, update the last known filename in the user dirty list so that we can offer is up as a restore file later
 	if(IsAutoSaving())
 	{
diff --git a/Engine/Source/Editor/UnrealEd/Private/PackageAutoSaver.h b/Engine/Source/Editor/UnrealEd/Private/PackageAutoSaver.h
--- a/Engine/Source/Editor/UnrealEd/Private/PackageAutoSaver.h
+++ b/Engine/Source/Editor/UnrealEd/Private/PackageAutoSaver.h
@@ -7,6 +7,8 @@
 #include "UObject/WeakObjectPtr.h"
 #include "IPackageAutoSaver.h"
 
+class FObjectPostSaveContext;
+
 namespace ECloseNotification
 {
 	enum Type
@@ -61,7 +63,7 @@ private:
 	 * @param Filename The filename the package was saved to
 	 * @param Obj The package that was saved
 	 */
-	void OnPackageSaved(const FString& Filename, UObject* Obj);
+	void OnPackageSaved(const FString& Filename, UPackage* Pkg, FObjectPostSaveContext ObjectSaveContext);
 
 	/**
 	 * Update the dirty lists for the current package
diff --git a/Engine/Source/Editor/UnrealEd/Private/UnrealEdMisc.cpp b/Engine/Source/Editor/UnrealEd/Private/UnrealEdMisc.cpp
--- a/Engine/Source/Editor/UnrealEd/Private/UnrealEdMisc.cpp
+++ b/Engine/Source/Editor/UnrealEd/Private/UnrealEdMisc.cpp
@@ -1,6 +1,7 @@
 // Copyright Epic Games, Inc. All Rights Reserved.
 
 #include "UnrealEdMisc.h"
+
 #include "TickableEditorObject.h"
 #include "Components/PrimitiveComponent.h"
 #include "Misc/MessageDialog.h"
@@ -27,6 +28,7 @@
 #include "HAL/PlatformSplash.h"
 #include "Internationalization/Culture.h"
 #include "Misc/ConfigCacheIni.h"
+#include "UObject/ObjectSaveContext.h"
 #include "UObject/UObjectIterator.h"
 #include "EngineUtils.h"
 #include "EditorViewportClient.h"
@@ -245,12 +247,11 @@ void FUnrealEdMisc::OnInit()
 	FEditorDelegates::DisplayLoadErrors.AddRaw(this, &FUnrealEdMisc::CB_DisplayLoadErrors);
 	FEditorDelegates::MapChange.AddRaw(this, &FUnrealEdMisc::CB_MapChange);
 	FEditorDelegates::RefreshEditor.AddRaw(this, &FUnrealEdMisc::CB_RefreshEditor);
-	FEditorDelegates::PreSaveWorld.AddRaw(this, &FUnrealEdMisc::PreSaveWorld);
+	FEditorDelegates::PreSaveWorldWithContext.AddRaw(this, &FUnrealEdMisc::PreSaveWorld);
 	FEditorSupportDelegates::RedrawAllViewports.AddRaw(this, &FUnrealEdMisc::CB_RedrawAllViewports);
 	GEngine->OnLevelActorAdded().AddRaw( this, &FUnrealEdMisc::CB_LevelActorsAdded );
 
 	FCoreUObjectDelegates::OnObjectSaved.AddRaw(this, &FUnrealEdMisc::OnObjectSaved);
-	FEditorDelegates::PreSaveWorld.AddRaw(this, &FUnrealEdMisc::OnWorldSaved);
 
 #if USE_UNIT_TESTS
 	FAutomationTestFramework::Get().PreTestingEvent.AddRaw(this, &FUnrealEdMisc::CB_PreAutomationTesting);
@@ -987,7 +988,7 @@ void FUnrealEdMisc::OnExit()
 	FEditorDelegates::DisplayLoadErrors.RemoveAll(this);
 	FEditorDelegates::MapChange.RemoveAll(this);
 	FEditorDelegates::RefreshEditor.RemoveAll(this);
-	FEditorDelegates::PreSaveWorld.RemoveAll(this);
+	FEditorDelegates::PreSaveWorldWithContext.RemoveAll(this);
 	FEditorSupportDelegates::RedrawAllViewports.RemoveAll(this);
 	GEngine->OnLevelActorAdded().RemoveAll(this);
 
@@ -1091,10 +1092,11 @@ void FUnrealEdMisc::CB_RefreshEditor()
 	FEditorDelegates::RefreshAllBrowsers.Broadcast();
 }
 
-void FUnrealEdMisc::PreSaveWorld(uint32 SaveFlags, UWorld* World)
+void FUnrealEdMisc::PreSaveWorld(UWorld* World, FObjectPreSaveContext ObjectSaveContext)
 {
-	const bool bAutosaveOrPIE = (SaveFlags & SAVE_FromAutosave) != 0;
-	if (bAutosaveOrPIE || World == NULL || World != GEditor->GetEditorWorldContext().World() || !FEngineAnalytics::IsAvailable())
+	LogAssetUpdate(World);
+	const bool bAutosaveOrPIE = (ObjectSaveContext.GetSaveFlags() & SAVE_FromAutosave) != 0;
+	if (bAutosaveOrPIE || !World || World != GEditor->GetEditorWorldContext().World() || !FEngineAnalytics::IsAvailable())
 	{
 		return;
 	}
@@ -1649,11 +1651,6 @@ void FUnrealEdMisc::OnObjectSaved(UObject* SavedObject)
 	}
 }
 
-void FUnrealEdMisc::OnWorldSaved(uint32 SaveFlags, UWorld* SavedWorld)
-{
-	LogAssetUpdate(SavedWorld);
-}
-
 void FUnrealEdMisc::LogAssetUpdate(UObject* UpdatedAsset)
 {
 	UPackage* AssetPackage = UpdatedAsset->GetOutermost();
diff --git a/Engine/Source/Editor/UnrealEd/Public/Editor.h b/Engine/Source/Editor/UnrealEd/Public/Editor.h
--- a/Engine/Source/Editor/UnrealEd/Public/Editor.h
+++ b/Engine/Source/Editor/UnrealEd/Public/Editor.h
@@ -70,8 +70,12 @@ struct UNREALED_API FEditorDelegates
 	DECLARE_MULTICAST_DELEGATE_TwoParams(FOnDollyPerspectiveCamera, const FVector&, int32 );
 	/** delegate type for pre save world events ( uint32 SaveFlags, UWorld* World ) */
 	DECLARE_MULTICAST_DELEGATE_TwoParams(FOnPreSaveWorld, uint32, class UWorld*);
+	/** delegate type for pre save world events ( UWorld* World, FObjectPreSaveContext ObjectSaveContext ) */
+	DECLARE_MULTICAST_DELEGATE_TwoParams(FOnPreSaveWorldWithContext, class UWorld*, FObjectPreSaveContext);
 	/** delegate type for post save world events ( uint32 SaveFlags, UWorld* World, bool bSuccess ) */
 	DECLARE_MULTICAST_DELEGATE_ThreeParams(FOnPostSaveWorld, uint32, class UWorld*, bool);
+	/** delegate type for post save world events ( UWorld* World, FObjectPostSaveContext ObjectSaveContext, bool bSuccess ) */
+	DECLARE_MULTICAST_DELEGATE_TwoParams(FOnPostSaveWorldWithContext, class UWorld*, FObjectPostSaveContext);
 	/** delegate for a PIE event (begin, end, pause/resume, etc) (Params: bool bIsSimulating) */
 	DECLARE_MULTICAST_DELEGATE_OneParam(FOnPIEEvent, const bool);
 	/** delegate for a standalone local play event (Params: uint32 processID) */
@@ -211,10 +215,14 @@ struct UNREALED_API FEditorDelegates
 	static FSimpleMulticastDelegate PropertySelectionChange;
 	/** Called after Landscape layer infomap update have completed */
 	static FSimpleMulticastDelegate PostLandscapeLayerUpdated;
-	/** Called before SaveWorld is processed */
+	UE_DEPRECATED(5.0, "Use PreSaveWorldWithContext instead.")
 	static FOnPreSaveWorld PreSaveWorld;
-	/** Called after SaveWorld is processed */
+	/** Called before SaveWorld is processed */
+	static FOnPreSaveWorldWithContext PreSaveWorldWithContext;
+	UE_DEPRECATED(5.0, "Use PostSaveWorldWithContext instead.")
 	static FOnPostSaveWorld PostSaveWorld;
+	/** Called after SaveWorld is processed */
+	static FOnPostSaveWorldWithContext PostSaveWorldWithContext;
 	/** Called when finishing picking a new blueprint class during construction */
 	static FOnFinishPickingBlueprintClass OnFinishPickingBlueprintClass;
 	/** Called when beginning configuration of a new asset */
diff --git a/Engine/Source/Editor/UnrealEd/Public/EditorActorFolders.h b/Engine/Source/Editor/UnrealEd/Public/EditorActorFolders.h
--- a/Engine/Source/Editor/UnrealEd/Public/EditorActorFolders.h
+++ b/Engine/Source/Editor/UnrealEd/Public/EditorActorFolders.h
@@ -10,6 +10,7 @@
 #include "EditorActorFolders.generated.h"
 
 class AActor;
+class FObjectPostSaveContext;
 
 /** Multicast delegates for broadcasting various folder events */
 DECLARE_MULTICAST_DELEGATE_TwoParams(FOnActorFolderCreate, UWorld&, FName);
@@ -138,7 +139,7 @@ private:
 	void OnMapChange(uint32 MapChangeFlags);
 
 	/** Called after a world has been saved */
-	void OnWorldSaved(uint32 SaveFlags, UWorld* World, bool bSuccess);
+	void OnWorldSaved(UWorld* World, FObjectPostSaveContext ObjectSaveContext);
 
 	/** Remove any references to folder arrays for dead worlds */
 	void Housekeeping();
diff --git a/Engine/Source/Editor/UnrealEd/Public/UnrealEdMisc.h b/Engine/Source/Editor/UnrealEd/Public/UnrealEdMisc.h
--- a/Engine/Source/Editor/UnrealEd/Public/UnrealEdMisc.h
+++ b/Engine/Source/Editor/UnrealEd/Public/UnrealEdMisc.h
@@ -229,7 +229,6 @@ public:
 	void CB_SelectedProps();
 	void CB_DisplayLoadErrors();
 	void CB_RefreshEditor();
-	void PreSaveWorld(uint32 SaveFlags, class UWorld* World);
 
 	/** Tells the editor that something has been done to change the map.  Can be anything from loading a whole new map to changing the BSP. */
 	void CB_MapChange( uint32 InFlags );
@@ -291,9 +290,6 @@ public:
 	/** Delegate used to update the map of asset update counts */
 	void OnObjectSaved(UObject* SavedObject);
 
-	/** Delegate used to update the map of asset update counts (for UWorlds specifically) */
-	void OnWorldSaved(uint32 SaveFlags, UWorld* SavedWorld);
-
 	/** Logs an update to an asset */
 	void LogAssetUpdate(UObject* UpdatedAsset);
 
@@ -312,6 +308,7 @@ public:
 	FString GetProjectEditorBinaryPath();
 
 private:
+	void PreSaveWorld(class UWorld* World, FObjectPreSaveContext ObjectSaveContext);
 
 	/** The current state of the autosave */
 	EAutosaveState::Type AutosaveState;
diff --git a/Engine/Source/Editor/WorldBrowser/Private/SWorldHierarchyImpl.cpp b/Engine/Source/Editor/WorldBrowser/Private/SWorldHierarchyImpl.cpp
--- a/Engine/Source/Editor/WorldBrowser/Private/SWorldHierarchyImpl.cpp
+++ b/Engine/Source/Editor/WorldBrowser/Private/SWorldHierarchyImpl.cpp
@@ -1,6 +1,7 @@
 // Copyright Epic Games, Inc. All Rights Reserved.
 
 #include "SWorldHierarchyImpl.h"
+
 #include "SLevelsTreeWidget.h"
 #include "SWorldHierarchyItem.h"
 #include "SWorldHierarchy.h"
@@ -11,6 +12,7 @@
 #include "Framework/Application/SlateApplication.h"
 #include "Framework/MultiBox/MultiBoxBuilder.h"
 #include "ToolMenus.h"
+#include "UObject/ObjectSaveContext.h"
 #include "Widgets/Layout/SSeparator.h"
 #include "Widgets/Images/SImage.h"
 #include "Widgets/Input/SComboButton.h"
@@ -59,7 +61,7 @@ SWorldHierarchyImpl::~SWorldHierarchyImpl()
 		LevelFolders.OnFolderDelete.RemoveAll(this);
 	}
 
-	FEditorDelegates::PostSaveWorld.RemoveAll(this);
+	FEditorDelegates::PostSaveWorldWithContext.RemoveAll(this);
 }
 
 void SWorldHierarchyImpl::Construct(const FArguments& InArgs)
@@ -326,7 +328,7 @@ void SWorldHierarchyImpl::Construct(const FArguments& InArgs)
 
 		if (!bFoldersOnlyMode)
 		{
-			FEditorDelegates::PostSaveWorld.AddSP(this, &SWorldHierarchyImpl::OnWorldSaved);
+			FEditorDelegates::PostSaveWorldWithContext.AddSP(this, &SWorldHierarchyImpl::OnWorldSaved);
 		}
 	}
 
@@ -364,7 +366,7 @@ void SWorldHierarchyImpl::Tick( const FGeometry& AllotedGeometry, const double I
 	}
 }
 
-void SWorldHierarchyImpl::OnWorldSaved(uint32 SaveFlags, UWorld* World, bool bSuccess)
+void SWorldHierarchyImpl::OnWorldSaved(UWorld* World, FObjectPostSaveContext ObjectSaveContext)
 {
 	if (FLevelFolders::IsAvailable())
 	{
diff --git a/Engine/Source/Editor/WorldBrowser/Private/SWorldHierarchyImpl.h b/Engine/Source/Editor/WorldBrowser/Private/SWorldHierarchyImpl.h
--- a/Engine/Source/Editor/WorldBrowser/Private/SWorldHierarchyImpl.h
+++ b/Engine/Source/Editor/WorldBrowser/Private/SWorldHierarchyImpl.h
@@ -14,6 +14,7 @@
 #include "LevelCollectionModel.h"
 #include "Misc/TextFilter.h"
 #include "SLevelsTreeWidget.h"
+#include "UObject/ObjectSaveContext.h"
 
 #include "IWorldTreeItem.h"
 
@@ -301,7 +302,7 @@ private:
 	bool GetDisplayActorsCountState() const;
 
 	/** Callback for when the world is saved */
-	void OnWorldSaved(uint32 SaveFlags, UWorld* World, bool bSuccess);
+	void OnWorldSaved(UWorld* World, FObjectPostSaveContext ObjectSaveContext);
 
 private:
 	/**	Whether the view is currently updating the viewmodel selection */
diff --git a/Engine/Source/Editor/WorldBrowser/Private/Tiles/WorldTileCollectionModel.cpp b/Engine/Source/Editor/WorldBrowser/Private/Tiles/WorldTileCollectionModel.cpp
--- a/Engine/Source/Editor/WorldBrowser/Private/Tiles/WorldTileCollectionModel.cpp
+++ b/Engine/Source/Editor/WorldBrowser/Private/Tiles/WorldTileCollectionModel.cpp
@@ -1,5 +1,7 @@
 // Copyright Epic Games, Inc. All Rights Reserved.
+
 #include "Tiles/WorldTileCollectionModel.h"
+
 #include "Misc/PackageName.h"
 #include "Components/PrimitiveComponent.h"
 #include "Misc/CoreDelegates.h"
@@ -47,6 +49,7 @@
 #include "IMeshReductionManagerModule.h"
 #include "IMeshMergeUtilities.h"
 #include "MeshMergeModule.h"
+#include "UObject/ObjectSaveContext.h"
 
 #define LOCTEXT_NAMESPACE "WorldBrowser"
 
@@ -71,8 +74,8 @@ FWorldTileCollectionModel::~FWorldTileCollectionModel()
 	GEditor->UnregisterForUndo(this);
 	FCoreDelegates::PreWorldOriginOffset.RemoveAll(this);
 	FCoreDelegates::PostWorldOriginOffset.RemoveAll(this);
-	FEditorDelegates::PreSaveWorld.RemoveAll(this);
-	FEditorDelegates::PostSaveWorld.RemoveAll(this);
+	FEditorDelegates::PreSaveWorldWithContext.RemoveAll(this);
+	FEditorDelegates::PostSaveWorldWithContext.RemoveAll(this);
 	FEditorDelegates::NewCurrentLevel.RemoveAll(this);
 }
 
@@ -86,8 +89,8 @@ void FWorldTileCollectionModel::Initialize(UWorld* InWorld)
 	GEditor->RegisterForUndo(this);
 	FCoreDelegates::PreWorldOriginOffset.AddSP(this, &FWorldTileCollectionModel::PreWorldOriginOffset);
 	FCoreDelegates::PostWorldOriginOffset.AddSP(this, &FWorldTileCollectionModel::PostWorldOriginOffset);
-	FEditorDelegates::PreSaveWorld.AddSP(this, &FWorldTileCollectionModel::OnPreSaveWorld);
-	FEditorDelegates::PostSaveWorld.AddSP(this, &FWorldTileCollectionModel::OnPostSaveWorld);
+	FEditorDelegates::PreSaveWorldWithContext.AddSP(this, &FWorldTileCollectionModel::OnPreSaveWorld);
+	FEditorDelegates::PostSaveWorldWithContext.AddSP(this, &FWorldTileCollectionModel::OnPostSaveWorld);
 	FEditorDelegates::NewCurrentLevel.AddSP(this, &FWorldTileCollectionModel::OnNewCurrentLevel);
 	BindCommands();
 			
@@ -1895,14 +1898,14 @@ void FWorldTileCollectionModel::ResetLevelOrigin_Executed()
 	RequestUpdateAllLevels();
 }
 
-void FWorldTileCollectionModel::OnPreSaveWorld(uint32 SaveFlags, UWorld* World)
+void FWorldTileCollectionModel::OnPreSaveWorld(UWorld* World, FObjectPreSaveContext ObjectSaveContext)
 {
 	// Levels during OnSave procedure might be moved to original position
 	// and then back to position with offset
 	bIsSavingLevel = true;
 }
 
-void FWorldTileCollectionModel::OnPostSaveWorld(uint32 SaveFlags, UWorld* World, bool bSuccess)
+void FWorldTileCollectionModel::OnPostSaveWorld(UWorld* World, FObjectPostSaveContext ObjectSaveContext)
 {
 	bIsSavingLevel = false;
 }
