Select the group actor when creating a group from the active selection

This preserves group selection rules, and means that gizmo manipulation immediately after making a group also moves the group actor itself.

diff --git a/Engine/Source/Editor/UnrealEd/Private/ActorGroupingUtils.cpp b/Engine/Source/Editor/UnrealEd/Private/ActorGroupingUtils.cpp
--- a/Engine/Source/Editor/UnrealEd/Private/ActorGroupingUtils.cpp
+++ b/Engine/Source/Editor/UnrealEd/Private/ActorGroupingUtils.cpp
@@ -22,7 +22,7 @@ UActorGroupingUtils* UActorGroupingUtils::Get()
 	return GEditor->GetActorGroupingUtils();
 }
 
-void UActorGroupingUtils::GroupSelected()
+AGroupActor* UActorGroupingUtils::GroupSelected()
 {
 	if (IsGroupingActive())
 	{
@@ -34,12 +34,21 @@ void UActorGroupingUtils::GroupSelected()
 
 		if (ActorsToAdd.Num() > 0)
 		{
-			GroupActors(ActorsToAdd);
+			const FScopedTransaction Transaction(NSLOCTEXT("UnrealEd", "Group_Selected", "Group Selected"));
+
+			AGroupActor* GroupActor = GroupActors(ActorsToAdd);
+			if (GroupActor)
+			{
+				GEditor->SelectActor(GroupActor, /*bIsSelected*/true, /*bNotify*/false);
+			}
+			return GroupActor;
 		}
 	}
+
+	return nullptr;
 }
 
-void UActorGroupingUtils::GroupActors(const TArray<AActor*>& ActorsToGroup)
+AGroupActor* UActorGroupingUtils::GroupActors(const TArray<AActor*>& ActorsToGroup)
 {
 	if(IsGroupingActive())
 	{
@@ -106,6 +115,8 @@ void UActorGroupingUtils::GroupActors(const TArray<AActor*>& ActorsToGroup)
 					SpawnedGroupActor->SetFolderPath(FolderPath);
 					SpawnedGroupActor->CenterGroupLocation();
 					SpawnedGroupActor->Lock();
+
+					return SpawnedGroupActor;
 				}
 			}
 		}
@@ -117,6 +128,8 @@ void UActorGroupingUtils::GroupActors(const TArray<AActor*>& ActorsToGroup)
 			FSlateNotificationManager::Get().AddNotification(Info);
 		}
 	}
+
+	return nullptr;
 }
 
 void UActorGroupingUtils::UngroupSelected()
@@ -133,6 +146,8 @@ void UActorGroupingUtils::UngroupSelected()
 
 		if (ActorsToUngroup.Num())
 		{
+			const FScopedTransaction Transaction(NSLOCTEXT("UnrealEd", "Ungroup_Selected", "Ungroup Selected"));
+
 			UngroupActors(ActorsToUngroup);
 		}
 	}
diff --git a/Engine/Source/Editor/UnrealEd/Public/ActorGroupingUtils.h b/Engine/Source/Editor/UnrealEd/Public/ActorGroupingUtils.h
--- a/Engine/Source/Editor/UnrealEd/Public/ActorGroupingUtils.h
+++ b/Engine/Source/Editor/UnrealEd/Public/ActorGroupingUtils.h
@@ -30,13 +30,13 @@ public:
 	 * Creates a new group from the current selection removing the actors from any existing groups they are already in
 	 */
 	UFUNCTION(BlueprintCallable, Category = "Editor Scripting | Actor Grouping")
-	virtual void GroupSelected();
+	virtual AGroupActor* GroupSelected();
 
 	/**
 	 * Creates a new group from the provided list of actors removing the actors from any existing groups they are already in
 	 */
 	UFUNCTION(BlueprintCallable, Category = "Editor Scripting | Actor Grouping")
-	virtual void GroupActors(const TArray<AActor*>& ActorsToGroup);
+	virtual AGroupActor* GroupActors(const TArray<AActor*>& ActorsToGroup);
 
 	/**
 	 * Disbands any groups in the current selection, does not attempt to maintain any hierarchy
