Avoid sub-object resolution during SnapshotObject when locating the object record

diff --git a/Engine/Source/Editor/UnrealEd/Classes/Editor/Transactor.h b/Engine/Source/Editor/UnrealEd/Classes/Editor/Transactor.h
--- a/Engine/Source/Editor/UnrealEd/Classes/Editor/Transactor.h
+++ b/Engine/Source/Editor/UnrealEd/Classes/Editor/Transactor.h
@@ -64,8 +64,19 @@ protected:
 				return Ar;
 			}
 
-			FPersistentObjectRef() = default;
+			friend bool operator==(const FPersistentObjectRef& LHS, const FPersistentObjectRef& RHS)
+			{
+				return LHS.ReferenceType == RHS.ReferenceType
+					&& LHS.RootObject == RHS.RootObject
+					&& (LHS.ReferenceType != EReferenceType::SubObject || LHS.SubObjectHierarchyIDs == RHS.SubObjectHierarchyIDs);
+			}
 
+			friend bool operator!=(const FPersistentObjectRef& LHS, const FPersistentObjectRef& RHS)
+			{
+				return !(LHS == RHS);
+			}
+
+			FPersistentObjectRef() = default;
 			explicit FPersistentObjectRef(UObject* InObject);
 
 			bool IsRootObjectReference() const
diff --git a/Engine/Source/Editor/UnrealEd/Private/EditorTransaction.cpp b/Engine/Source/Editor/UnrealEd/Private/EditorTransaction.cpp
--- a/Engine/Source/Editor/UnrealEd/Private/EditorTransaction.cpp
+++ b/Engine/Source/Editor/UnrealEd/Private/EditorTransaction.cpp
@@ -879,9 +879,9 @@ void FTransaction::SnapshotObject( UObject* InObject, TArrayView<const FProperty
 {
 	if (InObject && ObjectMap.Contains(InObject))
 	{
-		FObjectRecord* FoundObjectRecord = Algo::FindByPredicate(Records, [InObject](const FObjectRecord& ObjRecord)
+		FObjectRecord* FoundObjectRecord = Algo::FindByPredicate(Records, [ObjRefToFind = FObjectRecord::FPersistentObjectRef(InObject)](const FObjectRecord& ObjRecord)
 		{
-			return ObjRecord.Object.Get() == InObject;
+			return ObjRecord.Object == ObjRefToFind;
 		});
 
 		if (FoundObjectRecord)
