Fixes misalignment issues in New Blueprint Pick Parent dialog when documentation was available

diff --git a/Engine/Source/Editor/UnrealEd/Private/SClassPickerDialog.cpp b/Engine/Source/Editor/UnrealEd/Private/SClassPickerDialog.cpp
--- a/Engine/Source/Editor/UnrealEd/Private/SClassPickerDialog.cpp
+++ b/Engine/Source/Editor/UnrealEd/Private/SClassPickerDialog.cpp
@@ -261,11 +261,10 @@ TSharedRef<ITableRow> SClassPickerDialog::GenerateListRow(TSharedPtr<FClassPicke
 						SNew(STextBlock)
 						.Text(Obj->GetName())
 					]
-
 				]
 			]
 			+SHorizontalBox::Slot()
-			.Padding(10.0f, 0.0f)
+			.Padding(10.0f, 0.0f, 4.0f, 0.0f)
 			[
 				SNew(STextBlock)
 				.Text(Obj->GetDescription())
@@ -275,7 +274,11 @@ TSharedRef<ITableRow> SClassPickerDialog::GenerateListRow(TSharedPtr<FClassPicke
 			.VAlign(VAlign_Center)
 			.AutoWidth()
 			[
-				FEditorClassUtils::GetDocumentationLinkWidget(ItemClass)
+				SNew(SBox)
+				.WidthOverride(60)
+				[
+					FEditorClassUtils::GetDocumentationLinkWidget(ItemClass)
+				]
 			]
 		]
 	];
