Fixed an issue in the FBX exporting where multiple meshes with same material would receive the default material instead.

diff --git a/Engine/Source/Editor/UnrealEd/Private/Fbx/FbxMainExport.cpp b/Engine/Source/Editor/UnrealEd/Private/Fbx/FbxMainExport.cpp
--- a/Engine/Source/Editor/UnrealEd/Private/Fbx/FbxMainExport.cpp
+++ b/Engine/Source/Editor/UnrealEd/Private/Fbx/FbxMainExport.cpp
@@ -4669,7 +4669,7 @@ FbxNode* FFbxExporter::ExportStaticMeshToFbx(const UStaticMesh* StaticMesh, int3
 			}
 			else
 			{
-				StaticMesh->GetMaterial(Polygons.MaterialIndex);
+				Material = StaticMesh->GetMaterial(Polygons.MaterialIndex);
 			}
 
 			FbxSurfaceMaterial* FbxMaterial = Material ? ExportMaterial(Material) : NULL;
