Fix thumbnails not showing up or refreshing properly when shaders are finished compiling
  - This bug is a combination of CL 14673948 and 14688146 that were tested separately

diff --git a/Engine/Source/Editor/UnrealEd/Private/AssetThumbnail.cpp b/Engine/Source/Editor/UnrealEd/Private/AssetThumbnail.cpp
--- a/Engine/Source/Editor/UnrealEd/Private/AssetThumbnail.cpp
+++ b/Engine/Source/Editor/UnrealEd/Private/AssetThumbnail.cpp
@@ -948,7 +948,7 @@ TStatId FAssetThumbnailPool::GetStatId() const
 
 bool FAssetThumbnailPool::IsTickable() const
 {
-	return RecentlyLoadedAssets.Num() > 0 || ThumbnailsToRenderStack.Num() > 0 || RealTimeThumbnails.Num() > 0;
+	return RecentlyLoadedAssets.Num() > 0 || ThumbnailsToRenderStack.Num() > 0 || RealTimeThumbnails.Num() > 0 || bWereShadersCompilingLastFrame || (GShaderCompilingManager && GShaderCompilingManager->IsCompiling());
 }
 
 void FAssetThumbnailPool::Tick( float DeltaTime )
@@ -959,6 +959,18 @@ void FAssetThumbnailPool::Tick( float DeltaTime )
 		return;
 	}
 
+	const bool bAreShadersCompiling = (GShaderCompilingManager && GShaderCompilingManager->IsCompiling());
+	if (bWereShadersCompilingLastFrame && !bAreShadersCompiling)
+	{
+		ThumbnailsToRenderStack.Reset();
+		// Reschedule visible thumbnails to be rerendered now that shaders are finished compiling
+		for (auto ThumbIt = ThumbnailToTextureMap.CreateIterator(); ThumbIt; ++ThumbIt)
+		{
+			ThumbnailsToRenderStack.Push(ThumbIt.Value());
+		}
+	}
+	bWereShadersCompilingLastFrame = bAreShadersCompiling;
+
 	TRACE_CPUPROFILER_EVENT_SCOPE(FAssetThumbnailPool::Tick);
 	// If there were any assets loaded since last frame that we are currently displaying thumbnails for, push them on the render stack now.
 	if ( RecentlyLoadedAssets.Num() > 0 )
@@ -1021,31 +1033,28 @@ void FAssetThumbnailPool::Tick( float DeltaTime )
 
 			if( Info.IsValid() )
 			{
+				bool bIsAssetStillCompiling = false;
 				TSharedRef<FThumbnailInfo> InfoRef = Info.ToSharedRef();
 
 				if ( InfoRef->AssetData.IsValid() )
 				{
-					const FObjectThumbnail* ObjectThumbnail = NULL;
+					const FObjectThumbnail* ObjectThumbnail = nullptr;
 					bool bLoadedThumbnail = false;
 
-					// Hold off thumbnail rendering during asset compilation to avoid making the game-thread slower than it already is.
+					// Hold off thumbnail rendering during shader compilation to avoid making the game-thread slower than it already is.
 					// It will fallback on the thumbnail cached on disk during asset compile time, which is even better than showing wrong thumbnails.
-					const bool bAssetsBeingCompiled = 
-						(GShaderCompilingManager && GShaderCompilingManager->IsCompiling()) ||
-						FAssetCompilingManager::Get().GetNumRemainingAssets() > 0;
-
 					// If this is a loaded asset and we have a rendering info for it, render a fresh thumbnail here
-					if( InfoRef->AssetData.IsAssetLoaded() && !bAssetsBeingCompiled)
+					if( InfoRef->AssetData.IsAssetLoaded() && !bAreShadersCompiling)
 					{
 						UObject* Asset = InfoRef->AssetData.GetAsset();
 						
 						//Avoid rendering the thumbnail of an asset that is currently edited asynchronously
 						const IInterface_AsyncCompilation* Interface_AsyncCompilation = Cast<IInterface_AsyncCompilation>(Asset);
-						const bool bIsAssetCurrentlyEditedAsynchronously = Interface_AsyncCompilation && Interface_AsyncCompilation->IsCompiling();
-						if (!bIsAssetCurrentlyEditedAsynchronously)
+						bIsAssetStillCompiling = Interface_AsyncCompilation && Interface_AsyncCompilation->IsCompiling();
+						if (!bIsAssetStillCompiling)
 						{
 							FThumbnailRenderingInfo* RenderInfo = GUnrealEd->GetThumbnailManager()->GetRenderingInfo(Asset);
-							if (RenderInfo != NULL && RenderInfo->Renderer != NULL)
+							if (RenderInfo != nullptr && RenderInfo->Renderer != nullptr)
 							{
 								FThumbnailInfo_RenderThread ThumbInfo = InfoRef.Get();
 								ENQUEUE_RENDER_COMMAND(SyncSlateTextureCommand)(
@@ -1075,7 +1084,7 @@ void FAssetThumbnailPool::Tick( float DeltaTime )
 							}
 						}
 					}
-				
+
 					FThumbnailMap ThumbnailMap;
 					// If we could not render a fresh thumbnail, see if we already have a cached one to load
 					if ( !bLoadedThumbnail )
@@ -1146,9 +1155,12 @@ void FAssetThumbnailPool::Tick( float DeltaTime )
 						// Notify listeners that a thumbnail has been rendered
 						ThumbnailRenderedEvent.Broadcast(InfoRef->AssetData);
 					}
-					else
+					// Do not send a failure event for this asset yet if shaders are still compiling or the asset itself is compiling.
+					// The failure event will disable the rendering of this asset for good and we need to have a chance to 
+					// rerender it when everything settles down.
+					else if (!bAreShadersCompiling && !bIsAssetStillCompiling)
 					{
-						// Notify listeners that a thumbnail has been rendered
+						// Notify listeners that a thumbnail render has failed
 						ThumbnailRenderFailedEvent.Broadcast(InfoRef->AssetData);
 					}
 				}
diff --git a/Engine/Source/Editor/UnrealEd/Public/AssetThumbnail.h b/Engine/Source/Editor/UnrealEd/Public/AssetThumbnail.h
--- a/Engine/Source/Editor/UnrealEd/Public/AssetThumbnail.h
+++ b/Engine/Source/Editor/UnrealEd/Public/AssetThumbnail.h
@@ -346,6 +346,9 @@ private:
 	/** Max number of thumbnails in the pool */
 	uint32 NumInPool;
 
+	/** Shaders are still building */
+	bool bWereShadersCompilingLastFrame = false;
+
 	/** Max number of dynamic thumbnails to update per frame */
 	uint32 MaxRealTimeThumbnailsPerFrame;
 
