Split legacy viewport input handling from UEdMode.

Move potentially duplicated work (like bracket draws) to the mode manager. Input and drawing of mode specifics should be handled going forward by ITF tools and gizmos in the modes.

#rnx

diff --git a/Engine/Plugins/Experimental/GizmoEdMode/Source/GizmoEdMode/Private/GizmoEdMode.cpp b/Engine/Plugins/Experimental/GizmoEdMode/Source/GizmoEdMode/Private/GizmoEdMode.cpp
--- a/Engine/Plugins/Experimental/GizmoEdMode/Source/GizmoEdMode/Private/GizmoEdMode.cpp
+++ b/Engine/Plugins/Experimental/GizmoEdMode/Source/GizmoEdMode/Private/GizmoEdMode.cpp
@@ -9,6 +9,7 @@
 #include "BaseGizmos/TransformGizmo.h"
 #include "UnrealWidgetFwd.h"
 #include "Utils.h"
+#include "Stats/Stats2.h"
 
 #define LOCTEXT_NAMESPACE "FGizmoEdMode"
 
@@ -85,14 +86,10 @@ void UGizmoEdMode::Exit()
 	Super::Exit();
 }
 
-bool UGizmoEdMode::InputKey(FEditorViewportClient* ViewportClient, FViewport* Viewport, FKey Key, EInputEvent Event)
+void UGizmoEdMode::ModeTick(float DeltaTime)
 {
-	return ToolsContext->InputKey(ViewportClient, Viewport, Key, Event);
-}
+	Super::ModeTick(DeltaTime);
 
-void UGizmoEdMode::Tick(FEditorViewportClient* ViewportClient, float DeltaTime)
-{
-	Super::Tick(ViewportClient, DeltaTime);
 	if ( bNeedInitialGizmos )
 	{
 		RecreateGizmo();
diff --git a/Engine/Plugins/Experimental/GizmoEdMode/Source/GizmoEdMode/Public/GizmoEdMode.h b/Engine/Plugins/Experimental/GizmoEdMode/Source/GizmoEdMode/Public/GizmoEdMode.h
--- a/Engine/Plugins/Experimental/GizmoEdMode/Source/GizmoEdMode/Public/GizmoEdMode.h
+++ b/Engine/Plugins/Experimental/GizmoEdMode/Source/GizmoEdMode/Public/GizmoEdMode.h
@@ -32,14 +32,14 @@ public:
 	{
 		return false;
 	}
+
 private:
 	void ActorSelectionChangeNotify() override;
 	void Enter() override;
 	void Exit() override;
-	bool InputKey(FEditorViewportClient* ViewportClient, FViewport* Viewport, FKey Key, EInputEvent Event) override;
+	void ModeTick(float DeltaTime) override;
 
 	bool IsCompatibleWith(FEditorModeID OtherModeID) const override { return true; }
-	void Tick(FEditorViewportClient* ViewportClient, float DeltaTime) override;
 
 	void RecreateGizmo();
 	void DestroyGizmo();
diff --git a/Engine/Plugins/Experimental/ModelingToolsEditorMode/Source/ModelingToolsEditorMode/Private/ModelingToolsEditorMode.cpp b/Engine/Plugins/Experimental/ModelingToolsEditorMode/Source/ModelingToolsEditorMode/Private/ModelingToolsEditorMode.cpp
--- a/Engine/Plugins/Experimental/ModelingToolsEditorMode/Source/ModelingToolsEditorMode/Private/ModelingToolsEditorMode.cpp
+++ b/Engine/Plugins/Experimental/ModelingToolsEditorMode/Source/ModelingToolsEditorMode/Private/ModelingToolsEditorMode.cpp
@@ -194,7 +194,7 @@ bool UModelingToolsEditorMode::ShouldDrawWidget() const
 
 void UModelingToolsEditorMode::Tick(FEditorViewportClient* ViewportClient, float DeltaTime)
 {
-	UEdMode::Tick(ViewportClient, DeltaTime);
+	Super::Tick(ViewportClient, DeltaTime);
 
 	if (Toolkit.IsValid())
 	{
diff --git a/Engine/Plugins/Experimental/MotionTrailEditorMode/Source/MotionTrailEditorMode/Private/MotionTrailEditorMode.cpp b/Engine/Plugins/Experimental/MotionTrailEditorMode/Source/MotionTrailEditorMode/Private/MotionTrailEditorMode.cpp
--- a/Engine/Plugins/Experimental/MotionTrailEditorMode/Source/MotionTrailEditorMode/Private/MotionTrailEditorMode.cpp
+++ b/Engine/Plugins/Experimental/MotionTrailEditorMode/Source/MotionTrailEditorMode/Private/MotionTrailEditorMode.cpp
@@ -154,11 +154,6 @@ void UMotionTrailEditorMode::DrawHUD(FEditorViewportClient* ViewportClient, FVie
 	Super::DrawHUD(ViewportClient, Viewport, View, Canvas);
 }
 
-bool UMotionTrailEditorMode::UsesToolkits() const
-{
-	return true;
-}
-
 TMap<FName, TArray<TSharedPtr<FUICommandInfo>>> UMotionTrailEditorMode::GetModeCommands() const
 {
 	const TMap<FName, TArray<TSharedPtr<FUICommandInfo>>>& Commands = UE::MotionTrailEditor::FMotionTrailEditorModeCommands::Get().GetCommands();
@@ -214,26 +209,6 @@ void UMotionTrailEditorMode::ActivateDefaultTool()
 	ToolsContext->StartTool(UMotionTrailEditorMode::DefaultToolName);
 }
 
-FEdMode* UMotionTrailEditorMode::AsLegacyMode()
-{
-	class FMotionTrailEditorMode : public FEdMode
-	{
-	public:
-		FMotionTrailEditorMode()
-			: FEdMode()
-		{
-			Owner = &GLevelEditorModeTools();
-		}
-
-		virtual ~FMotionTrailEditorMode() {};
-		virtual bool UsesTransformWidget() const { return true; }
-	};
-
-	static TUniquePtr<FMotionTrailEditorMode> LegacyEdMode = MakeUnique<FMotionTrailEditorMode>();
-
-	return LegacyEdMode.Get();
-}
-
 bool UMotionTrailEditorMode::IsCompatibleWith(FEditorModeID OtherModeID) const
 {
 	return OtherModeID == FName(TEXT("EM_SequencerMode"), FNAME_Find) || OtherModeID == FName(TEXT("EditMode.ControlRig"), FNAME_Find) || OtherModeID == FName(TEXT("EditMode.ControlRigEditor"), FNAME_Find);
diff --git a/Engine/Plugins/Experimental/MotionTrailEditorMode/Source/MotionTrailEditorMode/Public/MotionTrailEditorMode.h b/Engine/Plugins/Experimental/MotionTrailEditorMode/Source/MotionTrailEditorMode/Public/MotionTrailEditorMode.h
--- a/Engine/Plugins/Experimental/MotionTrailEditorMode/Source/MotionTrailEditorMode/Public/MotionTrailEditorMode.h
+++ b/Engine/Plugins/Experimental/MotionTrailEditorMode/Source/MotionTrailEditorMode/Public/MotionTrailEditorMode.h
@@ -3,7 +3,7 @@
 #pragma once
 
 #include "CoreMinimal.h"
-#include "Tools/UEdMode.h"
+#include "Tools/LegacyEdModeWidgetHelpers.h"
 #include "ISequencer.h"
 
 #include "TrailHierarchy.h"
@@ -71,7 +71,7 @@ public:
 };
 
 UCLASS()
-class UMotionTrailEditorMode : public UEdMode
+class UMotionTrailEditorMode : public UBaseLegacyWidgetEdMode
 {
 	GENERATED_BODY()
 public:
@@ -91,13 +91,10 @@ public:
 	
 	virtual void Render(const FSceneView* View, FViewport* Viewport, FPrimitiveDrawInterface* PDI) override;
 	virtual void DrawHUD(FEditorViewportClient* ViewportClient, FViewport* Viewport, const FSceneView* View, FCanvas* Canvas) override;
-	
-	bool UsesToolkits() const override;
 
 	virtual TMap<FName, TArray<TSharedPtr<FUICommandInfo>>> GetModeCommands() const override;
 	virtual void ActivateDefaultTool() override;
 
-	virtual class FEdMode* AsLegacyMode() override;
 	virtual bool IsCompatibleWith(FEditorModeID OtherModeID) const override;
 	// End of UEdMode interface
 
diff --git a/Engine/Plugins/MeshPainting/Source/MeshPaintEditorMode/Private/MeshPaintMode.cpp b/Engine/Plugins/MeshPainting/Source/MeshPaintEditorMode/Private/MeshPaintMode.cpp
--- a/Engine/Plugins/MeshPainting/Source/MeshPaintEditorMode/Private/MeshPaintMode.cpp
+++ b/Engine/Plugins/MeshPainting/Source/MeshPaintEditorMode/Private/MeshPaintMode.cpp
@@ -205,8 +205,6 @@ void UMeshPaintMode::CreateToolkit()
 
 void UMeshPaintMode::Tick(FEditorViewportClient* ViewportClient, float DeltaTime)
 {
-	UEdMode::Tick(ViewportClient, DeltaTime);
-
 	if (ViewportClient->IsPerspective())
 	{
 		// Make sure perspective viewports are still set to real-time
diff --git a/Engine/Plugins/MeshPainting/Source/MeshPaintEditorMode/Private/MeshPaintMode.h b/Engine/Plugins/MeshPainting/Source/MeshPaintEditorMode/Private/MeshPaintMode.h
--- a/Engine/Plugins/MeshPainting/Source/MeshPaintEditorMode/Private/MeshPaintMode.h
+++ b/Engine/Plugins/MeshPainting/Source/MeshPaintEditorMode/Private/MeshPaintMode.h
@@ -14,6 +14,7 @@
 #include "IMeshPaintComponentAdapter.h"
 #include "MeshPaintHelpers.h"
 #include "MeshVertexPaintingTool.h"
+#include "Tools/LegacyEdModeInterfaces.h"
 #include "MeshPaintMode.generated.h"
 
 class UMeshVertexPaintingToolProperties;
@@ -28,7 +29,7 @@ class UMeshToolManager;
  * Mesh paint Mode.  Extends editor viewports with the ability to paint data on meshes
  */
 UCLASS()
-class UMeshPaintMode : public UEdMode
+class UMeshPaintMode : public UEdMode, public ILegacyEdModeViewportInterface
 {
 public:
 	GENERATED_BODY()
diff --git a/Engine/Source/Editor/UnrealEd/Private/EdMode.cpp b/Engine/Source/Editor/UnrealEd/Private/EdMode.cpp
--- a/Engine/Source/Editor/UnrealEd/Private/EdMode.cpp
+++ b/Engine/Source/Editor/UnrealEd/Private/EdMode.cpp
@@ -105,27 +105,6 @@ bool FEdMode::InputKey(FEditorViewportClient* ViewportClient, FViewport* Viewpor
 	{
 		return true;
 	}
-	else
-	{
-		// Next pass input to the mode toolkit
-		if (Toolkit.IsValid() && ((Event == IE_Pressed) || (Event == IE_Repeat)))
-		{
-			if (Toolkit->GetToolkitCommands()->ProcessCommandBindings(Key, FSlateApplication::Get().GetModifierKeys(), (Event == IE_Repeat)))
-			{
-				return true;
-			}
-		}
-
-		// Finally, pass input up to selected actors if not in a tool mode
-		TArray<AActor*> SelectedActors;
-		Owner->GetSelectedActors()->GetSelectedObjects<AActor>(SelectedActors);
-
-		for (TArray<AActor*>::TIterator It(SelectedActors); It; ++It)
-		{
-			// Tell the object we've had a key press
-			(*It)->EditorKeyPressed(Key, Event);
-		}
-	}
 
 	return false;
 }
@@ -290,130 +269,17 @@ void FEdMode::Render(const FSceneView* View,FViewport* Viewport,FPrimitiveDrawIn
 
 void FEdMode::DrawHUD(FEditorViewportClient* ViewportClient,FViewport* Viewport,const FSceneView* View,FCanvas* Canvas)
 {
-	// Render the drag tool.
-	ViewportClient->RenderDragTool( View, Canvas );
-
 	// Let the current mode tool draw a HUD if it wants to
-	FModeTool* tool = GetCurrentTool();
-	if( tool )
-	{
-		tool->DrawHUD( ViewportClient, Viewport, View, Canvas );
-	}
-
-	if (ViewportClient->IsPerspective() && GetDefault<ULevelEditorViewportSettings>()->bHighlightWithBrackets)
+	FModeTool* Tool = GetCurrentTool();
+	if( Tool )
 	{
-		DrawBrackets( ViewportClient, Viewport, View, Canvas );
+		Tool->DrawHUD( ViewportClient, Viewport, View, Canvas );
 	}
 
 	// If this viewport doesn't show mode widgets or the mode itself doesn't want them, leave.
-	if( !(ViewportClient->EngineShowFlags.ModeWidgets) || !ShowModeWidgets() )
-	{
-		return;
-	}
-
-	// Clear Hit proxies
-	const bool bIsHitTesting = Canvas->IsHitTesting();
-	if ( !bIsHitTesting )
-	{
-		Canvas->SetHitProxy(NULL);
-	}
-
-	// Draw vertices for selected BSP brushes and static meshes if the large vertices show flag is set.
-	if ( !ViewportClient->bDrawVertices )
+	if (ShowModeWidgets())
 	{
-		return;
-	}
-
-	const bool bLargeVertices		= View->Family->EngineShowFlags.LargeVertices;
-	const bool bShowBrushes			= View->Family->EngineShowFlags.Brushes;
-	const bool bShowBSP				= View->Family->EngineShowFlags.BSP;
-	const bool bShowBuilderBrush	= View->Family->EngineShowFlags.BuilderBrush != 0;
-
-	UTexture2D* VertexTexture = GetVertexTexture();
-	const float TextureSizeX		= VertexTexture->GetSizeX() * ( bLargeVertices ? 1.0f : 0.5f );
-	const float TextureSizeY		= VertexTexture->GetSizeY() * ( bLargeVertices ? 1.0f : 0.5f );
-
-	// Temporaries.
-	TArray<FVector> Vertices;
-
-	for ( FSelectionIterator It( *Owner->GetSelectedActors() ) ; It ; ++It )
-	{
-		AActor* SelectedActor = static_cast<AActor*>( *It );
-		checkSlow( SelectedActor->IsA(AActor::StaticClass()) );
-
-		if( bLargeVertices )
-		{
-			FCanvasItemTestbed::bTestState = !FCanvasItemTestbed::bTestState;
-
-			// Static mesh vertices
-			AStaticMeshActor* Actor = Cast<AStaticMeshActor>( SelectedActor );
-			if( Actor && Actor->GetStaticMeshComponent() && Actor->GetStaticMeshComponent()->GetStaticMesh()
-				&& Actor->GetStaticMeshComponent()->GetStaticMesh()->GetRenderData())
-			{
-				FTransform ActorToWorld = Actor->ActorToWorld();
-				Vertices.Empty();
-				const FPositionVertexBuffer& VertexBuffer = Actor->GetStaticMeshComponent()->GetStaticMesh()->GetRenderData()->LODResources[0].VertexBuffers.PositionVertexBuffer;
-				for( uint32 i = 0 ; i < VertexBuffer.GetNumVertices() ; i++ )
-				{
-					Vertices.AddUnique( ActorToWorld.TransformPosition( VertexBuffer.VertexPosition(i) ) );
-				}
-
-				const float InvDpiScale = 1.0f / Canvas->GetDPIScale();
-
-				FCanvasTileItem TileItem( FVector2D( 0.0f, 0.0f ), FVector2D( 0.0f, 0.0f ), FLinearColor::White );
-				TileItem.BlendMode = SE_BLEND_Translucent;
-				for( int32 VertexIndex = 0 ; VertexIndex < Vertices.Num() ; ++VertexIndex )
-				{				
-					const FVector& Vertex = Vertices[VertexIndex];
-					FVector2D PixelLocation;
-					if(View->ScreenToPixel(View->WorldToScreen(Vertex),PixelLocation))
-					{
-						PixelLocation *= InvDpiScale;
-
-						const bool bOutside =
-							PixelLocation.X < 0.0f || PixelLocation.X > View->UnscaledViewRect.Width()*InvDpiScale ||
-							PixelLocation.Y < 0.0f || PixelLocation.Y > View->UnscaledViewRect.Height()*InvDpiScale;
-						if( !bOutside )
-						{
-							const float X = PixelLocation.X - (TextureSizeX/2);
-							const float Y = PixelLocation.Y - (TextureSizeY/2);
-							if( bIsHitTesting ) 
-							{
-								Canvas->SetHitProxy( new HStaticMeshVert(Actor,Vertex) );
-							}
-							TileItem.Texture = VertexTexture->Resource;
-							
-							TileItem.Size = FVector2D( TextureSizeX, TextureSizeY );
-							Canvas->DrawItem( TileItem, FVector2D( X, Y ) );							
-							if( bIsHitTesting )
-							{
-								Canvas->SetHitProxy( NULL );
-							}
-						}
-					}
-				}
-			}
-		}
-	}
-
-	FLegacyEdModeWidgetHelper::DrawHUD(ViewportClient, Viewport, View, Canvas);
-}
-
-void FEdMode::DrawBrackets( FEditorViewportClient* ViewportClient, FViewport* Viewport, const FSceneView* View, FCanvas* Canvas )
-{
-	USelection& SelectedActors = *Owner->GetSelectedActors();
-	for( int32 CurSelectedActorIndex = 0; CurSelectedActorIndex < SelectedActors.Num(); ++CurSelectedActorIndex )
-	{
-		AActor* SelectedActor = Cast<AActor>( SelectedActors.GetSelectedObject(CurSelectedActorIndex ) );
-		if( SelectedActor != NULL )
-		{
-			// Draw a bracket for selected "paintable" static mesh actors
-			const bool bIsValidActor = ( Cast< AStaticMeshActor >( SelectedActor ) != NULL );
-
-			const FLinearColor SelectedActorBoxColor( 0.6f, 0.6f, 1.0f );
-			const bool bDrawBracket = bIsValidActor;
-			ViewportClient->DrawActorScreenSpaceBoundingBox( Canvas, View, Viewport, SelectedActor, SelectedActorBoxColor, bDrawBracket );
-		}
+		FLegacyEdModeWidgetHelper::DrawHUD(ViewportClient, Viewport, View, Canvas);
 	}
 }
 
diff --git a/Engine/Source/Editor/UnrealEd/Private/EditorModeManager.cpp b/Engine/Source/Editor/UnrealEd/Private/EditorModeManager.cpp
--- a/Engine/Source/Editor/UnrealEd/Private/EditorModeManager.cpp
+++ b/Engine/Source/Editor/UnrealEd/Private/EditorModeManager.cpp
@@ -36,6 +36,10 @@
 #include "InteractiveGizmoManager.h"
 #include "EdModeInteractiveToolsContext.h"
 #include "Tools/LegacyEdModeInterfaces.h"
+#include "CanvasTypes.h"
+#include "CanvasItem.h"
+#include "Engine/StaticMeshActor.h"
+#include "EngineUtils.h"
 
 /*------------------------------------------------------------------------------
 	FEditorModeTools.
@@ -240,10 +244,11 @@ bool FEditorModeTools::IsSelectionAllowed(AActor* InActor, const bool bInSelecte
 bool FEditorModeTools::IsSelectionHandled(AActor* InActor, const bool bInSelected) const
 {
 	bool bSelectionHandled = false;
-	for (UEdMode* Mode : ActiveScriptableModes)
-	{
-		bSelectionHandled |= Mode->Select(InActor, bInSelected);
-	}
+	ForEachEdMode([&bSelectionHandled, bInSelected, InActor](UEdMode* Mode)
+		{
+			bSelectionHandled |= Mode->Select(InActor, bInSelected);
+			return true;
+		});
 
 	return bSelectionHandled;
 }
@@ -251,10 +256,11 @@ bool FEditorModeTools::IsSelectionHandled(AActor* InActor, const bool bInSelecte
 bool FEditorModeTools::ProcessEditDuplicate()
 {
 	bool bHandled = false;
-	for (UEdMode* Mode : ActiveScriptableModes)
-	{
-		bHandled = Mode->ProcessEditDuplicate();
-	}
+	ForEachEdMode([&bHandled](UEdMode* Mode)
+		{
+			bHandled |= Mode->ProcessEditDuplicate();
+			return true;
+		});
 
 	return bHandled;
 }
@@ -262,121 +268,139 @@ bool FEditorModeTools::ProcessEditDuplicate()
 bool FEditorModeTools::ProcessEditDelete()
 {
 	bool bHandled = InteractiveToolsContext->ProcessEditDelete();
-	for (UEdMode* Mode : ActiveScriptableModes)
-	{
-		bHandled = Mode->ProcessEditDelete();
-	}
+	ForEachEdMode([&bHandled](UEdMode* Mode)
+		{
+			bHandled |= Mode->ProcessEditDelete();
+			return true;
+		});
 
 	return bHandled;
 }
 
 bool FEditorModeTools::ProcessEditCut()
 {
-	for (UEdMode* Mode : ActiveScriptableModes)
-	{
-		if (Mode->ProcessEditCut())
+	bool bHandled = false;
+	ForEachEdMode([&bHandled](UEdMode* Mode)
 		{
-			return true;
-		}
-	}
+			bHandled = Mode->ProcessEditCut();
+			return !bHandled;
+		});
 
-	return false;
+	return bHandled;
 }
 
 bool FEditorModeTools::ProcessEditCopy()
 {
-	for (UEdMode* Mode : ActiveScriptableModes)
-	{
-		if (Mode->ProcessEditCopy())
+	bool bHandled = false;
+	ForEachEdMode([&bHandled](UEdMode* Mode)
 		{
-			return true;
-		}
-	}
+			bHandled = Mode->ProcessEditCopy();
+			return !bHandled;
+		});
 
-	return false;
+	return bHandled;
 }
 
 bool FEditorModeTools::ProcessEditPaste()
 {
-	for (UEdMode* Mode : ActiveScriptableModes)
-	{
-		if (Mode->ProcessEditPaste())
+	bool bHandled = false;
+	ForEachEdMode([&bHandled](UEdMode* Mode)
 		{
-			return true;
-		}
-	}
+			bHandled = Mode->ProcessEditPaste();
+			return !bHandled;
+		});
 
-	return false;
+	return bHandled;
 }
 
 EEditAction::Type FEditorModeTools::GetActionEditDuplicate()
 {
-	for (UEdMode* Mode : ActiveScriptableModes)
-	{
-		const EEditAction::Type EditAction = Mode->GetActionEditDuplicate();
-		if (EditAction == EEditAction::Process || EditAction == EEditAction::Halt)
+	EEditAction::Type ReturnedAction = EEditAction::Skip;
+	ForEachEdMode([&ReturnedAction](UEdMode* Mode)
 		{
-			return EditAction;
-		}
-	}
+			const EEditAction::Type EditAction = Mode->GetActionEditDuplicate();
+			if (EditAction == EEditAction::Process || EditAction == EEditAction::Halt)
+			{
+				ReturnedAction = EditAction;
+				return false;
+			}
+
+			return true;
+		});
 
-	return EEditAction::Skip;
+	return ReturnedAction;
 }
 
 EEditAction::Type FEditorModeTools::GetActionEditDelete()
 {
-	for (UEdMode* Mode : ActiveScriptableModes)
+	EEditAction::Type ReturnedAction = EEditAction::Skip;
+	ForEachEdMode([&ReturnedAction](UEdMode* Mode)
 	{
 		const EEditAction::Type EditAction = Mode->GetActionEditDelete();
 		if (EditAction == EEditAction::Process || EditAction == EEditAction::Halt)
 		{
-			return EditAction;
+			ReturnedAction = EditAction;
+			return false;
 		}
-	}
 
-	return EEditAction::Skip;
+		return true;
+	});
+
+	return ReturnedAction;
 }
 
 EEditAction::Type FEditorModeTools::GetActionEditCut()
 {
-	for (UEdMode* Mode : ActiveScriptableModes)
-	{
-		const EEditAction::Type EditAction = Mode->GetActionEditCut();
-		if (EditAction == EEditAction::Process || EditAction == EEditAction::Halt)
+	EEditAction::Type ReturnedAction = EEditAction::Skip;
+	ForEachEdMode([&ReturnedAction](UEdMode* Mode)
 		{
-			return EditAction;
-		}
-	}
+			const EEditAction::Type EditAction = Mode->GetActionEditCut();
+			if (EditAction == EEditAction::Process || EditAction == EEditAction::Halt)
+			{
+				ReturnedAction = EditAction;
+				return false;
+			}
+
+			return true;
+		});
 
-	return EEditAction::Skip;
+	return ReturnedAction;
 }
 
 EEditAction::Type FEditorModeTools::GetActionEditCopy()
 {
-	for (UEdMode* Mode : ActiveScriptableModes)
-	{
-		const EEditAction::Type EditAction = Mode->GetActionEditCopy();
-		if (EditAction == EEditAction::Process || EditAction == EEditAction::Halt)
+	EEditAction::Type ReturnedAction = EEditAction::Skip;
+	ForEachEdMode([&ReturnedAction](UEdMode* Mode)
 		{
-			return EditAction;
-		}
-	}
+			const EEditAction::Type EditAction = Mode->GetActionEditCopy();
+			if (EditAction == EEditAction::Process || EditAction == EEditAction::Halt)
+			{
+				ReturnedAction = EditAction;
+				return false;
+			}
+
+			return true;
+		});
 
-	return EEditAction::Skip;
+	return ReturnedAction;
 }
 
 EEditAction::Type FEditorModeTools::GetActionEditPaste()
 {
-	for (UEdMode* Mode : ActiveScriptableModes)
-	{
-		const EEditAction::Type EditAction = Mode->GetActionEditPaste();
-		if (EditAction == EEditAction::Process || EditAction == EEditAction::Halt)
+	EEditAction::Type ReturnedAction = EEditAction::Skip;
+	ForEachEdMode([&ReturnedAction](UEdMode* Mode)
 		{
-			return EditAction;
-		}
-	}
+			const EEditAction::Type EditAction = Mode->GetActionEditPaste();
+			if (EditAction == EEditAction::Process || EditAction == EEditAction::Halt)
+			{
+				ReturnedAction = EditAction;
+				return false;
+			}
+
+			return true;
+		});
 
-	return EEditAction::Skip;
+	return ReturnedAction;
 }
 
 void FEditorModeTools::DeactivateOtherVisibleModes(FEditorModeID InMode)
@@ -393,88 +417,90 @@ void FEditorModeTools::DeactivateOtherVisibleModes(FEditorModeID InMode)
 
 bool FEditorModeTools::IsSnapRotationEnabled() const
 {
-	for (UEdMode* Mode : ActiveScriptableModes)
-	{
-		if (Mode->IsSnapRotationEnabled())
+	bool bRetVal = false;
+	ForEachEdMode([&bRetVal](UEdMode* Mode)
 		{
-			return true;
-		}
-	}
+			bRetVal = Mode->IsSnapRotationEnabled();
+			return !bRetVal;
+		});
 
-	return false;
+	return bRetVal;
 }
 
 bool FEditorModeTools::SnapRotatorToGridOverride(FRotator& InRotation) const
 {
-	for (UEdMode* Mode : ActiveScriptableModes)
-	{
-		if (Mode->SnapRotatorToGridOverride(InRotation))
+	bool bRetVal = false;
+	ForEachEdMode([&bRetVal, &InRotation](UEdMode* Mode)
 		{
-			return true;
-		}
-	}
+			bRetVal = Mode->SnapRotatorToGridOverride(InRotation);
+			return !bRetVal;
+		});
 
-	return false;
+	return bRetVal;
 }
 
 void FEditorModeTools::ActorsDuplicatedNotify(TArray<AActor*>& InPreDuplicateSelection, TArray<AActor*>& InPostDuplicateSelection, const bool bOffsetLocations)
 {
-	for (UEdMode* Mode : ActiveScriptableModes)
+	ForEachEdMode([&InPreDuplicateSelection, &InPostDuplicateSelection, bOffsetLocations](UEdMode* Mode)
 	{
 		// Tell the tools about the duplication
 		Mode->ActorsDuplicatedNotify(InPreDuplicateSelection, InPostDuplicateSelection, bOffsetLocations);
-	}
+		return true;
+	});
 }
 
 void FEditorModeTools::ActorMoveNotify()
 {
-	for (UEdMode* Mode : ActiveScriptableModes)
+	ForEachEdMode([](UEdMode* Mode)
 	{
 		// Also notify the current editing modes if they are interested.
 		Mode->ActorMoveNotify();
-	}
+		return true;
+	});
 }
 
 void FEditorModeTools::ActorSelectionChangeNotify()
 {
-	for (UEdMode* Mode : ActiveScriptableModes)
+	ForEachEdMode([](UEdMode* Mode)
 	{
 		Mode->ActorSelectionChangeNotify();
-	}
+		return true;
+	});
 }
 
 void FEditorModeTools::ActorPropChangeNotify()
 {
-	for (UEdMode* Mode : ActiveScriptableModes)
+	ForEachEdMode([](UEdMode* Mode)
 	{
 		Mode->ActorPropChangeNotify();
-	}
+		return true;
+	});
 }
 
 void FEditorModeTools::UpdateInternalData()
 {
-	for (UEdMode* Mode : ActiveScriptableModes)
+	ForEachEdMode([](UEdMode* Mode)
 	{
 		Mode->UpdateInternalData();
-	}
+		return true;
+	});
 }
 
 bool FEditorModeTools::IsOnlyVisibleActiveMode(FEditorModeID InMode) const
 {
 	// Only return true if this is the *only* active mode
-	for (const UEdMode* Mode : ActiveScriptableModes)
+	bool bFoundAnotherVisibleMode = false;
+	ForEachEdMode([&bFoundAnotherVisibleMode, InMode](UEdMode* Mode)
 	{
-		if (Mode->GetModeInfo().IsVisible() && Mode->GetID() != InMode)
-		{
-			return false;
-		}
-	}
-	return true;
+		bFoundAnotherVisibleMode = (Mode->GetModeInfo().IsVisible() && Mode->GetID() != InMode);
+		return !bFoundAnotherVisibleMode;
+	});
+	return !bFoundAnotherVisibleMode;
 }
 
 void FEditorModeTools::OnEditorSelectionChanged(UObject* NewSelection)
 {
-	if(NewSelection == GetSelectedActors())
+	if (NewSelection == GetSelectedActors())
 	{
 		// when actors are selected check if there is at least one component selected and cache that off
 		// Editor modes use this primarily to determine of transform gizmos should be drawn.  
@@ -531,6 +557,39 @@ void FEditorModeTools::OnEditorSelectNone()
 	GEditor->ActorsThatWereSelected.Empty();
 }
 
+void FEditorModeTools::DrawBrackets(FEditorViewportClient* ViewportClient, FViewport* Viewport, const FSceneView* View, FCanvas* Canvas)
+{
+	if (!ViewportClient->IsPerspective() || !GetDefault<ULevelEditorViewportSettings>()->bHighlightWithBrackets)
+	{
+		return;
+	}
+
+	if (UTypedElementSelectionSet* CurrentSelection = GetEditorSelectionSet())
+	{
+		CurrentSelection->ForEachSelectedObject<AActor>([Canvas, View, Viewport, ViewportClient](AActor* Actor)
+			{
+				const FLinearColor SelectedActorBoxColor(0.6f, 0.6f, 1.0f);
+				const bool bDrawBracket = Actor->IsA<AStaticMeshActor>();
+				ViewportClient->DrawActorScreenSpaceBoundingBox(Canvas, View, Viewport, Actor, SelectedActorBoxColor, bDrawBracket);
+				return true;
+			});
+	}
+}
+
+void FEditorModeTools::ForEachEdMode(TFunctionRef<bool(UEdMode*)> InCalllback) const
+{
+	for (UEdMode* Mode : ActiveScriptableModes)
+	{
+		if (Mode && !Mode->IsPendingDeletion())
+		{
+			if (!InCalllback(Mode))
+			{
+				break;
+			}
+		}
+	}
+}
+
 void FEditorModeTools::SetPivotLocation( const FVector& Location, const bool bIncGridBase )
 {
 	CachedLocation = PivotLocation = SnappedLocation = Location;
@@ -821,24 +880,17 @@ void FEditorModeTools::DeactivateMode( FEditorModeID InID )
 		auto& Mode = ActiveScriptableModes[Index];
 		if (Mode->GetID() == InID)
 		{
-			DeactivateScriptableModeAtIndex(Index);
+			Mode->RequestDeletion();
 			break;
 		}
 	}
-
-
-	if ((ActiveScriptableModes.Num() == 0))
-	{
-		// Ensure the default mode is active if there are no active modes.
-		ActivateDefaultMode();
-	}
 }
 
 void FEditorModeTools::DeactivateAllModes()
 {
 	for (int32 Index = ActiveScriptableModes.Num() - 1; Index >= 0; --Index)
 	{
-		DeactivateScriptableModeAtIndex(Index);
+		ActiveScriptableModes[Index]->RequestDeletion();
 	}
 }
 
@@ -963,7 +1015,7 @@ void FEditorModeTools::ActivateMode(FEditorModeID InID, bool bToggle)
 		const bool bModesAreCompatible = ScriptableMode->IsCompatibleWith(ActiveScriptableModes[ModeIndex]->GetID()) || ActiveScriptableModes[ModeIndex]->IsCompatibleWith(ScriptableMode->GetID());
 		if (!bModesAreCompatible || (bIsVisibleMode && ActiveScriptableModes[ModeIndex]->GetModeInfo().IsVisible()))
 		{
-			DeactivateScriptableModeAtIndex(ModeIndex);
+			ActiveScriptableModes[ModeIndex]->RequestDeletion();
 		}
 	}
 
@@ -1029,20 +1081,18 @@ bool FEditorModeTools::EnsureNotInMode(FEditorModeID ModeID, const FText& ErrorM
 
 UEdMode* FEditorModeTools::GetActiveScriptableMode(FEditorModeID InID) const
 {
-	for (auto& Mode : ActiveScriptableModes)
+	if (UEdMode* const* FoundMode = ActiveScriptableModes.FindByPredicate([InID](UEdMode* Mode) { return (Mode->GetID() == InID); }))
 	{
-		if (Mode->GetID() == InID)
-		{
-			return Mode;
-		}
+		return const_cast<UEdMode*>(*FoundMode);
 	}
 
 	return nullptr;
 }
 
-/**
- * Returns a coordinate system that should be applied on top of the worldspace system.
- */
+UTexture2D* FEditorModeTools::GetVertexTexture() const
+{
+	return GEngine->DefaultBSPVertexTexture;
+}
 
 FMatrix FEditorModeTools::GetCustomDrawingCoordinateSystem()
 {
@@ -1078,15 +1128,11 @@ FMatrix FEditorModeTools::GetLocalCoordinateSystem()
 	// If it doesn't want to, create it by looking at the currently selected actors list.
 
 	bool CustomCoordinateSystemProvided = false;
-	for (UEdMode* Mode : ActiveScriptableModes)
-	{
-		ILegacyEdModeWidgetInterface* LegacyMode = Cast<ILegacyEdModeWidgetInterface>(Mode);
-		if (LegacyMode && LegacyMode->GetCustomDrawingCoordinateSystem(Matrix, nullptr))
+	ForEachEdMode<ILegacyEdModeWidgetInterface>([&Matrix, &CustomCoordinateSystemProvided](ILegacyEdModeWidgetInterface* LegacyMode)
 		{
-			CustomCoordinateSystemProvided = true;
-			break;
-		}
-	}
+			CustomCoordinateSystemProvided = LegacyMode->GetCustomDrawingCoordinateSystem(Matrix, nullptr);
+			return !CustomCoordinateSystemProvided;
+		});
 
 	if (!CustomCoordinateSystemProvided)
 	{
@@ -1138,10 +1184,11 @@ bool FEditorModeTools::StartTracking(FEditorViewportClient* InViewportClient, FV
 	CachedLocation = PivotLocation;	// Cache the pivot location
 
 	bool bTransactionHandled = InteractiveToolsContext->StartTracking(InViewportClient, InViewport);
-	for (UEdMode* Mode : ActiveScriptableModes)
-	{
-		bTransactionHandled |= Mode->StartTracking(InViewportClient, InViewportClient->Viewport);
-	}
+	ForEachEdMode<ILegacyEdModeViewportInterface>([&bTransactionHandled, InViewportClient, InViewport](ILegacyEdModeViewportInterface* ViewportInterface)
+		{
+			bTransactionHandled |= ViewportInterface->StartTracking(InViewportClient, InViewport);
+			return true;
+		});
 
 	return bTransactionHandled;
 }
@@ -1152,10 +1199,11 @@ bool FEditorModeTools::EndTracking(FEditorViewportClient* InViewportClient, FVie
 	bIsTracking = false;
 	bool bTransactionHandled = InteractiveToolsContext->EndTracking(InViewportClient, InViewport);
 
-	for (UEdMode* Mode : ActiveScriptableModes)
-	{
-		bTransactionHandled |= Mode->EndTracking(InViewportClient, InViewportClient->Viewport);
-	}
+	ForEachEdMode<ILegacyEdModeViewportInterface>([&bTransactionHandled, InViewportClient, InViewport](ILegacyEdModeViewportInterface* ViewportInterface)
+		{
+			bTransactionHandled |= ViewportInterface->EndTracking(InViewportClient, InViewport);
+			return true;
+		});
 
 	CachedLocation = PivotLocation;	// Clear the pivot location
 	
@@ -1165,45 +1213,44 @@ bool FEditorModeTools::EndTracking(FEditorViewportClient* InViewportClient, FVie
 bool FEditorModeTools::AllowsViewportDragTool() const
 {
 	bool bCanUseDragTool = false;
-	for (UEdMode* Mode : ActiveScriptableModes)
-	{
-		if (ILegacyEdModeWidgetInterface* LegacyMode = Cast<ILegacyEdModeWidgetInterface>(Mode))
+	ForEachEdMode<const ILegacyEdModeViewportInterface>([&bCanUseDragTool](const ILegacyEdModeViewportInterface* LegacyMode)
 		{
 			bCanUseDragTool |= LegacyMode->AllowsViewportDragTool();
-		}
-	}
+			return true;
+		});
+
 	return bCanUseDragTool;
 }
 
 /** Notifies all active modes that a map change has occured */
 void FEditorModeTools::MapChangeNotify()
 {
-	for (UEdMode* Mode : ActiveScriptableModes)
-	{
-		Mode->MapChangeNotify();
-	}
+	ForEachEdMode([](UEdMode* Mode)
+		{
+			Mode->MapChangeNotify();
+			return true;
+		});
 }
 
 /** Notifies all active modes to empty their selections */
 void FEditorModeTools::SelectNone()
 {
-	for (UEdMode* Mode : ActiveScriptableModes)
-	{
-		Mode->SelectNone();
-	}
+	ForEachEdMode([](UEdMode* Mode)
+		{
+			Mode->SelectNone();
+			return true;
+		});
 }
 
 /** Notifies all active modes of box selection attempts */
 bool FEditorModeTools::BoxSelect( FBox& InBox, bool InSelect )
 {
 	bool bHandled = false;
-	for (UEdMode* Mode : ActiveScriptableModes)
-	{
-		if (ILegacyEdModeSelectInterface* LegacyMode = Cast<ILegacyEdModeSelectInterface>(Mode))
+	ForEachEdMode<ILegacyEdModeSelectInterface>([&bHandled, &InBox, InSelect](ILegacyEdModeSelectInterface* LegacyMode)
 		{
 			bHandled |= LegacyMode->BoxSelect(InBox, InSelect);
-		}
-	}
+			return true;
+		});
 	return bHandled;
 }
 
@@ -1211,13 +1258,11 @@ bool FEditorModeTools::BoxSelect( FBox& InBox, bool InSelect )
 bool FEditorModeTools::FrustumSelect( const FConvexVolume& InFrustum, FEditorViewportClient* InViewportClient, bool InSelect )
 {
 	bool bHandled = false;
-	for (UEdMode* Mode : ActiveScriptableModes)
-	{
-		if (ILegacyEdModeSelectInterface* LegacyMode = Cast<ILegacyEdModeSelectInterface>(Mode))
+	ForEachEdMode<ILegacyEdModeSelectInterface>([&bHandled, InFrustum, InViewportClient, InSelect](ILegacyEdModeSelectInterface* LegacyMode)
 		{
 			bHandled |= LegacyMode->FrustumSelect(InFrustum, InViewportClient, InSelect);
-		}
-	}
+			return true;
+		});
 	return bHandled;
 }
 
@@ -1226,13 +1271,11 @@ bool FEditorModeTools::FrustumSelect( const FConvexVolume& InFrustum, FEditorVie
 bool FEditorModeTools::UsesTransformWidget() const
 {
 	bool bUsesTransformWidget = false;
-	for (UEdMode* Mode : ActiveScriptableModes)
-	{
-		if (ILegacyEdModeWidgetInterface* LegacyMode = Cast<ILegacyEdModeWidgetInterface>(Mode))
+	ForEachEdMode<const ILegacyEdModeWidgetInterface>([&bUsesTransformWidget](const ILegacyEdModeWidgetInterface* LegacyMode)
 		{
 			bUsesTransformWidget |= LegacyMode->UsesTransformWidget();
-		}
-	}
+			return true;
+		});
 
 	return bUsesTransformWidget;
 }
@@ -1241,13 +1284,11 @@ bool FEditorModeTools::UsesTransformWidget() const
 bool FEditorModeTools::UsesTransformWidget( UE::Widget::EWidgetMode CheckMode ) const
 {
 	bool bUsesTransformWidget = false;
-	for (UEdMode* Mode : ActiveScriptableModes)
-	{
-		if (ILegacyEdModeWidgetInterface* LegacyMode = Cast<ILegacyEdModeWidgetInterface>(Mode))
+	ForEachEdMode<const ILegacyEdModeWidgetInterface>([&bUsesTransformWidget, CheckMode](const ILegacyEdModeWidgetInterface* LegacyMode)
 		{
 			bUsesTransformWidget |= LegacyMode->UsesTransformWidget(CheckMode);
-		}
-	}
+			return true;
+		});
 
 	return bUsesTransformWidget;
 }
@@ -1255,23 +1296,22 @@ bool FEditorModeTools::UsesTransformWidget( UE::Widget::EWidgetMode CheckMode )
 /** Sets the current widget axis */
 void FEditorModeTools::SetCurrentWidgetAxis( EAxisList::Type NewAxis )
 {
-	for (UEdMode* Mode : ActiveScriptableModes)
-	{
-		if (ILegacyEdModeWidgetInterface* LegacyMode = Cast<ILegacyEdModeWidgetInterface>(Mode))
+	ForEachEdMode<ILegacyEdModeWidgetInterface>([NewAxis](ILegacyEdModeWidgetInterface* LegacyMode)
 		{
 			LegacyMode->SetCurrentWidgetAxis(NewAxis);
-		}
-	}
+			return true;
+		});
 }
 
 /** Notifies all active modes of mouse click messages. */
 bool FEditorModeTools::HandleClick(FEditorViewportClient* InViewportClient,  HHitProxy *HitProxy, const FViewportClick& Click )
 {
 	bool bHandled = false;
-	for (UEdMode* Mode : ActiveScriptableModes)
-	{
-		bHandled |= Mode->HandleClick(InViewportClient, HitProxy, Click);
-	}
+	ForEachEdMode<ILegacyEdModeViewportInterface>([&bHandled, InViewportClient, HitProxy, Click](ILegacyEdModeViewportInterface* Mode)
+		{
+			bHandled |= Mode->HandleClick(InViewportClient, HitProxy, Click);
+			return true;
+		});
 
 	return bHandled;
 }
@@ -1330,7 +1370,6 @@ void FEditorModeTools::Tick( FEditorViewportClient* ViewportClient, float DeltaT
 		}
 	}
 
-
 	if (ActiveScriptableModes.Num() == 0)
 	{
 		// Ensure the default mode is active if there are no active modes.
@@ -1338,20 +1377,27 @@ void FEditorModeTools::Tick( FEditorViewportClient* ViewportClient, float DeltaT
 	}
 
 	InteractiveToolsContext->Tick(ViewportClient, DeltaTime);
-	for (UEdMode* Mode : ActiveScriptableModes)
-	{
-		Mode->Tick(ViewportClient, DeltaTime);
-	}
+	ForEachEdMode([ViewportClient, DeltaTime](UEdMode* Mode)
+		{
+			if (ILegacyEdModeViewportInterface* ViewportInterface = Cast<ILegacyEdModeViewportInterface>(Mode))
+			{
+				ViewportInterface->Tick(ViewportClient, DeltaTime);
+			}
+
+			Mode->ModeTick(DeltaTime);
+			return true;
+		});
 }
 
 /** Notifies all active modes of any change in mouse movement */
 bool FEditorModeTools::InputDelta( FEditorViewportClient* InViewportClient,FViewport* InViewport,FVector& InDrag,FRotator& InRot,FVector& InScale )
 {
 	bool bHandled = false;
-	for (UEdMode* Mode : ActiveScriptableModes)
-	{
-		bHandled |= Mode->InputDelta(InViewportClient, InViewport, InDrag, InRot, InScale);
-	}
+	ForEachEdMode<ILegacyEdModeViewportInterface>([&bHandled, InViewportClient, InViewport, &InDrag, &InRot, &InScale](ILegacyEdModeViewportInterface* Mode)
+		{
+			bHandled |= Mode->InputDelta(InViewportClient, InViewport, InDrag, InRot, InScale);
+			return true;
+		});
 	return bHandled;
 }
 
@@ -1359,10 +1405,11 @@ bool FEditorModeTools::InputDelta( FEditorViewportClient* InViewportClient,FView
 bool FEditorModeTools::CapturedMouseMove( FEditorViewportClient* InViewportClient, FViewport* InViewport, int32 InMouseX, int32 InMouseY )
 {
 	bool bHandled = InteractiveToolsContext->CapturedMouseMove(InViewportClient, InViewport, InMouseX, InMouseY);
-	for (UEdMode* Mode : ActiveScriptableModes)
-	{
-		bHandled |= Mode->CapturedMouseMove(InViewportClient, InViewport, InMouseX, InMouseY);
-	}
+	ForEachEdMode<ILegacyEdModeViewportInterface>([&bHandled, InViewportClient, InViewport, InMouseX, InMouseY](ILegacyEdModeViewportInterface* Mode)
+		{
+			bHandled |= Mode->CapturedMouseMove(InViewportClient, InViewport, InMouseX, InMouseY);
+			return true;
+		});
 	return bHandled;
 }
 
@@ -1370,10 +1417,11 @@ bool FEditorModeTools::CapturedMouseMove( FEditorViewportClient* InViewportClien
 bool FEditorModeTools::ProcessCapturedMouseMoves( FEditorViewportClient* InViewportClient, FViewport* InViewport, const TArrayView<FIntPoint>& CapturedMouseMoves )
 {
 	bool bHandled = false;
-	for (UEdMode* Mode : ActiveScriptableModes)
-	{
-		bHandled |= Mode->ProcessCapturedMouseMoves(InViewportClient, InViewport, CapturedMouseMoves);
-	}
+	ForEachEdMode<ILegacyEdModeViewportInterface>([&bHandled, InViewportClient, InViewport, &CapturedMouseMoves](ILegacyEdModeViewportInterface* Mode)
+		{
+			bHandled |= Mode->ProcessCapturedMouseMoves(InViewportClient, InViewport, CapturedMouseMoves);
+			return true;
+		});
 	return bHandled;
 }
 
@@ -1386,12 +1434,27 @@ bool FEditorModeTools::InputKey(FEditorViewportClient* InViewportClient, FViewpo
 		bHandled |= InteractiveToolsContext->InputKey(InViewportClient, Viewport, Key, Event);
 	}
 	
+
 	// Copy the modes and iterate of that since a key may remove the edit mode and change ActiveScriptableModes
 	TArray<UEdMode*> CopyActiveScriptableModes(ActiveScriptableModes);
 	for (UEdMode* Mode : CopyActiveScriptableModes)
 	{
-		bHandled |= Mode->InputKey(InViewportClient, Viewport, Key, Event);
+		if (ILegacyEdModeViewportInterface* ViewportInterface = Cast<ILegacyEdModeViewportInterface>(Mode))
+		{
+			bHandled |= ViewportInterface->InputKey(InViewportClient, Viewport, Key, Event);
+		}
 	}
+
+	// Finally, pass input up to selected actors if not in a tool mode
+	GetEditorSelectionSet()->ForEachSelectedObject<AActor>([Key, Event](AActor* ActorPtr)
+		{
+			if (ActorPtr)
+			{
+				// Tell the object we've had a key press
+				ActorPtr->EditorKeyPressed(Key, Event);
+			}
+			return true;
+		});
 	return bHandled;
 }
 
@@ -1399,34 +1462,40 @@ bool FEditorModeTools::InputKey(FEditorViewportClient* InViewportClient, FViewpo
 bool FEditorModeTools::InputAxis(FEditorViewportClient* InViewportClient, FViewport* Viewport, int32 ControllerId, FKey Key, float Delta, float DeltaTime)
 {
 	bool bHandled = false;
-	for (UEdMode* Mode : ActiveScriptableModes)
-	{
-		bHandled |= Mode->InputAxis(InViewportClient, Viewport, ControllerId, Key, Delta, DeltaTime);
-	}
+
+	ForEachEdMode<ILegacyEdModeViewportInterface>([&bHandled, InViewportClient, Viewport, ControllerId, Key, Delta, DeltaTime](ILegacyEdModeViewportInterface* Mode)
+		{
+			bHandled |= Mode->InputAxis(InViewportClient, Viewport, ControllerId, Key, Delta, DeltaTime);
+			return true;
+		});
+
 	return bHandled;
 }
 
 bool FEditorModeTools::GetPivotForOrbit( FVector& Pivot ) const
 {
+	bool bHandled = false;
 	// Just return the first pivot point specified by a mode
-	for (UEdMode* Mode : ActiveScriptableModes)
-	{
-		if (Mode->GetPivotForOrbit(Pivot))
+	ForEachEdMode([&Pivot, &bHandled](const UEdMode* Mode)
 		{
-			return true;
-		}
-	}
-	return false;
+			bHandled = Mode->GetPivotForOrbit(Pivot);
+			return !bHandled;
+		});
+		
+	return bHandled;
 }
 
 bool FEditorModeTools::MouseEnter( FEditorViewportClient* InViewportClient, FViewport* Viewport, int32 X, int32 Y )
 {
 	HoveredViewportClient = InViewportClient;
-	bool bHandled = InteractiveToolsContext->MouseEnter(InViewportClient, Viewport, X, Y);;
-	for (UEdMode* Mode : ActiveScriptableModes)
-	{
-		bHandled |= Mode->MouseEnter(InViewportClient, Viewport, X, Y);
-	}
+	bool bHandled = InteractiveToolsContext->MouseEnter(InViewportClient, Viewport, X, Y);
+
+	ForEachEdMode<ILegacyEdModeViewportInterface>([&bHandled, InViewportClient, Viewport, X, Y](ILegacyEdModeViewportInterface* Mode)
+		{
+			bHandled |= Mode->MouseEnter(InViewportClient, Viewport, X, Y);
+			return true;
+		});
+
 	return bHandled;
 }
 
@@ -1434,10 +1503,12 @@ bool FEditorModeTools::MouseLeave( FEditorViewportClient* InViewportClient, FVie
 {
 	bool bHandled = InteractiveToolsContext->MouseLeave(InViewportClient, Viewport);
 
-	for (UEdMode* Mode : ActiveScriptableModes)
-	{
-		bHandled |= Mode->MouseLeave(InViewportClient, Viewport);
-	}
+	ForEachEdMode<ILegacyEdModeViewportInterface>([&bHandled, InViewportClient, Viewport](ILegacyEdModeViewportInterface* Mode)
+		{
+			bHandled |= Mode->MouseLeave(InViewportClient, Viewport);
+			return true;
+		});
+
 	return bHandled;
 }
 
@@ -1446,10 +1517,12 @@ bool FEditorModeTools::MouseMove( FEditorViewportClient* InViewportClient, FView
 {
 	bool bHandled = InteractiveToolsContext->MouseMove(InViewportClient, Viewport, X, Y);
 
-	for (UEdMode* Mode : ActiveScriptableModes)
-	{
-		bHandled |= Mode->MouseMove(InViewportClient, Viewport, X, Y);
-	}
+	ForEachEdMode<ILegacyEdModeViewportInterface>([&bHandled, InViewportClient, Viewport, X, Y](ILegacyEdModeViewportInterface* Mode)
+		{
+			bHandled |= Mode->MouseMove(InViewportClient, Viewport, X, Y);
+			return true;
+		});
+
 	return bHandled;
 }
 
@@ -1457,53 +1530,149 @@ bool FEditorModeTools::ReceivedFocus( FEditorViewportClient* InViewportClient, F
 {
 	FocusedViewportClient = InViewportClient;
 	bool bHandled = false;
-	for (UEdMode* Mode : ActiveScriptableModes)
-	{
-		bHandled |= Mode->ReceivedFocus(InViewportClient, Viewport);
-	}
+	ForEachEdMode<ILegacyEdModeViewportInterface>([&bHandled, InViewportClient, Viewport](ILegacyEdModeViewportInterface* Mode)
+		{
+			bHandled |= Mode->ReceivedFocus(InViewportClient, Viewport);
+			return true;
+		});
+
 	return bHandled;
 }
 
 bool FEditorModeTools::LostFocus( FEditorViewportClient* InViewportClient, FViewport* Viewport )
 {
 	bool bHandled = false;
-	for (UEdMode* Mode : ActiveScriptableModes)
-	{
-		bHandled |= Mode->LostFocus(InViewportClient, Viewport);
-	}
+	ForEachEdMode<ILegacyEdModeViewportInterface>([&bHandled, InViewportClient, Viewport](ILegacyEdModeViewportInterface* Mode)
+		{
+			bHandled |= Mode->LostFocus(InViewportClient, Viewport);
+			return true;
+		});
+
 	return bHandled;
 }
 
 /** Draws all active mode components */	
 void FEditorModeTools::DrawActiveModes( const FSceneView* InView, FPrimitiveDrawInterface* PDI )
 {
-	for (UEdMode* Mode : ActiveScriptableModes)
-	{
-		if (ILegacyEdModeDrawHelperInterface* DrawHelper = Cast<ILegacyEdModeDrawHelperInterface>(Mode))
+	ForEachEdMode<ILegacyEdModeDrawHelperInterface>([InView, PDI](ILegacyEdModeDrawHelperInterface* DrawHelper)
 		{
 			DrawHelper->Draw(InView, PDI);
-		}
-	}
+			return true;
+		});
 }
 
 /** Renders all active modes */
 void FEditorModeTools::Render( const FSceneView* InView, FViewport* Viewport, FPrimitiveDrawInterface* PDI )
 {
 	InteractiveToolsContext->Render(InView, Viewport, PDI);
-	for (UEdMode* Mode : ActiveScriptableModes)
-	{
-		Mode->Render(InView, Viewport, PDI);
-	}
+
+	ForEachEdMode<ILegacyEdModeWidgetInterface>([InView, Viewport, PDI](ILegacyEdModeWidgetInterface* Mode)
+		{
+			Mode->Render(InView, Viewport, PDI);
+			return true;
+		});
 }
 
 /** Draws the HUD for all active modes */
 void FEditorModeTools::DrawHUD( FEditorViewportClient* InViewportClient,FViewport* Viewport, const FSceneView* View, FCanvas* Canvas )
 {
 	InteractiveToolsContext->DrawHUD(InViewportClient, Viewport, View, Canvas);
-	for (UEdMode* Mode : ActiveScriptableModes)
+
+	DrawBrackets(InViewportClient, Viewport, View, Canvas);
+
+	if (!(InViewportClient->EngineShowFlags.ModeWidgets))
+	{
+		return;
+	}
+
+	// Clear Hit proxies
+	const bool bIsHitTesting = Canvas->IsHitTesting();
+	if (!bIsHitTesting)
+	{
+		Canvas->SetHitProxy(nullptr);
+	}
+
+	ForEachEdMode<ILegacyEdModeWidgetInterface>([InViewportClient, Viewport, View, Canvas](ILegacyEdModeWidgetInterface* Mode)
+		{
+			Mode->DrawHUD(InViewportClient, Viewport, View, Canvas);
+			return true;
+		});
+
+	// Draw vertices for selected BSP brushes and static meshes if the large vertices show flag is set.
+	if (!InViewportClient->bDrawVertices)
+	{
+		return;
+	}
+
+	const bool bLargeVertices = View->Family->EngineShowFlags.LargeVertices;
+	if (!bLargeVertices)
 	{
-		Mode->DrawHUD(InViewportClient, Viewport, View, Canvas);
+		return;
 	}
+
+	// Temporaries.
+	const bool bShowBrushes = View->Family->EngineShowFlags.Brushes;
+	const bool bShowBSP = View->Family->EngineShowFlags.BSP;
+	const bool bShowBuilderBrush = View->Family->EngineShowFlags.BuilderBrush != 0;
+
+	UTexture2D* VertexTexture = GetVertexTexture();
+	const float TextureSizeX = VertexTexture->GetSizeX() * (bLargeVertices ? 1.0f : 0.5f);
+	const float TextureSizeY = VertexTexture->GetSizeY() * (bLargeVertices ? 1.0f : 0.5f);
+
+	GetEditorSelectionSet()->ForEachSelectedObject<AStaticMeshActor>([View, Canvas, VertexTexture, TextureSizeX, TextureSizeY, bIsHitTesting](AStaticMeshActor* Actor)
+		{
+			TArray<FVector> Vertices;
+			FCanvasItemTestbed::bTestState = !FCanvasItemTestbed::bTestState;
+
+			// Static mesh vertices
+			if (Actor->GetStaticMeshComponent() && Actor->GetStaticMeshComponent()->GetStaticMesh()
+				&& Actor->GetStaticMeshComponent()->GetStaticMesh()->GetRenderData())
+			{
+				FTransform ActorToWorld = Actor->ActorToWorld();
+				const FPositionVertexBuffer& VertexBuffer = Actor->GetStaticMeshComponent()->GetStaticMesh()->GetRenderData()->LODResources[0].VertexBuffers.PositionVertexBuffer;
+				for (uint32 i = 0; i < VertexBuffer.GetNumVertices(); i++)
+				{
+					Vertices.AddUnique(ActorToWorld.TransformPosition(VertexBuffer.VertexPosition(i)));
+				}
+
+				const float InvDpiScale = 1.0f / Canvas->GetDPIScale();
+
+				FCanvasTileItem TileItem(FVector2D(0.0f, 0.0f), FVector2D(0.0f, 0.0f), FLinearColor::White);
+				TileItem.BlendMode = SE_BLEND_Translucent;
+				for (int32 VertexIndex = 0; VertexIndex < Vertices.Num(); ++VertexIndex)
+				{
+					const FVector& Vertex = Vertices[VertexIndex];
+					FVector2D PixelLocation;
+					if (View->ScreenToPixel(View->WorldToScreen(Vertex), PixelLocation))
+					{
+						PixelLocation *= InvDpiScale;
+
+						const bool bOutside =
+							PixelLocation.X < 0.0f || PixelLocation.X > View->UnscaledViewRect.Width() * InvDpiScale ||
+							PixelLocation.Y < 0.0f || PixelLocation.Y > View->UnscaledViewRect.Height() * InvDpiScale;
+						if (!bOutside)
+						{
+							const float X = PixelLocation.X - (TextureSizeX / 2);
+							const float Y = PixelLocation.Y - (TextureSizeY / 2);
+							if (bIsHitTesting)
+							{
+								Canvas->SetHitProxy(new HStaticMeshVert(Actor, Vertex));
+							}
+							TileItem.Texture = VertexTexture->Resource;
+
+							TileItem.Size = FVector2D(TextureSizeX, TextureSizeY);
+							Canvas->DrawItem(TileItem, FVector2D(X, Y));
+							if (bIsHitTesting)
+							{
+								Canvas->SetHitProxy(nullptr);
+							}
+						}
+					}
+				}
+			}
+
+			return true;
+		});
 }
 
 /** Calls PostUndo on all active modes */
@@ -1511,10 +1680,11 @@ void FEditorModeTools::PostUndo(bool bSuccess)
 {
 	if (bSuccess)
 	{
-		for (UEdMode* Mode : ActiveScriptableModes)
-		{
-			Mode->PostUndo();
-		}
+		ForEachEdMode([](UEdMode* Mode)
+			{
+				Mode->PostUndo();
+				return true;
+			});
 	}
 }
 void FEditorModeTools::PostRedo(bool bSuccess)
@@ -1526,23 +1696,24 @@ void FEditorModeTools::PostRedo(bool bSuccess)
 bool FEditorModeTools::AllowWidgetMove() const
 {
 	bool bAllow = false;
-	for (UEdMode* Mode : ActiveScriptableModes)
-	{
-		if (ILegacyEdModeWidgetInterface* LegacyMode = Cast<ILegacyEdModeWidgetInterface>(Mode))
+	ForEachEdMode<ILegacyEdModeWidgetInterface>([&bAllow](ILegacyEdModeWidgetInterface* LegacyMode)
 		{
 			bAllow |= LegacyMode->AllowWidgetMove();
-		}
-	}
+			return true;
+		});
+
 	return bAllow;
 }
 
 bool FEditorModeTools::DisallowMouseDeltaTracking() const
 {
 	bool bDisallow = false;
-	for (const UEdMode* Mode : ActiveScriptableModes)
-	{
-		bDisallow |= Mode->DisallowMouseDeltaTracking();
-	}
+	ForEachEdMode<ILegacyEdModeViewportInterface>([&bDisallow](ILegacyEdModeViewportInterface* LegacyMode)
+		{
+			bDisallow |= LegacyMode->DisallowMouseDeltaTracking();
+			return true;
+		});
+
 	return bDisallow;
 }
 
@@ -1569,20 +1740,24 @@ bool FEditorModeTools::GetOverrideCursorVisibility(bool& bWantsOverride, bool& b
 bool FEditorModeTools::PreConvertMouseMovement(FEditorViewportClient* InViewportClient)
 {
 	bool bHandled = false;
-	for (UEdMode* Mode : ActiveScriptableModes)
-	{
-		bHandled |= Mode->PreConvertMouseMovement(InViewportClient);
-	}
+	ForEachEdMode<ILegacyEdModeViewportInterface>([&bHandled, InViewportClient](ILegacyEdModeViewportInterface* Mode)
+		{
+			bHandled |= Mode->PreConvertMouseMovement(InViewportClient);
+			return true;
+		});
+
 	return bHandled;
 }
 
 bool FEditorModeTools::PostConvertMouseMovement(FEditorViewportClient* InViewportClient)
 {
 	bool bHandled = false;
-	for (UEdMode* Mode : ActiveScriptableModes)
-	{
-		bHandled |= Mode->PostConvertMouseMovement(InViewportClient);
-	}
+	ForEachEdMode<ILegacyEdModeViewportInterface>([&bHandled, InViewportClient](ILegacyEdModeViewportInterface* ViewportInterface)
+		{
+			bHandled |= ViewportInterface->PostConvertMouseMovement(InViewportClient);
+			return true;
+		});
+
 	return bHandled;
 }
 
@@ -1590,13 +1765,11 @@ bool FEditorModeTools::GetShowWidget() const
 {
 	bool bDrawModeSupportsWidgetDrawing = false;
 	// Check to see of any active modes support widget drawing
-	for (UEdMode* Mode : ActiveScriptableModes)
-	{
-		if (ILegacyEdModeWidgetInterface* LegacyMode = Cast<ILegacyEdModeWidgetInterface>(Mode))
+	ForEachEdMode<ILegacyEdModeWidgetInterface>([&bDrawModeSupportsWidgetDrawing](ILegacyEdModeWidgetInterface* LegacyMode)
 		{
 			bDrawModeSupportsWidgetDrawing |= LegacyMode->ShouldDrawWidget();
-		}
-	}
+			return true;
+		});
 	return bDrawModeSupportsWidgetDrawing && bShowWidget;
 }
 
@@ -1780,31 +1953,31 @@ bool FEditorModeTools::IsDefaultModeActive() const
 
 bool FEditorModeTools::CanCycleWidgetMode() const
 {
-	for (UEdMode* Mode : ActiveScriptableModes)
-	{
-		if (ILegacyEdModeWidgetInterface* LegacyMode = Cast<ILegacyEdModeWidgetInterface>(Mode))
+	bool bCanCycleWidget = false;
+	ForEachEdMode<ILegacyEdModeWidgetInterface>([&bCanCycleWidget](ILegacyEdModeWidgetInterface* LegacyMode)
 		{
-			if (LegacyMode->CanCycleWidgetMode())
-			{
-				return true;
-			}
-		}
-	}
-	return false;
+			bCanCycleWidget = LegacyMode->CanCycleWidgetMode();
+			return !bCanCycleWidget;
+		});
+	return bCanCycleWidget;
 }
 
 
 bool FEditorModeTools::CanAutoSave() const
 {
-	for (const UEdMode* Mode : ActiveScriptableModes)
-	{
-		if (Mode->CanAutoSave() == false)
+	bool bCanAutoSave = true;
+	ForEachEdMode([&bCanAutoSave](const UEdMode* Mode)
 		{
-			return false;
-		}
-	}
+			if (!Mode->CanAutoSave())
+			{
+				bCanAutoSave = false;
+				return false;
+			}
 
-	return true;
+			return true;
+		});
+
+	return bCanAutoSave;
 }
 
 UEdModeInteractiveToolsContext* FEditorModeTools::GetInteractiveToolsContext() const
diff --git a/Engine/Source/Editor/UnrealEd/Private/EditorViewportClient.cpp b/Engine/Source/Editor/UnrealEd/Private/EditorViewportClient.cpp
--- a/Engine/Source/Editor/UnrealEd/Private/EditorViewportClient.cpp
+++ b/Engine/Source/Editor/UnrealEd/Private/EditorViewportClient.cpp
@@ -3722,6 +3722,10 @@ void FEditorViewportClient::DrawCanvas(FViewport& InViewport, FSceneView& View,
 	// Information string
 	Canvas.DrawShadowedString(4, 4, *ModeTools->InfoString, GEngine->GetSmallFont(), FColor::White);
 
+	// Render the marquee drag tool
+	RenderDragTool(&View, &Canvas);
+
+	// Draw any HUD from modes
 	ModeTools->DrawHUD(this, &InViewport, &View, &Canvas);
 }
 
diff --git a/Engine/Source/Editor/UnrealEd/Private/Toolkits/SStandaloneAssetEditorToolkitHost.cpp b/Engine/Source/Editor/UnrealEd/Private/Toolkits/SStandaloneAssetEditorToolkitHost.cpp
--- a/Engine/Source/Editor/UnrealEd/Private/Toolkits/SStandaloneAssetEditorToolkitHost.cpp
+++ b/Engine/Source/Editor/UnrealEd/Private/Toolkits/SStandaloneAssetEditorToolkitHost.cpp
@@ -345,13 +345,16 @@ UTypedElementCommonActions* SStandaloneAssetEditorToolkitHost::GetCommonActions(
 
 FReply SStandaloneAssetEditorToolkitHost::OnKeyDown( const FGeometry& MyGeometry, const FKeyEvent& InKeyEvent )
 {
-	// Check to see if any of the actions for the level editor can be processed by the current event
+	// Check to see if any of the actions for the toolkits can be processed by the current event
 	// If we are in debug mode do not process commands
 	if (FSlateApplication::Get().IsNormalExecution())
 	{
-		if (HostedAssetEditorToolkit->ProcessCommandBindings(InKeyEvent))
+		for (TSharedPtr<IToolkit>& HostedToolkit : HostedToolkits)
 		{
-			return FReply::Handled();
+			if (HostedToolkit->ProcessCommandBindings(InKeyEvent))
+			{
+				return FReply::Handled();
+			}
 		}
 	}
 
diff --git a/Engine/Source/Editor/UnrealEd/Private/Tools/LegacyEdMode.cpp b/Engine/Source/Editor/UnrealEd/Private/Tools/LegacyEdMode.cpp
--- a/Engine/Source/Editor/UnrealEd/Private/Tools/LegacyEdMode.cpp
+++ b/Engine/Source/Editor/UnrealEd/Private/Tools/LegacyEdMode.cpp
@@ -282,11 +282,6 @@ bool ULegacyEdModeWrapper::PostConvertMouseMovement(FEditorViewportClient* InVie
 	return LegacyEditorMode->PostConvertMouseMovement(InViewportClient);
 }
 
-void ULegacyEdModeWrapper::DrawBrackets(FEditorViewportClient* ViewportClient, FViewport* Viewport, const FSceneView* View, FCanvas* Canvas)
-{
-	LegacyEditorMode->DrawBrackets(ViewportClient, Viewport, View, Canvas);
-}
-
 bool ULegacyEdModeWrapper::ShouldDrawBrushWireframe(AActor* InActor) const
 {
 	return LegacyEditorMode->ShouldDrawBrushWireframe(InActor);
diff --git a/Engine/Source/Editor/UnrealEd/Private/Tools/LegacyEdMode.h b/Engine/Source/Editor/UnrealEd/Private/Tools/LegacyEdMode.h
--- a/Engine/Source/Editor/UnrealEd/Private/Tools/LegacyEdMode.h
+++ b/Engine/Source/Editor/UnrealEd/Private/Tools/LegacyEdMode.h
@@ -23,7 +23,7 @@ class FViewport;
 class UPrimitiveComponent;
 
 UCLASS()
-class UNREALED_API ULegacyEdModeWrapper final : public UEdMode, public ILegacyEdModeWidgetInterface, public ILegacyEdModeToolInterface, public ILegacyEdModeSelectInterface, public ILegacyEdModeDrawHelperInterface
+class UNREALED_API ULegacyEdModeWrapper final : public UEdMode, public ILegacyEdModeWidgetInterface, public ILegacyEdModeToolInterface, public ILegacyEdModeSelectInterface, public ILegacyEdModeDrawHelperInterface, public ILegacyEdModeViewportInterface
 {
 	GENERATED_BODY()
 
@@ -84,7 +84,6 @@ public:
 	virtual void DrawHUD(FEditorViewportClient* ViewportClient, FViewport* Viewport, const FSceneView* View, FCanvas* Canvas) override;
 	virtual bool PreConvertMouseMovement(FEditorViewportClient* InViewportClient) override;
 	virtual bool PostConvertMouseMovement(FEditorViewportClient* InViewportClient) override;
-	virtual void DrawBrackets(FEditorViewportClient* ViewportClient, FViewport* Viewport, const FSceneView* View, FCanvas* Canvas) override;
 
 	virtual bool ShouldDrawBrushWireframe(AActor* InActor) const override;
 	virtual void Enter() override;
diff --git a/Engine/Source/Editor/UnrealEd/Private/Tools/LegacyEdModeWidgetHelpers.cpp b/Engine/Source/Editor/UnrealEd/Private/Tools/LegacyEdModeWidgetHelpers.cpp
--- a/Engine/Source/Editor/UnrealEd/Private/Tools/LegacyEdModeWidgetHelpers.cpp
+++ b/Engine/Source/Editor/UnrealEd/Private/Tools/LegacyEdModeWidgetHelpers.cpp
@@ -771,7 +771,7 @@ void UBaseLegacyWidgetEdMode::Initialize()
 	WidgetHelper->Owner = this->Owner;
 	WidgetHelper->ParentModeInterface = this;
 
-	UEdMode::Initialize();
+	Super::Initialize();
 }
 
 TSharedRef<FLegacyEdModeWidgetHelper> UBaseLegacyWidgetEdMode::CreateWidgetHelper()
@@ -786,7 +786,7 @@ bool UBaseLegacyWidgetEdMode::InputDelta(FEditorViewportClient* InViewportClient
 		return true;
 	}
 
-	return UEdMode::InputDelta(InViewportClient, InViewport, InDrag, InRot, InScale);
+	return ILegacyEdModeViewportInterface::InputDelta(InViewportClient, InViewport, InDrag, InRot, InScale);
 }
 
 bool UBaseLegacyWidgetEdMode::HandleClick(FEditorViewportClient* InViewportClient, HHitProxy* HitProxy, const FViewportClick& Click)
@@ -796,21 +796,21 @@ bool UBaseLegacyWidgetEdMode::HandleClick(FEditorViewportClient* InViewportClien
 		return true;
 	}
 
-	return UEdMode::HandleClick(InViewportClient, HitProxy, Click);
+	return ILegacyEdModeViewportInterface::HandleClick(InViewportClient, HitProxy, Click);
 }
 
 void UBaseLegacyWidgetEdMode::Render(const FSceneView* View, FViewport* Viewport, FPrimitiveDrawInterface* PDI)
 {
 	WidgetHelper->Render(View, Viewport, PDI);
 
-	UEdMode::Render(View, Viewport, PDI);
+	ILegacyEdModeWidgetInterface::Render(View, Viewport, PDI);
 }
 
 void UBaseLegacyWidgetEdMode::DrawHUD(FEditorViewportClient* ViewportClient, FViewport* Viewport, const FSceneView* View, FCanvas* Canvas)
 {
 	WidgetHelper->DrawHUD(ViewportClient, Viewport, View, Canvas);
 
-	UEdMode::DrawHUD(ViewportClient, Viewport, View, Canvas);
+	ILegacyEdModeWidgetInterface::DrawHUD(ViewportClient, Viewport, View, Canvas);
 }
 
 bool UBaseLegacyWidgetEdMode::AllowWidgetMove()
diff --git a/Engine/Source/Editor/UnrealEd/Private/Tools/UEdMode.cpp b/Engine/Source/Editor/UnrealEd/Private/Tools/UEdMode.cpp
--- a/Engine/Source/Editor/UnrealEd/Private/Tools/UEdMode.cpp
+++ b/Engine/Source/Editor/UnrealEd/Private/Tools/UEdMode.cpp
@@ -1,65 +1,17 @@
 // Copyright Epic Games, Inc. All Rights Reserved.
 
 #include "Tools/UEdMode.h"
-#include "EditorModeTools.h"
-#include "EditorViewportClient.h"
-#include "Framework/Application/SlateApplication.h"
-#include "CanvasItem.h"
-#include "Engine/StaticMeshActor.h"
-#include "Engine/Selection.h"
-#include "EngineUtils.h"
+
 #include "EditorModeManager.h"
 #include "EditorModes.h"
-#include "StaticMeshResources.h"
 #include "Toolkits/BaseToolkit.h"
 #include "EdModeInteractiveToolsContext.h"
-#include "CanvasTypes.h"
-#include "ScopedTransaction.h"
-#include "Tools/EditorToolAssetAPI.h"
-#include "Editor.h"
 #include "Toolkits/ToolkitManager.h"
-#include "Materials/Material.h"
-#include "Materials/MaterialInterface.h"
-#include "EditorViewportClient.h"
 #include "InteractiveToolManager.h"
-#include "InteractiveToolObjects.h"
-#include "Editor/EditorEngine.h"
 #include "GameFramework/Actor.h"
-
 #include "Elements/Framework/TypedElementSelectionSet.h"
 #include "Elements/Interfaces/TypedElementObjectInterface.h"
 
-/** Hit proxy used for editable properties */
-struct HPropertyWidgetProxyTools : public HHitProxy
-{
-	DECLARE_HIT_PROXY();
-
-	/** Name of property this is the widget for */
-	FString	PropertyName;
-
-	/** If the property is an array property, the index into that array that this widget is for */
-	int32	PropertyIndex;
-
-	/** This property is a transform */
-	bool	bPropertyIsTransform;
-
-	HPropertyWidgetProxyTools(FString InPropertyName, int32 InPropertyIndex, bool bInPropertyIsTransform)
-		: HHitProxy(HPP_Foreground)
-		, PropertyName(InPropertyName)
-		, PropertyIndex(InPropertyIndex)
-		, bPropertyIsTransform(bInPropertyIsTransform)
-	{}
-
-	/** Show cursor as cross when over this handle */
-	virtual EMouseCursor::Type GetMouseCursor() override
-	{
-		return EMouseCursor::Crosshairs;
-	}
-};
-
-IMPLEMENT_HIT_PROXY(HPropertyWidgetProxyTools, HHitProxy);
-
-
 
 //////////////////////////////////
 // UEdMode
@@ -76,80 +28,6 @@ void UEdMode::Initialize()
 {
 }
 
-bool UEdMode::MouseEnter(FEditorViewportClient* ViewportClient, FViewport* Viewport, int32 x, int32 y)
-{
-	return false;
-}
-
-bool UEdMode::MouseLeave(FEditorViewportClient* ViewportClient, FViewport* Viewport)
-{
-	return false;
-}
-
-bool UEdMode::MouseMove(FEditorViewportClient* ViewportClient, FViewport* Viewport, int32 x, int32 y)
-{
-	return false;
-}
-
-bool UEdMode::ReceivedFocus(FEditorViewportClient* ViewportClient, FViewport* Viewport)
-{
-	return false;
-}
-
-bool UEdMode::LostFocus(FEditorViewportClient* ViewportClient, FViewport* Viewport)
-{
-	return false;
-}
-
-bool UEdMode::CapturedMouseMove(FEditorViewportClient* InViewportClient, FViewport* InViewport, int32 InMouseX, int32 InMouseY)
-{
-	return false;
-}
-
-bool UEdMode::InputKey(FEditorViewportClient* ViewportClient, FViewport* Viewport, FKey Key, EInputEvent Event)
-{
-	if (!Viewport)
-	{
-		return false;
-	}
-
-	if (Event != IE_Released)
-	{
-		if (ToolCommandList->ProcessCommandBindings(Key, FSlateApplication::Get().GetModifierKeys(), false/*Event == IE_Repeat*/))
-		{
-			return true;
-		}
-	}
-
-	// Next pass input to the mode toolkit
-	if (Toolkit.IsValid() && ((Event == IE_Pressed) || (Event == IE_Repeat)))
-	{
-		if (Toolkit->GetToolkitCommands()->ProcessCommandBindings(Key, FSlateApplication::Get().GetModifierKeys(), (Event == IE_Repeat)))
-		{
-			return true;
-		}
-	}
-
-	// Finally, pass input up to selected actors if not in a tool mode
-	Owner->GetEditorSelectionSet()->ForEachSelectedObject<AActor>([Key, Event](AActor* ActorPtr)
-		{
-			ActorPtr->EditorKeyPressed(Key, Event);
-			return true;
-		});
-
-	return false;
-}
-
-bool UEdMode::InputAxis(FEditorViewportClient* InViewportClient, FViewport* Viewport, int32 ControllerId, FKey Key, float Delta, float DeltaTime)
-{
-	return false;
-}
-
-bool UEdMode::InputDelta(FEditorViewportClient* InViewportClient, FViewport* InViewport, FVector& InDrag, FRotator& InRot, FVector& InScale)
-{
-	return false;
-}
-
 void UEdMode::SelectNone()
 {
 }
@@ -159,16 +37,6 @@ bool UEdMode::ProcessEditDelete()
 	return false;
 }
 
-void UEdMode::Tick(FEditorViewportClient* ViewportClient, float DeltaTime)
-{
-}
-
-
-bool UEdMode::HandleClick(FEditorViewportClient* InViewportClient, HHitProxy *HitProxy, const FViewportClick &Click)
-{
-	return false;
-}
-
 void UEdMode::Enter()
 {
 	// Update components for selected actors, in case the mode we just exited
@@ -214,7 +82,6 @@ void UEdMode::Enter()
 	}
 
 	FEditorDelegates::EditorModeIDEnter.Broadcast(GetID());
-
 }
 
 void UEdMode::RegisterTool(TSharedPtr<FUICommandInfo> UICommand, FString ToolIdentifier, UInteractiveToolBuilder* Builder)
@@ -273,125 +140,6 @@ void UEdMode::Exit()
 	FEditorDelegates::EditorModeIDExit.Broadcast(GetID());
 }
 
-UTexture2D* UEdMode::GetVertexTexture()
-{
-	return GEngine->DefaultBSPVertexTexture;
-}
-
-void UEdMode::Render(const FSceneView* View, FViewport* Viewport, FPrimitiveDrawInterface* PDI)
-{
-}
-
-void UEdMode::DrawHUD(FEditorViewportClient* ViewportClient, FViewport* Viewport, const FSceneView* View, FCanvas* Canvas)
-{
-	// Render the drag tool.
-	ViewportClient->RenderDragTool(View, Canvas);
-
-	if (ViewportClient->IsPerspective() && GetDefault<ULevelEditorViewportSettings>()->bHighlightWithBrackets)
-	{
-		DrawBrackets(ViewportClient, Viewport, View, Canvas);
-	}
-
-	// If this viewport doesn't show mode widgets, leave.
-	if (!(ViewportClient->EngineShowFlags.ModeWidgets))
-	{
-		return;
-	}
-
-	// Clear Hit proxies
-	const bool bIsHitTesting = Canvas->IsHitTesting();
-	if (!bIsHitTesting)
-	{
-		Canvas->SetHitProxy(NULL);
-	}
-
-	// Draw vertices for selected BSP brushes and static meshes if the large vertices show flag is set.
-	if (!ViewportClient->bDrawVertices)
-	{
-		return;
-	}
-
-	const bool bLargeVertices = View->Family->EngineShowFlags.LargeVertices;
-	if (!bLargeVertices)
-	{
-		return;
-	}
-
-	// Temporaries.
-	UTexture2D* VertexTexture = GetVertexTexture();
-	const float TextureSizeX = VertexTexture->GetSizeX() * (bLargeVertices ? 1.0f : 0.5f);
-	const float TextureSizeY = VertexTexture->GetSizeY() * (bLargeVertices ? 1.0f : 0.5f);
-
-	// Static mesh vertices
-	Owner->GetEditorSelectionSet()->ForEachSelectedObject<AStaticMeshActor>([View, Canvas, VertexTexture, TextureSizeX, TextureSizeY, bIsHitTesting](AStaticMeshActor* Actor)
-		{
-			TArray<FVector> Vertices;
-			FCanvasItemTestbed::bTestState = !FCanvasItemTestbed::bTestState;
-
-			if (Actor->GetStaticMeshComponent() && Actor->GetStaticMeshComponent()->GetStaticMesh()
-				&& Actor->GetStaticMeshComponent()->GetStaticMesh()->GetRenderData())
-			{
-				FTransform ActorToWorld = Actor->ActorToWorld();
-				const FPositionVertexBuffer& VertexBuffer = Actor->GetStaticMeshComponent()->GetStaticMesh()->GetRenderData()->LODResources[0].VertexBuffers.PositionVertexBuffer;
-				for (uint32 i = 0; i < VertexBuffer.GetNumVertices(); i++)
-				{
-					Vertices.AddUnique(ActorToWorld.TransformPosition(VertexBuffer.VertexPosition(i)));
-				}
-
-				const float InvDpiScale = 1.0f / Canvas->GetDPIScale();
-
-				FCanvasTileItem TileItem(FVector2D(0.0f, 0.0f), FVector2D(0.0f, 0.0f), FLinearColor::White);
-				TileItem.BlendMode = SE_BLEND_Translucent;
-				for (int32 VertexIndex = 0; VertexIndex < Vertices.Num(); ++VertexIndex)
-				{
-					const FVector& Vertex = Vertices[VertexIndex];
-					FVector2D PixelLocation;
-					if (View->ScreenToPixel(View->WorldToScreen(Vertex), PixelLocation))
-					{
-						PixelLocation *= InvDpiScale;
-
-						const bool bOutside =
-							PixelLocation.X < 0.0f || PixelLocation.X > View->UnscaledViewRect.Width() * InvDpiScale ||
-							PixelLocation.Y < 0.0f || PixelLocation.Y > View->UnscaledViewRect.Height() * InvDpiScale;
-						if (!bOutside)
-						{
-							const float X = PixelLocation.X - (TextureSizeX / 2);
-							const float Y = PixelLocation.Y - (TextureSizeY / 2);
-							if (bIsHitTesting)
-							{
-								Canvas->SetHitProxy(new HStaticMeshVert(Actor, Vertex));
-							}
-							TileItem.Texture = VertexTexture->Resource;
-
-							TileItem.Size = FVector2D(TextureSizeX, TextureSizeY);
-							Canvas->DrawItem(TileItem, FVector2D(X, Y));
-							if (bIsHitTesting)
-							{
-								Canvas->SetHitProxy(NULL);
-							}
-						}
-					}
-				}
-			}
-			return true;
-		});
-}
-
-void UEdMode::DrawBrackets(FEditorViewportClient* ViewportClient, FViewport* Viewport, const FSceneView* View, FCanvas* Canvas)
-{
-	Owner->GetEditorSelectionSet()->ForEachSelectedObject<AActor>([ViewportClient, Canvas, View, Viewport](AActor* SelectedActor)
-		{
-			// Draw a bracket for selected "paintable" static mesh actors
-			const bool bIsValidActor = (Cast< AStaticMeshActor >(SelectedActor) != NULL);
-
-			const FLinearColor SelectedActorBoxColor(0.6f, 0.6f, 1.0f);
-			const bool bDrawBracket = bIsValidActor;
-			ViewportClient->DrawActorScreenSpaceBoundingBox(Canvas, View, Viewport, SelectedActor, SelectedActorBoxColor, bDrawBracket);
-
-			return true;
-		});
-}
-
 bool UEdMode::UsesToolkits() const
 {
 	return true;
@@ -407,16 +155,6 @@ class FEditorModeTools* UEdMode::GetModeManager() const
 	return Owner;
 }
 
-bool UEdMode::StartTracking(FEditorViewportClient* InViewportClient, FViewport* InViewport)
-{
-	return false; 
-}
-
-bool UEdMode::EndTracking(FEditorViewportClient* InViewportClient, FViewport* InViewport)
-{
-	return false; 
-}
-
 AActor* UEdMode::GetFirstSelectedActorInstance() const
 {
 	return Owner->GetEditorSelectionSet()->GetTopSelectedObject<AActor>();
diff --git a/Engine/Source/Editor/UnrealEd/Public/EdMode.h b/Engine/Source/Editor/UnrealEd/Public/EdMode.h
--- a/Engine/Source/Editor/UnrealEd/Public/EdMode.h
+++ b/Engine/Source/Editor/UnrealEd/Public/EdMode.h
@@ -232,11 +232,6 @@ public:
 	 */
 	virtual bool HandleDragDuplicate() { return false; }
 
-	/**
-	 * Called when the mode wants to draw brackets around selected objects
-	 */
-	virtual void DrawBrackets( FEditorViewportClient* ViewportClient, FViewport* Viewport, const FSceneView* View, FCanvas* Canvas );
-
 	/** True if this mode uses a toolkit mode (eventually they all should) */
 	virtual bool UsesToolkits() const;
 
diff --git a/Engine/Source/Editor/UnrealEd/Public/EditorModeManager.h b/Engine/Source/Editor/UnrealEd/Public/EditorModeManager.h
--- a/Engine/Source/Editor/UnrealEd/Public/EditorModeManager.h
+++ b/Engine/Source/Editor/UnrealEd/Public/EditorModeManager.h
@@ -129,6 +129,8 @@ public:
 
 	UEdMode* GetActiveScriptableMode(FEditorModeID InID) const;
 
+	virtual UTexture2D* GetVertexTexture() const;
+
 	/**
 	 * Returns true if the current mode is not the specified ModeID.  Also optionally warns the user.
 	 *
@@ -536,6 +538,24 @@ protected:
 	void OnEditorSelectionChanged(UObject* NewSelection);
 	void OnEditorSelectNone();
 
+	virtual void DrawBrackets(FEditorViewportClient* ViewportClient, FViewport* Viewport, const FSceneView* View, FCanvas* Canvas);
+
+	void ForEachEdMode(TFunctionRef<bool(UEdMode*)> InCalllback) const;
+
+	template <class InterfaceToCastTo>
+	void ForEachEdMode(TFunctionRef<bool(InterfaceToCastTo*)> InCallback) const
+	{
+		ForEachEdMode([InCallback](UEdMode* Mode)
+		{
+			if (InterfaceToCastTo* CastedMode = Cast<InterfaceToCastTo>(Mode))
+			{
+				return InCallback(CastedMode);
+			}
+
+			return true;
+		});
+	}
+
 	/** List of default modes for this tool.  These must all be compatible with each other. */
 	TArray<FEditorModeID> DefaultModeIDs;
 
diff --git a/Engine/Source/Editor/UnrealEd/Public/Tools/LegacyEdModeInterfaces.h b/Engine/Source/Editor/UnrealEd/Public/Tools/LegacyEdModeInterfaces.h
--- a/Engine/Source/Editor/UnrealEd/Public/Tools/LegacyEdModeInterfaces.h
+++ b/Engine/Source/Editor/UnrealEd/Public/Tools/LegacyEdModeInterfaces.h
@@ -3,9 +3,12 @@
 #pragma once
 
 #include "UnrealWidgetFwd.h"
+#include "Engine/EngineBaseTypes.h"
+#include "InputCoreTypes.h"
 #include "Math/Axis.h"
 #include "Math/Vector.h"
 #include "UObject/Interface.h"
+#include "Editor.h"
 
 #include "LegacyEdModeInterfaces.generated.h"
 
@@ -17,6 +20,11 @@ struct FBox;
 struct FMatrix;
 class FSceneView;
 class FPrimitiveDrawInterface;
+class FViewport;
+class HHitProxy;
+struct FViewportClick;
+class FCanvas;
+class UTexture2D;
 
 UINTERFACE(NotBlueprintable, MinimalAPI)
 class ULegacyEdModeSelectInterface : public UInterface
@@ -109,8 +117,9 @@ public:
 	virtual bool GetCustomDrawingCoordinateSystem(FMatrix& InMatrix, void* InData) = 0;
 	virtual bool GetCustomInputCoordinateSystem(FMatrix& InMatrix, void* InData) = 0;
 
-	/** @return True if this mode allows the viewport to use a drag tool */
-	virtual bool AllowsViewportDragTool() const = 0;
+	virtual void Render(const FSceneView* View, FViewport* Viewport, FPrimitiveDrawInterface* PDI) {}
+	virtual void DrawHUD(FEditorViewportClient* ViewportClient, FViewport* Viewport, const FSceneView* View, FCanvas* Canvas) {}
+	virtual UTexture2D* GetVertexTexture() { return GEditor->DefaultBSPVertexTexture; }
 };
 
 UINTERFACE(NotBlueprintable, MinimalAPI)
@@ -149,3 +158,62 @@ class UNREALED_API ILegacyEdModeDrawHelperInterface
 public:
 	virtual void Draw(const FSceneView* View, FPrimitiveDrawInterface* PDI) = 0;
 };
+
+UINTERFACE(NotBlueprintable, MinimalAPI)
+class ULegacyEdModeViewportInterface : public UInterface
+{
+	GENERATED_BODY()
+};
+
+class UNREALED_API ILegacyEdModeViewportInterface
+{
+	GENERATED_BODY()
+
+public:
+	virtual bool MouseEnter(FEditorViewportClient* ViewportClient, FViewport* Viewport, int32 x, int32 y) { return false; }
+
+	virtual bool MouseLeave(FEditorViewportClient* ViewportClient, FViewport* Viewport) { return false; }
+
+	virtual bool MouseMove(FEditorViewportClient* ViewportClient, FViewport* Viewport, int32 x, int32 y) { return false; }
+
+	virtual bool ReceivedFocus(FEditorViewportClient* ViewportClient, FViewport* Viewport) { return false; }
+
+	virtual bool LostFocus(FEditorViewportClient* ViewportClient, FViewport* Viewport) { return false; }
+
+	/**
+	 * Called when the mouse is moved while a window input capture is in effect
+	 *
+	 * @param	InViewportClient	Level editor viewport client that captured the mouse input
+	 * @param	InViewport			Viewport that captured the mouse input
+	 * @param	InMouseX			New mouse cursor X coordinate
+	 * @param	InMouseY			New mouse cursor Y coordinate
+	 *
+	 * @return	true if input was handled
+	 */
+	virtual bool CapturedMouseMove(FEditorViewportClient* InViewportClient, FViewport* InViewport, int32 InMouseX, int32 InMouseY) { return false; }
+
+	/** Process all captured mouse moves that occurred during the current frame */
+	virtual bool ProcessCapturedMouseMoves(FEditorViewportClient* InViewportClient, FViewport* InViewport, const TArrayView<FIntPoint>& CapturedMouseMoves) { return false; }
+
+	virtual bool InputKey(FEditorViewportClient* ViewportClient, FViewport* Viewport, FKey Key, EInputEvent Event) { return false; }
+	virtual bool InputAxis(FEditorViewportClient* InViewportClient, FViewport* Viewport, int32 ControllerId, FKey Key, float Delta, float DeltaTime) { return false; }
+	virtual bool InputDelta(FEditorViewportClient* InViewportClient, FViewport* InViewport, FVector& InDrag, FRotator& InRot, FVector& InScale) { return false; }
+	virtual bool StartTracking(FEditorViewportClient* InViewportClient, FViewport* InViewport) { return false; }
+	virtual bool EndTracking(FEditorViewportClient* InViewportClient, FViewport* InViewport) { return false; }
+
+	/** Called before mouse movement is converted to drag/rot */
+	virtual bool PreConvertMouseMovement(FEditorViewportClient* InViewportClient) { return false; }
+
+	/** Called after mouse movement is converted to drag/rot */
+	virtual bool PostConvertMouseMovement(FEditorViewportClient* InViewportClient) { return false; }
+
+	virtual bool HandleClick(FEditorViewportClient* InViewportClient, HHitProxy* HitProxy, const FViewportClick& Click) { return false; }
+
+	virtual void Tick(FEditorViewportClient* ViewportClient, float DeltaTime) {}
+
+	/** @return True if this mode allows the viewport to use a drag tool */
+	virtual bool AllowsViewportDragTool() const { return false; }
+
+	/** If the Edmode is handling its own mouse deltas, it can disable the MouseDeltaTacker */
+	virtual bool DisallowMouseDeltaTracking() const { return false; }
+};
diff --git a/Engine/Source/Editor/UnrealEd/Public/Tools/LegacyEdModeWidgetHelpers.h b/Engine/Source/Editor/UnrealEd/Public/Tools/LegacyEdModeWidgetHelpers.h
--- a/Engine/Source/Editor/UnrealEd/Public/Tools/LegacyEdModeWidgetHelpers.h
+++ b/Engine/Source/Editor/UnrealEd/Public/Tools/LegacyEdModeWidgetHelpers.h
@@ -128,7 +128,7 @@ protected:
 //   2. Transition your native FEdMode to a UObject, by inheriting from this class
 //   3. Override CreateWidgetHelper function to return a SharedRef to the class you created in step 1.
 UCLASS(Abstract)
-class UNREALED_API UBaseLegacyWidgetEdMode : public UEdMode, public ILegacyEdModeWidgetInterface
+class UNREALED_API UBaseLegacyWidgetEdMode : public UEdMode, public ILegacyEdModeWidgetInterface, public ILegacyEdModeViewportInterface
 {
 	GENERATED_BODY()
 
diff --git a/Engine/Source/Editor/UnrealEd/Public/Tools/UEdMode.h b/Engine/Source/Editor/UnrealEd/Public/Tools/UEdMode.h
--- a/Engine/Source/Editor/UnrealEd/Public/Tools/UEdMode.h
+++ b/Engine/Source/Editor/UnrealEd/Public/Tools/UEdMode.h
@@ -3,41 +3,23 @@
 #pragma once
 
 #include "Tools/Modes.h"
-#include "InputCoreTypes.h"
-#include "UnrealWidgetFwd.h"
-#include "EditorComponents.h"
 #include "Templates/SharedPointer.h"
 #include "UObject/Object.h"
-#include "InputState.h"
-#include "Math/Ray.h"
-#include "UObject/SoftObjectPtr.h"
-#include "Engine/EngineBaseTypes.h"
 #include "GenericPlatform/ICursor.h"
-#include "UEdMode.generated.h"
 
+#include "UEdMode.generated.h"
 
-class FCanvas;
 class FEditorModeTools;
-class FEditorViewportClient;
 class FModeToolkit;
-class FPrimitiveDrawInterface;
-class FSceneView;
-class FViewport;
-class UTexture2D;
 class UEdModeInteractiveToolsContext;
 class UInteractiveToolManager;
 class UInteractiveTool;
 
-class FEditorViewportClient;
-class HHitProxy;
-struct FViewportClick;
-class FEditorViewportClient;
 class UInteractiveToolBuilder;
 class FUICommandInfo;
 class FUICommandList;
 class FEdMode;
 
-
 /** Outcomes when determining whether it's possible to perform an action on the edit modes*/
 namespace EEditAction
 {
@@ -68,36 +50,6 @@ public:
 
 	virtual void Initialize();
 
-	virtual bool MouseEnter(FEditorViewportClient* ViewportClient, FViewport* Viewport, int32 x, int32 y);
-
-	virtual bool MouseLeave(FEditorViewportClient* ViewportClient, FViewport* Viewport);
-
-	virtual bool MouseMove(FEditorViewportClient* ViewportClient, FViewport* Viewport, int32 x, int32 y);
-
-	virtual bool ReceivedFocus(FEditorViewportClient* ViewportClient, FViewport* Viewport);
-
-	virtual bool LostFocus(FEditorViewportClient* ViewportClient, FViewport* Viewport);
-
-	/**
-	 * Called when the mouse is moved while a window input capture is in effect
-	 *
-	 * @param	InViewportClient	Level editor viewport client that captured the mouse input
-	 * @param	InViewport			Viewport that captured the mouse input
-	 * @param	InMouseX			New mouse cursor X coordinate
-	 * @param	InMouseY			New mouse cursor Y coordinate
-	 *
-	 * @return	true if input was handled
-	 */
-	virtual bool CapturedMouseMove(FEditorViewportClient* InViewportClient, FViewport* InViewport, int32 InMouseX, int32 InMouseY);
-
-	/** Process all captured mouse moves that occurred during the current frame */
-	virtual bool ProcessCapturedMouseMoves(FEditorViewportClient* InViewportClient, FViewport* InViewport, const TArrayView<FIntPoint>& CapturedMouseMoves) { return false; }
-
-	virtual bool InputKey(FEditorViewportClient* ViewportClient, FViewport* Viewport, FKey Key, EInputEvent Event);
-	virtual bool InputAxis(FEditorViewportClient* InViewportClient, FViewport* Viewport, int32 ControllerId, FKey Key, float Delta, float DeltaTime);
-	virtual bool InputDelta(FEditorViewportClient* InViewportClient, FViewport* InViewport, FVector& InDrag, FRotator& InRot, FVector& InScale);
-	virtual bool StartTracking(FEditorViewportClient* InViewportClient, FViewport* InViewport);
-	virtual bool EndTracking(FEditorViewportClient* InViewportClient, FViewport* InViewport);
 	// Added for handling EDIT Command...
 	virtual EEditAction::Type GetActionEditDuplicate() { return EEditAction::Skip; }
 	virtual EEditAction::Type GetActionEditDelete() { return EEditAction::Skip; }
@@ -110,8 +62,6 @@ public:
 	virtual bool ProcessEditCopy() { return false; }
 	virtual bool ProcessEditPaste() { return false; }
 
-	virtual void Tick(FEditorViewportClient* ViewportClient, float DeltaTime);
-
 	virtual bool IsCompatibleWith(FEditorModeID OtherModeID) const { return false; }
 	virtual void ActorMoveNotify() {}
 	virtual void ActorsDuplicatedNotify(TArray<AActor*>& PreDuplicateSelection, TArray<AActor*>& PostDuplicateSelection, bool bOffsetLocations) {}
@@ -119,9 +69,6 @@ public:
 	virtual void ActorPropChangeNotify() {}
 	virtual void MapChangeNotify() {}
 
-	/** If the Edmode is handling its own mouse deltas, it can disable the MouseDeltaTacker */
-	virtual bool DisallowMouseDeltaTracking() const { return false; }
-
 	/**
 	 * Lets each mode/tool specify a pivot point around which the camera should orbit
 	 * @param	OutPivot	The custom pivot point returned by the mode/tool
@@ -138,12 +85,6 @@ public:
 	/** Get override cursor visibility settings */
 	virtual bool GetOverrideCursorVisibility(bool& bWantsOverride, bool& bHardwareCursorVisible, bool bSoftwareCursorVisible) const { return false; }
 
-	/** Called before mouse movement is converted to drag/rot */
-	virtual bool PreConvertMouseMovement(FEditorViewportClient* InViewportClient) { return false; }
-
-	/** Called after mouse movement is converted to drag/rot */
-	virtual bool PostConvertMouseMovement(FEditorViewportClient* InViewportClient) { return false; }
-
 	virtual bool ShouldDrawBrushWireframe(AActor* InActor) const { return true; }
 
 	/** If Rotation Snap should be enabled for this mode*/
@@ -156,7 +97,6 @@ public:
 	*/
 	virtual bool SnapRotatorToGridOverride(FRotator& Rotation) { return false; };
 
-
 	virtual void UpdateInternalData() {}
 
 	virtual void Enter();
@@ -168,12 +108,12 @@ public:
 	 * By default the check disallows starting tools during play/simulate in editor.
 	 */
 	virtual bool ShouldToolStartBeAllowed(const FString& ToolIdentifier) const;
-
 	virtual void Exit();
-	virtual UTexture2D* GetVertexTexture();
 
 	virtual void PostUndo() {}
 
+	virtual void ModeTick(float DeltaTime) {}
+
 	/**
 	 * Check to see if this UEdMode wants to disallow AutoSave
 	 * @return true if AutoSave can be applied right now
@@ -183,8 +123,6 @@ public:
 	virtual void SelectNone();
 	virtual void SelectionChanged() {}
 
-	virtual bool HandleClick(FEditorViewportClient* InViewportClient, HHitProxy* HitProxy, const FViewportClick& Click);
-
 	/**
 	 * Allows an editor mode to override the bounding box used to focus the viewport on a selection
 	 *
@@ -207,21 +145,6 @@ public:
 	/** Returns the editor mode information. */
 	const FEditorModeInfo& GetModeInfo() const { return Info; }
 
-	/** @name Rendering */
-	//@{
-	/** Draws translucent polygons on brushes and volumes. */
-	virtual void Render(const FSceneView* View, FViewport* Viewport, FPrimitiveDrawInterface* PDI);
-	//void DrawGridSection(int32 ViewportLocX,int32 ViewportGridY,FVector* A,FVector* B,float* AX,float* BX,int32 Axis,int32 AlphaCase,FSceneView* View,FPrimitiveDrawInterface* PDI);
-
-	/** Overlays the editor hud (brushes, drag tools, static mesh vertices, etc*. */
-	virtual void DrawHUD(FEditorViewportClient* ViewportClient, FViewport* Viewport, const FSceneView* View, FCanvas* Canvas);
-	//@}
-
-	/**
-	 * Called when the mode wants to draw brackets around selected objects
-	 */
-	virtual void DrawBrackets(FEditorViewportClient* ViewportClient, FViewport* Viewport, const FSceneView* View, FCanvas* Canvas);
-
 	/** True if this mode uses a toolkit mode (eventually they all should) */
 	virtual bool UsesToolkits() const;
 
