ObjectTools::IsObjectBrowsable now returns false for actors
This makes sure we don't generate/save thumbnails for actors.

diff --git a/Engine/Source/Editor/UnrealEd/Private/ObjectTools.cpp b/Engine/Source/Editor/UnrealEd/Private/ObjectTools.cpp
--- a/Engine/Source/Editor/UnrealEd/Private/ObjectTools.cpp
+++ b/Engine/Source/Editor/UnrealEd/Private/ObjectTools.cpp
@@ -170,7 +170,8 @@ namespace ObjectTools
 			{
 				if( ObjectPackage != GetTransientPackage()
 					&& (ObjectPackage->HasAnyPackageFlags(PKG_PlayInEditor) == false)
-					&& !Obj->IsPendingKill() )
+					&& !Obj->IsPendingKill() 
+					&& !Obj->IsA<AActor>())
 				{
 					bIsSupported = true;
 				}
