Updated references to UE4 to Unreal in NetworkFileConnection.cpp
This required references to GPackageFileUE4Version and GPackageFileLicenseeUE4Version to also be updated to GPackageFileUnrealVersion and GPackageFileLicenseeUnrealVersion accordingly

diff --git a/Engine/Plugins/Developer/Concert/ConcertMain/Source/Concert/Private/ConcertVersion.cpp b/Engine/Plugins/Developer/Concert/ConcertMain/Source/Concert/Private/ConcertVersion.cpp
--- a/Engine/Plugins/Developer/Concert/ConcertMain/Source/Concert/Private/ConcertVersion.cpp
+++ b/Engine/Plugins/Developer/Concert/ConcertMain/Source/Concert/Private/ConcertVersion.cpp
@@ -49,8 +49,8 @@ bool ValidateVersion(const int32 InCurrent, const int32 InOther, const FText InV
 
 void FConcertFileVersionInfo::Initialize()
 {
-	FileVersionUE4 = GPackageFileUE4Version;
-	FileVersionLicenseeUE4 = GPackageFileLicenseeUE4Version;
+	FileVersionUE4 = GPackageFileUEVersion;
+	FileVersionLicenseeUE4 = GPackageFileLicenseeUEVersion;
 }
 
 bool FConcertFileVersionInfo::Validate(const FConcertFileVersionInfo& InOther, const EConcertVersionValidationMode InValidationMode, FText* OutFailureReason) const
diff --git a/Engine/Source/Editor/UnrealEd/Private/Commandlets/ContentCommandlets.cpp b/Engine/Source/Editor/UnrealEd/Private/Commandlets/ContentCommandlets.cpp
--- a/Engine/Source/Editor/UnrealEd/Private/Commandlets/ContentCommandlets.cpp
+++ b/Engine/Source/Editor/UnrealEd/Private/Commandlets/ContentCommandlets.cpp
@@ -417,12 +417,12 @@ int32 UResavePackagesCommandlet::InitializeResaveParameters( const TArray<FStrin
 	if ( Switches.Contains(TEXT("CHECKLICENSEEVER")) )
 	{
 		// Limits resaving to packages with this licensee package version or lower.
-		MaxResaveLicenseeUE4Version = FMath::Max<int32>(GPackageFileLicenseeUE4Version - 1, 0);
+		MaxResaveLicenseeUE4Version = FMath::Max<int32>(GPackageFileLicenseeUEVersion - 1, 0);
 	}
 	if ( Switches.Contains(TEXT("CHECKUE4VER")) )
 	{
 		// Limits resaving to packages with this ue4 package version or lower.
-		MaxResaveUE4Version = FMath::Max<int32>(GPackageFileUE4Version - 1, 0);
+		MaxResaveUE4Version = FMath::Max<int32>(GPackageFileUEVersion - 1, 0);
 	}
 	else if ( Switches.Contains(TEXT("RESAVEDEPRECATED")) )
 	{
@@ -439,7 +439,7 @@ int32 UResavePackagesCommandlet::InitializeResaveParameters( const TArray<FStrin
 			{
 				if ( MinResaveUE4Version == CURRENT_PACKAGE_VERSION )
 				{
-					MinResaveUE4Version = GPackageFileUE4Version;
+					MinResaveUE4Version = GPackageFileUEVersion;
 				}
 			}
 
@@ -447,7 +447,7 @@ int32 UResavePackagesCommandlet::InitializeResaveParameters( const TArray<FStrin
 			{
 				if ( MaxResaveUE4Version == CURRENT_PACKAGE_VERSION )
 				{
-					MaxResaveUE4Version = GPackageFileUE4Version;
+					MaxResaveUE4Version = GPackageFileUEVersion;
 				}
 			}
 		}
@@ -773,7 +773,7 @@ void UResavePackagesCommandlet::LoadAndSaveOnePackage(const FString& Filename)
 				if( bIsReadOnly == true && bVerifyContent == true && bAutoCheckOut == false )
 				{
 					UE_LOG(LogContentCommandlet, Warning, TEXT("Package [%s] is read-only but needs to be resaved (UE4 Version: %i, Licensee Version: %i  Current UE4 Version: %i, Current Licensee Version: %i)"),
-						*Filename, Linker->Summary.GetFileVersionUE4(), Linker->Summary.GetFileVersionLicenseeUE4(), GPackageFileUE4Version, VER_LATEST_ENGINE_LICENSEEUE4 );
+						*Filename, Linker->Summary.GetFileVersionUE4(), Linker->Summary.GetFileVersionLicenseeUE4(), GPackageFileUEVersion, VER_LATEST_ENGINE_LICENSEEUE4 );
 					if( SavePackageHelper(Package, FString(TEXT("Temp.temp"))) )
 					{
 						UE_LOG(LogContentCommandlet, Warning, TEXT("Correctly saved:  [Temp.temp].") );
@@ -839,7 +839,7 @@ void UResavePackagesCommandlet::LoadAndSaveOnePackage(const FString& Filename)
 					if (Verbosity != ONLY_ERRORS)
 					{
 						UE_LOG(LogContentCommandlet, Display, TEXT("Resaving package [%s] (UE4 Version: %i, Licensee Version: %i  Saved UE4 Version: %i, Saved Licensee Version: %i)"),
-							*Filename,Linker->Summary.GetFileVersionUE4(), Linker->Summary.GetFileVersionLicenseeUE4(), GPackageFileUE4Version, VER_LATEST_ENGINE_LICENSEEUE4 );
+							*Filename,Linker->Summary.GetFileVersionUE4(), Linker->Summary.GetFileVersionLicenseeUE4(), GPackageFileUEVersion, VER_LATEST_ENGINE_LICENSEEUE4 );
 					}
 
 					const static bool bKeepPackageGUIDOnSave = FParse::Param(FCommandLine::Get(), TEXT("KeepPackageGUIDOnSave"));
diff --git a/Engine/Source/Editor/UnrealEd/Private/CookOnTheFlyServer.cpp b/Engine/Source/Editor/UnrealEd/Private/CookOnTheFlyServer.cpp
--- a/Engine/Source/Editor/UnrealEd/Private/CookOnTheFlyServer.cpp
+++ b/Engine/Source/Editor/UnrealEd/Private/CookOnTheFlyServer.cpp
@@ -4586,8 +4586,8 @@ void GetAdditionalCurrentIniVersionStrings( const ITargetPlatform* TargetPlatfor
 		IniVersionMap.Add(CustomVersionString, CustomVersionValue);
 	}
 
-	FString UE4Ver = FString::Printf(TEXT("PackageFileVersions:%d"), GPackageFileUE4Version);
-	FString UE4Value = FString::Printf(TEXT("%d"), GPackageFileLicenseeUE4Version);
+	FString UE4Ver = FString::Printf(TEXT("PackageFileVersions:%d"), GPackageFileUEVersion);
+	FString UE4Value = FString::Printf(TEXT("%d"), GPackageFileLicenseeUEVersion);
 	IniVersionMap.Add(UE4Ver, UE4Value);
 
 	/*FString UE4EngineVersionCompatibleName = TEXT("EngineVersionCompatibleWith");
diff --git a/Engine/Source/Editor/UnrealEd/Private/EditorServer.cpp b/Engine/Source/Editor/UnrealEd/Private/EditorServer.cpp
--- a/Engine/Source/Editor/UnrealEd/Private/EditorServer.cpp
+++ b/Engine/Source/Editor/UnrealEd/Private/EditorServer.cpp
@@ -2433,22 +2433,22 @@ bool UEditorEngine::PackageIsAMapFile( const TCHAR* PackageFilename, FText& OutN
 		const int32 UE4LicenseeVersion = Summary.GetFileVersionLicenseeUE4();
 
 		// Don't load packages that were saved with an engine version newer than the current one.
-		if( UE4Version > GPackageFileUE4Version )
+		if( UE4Version > GPackageFileUEVersion )
 		{
 			FFormatNamedArguments Arguments;
 			Arguments.Add(TEXT("File"), FText::FromString( FString( PackageFilename ) ));
 			Arguments.Add(TEXT("Version"), UE4Version);
-			Arguments.Add(TEXT("Last"), GPackageFileUE4Version);
+			Arguments.Add(TEXT("Last"), GPackageFileUEVersion);
 			OutNotMapReason = FText::Format( LOCTEXT( "UE4FileIsNewer", "{File} is a UE4 map [File:v{Version}], from an engine release newer than this [Cur:v{Last}]." ), 
 				Arguments);
 			return false;
 		}
-		else if( UE4LicenseeVersion > GPackageFileLicenseeUE4Version )
+		else if( UE4LicenseeVersion > GPackageFileLicenseeUEVersion )
 		{
 			FFormatNamedArguments Arguments;
 			Arguments.Add(TEXT("File"), FText::FromString( FString( PackageFilename ) ));
 			Arguments.Add(TEXT("Version"), UE4LicenseeVersion);
-			Arguments.Add(TEXT("Last"), GPackageFileLicenseeUE4Version);
+			Arguments.Add(TEXT("Last"), GPackageFileLicenseeUEVersion);
 			OutNotMapReason = FText::Format( LOCTEXT( "UE4FileIsNewer", "{File} is a UE4 map [File:v{Version}], from an engine release newer than this [Cur:v{Last}]." ), 
 				Arguments);
 			return false;
diff --git a/Engine/Source/Runtime/AssetRegistry/Private/PackageReader.cpp b/Engine/Source/Runtime/AssetRegistry/Private/PackageReader.cpp
--- a/Engine/Source/Runtime/AssetRegistry/Private/PackageReader.cpp
+++ b/Engine/Source/Runtime/AssetRegistry/Private/PackageReader.cpp
@@ -88,7 +88,7 @@ bool FPackageReader::OpenPackageFile(EOpenPackageResult* OutErrorCode)
 	}
 
 	// Don't read packages that were saved with an package version newer than the current one.
-	if( (PackageFileSummary.GetFileVersionUE4() > GPackageFileUE4Version) || (PackageFileSummary.GetFileVersionLicenseeUE4() > GPackageFileLicenseeUE4Version) )
+	if( (PackageFileSummary.GetFileVersionUE4() > GPackageFileUEVersion) || (PackageFileSummary.GetFileVersionLicenseeUE4() > GPackageFileLicenseeUEVersion) )
 	{
 		UE_LOG(LogAssetRegistry, Error, TEXT("Package %s is too new"), *PackageFilename);
 		SetPackageErrorCode(EOpenPackageResult::VersionTooNew);
diff --git a/Engine/Source/Runtime/Core/Private/Serialization/Archive.cpp b/Engine/Source/Runtime/Core/Private/Serialization/Archive.cpp
--- a/Engine/Source/Runtime/Core/Private/Serialization/Archive.cpp
+++ b/Engine/Source/Runtime/Core/Private/Serialization/Archive.cpp
@@ -150,8 +150,8 @@ void FArchiveState::Reset()
 #if DEVIRTUALIZE_FLinkerLoad_Serialize
 	ActiveFPLB->Reset();
 #endif
-	ArUE4Ver							= GPackageFileUE4Version;
-	ArLicenseeUE4Ver					= GPackageFileLicenseeUE4Version;
+	ArUE4Ver							= GPackageFileUEVersion;
+	ArLicenseeUE4Ver					= GPackageFileLicenseeUEVersion;
 	ArEngineVer							= FEngineVersion::Current();
 	ArEngineNetVer						= FNetworkVersion::GetEngineNetworkProtocolVersion();
 	ArGameNetVer						= FNetworkVersion::GetGameNetworkProtocolVersion();
diff --git a/Engine/Source/Runtime/Core/Private/UObject/ObjectVersion.cpp b/Engine/Source/Runtime/Core/Private/UObject/ObjectVersion.cpp
--- a/Engine/Source/Runtime/Core/Private/UObject/ObjectVersion.cpp
+++ b/Engine/Source/Runtime/Core/Private/UObject/ObjectVersion.cpp
@@ -7,5 +7,9 @@
 #include "UObject/ObjectVersion.h"
 
 // @see ObjectVersion.h for the list of changes/defines
-const int32 GPackageFileUE4Version = VER_LATEST_ENGINE_UE4;
-const int32 GPackageFileLicenseeUE4Version = VER_LATEST_ENGINE_LICENSEEUE4;
+const int32 GPackageFileUEVersion = VER_LATEST_ENGINE_UE4;
+const int32 GPackageFileLicenseeUEVersion = VER_LATEST_ENGINE_LICENSEEUE4;
+
+const int32 GPackageFileUE4Version = GPackageFileUEVersion;
+const int32 GPackageFileLicenseeUE4Version = GPackageFileLicenseeUEVersion;
+
diff --git a/Engine/Source/Runtime/Core/Public/UObject/ObjectVersion.h b/Engine/Source/Runtime/Core/Public/UObject/ObjectVersion.h
--- a/Engine/Source/Runtime/Core/Public/UObject/ObjectVersion.h
+++ b/Engine/Source/Runtime/Core/Public/UObject/ObjectVersion.h
@@ -666,5 +666,11 @@ enum EUnrealEngineObjectLicenseeUE4Version
 
 // Version access.
 
-extern const CORE_API int32			GPackageFileUE4Version;			// UE4 Version Number.
-extern const CORE_API int32			GPackageFileLicenseeUE4Version;	// Licensee Version Number.
+UE_DEPRECATED(5.0, "Use GPackageFileUEVersion instead")
+extern const CORE_API int32			GPackageFileUE4Version;				// Unreal Version Number.
+
+UE_DEPRECATED(5.0, "Use GPackageFileLicenseeUEVersion instead")
+extern const CORE_API int32			GPackageFileLicenseeUE4Version;		// Licensee Version Number.
+
+extern const CORE_API int32			GPackageFileUEVersion;			// Unreal Version Number.
+extern const CORE_API int32			GPackageFileLicenseeUEVersion;	// Licensee Version Number.
diff --git a/Engine/Source/Runtime/CoreUObject/Private/Serialization/AsyncLoading2.cpp b/Engine/Source/Runtime/CoreUObject/Private/Serialization/AsyncLoading2.cpp
--- a/Engine/Source/Runtime/CoreUObject/Private/Serialization/AsyncLoading2.cpp
+++ b/Engine/Source/Runtime/CoreUObject/Private/Serialization/AsyncLoading2.cpp
@@ -5298,8 +5298,8 @@ void FAsyncPackage2::CreateUPackage(const FPackageSummary* PackageSummary)
 		LinkerRoot->SetCanBeImportedFlag(Desc.CanBeImported());
 		LinkerRoot->SetPackageId(Desc.DiskPackageId);
 		LinkerRoot->SetPackageFlagsTo(PackageSummary->PackageFlags | PKG_Cooked);
-		LinkerRoot->LinkerPackageVersion = GPackageFileUE4Version;
-		LinkerRoot->LinkerLicenseeVersion = GPackageFileLicenseeUE4Version;
+		LinkerRoot->LinkerPackageVersion = GPackageFileUEVersion;
+		LinkerRoot->LinkerLicenseeVersion = GPackageFileLicenseeUEVersion;
 		// LinkerRoot->LinkerCustomVersion = PackageSummaryVersions; // only if (!bCustomVersionIsLatest)
 #if WITH_IOSTORE_IN_EDITOR
 		LinkerRoot->bIsCookedForEditor = !!(PackageSummary->PackageFlags & PKG_FilterEditorOnly);
@@ -5314,8 +5314,8 @@ void FAsyncPackage2::CreateUPackage(const FPackageSummary* PackageSummary)
 		check(LinkerRoot->CanBeImported() == Desc.CanBeImported());
 		check(LinkerRoot->GetPackageId() == Desc.DiskPackageId);
 		check(LinkerRoot->GetPackageFlags() == (PackageSummary->PackageFlags | PKG_Cooked));
-		check(LinkerRoot->LinkerPackageVersion == GPackageFileUE4Version);
-		check(LinkerRoot->LinkerLicenseeVersion == GPackageFileLicenseeUE4Version);
+		check(LinkerRoot->LinkerPackageVersion == GPackageFileUEVersion);
+		check(LinkerRoot->LinkerLicenseeVersion == GPackageFileLicenseeUEVersion);
 		check(LinkerRoot->HasAnyFlags(RF_WasLoaded));
 	}
 
diff --git a/Engine/Source/Runtime/CoreUObject/Private/UObject/LinkerLoad.cpp b/Engine/Source/Runtime/CoreUObject/Private/UObject/LinkerLoad.cpp
--- a/Engine/Source/Runtime/CoreUObject/Private/UObject/LinkerLoad.cpp
+++ b/Engine/Source/Runtime/CoreUObject/Private/UObject/LinkerLoad.cpp
@@ -991,8 +991,8 @@ bool FLinkerLoad::IsTimeLimitExceeded( const TCHAR* CurrentTask, int32 Granulari
 void FLinkerLoad::ResetStatusInfo()
 {
 	// Set status info.
-	this->SetUE4Ver(GPackageFileUE4Version);
-	this->SetLicenseeUE4Ver(GPackageFileLicenseeUE4Version);
+	this->SetUE4Ver(GPackageFileUEVersion);
+	this->SetLicenseeUE4Ver(GPackageFileLicenseeUEVersion);
 	this->SetEngineVer(FEngineVersion::Current());
 	this->SetIsLoading(true);
 	this->SetIsPersistent(true);
@@ -1290,10 +1290,10 @@ FLinkerLoad::ELinkerStatus FLinkerLoad::SerializePackageFileSummaryInternal()
 	}
 
 	// Don't load packages that were saved with package version newer than the current one.
-	if ((Summary.GetFileVersionUE4() > GPackageFileUE4Version) || (Summary.GetFileVersionLicenseeUE4() > GPackageFileLicenseeUE4Version))
+	if ((Summary.GetFileVersionUE4() > GPackageFileUEVersion) || (Summary.GetFileVersionLicenseeUE4() > GPackageFileLicenseeUEVersion))
 	{
 		UE_LOG(LogLinker, Warning, TEXT("Unable to load package (%s) PackageVersion %i, MaxExpected %i : LicenseePackageVersion %i, MaxExpected %i."),
-			*GetDebugName(), Summary.GetFileVersionUE4(), GPackageFileUE4Version, Summary.GetFileVersionLicenseeUE4(), GPackageFileLicenseeUE4Version);
+			*GetDebugName(), Summary.GetFileVersionUE4(), GPackageFileUEVersion, Summary.GetFileVersionLicenseeUE4(), GPackageFileLicenseeUEVersion);
 		return LINKER_Failed;
 	}
 
diff --git a/Engine/Source/Runtime/CoreUObject/Private/UObject/LinkerSave.cpp b/Engine/Source/Runtime/CoreUObject/Private/UObject/LinkerSave.cpp
--- a/Engine/Source/Runtime/CoreUObject/Private/UObject/LinkerSave.cpp
+++ b/Engine/Source/Runtime/CoreUObject/Private/UObject/LinkerSave.cpp
@@ -40,7 +40,7 @@ FLinkerSave::FLinkerSave(UPackage* InParent, const TCHAR* InFilename, bool bForc
 
 		// Set main summary info.
 		Summary.Tag           = PACKAGE_FILE_TAG;
-		Summary.SetFileVersions( GPackageFileUE4Version, GPackageFileLicenseeUE4Version, bInSaveUnversioned );
+		Summary.SetFileVersions( GPackageFileUEVersion, GPackageFileLicenseeUEVersion, bInSaveUnversioned );
 		Summary.SavedByEngineVersion = FEngineVersion::Current();
 		Summary.CompatibleWithEngineVersion = FEngineVersion::CompatibleWith();
 		Summary.SetPackageFlags(Package ? Package->GetPackageFlags() : 0);
@@ -94,7 +94,7 @@ FLinkerSave::FLinkerSave(UPackage* InParent, FArchive *InSaver, bool bForceByteS
 
 		// Set main summary info.
 		Summary.Tag = PACKAGE_FILE_TAG;
-		Summary.SetFileVersions(GPackageFileUE4Version, GPackageFileLicenseeUE4Version, bInSaveUnversioned);
+		Summary.SetFileVersions(GPackageFileUEVersion, GPackageFileLicenseeUEVersion, bInSaveUnversioned);
 		Summary.SavedByEngineVersion = FEngineVersion::Current();
 		Summary.CompatibleWithEngineVersion = FEngineVersion::CompatibleWith();
 		Summary.SetPackageFlags(Package ? Package->GetPackageFlags() : 0);
@@ -143,7 +143,7 @@ FLinkerSave::FLinkerSave(UPackage* InParent, bool bForceByteSwapping, bool bInSa
 
 		// Set main summary info.
 		Summary.Tag           = PACKAGE_FILE_TAG;
-		Summary.SetFileVersions( GPackageFileUE4Version, GPackageFileLicenseeUE4Version, bInSaveUnversioned );
+		Summary.SetFileVersions( GPackageFileUEVersion, GPackageFileLicenseeUEVersion, bInSaveUnversioned );
 		Summary.SavedByEngineVersion = FEngineVersion::Current();
 		Summary.CompatibleWithEngineVersion = FEngineVersion::CompatibleWith();
 		Summary.SetPackageFlags(Package ? Package->GetPackageFlags() : 0);
diff --git a/Engine/Source/Runtime/CoreUObject/Private/UObject/ObjectBaseUtility.cpp b/Engine/Source/Runtime/CoreUObject/Private/UObject/ObjectBaseUtility.cpp
--- a/Engine/Source/Runtime/CoreUObject/Private/UObject/ObjectBaseUtility.cpp
+++ b/Engine/Source/Runtime/CoreUObject/Private/UObject/ObjectBaseUtility.cpp
@@ -11,7 +11,7 @@
  * Returns the UE4 version of the linker for this object.
  *
  * @return	the UE4 version of the engine's package file when this object
- *			was last saved, or GPackageFileUE4Version (current version) if
+ *			was last saved, or GPackageFileUEVersion (current version) if
  *			this object does not have a linker, which indicates that
  *			a) this object is a native only class, or
  *			b) this object's linker has been detached, in which case it is already fully loaded
@@ -45,7 +45,7 @@ int32 UObjectBaseUtility::GetLinkerUE4Version() const
 	else
 	{
 		// We don't have a linker associated as we e.g. might have been saved or had loaders reset, ...
-		return GPackageFileUE4Version;
+		return GPackageFileUEVersion;
 	}
 }
 
@@ -118,6 +118,6 @@ int32 UObjectBaseUtility::GetLinkerLicenseeUE4Version() const
 	else
 	{
 		// We don't have a linker associated as we e.g. might have been saved or had loaders reset, ...
-		return GPackageFileLicenseeUE4Version;
+		return GPackageFileLicenseeUEVersion;
 	}
 }
diff --git a/Engine/Source/Runtime/CoreUObject/Private/UObject/Package.cpp b/Engine/Source/Runtime/CoreUObject/Private/UObject/Package.cpp
--- a/Engine/Source/Runtime/CoreUObject/Private/UObject/Package.cpp
+++ b/Engine/Source/Runtime/CoreUObject/Private/UObject/Package.cpp
@@ -47,8 +47,8 @@ void UPackage::PostInitProperties()
 	MetaData = nullptr;
 	PersistentGuid = FGuid::NewGuid();
 #endif
-	LinkerPackageVersion = GPackageFileUE4Version;
-	LinkerLicenseeVersion = GPackageFileLicenseeUE4Version;
+	LinkerPackageVersion = GPackageFileUEVersion;
+	LinkerLicenseeVersion = GPackageFileLicenseeUEVersion;
 	PIEInstanceID = INDEX_NONE;
 #if WITH_EDITORONLY_DATA
 	bIsCookedForEditor = false;
diff --git a/Engine/Source/Runtime/CoreUObject/Private/UObject/PackageFileSummary.cpp b/Engine/Source/Runtime/CoreUObject/Private/UObject/PackageFileSummary.cpp
--- a/Engine/Source/Runtime/CoreUObject/Private/UObject/PackageFileSummary.cpp
+++ b/Engine/Source/Runtime/CoreUObject/Private/UObject/PackageFileSummary.cpp
@@ -141,8 +141,8 @@ void operator<<(FStructuredArchive::FSlot Slot, FPackageFileSummary& Sum)
 #endif
 					// this file is unversioned, remember that, then use current versions
 					Sum.bUnversioned = true;
-					Sum.FileVersionUE4 = GPackageFileUE4Version;
-					Sum.FileVersionLicenseeUE4 = GPackageFileLicenseeUE4Version;
+					Sum.FileVersionUE4 = GPackageFileUEVersion;
+					Sum.FileVersionLicenseeUE4 = GPackageFileLicenseeUEVersion;
 					Sum.CustomVersionContainer = FCurrentCustomVersions::GetAll();
 				}
 			}
@@ -209,7 +209,7 @@ void operator<<(FStructuredArchive::FSlot Slot, FPackageFileSummary& Sum)
 		Record << SA_VALUE(TEXT("ImportCount"), Sum.ImportCount) << SA_VALUE(TEXT("ImportOffset"), Sum.ImportOffset);
 		Record << SA_VALUE(TEXT("DependsOffset"), Sum.DependsOffset);
 
-		if (BaseArchive.IsLoading() && (Sum.FileVersionUE4 < VER_UE4_OLDEST_LOADABLE_PACKAGE || Sum.FileVersionUE4 > GPackageFileUE4Version))
+		if (BaseArchive.IsLoading() && (Sum.FileVersionUE4 < VER_UE4_OLDEST_LOADABLE_PACKAGE || Sum.FileVersionUE4 > GPackageFileUEVersion))
 		{
 			return; // we can't safely load more than this because the below was different in older files.
 		}
diff --git a/Engine/Source/Runtime/CoreUObject/Public/UObject/UObjectBaseUtility.h b/Engine/Source/Runtime/CoreUObject/Public/UObject/UObjectBaseUtility.h
--- a/Engine/Source/Runtime/CoreUObject/Public/UObject/UObjectBaseUtility.h
+++ b/Engine/Source/Runtime/CoreUObject/Public/UObject/UObjectBaseUtility.h
@@ -611,7 +611,7 @@ public:
 	 * Returns the UE4 version of the linker for this object.
 	 *
 	 * @return	the UE4 version of the engine's package file when this object
-	 *			was last saved, or GPackageFileUE4Version (current version) if
+	 *			was last saved, or GPackageFileUEVersion (current version) if
 	 *			this object does not have a linker, which indicates that
 	 *			a) this object is a native only class, or
 	 *			b) this object's linker has been detached, in which case it is already fully loaded
diff --git a/Engine/Source/Runtime/Engine/Private/GameplayStatics.cpp b/Engine/Source/Runtime/Engine/Private/GameplayStatics.cpp
--- a/Engine/Source/Runtime/Engine/Private/GameplayStatics.cpp
+++ b/Engine/Source/Runtime/Engine/Private/GameplayStatics.cpp
@@ -106,7 +106,7 @@ FSaveGameHeader::FSaveGameHeader()
 FSaveGameHeader::FSaveGameHeader(TSubclassOf<USaveGame> ObjectType)
 	: FileTypeTag(UE_SAVEGAME_FILE_TYPE_TAG)
 	, SaveGameFileVersion(FSaveGameFileVersion::LatestVersion)
-	, PackageFileUE4Version(GPackageFileUE4Version)
+	, PackageFileUE4Version(GPackageFileUEVersion)
 	, SavedEngineVersion(FEngineVersion::Current())
 	, CustomVersionFormat(static_cast<int32>(ECustomVersionSerializationFormat::Latest))
 	, CustomVersions(FCurrentCustomVersions::GetAll())
@@ -144,7 +144,7 @@ void FSaveGameHeader::Read(FMemoryReader& MemoryReader)
 		// Note for 4.8 and beyond: if you get a crash loading a pre-4.8 version of your savegame file and 
 		// you don't want to delete it, try uncommenting these lines and changing them to use the version 
 		// information from your previous build. Then load and resave your savegame file.
-		//MemoryReader.SetUE4Ver(MyPreviousUE4Version);				// @see GPackageFileUE4Version
+		//MemoryReader.SetUE4Ver(MyPreviousUE4Version);				// @see GPackageFileUEVersion
 		//MemoryReader.SetEngineVer(MyPreviousEngineVersion);		// @see FEngineVersion::Current()
 	}
 	else
diff --git a/Engine/Source/Runtime/NetworkFileSystem/Private/NetworkFileServerConnection.cpp b/Engine/Source/Runtime/NetworkFileSystem/Private/NetworkFileServerConnection.cpp
--- a/Engine/Source/Runtime/NetworkFileSystem/Private/NetworkFileServerConnection.cpp
+++ b/Engine/Source/Runtime/NetworkFileSystem/Private/NetworkFileServerConnection.cpp
@@ -997,7 +997,7 @@ bool FNetworkFileServerClientConnection::ProcessGetFileList( FArchive& In, FArch
 		}
 		else
 		{
-			//@todo: This assumes the game is located in the UE4 Root directory
+			//@todo: This assumes the game is located in the Unreal Root directory
 			SandboxDirectory = FPaths::Combine(*FPaths::GetRelativePathToRoot(), *GameName, TEXT("Saved"), TEXT("Cooked"), *ConnectedPlatformName);
 		}
 	}
@@ -1096,10 +1096,10 @@ bool FNetworkFileServerClientConnection::ProcessGetFileList( FArchive& In, FArch
 
 	// report the package version information
 	// The downside of this is that ALL cooked data will get tossed on package version changes
-	int32 PackageFileUE4Version = GPackageFileUE4Version;
-	Out << PackageFileUE4Version;
-	int32 PackageFileLicenseeUE4Version = GPackageFileLicenseeUE4Version;
-	Out << PackageFileLicenseeUE4Version;
+	int32 PackageFileUnrealVersion = GPackageFileUEVersion;
+	Out << PackageFileUnrealVersion;
+	int32 PackageFileLicenseeUnrealVersion = GPackageFileLicenseeUEVersion;
+	Out << PackageFileLicenseeUnrealVersion;
 
 	// Send *our* engine and game dirs
 	Out << LocalEngineDir;
diff --git a/Engine/Source/Runtime/Unix/UnixCommonStartup/Private/UnixCommonStartup.cpp b/Engine/Source/Runtime/Unix/UnixCommonStartup/Private/UnixCommonStartup.cpp
--- a/Engine/Source/Runtime/Unix/UnixCommonStartup/Private/UnixCommonStartup.cpp
+++ b/Engine/Source/Runtime/Unix/UnixCommonStartup/Private/UnixCommonStartup.cpp
@@ -198,7 +198,7 @@ int CommonUnixMain(int argc, char *argv[], int (*RealMain)(const TCHAR * Command
 	if (UE_BUILD_SHIPPING)
 	{
 		// only printed in shipping
-		printf("%s %d %d\n", StringCast<ANSICHAR>(*FEngineVersion::Current().ToString()).Get(), GPackageFileUE4Version, GPackageFileLicenseeUE4Version);
+		printf("%s %d %d\n", StringCast<ANSICHAR>(*FEngineVersion::Current().ToString()).Get(), GPackageFileUEVersion, GPackageFileLicenseeUEVersion);
 	}
 
 	int ErrorLevel = 0;
