Clamp mouse sensitivity minimum to non-zero number to avoid divide by zero and difficult to track down WASD movement loss when this setting was 0.

Also gave the setting the default value, in case the setting is removed from a project's ini files to preserve some decent behavior in the editor.

diff --git a/Engine/Source/Editor/UnrealEd/Classes/Settings/LevelEditorViewportSettings.h b/Engine/Source/Editor/UnrealEd/Classes/Settings/LevelEditorViewportSettings.h
--- a/Engine/Source/Editor/UnrealEd/Classes/Settings/LevelEditorViewportSettings.h
+++ b/Engine/Source/Editor/UnrealEd/Classes/Settings/LevelEditorViewportSettings.h
@@ -280,7 +280,7 @@ class UNREALED_API ULevelEditorViewportSettings
 	int32 MouseScrollCameraSpeed;
 
 	/** The sensitivity of mouse movement when rotating the camera. */
-	UPROPERTY(EditAnywhere, config, Category=Controls, meta=(DisplayName="Mouse Sensitivity", ClampMin="0.0",ClampMax="1.0") )
+	UPROPERTY(EditAnywhere, config, Category=Controls, meta=(DisplayName="Mouse Sensitivity", ClampMin="0.01",ClampMax="1.0") )
 	float MouseSensitivty;
 	
 	/** Whether or not to invert mouse on the y axis in free look mode */
diff --git a/Engine/Source/Editor/UnrealEd/Private/Settings/SettingsClasses.cpp b/Engine/Source/Editor/UnrealEd/Private/Settings/SettingsClasses.cpp
--- a/Engine/Source/Editor/UnrealEd/Private/Settings/SettingsClasses.cpp
+++ b/Engine/Source/Editor/UnrealEd/Private/Settings/SettingsClasses.cpp
@@ -910,6 +910,7 @@ ULevelEditorViewportSettings::ULevelEditorViewportSettings( const FObjectInitial
 	MeasuringToolUnits = MeasureUnits_Centimeters;
 	bAllowArcballRotate = false;
 	bAllowScreenRotate = false;
+	MouseSensitivty = .2f;
 	// Set a default preview mesh
 	PreviewMeshes.Add(FSoftObjectPath("/Engine/EditorMeshes/ColorCalibrator/SM_ColorCalibrator.SM_ColorCalibrator"));
 	LastInViewportMenuLocation = FVector2D(EForceInit::ForceInitToZero);
@@ -920,6 +921,9 @@ void ULevelEditorViewportSettings::PostInitProperties()
 	Super::PostInitProperties();
 	UBillboardComponent::SetEditorScale(BillboardScale);
 	UArrowComponent::SetEditorScale(BillboardScale);
+
+	// Make sure the mouse sensitivity is not set below a valid value somehow. This can cause weird viewport interactions.
+	MouseSensitivty = FMath::Max(MouseSensitivty, .01f);
 }
 
 void ULevelEditorViewportSettings::PostEditChangeProperty( struct FPropertyChangedEvent& PropertyChangedEvent )
