Fix for ROBOMERGE issue related to change of API name deprecation

diff --git a/Engine/Source/Editor/UnrealEd/Private/EditorDomain/EditorDomainUtils.cpp b/Engine/Source/Editor/UnrealEd/Private/EditorDomain/EditorDomainUtils.cpp
--- a/Engine/Source/Editor/UnrealEd/Private/EditorDomain/EditorDomainUtils.cpp
+++ b/Engine/Source/Editor/UnrealEd/Private/EditorDomain/EditorDomainUtils.cpp
@@ -26,14 +26,14 @@ const TCHAR* EditorDomainPackageBucketName = TEXT("EditorDomainPackage");
 FPackageDigest GetPackageDigest(const FAssetPackageData& PackageData, FName PackageName)
 {
 	FCbWriter Writer;
-	int32 CurrentFileVersionUE4 = GPackageFileUE4Version;
-	int32 CurrentFileVersionLicenseeUE4 = GPackageFileLicenseeUE4Version;
+	int32 CurrentFileVersionUE = GPackageFileUEVersion;
+	int32 CurrentFileVersionLicenseeUE = GPackageFileLicenseeUEVersion;
 	Writer << EditorDomainVersion;
 	PRAGMA_DISABLE_DEPRECATION_WARNINGS
 	Writer << const_cast<FGuid&>(PackageData.PackageGuid);
 	PRAGMA_ENABLE_DEPRECATION_WARNINGS
-	Writer << CurrentFileVersionUE4;
-	Writer << CurrentFileVersionLicenseeUE4;
+	Writer << CurrentFileVersionUE;
+	Writer << CurrentFileVersionLicenseeUE;
 	check(Algo::IsSorted(PackageData.GetCustomVersions()));
 	for (const UE::AssetRegistry::FPackageCustomVersion& PackageVersion : PackageData.GetCustomVersions())
 	{
