Custom Menu Delegates for Toolbar Widgets + Viewport Transform ToolBar creates menu items for overflow widgets

diff --git a/Engine/Source/Editor/LevelEditor/Private/SLevelViewportToolBar.cpp b/Engine/Source/Editor/LevelEditor/Private/SLevelViewportToolBar.cpp
--- a/Engine/Source/Editor/LevelEditor/Private/SLevelViewportToolBar.cpp
+++ b/Engine/Source/Editor/LevelEditor/Private/SLevelViewportToolBar.cpp
@@ -282,7 +282,7 @@ void SLevelViewportToolBar::Construct( const FArguments& InArgs )
 				[
 					SNew(STransformViewportToolBar)
 					.Viewport(ViewportRef)
-					.CommandList(LevelEditorModule.GetGlobalLevelEditorActions())
+					.CommandList(ViewportRef->GetCommandList())
 					.Extenders(LevelEditorModule.GetToolBarExtensibilityManager()->GetAllExtenders())
 					.Visibility(ViewportRef, &SLevelViewport::GetTransformToolbarVisibility)
 				]
diff --git a/Engine/Source/Editor/UnrealEd/Private/STransformViewportToolbar.cpp b/Engine/Source/Editor/UnrealEd/Private/STransformViewportToolbar.cpp
--- a/Engine/Source/Editor/UnrealEd/Private/STransformViewportToolbar.cpp
+++ b/Engine/Source/Editor/UnrealEd/Private/STransformViewportToolbar.cpp
@@ -32,6 +32,7 @@
 #include "LevelEditorActions.h"
 #include "Styling/ToolBarStyle.h"
 #include "SEditorViewportToolBarMenu.h"
+#include "SEditorViewportToolBarButton.h"
 #include "Widgets/Layout/SSpacer.h"
 
 #define LOCTEXT_NAMESPACE "TransformToolBar"
@@ -93,7 +94,7 @@ TSharedRef<SWidget> STransformViewportToolBar::GenerateSurfaceSnappingMenu()
 					return Settings.bSnapRotation;
 					})),
 				NAME_None,
-			EUserInterfaceActionType::RadioButton);
+			EUserInterfaceActionType::Check);
 
 
 		MenuBuilder.AddWidget(
@@ -210,12 +211,20 @@ TSharedRef< SWidget > STransformViewportToolBar::MakeTransformToolBar( const TSh
 		ToolbarBuilder.SetIsFocusable( true );
 
 		ToolbarBuilder.AddToolBarButton( FEditorViewportCommands::Get().CycleTransformGizmoCoordSystem,
-										NAME_None,
-										TAttribute<FText>(),
-										TAttribute<FText>(),
-										TAttribute<FSlateIcon>(this, &STransformViewportToolBar::GetLocalToWorldIcon),
-										FName(TEXT("CycleTransformGizmoCoordSystem"))
-										);
+			NAME_None,
+			TAttribute<FText>(),
+			TAttribute<FText>(),
+			TAttribute<FSlateIcon>(this, &STransformViewportToolBar::GetLocalToWorldIcon),
+			FName(TEXT("CycleTransformGizmoCoordSystem")),
+
+			// explictly specify what this widget should look like as a menu item
+			FNewMenuDelegate::CreateLambda( []( FMenuBuilder& InMenuBuilder )
+			{
+				InMenuBuilder.AddMenuEntry(FEditorViewportCommands::Get().RelativeCoordinateSystem_World);
+				InMenuBuilder.AddMenuEntry(FEditorViewportCommands::Get().RelativeCoordinateSystem_Local);
+			}
+		));
+
 	}
 
 	ToolbarBuilder.EndSection();
@@ -223,56 +232,101 @@ TSharedRef< SWidget > STransformViewportToolBar::MakeTransformToolBar( const TSh
 	ToolbarBuilder.BeginSection("LocationGridSnap");
 	{
 		static FName SurfaceSnapName = FName(TEXT("SurfaceSnap"));
-		ToolbarBuilder.AddWidget( MakeSurfaceSnappingButton(), SurfaceSnapName );
+		ToolbarBuilder.AddWidget( MakeSurfaceSnappingButton(), SurfaceSnapName, false, HAlign_Fill, 
+			FNewMenuDelegate::CreateLambda( [this]( FMenuBuilder& InMenuBuilder )
+			{
+				InMenuBuilder.AddWrapperSubMenu(
+					LOCTEXT("SnapToSurfaceMenuSettings", "Surface Snap Settings"),
+					LOCTEXT("SnapToSurfaceMenuSettings_Tooltip", "Snap To Surface Settings"),
+					FOnGetContent::CreateSP(this, &STransformViewportToolBar::GenerateSurfaceSnappingMenu),
+					FSlateIcon(FAppStyle::GetAppStyleSetName(), "EditorViewport.ToggleSurfaceSnapping")
+				);
+			}
+		));
 
 		ToolbarBuilder.AddSeparator();
 
 		// Grab the existing UICommand 
-		FUICommandInfo* Command = FEditorViewportCommands::Get().LocationGridSnap.Get();
+		TSharedPtr<FUICommandInfo> Command = FEditorViewportCommands::Get().LocationGridSnap;
 
 		static FName PositionSnapName = FName(TEXT("PositionSnap"));
 
 		// Setup a GridSnapSetting with the UICommand
-		ToolbarBuilder.AddWidget(	SNew(SViewportToolBarComboMenu)
-									.Style(ToolBarStyle)
-									.IsChecked(this, &STransformViewportToolBar::IsLocationGridSnapChecked)
-									.OnCheckStateChanged(this, &STransformViewportToolBar::HandleToggleLocationGridSnap)
-									.Label(this, &STransformViewportToolBar::GetLocationGridLabel)
-									.OnGetMenuContent(this, &STransformViewportToolBar::FillLocationGridSnapMenu)
-									.ToggleButtonToolTip(Command->GetDescription())
-									.MenuButtonToolTip(LOCTEXT("LocationGridSnap_ToolTip", "Set the Position Grid Snap value"))
-									.Icon(Command->GetIcon())
-									.ParentToolBar(SharedThis(this))
-									,PositionSnapName );
+		ToolbarBuilder.AddWidget(	
+			SNew(SViewportToolBarComboMenu)
+				.Style(ToolBarStyle)
+				.IsChecked(this, &STransformViewportToolBar::IsLocationGridSnapChecked)
+				.OnCheckStateChanged(this, &STransformViewportToolBar::HandleToggleLocationGridSnap)
+				.Label(this, &STransformViewportToolBar::GetLocationGridLabel)
+				.OnGetMenuContent(this, &STransformViewportToolBar::FillLocationGridSnapMenu)
+				.ToggleButtonToolTip(Command->GetDescription())
+				.MenuButtonToolTip(LOCTEXT("LocationGridSnap_ToolTip", "Set the Position Grid Snap value"))
+				.Icon(Command->GetIcon())
+				.ParentToolBar(SharedThis(this)),
+			PositionSnapName, 
+			false, 
+			HAlign_Fill, 
+
+			// explictly specify what this widget should look like as a menu item
+			FNewMenuDelegate::CreateLambda( [this, Command]( FMenuBuilder& InMenuBuilder )
+			{
+				// TODO - debug why can't just use the Command / mapping isn't working 
+				InMenuBuilder.AddMenuEntry(Command);
+
+				InMenuBuilder.AddWrapperSubMenu(
+					LOCTEXT("GridSnapMenuSettings", "Grid Snap Settings"),
+					LOCTEXT("GridSnapMenuSettings_ToolTip", "Set the Position Grid Snap value"),
+					FOnGetContent::CreateSP(this, &STransformViewportToolBar::FillLocationGridSnapMenu),
+					FSlateIcon(Command->GetIcon())
+				);
+			}
+		));
 	}
 	ToolbarBuilder.EndSection();
 
 	ToolbarBuilder.BeginSection("RotationGridSnap");
 	{
 		// Grab the existing UICommand 
-		FUICommandInfo* Command = FEditorViewportCommands::Get().RotationGridSnap.Get();
+		TSharedPtr<FUICommandInfo> Command = FEditorViewportCommands::Get().RotationGridSnap;
 
 		static FName RotationSnapName = FName(TEXT("RotationSnap"));
 
 		// Setup a GridSnapSetting with the UICommand
-		ToolbarBuilder.AddWidget(	SNew(SViewportToolBarComboMenu)
-									.Style(ToolBarStyle)
-									.IsChecked(this, &STransformViewportToolBar::IsRotationGridSnapChecked)
-									.OnCheckStateChanged(this, &STransformViewportToolBar::HandleToggleRotationGridSnap)
-									.Label(this, &STransformViewportToolBar::GetRotationGridLabel)
-									.OnGetMenuContent(this, &STransformViewportToolBar::FillRotationGridSnapMenu)
-									.ToggleButtonToolTip(Command->GetDescription())
-									.MenuButtonToolTip(LOCTEXT("RotationGridSnap_ToolTip", "Set the Rotation Grid Snap value"))
-									.Icon(Command->GetIcon())
-									.ParentToolBar(SharedThis(this))
-									, RotationSnapName );
+		ToolbarBuilder.AddWidget(	
+			SNew(SViewportToolBarComboMenu)
+				.Style(ToolBarStyle)
+				.IsChecked(this, &STransformViewportToolBar::IsRotationGridSnapChecked)
+				.OnCheckStateChanged(this, &STransformViewportToolBar::HandleToggleRotationGridSnap)
+				.Label(this, &STransformViewportToolBar::GetRotationGridLabel)
+				.OnGetMenuContent(this, &STransformViewportToolBar::FillRotationGridSnapMenu)
+				.ToggleButtonToolTip(Command->GetDescription())
+				.MenuButtonToolTip(LOCTEXT("RotationGridSnap_ToolTip", "Set the Rotation Grid Snap value"))
+				.Icon(Command->GetIcon())
+				.ParentToolBar(SharedThis(this)),
+			RotationSnapName,
+			false, 
+			HAlign_Fill, 
+
+			// explictly specify what this widget should look like as a menu item
+			FNewMenuDelegate::CreateLambda( [this, Command]( FMenuBuilder& InMenuBuilder )
+			{
+				InMenuBuilder.AddMenuEntry(Command);
+
+				InMenuBuilder.AddWrapperSubMenu(
+					LOCTEXT("RotationGridSnapMenuSettings", "Rotation Snap Settings"),
+					LOCTEXT("RotationGridSnapMenuSettings_ToolTip", "Adjust the Grid Settings for Rotation Snap"),
+					FOnGetContent::CreateSP(this, &STransformViewportToolBar::FillRotationGridSnapMenu),
+					FSlateIcon(Command->GetIcon())
+				);
+			}
+		));
 	}
 	ToolbarBuilder.EndSection();
 
 	ToolbarBuilder.BeginSection("Layer2DSnap");
 	{
 		// Grab the existing UICommand 
-		FUICommandInfo* Command = FEditorViewportCommands::Get().Layer2DSnap.Get();
+		TSharedPtr<FUICommandInfo> Command = FEditorViewportCommands::Get().Layer2DSnap;
 
 		static FName Layer2DSnapName = FName(TEXT("Layer2DSnap"));
 
@@ -290,7 +344,28 @@ TSharedRef< SWidget > STransformViewportToolBar::MakeTransformToolBar( const TSh
 			.ParentToolBar(SharedThis(this))
 			.MinDesiredButtonWidth(88.0f);
 
-		ToolbarBuilder.AddWidget(SnapLayerPickerWidget, Layer2DSnapName);
+		ToolbarBuilder.AddWidget(
+			SnapLayerPickerWidget, 
+			Layer2DSnapName,
+			false, 
+			HAlign_Fill, 
+
+			// explictly specify what this widget should look like as a menu item
+			FNewMenuDelegate::CreateLambda( [this, Command]( FMenuBuilder& InMenuBuilder )
+			{
+				if (IsLayer2DSnapVisible() == EVisibility::Visible)
+				{
+					InMenuBuilder.AddMenuEntry(Command);
+
+					InMenuBuilder.AddWrapperSubMenu(
+						LOCTEXT("Layer2DSnapMenuSettings", "Layer 2D Snap Settings"),
+						LOCTEXT("Layer2DSnapMenuSettings_ToolTip", "Adjust the Grid Settings for Layer 2D Snap"),
+						FOnGetContent::CreateSP(this, &STransformViewportToolBar::FillLayer2DSnapMenu),
+						FSlateIcon(Command->GetIcon())
+					);
+				}
+			}
+		));
 	}
 	ToolbarBuilder.EndSection();
 
@@ -299,37 +374,73 @@ TSharedRef< SWidget > STransformViewportToolBar::MakeTransformToolBar( const TSh
 	ToolbarBuilder.BeginSection("ScaleGridSnap");
 	{
 		// Grab the existing UICommand 
-		FUICommandInfo* Command = FEditorViewportCommands::Get().ScaleGridSnap.Get();
+		TSharedPtr<FUICommandInfo> Command = FEditorViewportCommands::Get().ScaleGridSnap;
 
 		static FName ScaleSnapName = FName(TEXT("ScaleSnap"));
 
 		// Setup a GridSnapSetting with the UICommand
-		ToolbarBuilder.AddWidget(	SNew(SViewportToolBarComboMenu)
-									.Cursor( EMouseCursor::Default )
-									.Style(ToolBarStyle)
-									.IsChecked(this,&STransformViewportToolBar::IsScaleGridSnapChecked)
-									.OnCheckStateChanged(this, &STransformViewportToolBar::HandleToggleScaleGridSnap)
-									.Label(this ,&STransformViewportToolBar::GetScaleGridLabel)
-									.OnGetMenuContent(this, &STransformViewportToolBar::FillScaleGridSnapMenu)
-									.ToggleButtonToolTip(Command->GetDescription())
-									.MenuButtonToolTip(LOCTEXT("ScaleGridSnap_ToolTip", "Set scaling options"))
-									.Icon(Command->GetIcon())
-									.ParentToolBar(SharedThis(this))
-									, ScaleSnapName );
+		ToolbarBuilder.AddWidget(	
+			SNew(SViewportToolBarComboMenu)
+				.Cursor( EMouseCursor::Default )
+				.Style(ToolBarStyle)
+				.IsChecked(this,&STransformViewportToolBar::IsScaleGridSnapChecked)
+				.OnCheckStateChanged(this, &STransformViewportToolBar::HandleToggleScaleGridSnap)
+				.Label(this ,&STransformViewportToolBar::GetScaleGridLabel)
+				.OnGetMenuContent(this, &STransformViewportToolBar::FillScaleGridSnapMenu)
+				.ToggleButtonToolTip(Command->GetDescription())
+				.MenuButtonToolTip(LOCTEXT("ScaleGridSnap_ToolTip", "Set scaling options"))
+				.Icon(Command->GetIcon())
+				.ParentToolBar(SharedThis(this)),
+			ScaleSnapName,
+			false, 
+			HAlign_Fill, 
+
+			// explictly specify what this widget should look like as a menu item
+			FNewMenuDelegate::CreateLambda( [this, Command]( FMenuBuilder& InMenuBuilder )
+			{
+				InMenuBuilder.AddMenuEntry(Command);
+
+				InMenuBuilder.AddWrapperSubMenu(
+					LOCTEXT("ScaleGridSnapMenuSettings", "Scale Snap Settings"),
+					LOCTEXT("ScaleGridSnapMenuSettings_ToolTip", "Adjust the Grid Settings for Scale Snap"),
+					FOnGetContent::CreateSP(this, &STransformViewportToolBar::FillScaleGridSnapMenu),
+					FSlateIcon(Command->GetIcon())
+				);
+			}
+		));
 	}
 	ToolbarBuilder.EndSection();
 
 	ToolbarBuilder.BeginSection("CameraSpeed");
 	{
+		static FName CameraSpeedName = FName(TEXT("CameraSpeed"));
+
 		// Camera speed 
 		ToolbarBuilder.AddWidget(	
 			SNew(SEditorViewportToolbarMenu)
-			.ParentToolBar(SharedThis(this))
-			.AddMetaData<FTagMetaData>(FTagMetaData(TEXT("CameraSpeedButton")))
-			.ToolTipText(LOCTEXT("CameraSpeed_ToolTip", "Camera Speed"))
-			.LabelIcon(FAppStyle::Get().GetBrush("EditorViewport.CamSpeedSetting"))
-			.Label(this, &STransformViewportToolBar::GetCameraSpeedLabel)
-			.OnGetMenuContent(this, &STransformViewportToolBar::FillCameraSpeedMenu));
+				.ParentToolBar(SharedThis(this))
+				.AddMetaData<FTagMetaData>(FTagMetaData(TEXT("CameraSpeedButton")))
+				.ToolTipText(LOCTEXT("CameraSpeed_ToolTip", "Camera Speed"))
+				.LabelIcon(FAppStyle::Get().GetBrush("EditorViewport.CamSpeedSetting"))
+				.Label(this, &STransformViewportToolBar::GetCameraSpeedLabel)
+				.OnGetMenuContent(this, &STransformViewportToolBar::FillCameraSpeedMenu),
+			CameraSpeedName,
+			false,
+			HAlign_Fill,
+
+			// explictly specify what this widget should look like as a menu item
+			FNewMenuDelegate::CreateLambda( [this]( FMenuBuilder& InMenuBuilder )
+			{
+				InMenuBuilder.AddWrapperSubMenu(
+					LOCTEXT("CameraSpeedMenuSettings", "Camera Speed Settings"),
+					LOCTEXT("CameraSpeedMenuSettings_ToolTip", "Adjust the camera navigation speed"),
+					FOnGetContent::CreateSP(this, &STransformViewportToolBar::FillCameraSpeedMenu),
+					FSlateIcon(FAppStyle::GetAppStyleSetName(), "EditorViewport.CamSpeedSetting")
+				);
+			}
+		));
+
+
 	}
 	ToolbarBuilder.EndSection();
 
@@ -407,6 +518,16 @@ TSharedRef<SWidget> STransformViewportToolBar::FillCameraSpeedMenu()
 	return ReturnWidget;
 }
 
+FReply STransformViewportToolBar::OnCycleCoordinateSystem()
+{
+	if( Viewport.IsValid() )
+	{
+		Viewport.Pin()->OnCycleCoordinateSystem();
+	}
+
+	return FReply::Handled();
+}
+
 FSlateIcon STransformViewportToolBar::GetLocalToWorldIcon() const
 {
 	if( Viewport.IsValid() && Viewport.Pin()->IsCoordSystemActive(COORD_World) )
diff --git a/Engine/Source/Editor/UnrealEd/Public/STransformViewportToolbar.h b/Engine/Source/Editor/UnrealEd/Public/STransformViewportToolbar.h
--- a/Engine/Source/Editor/UnrealEd/Public/STransformViewportToolbar.h
+++ b/Engine/Source/Editor/UnrealEd/Public/STransformViewportToolbar.h
@@ -50,6 +50,9 @@ private:
 	FSlateIcon GetLocalToWorldIcon() const;
 	FSlateIcon GetSurfaceSnappingIcon() const;
 
+	/** Callback to toggle between local and world space */
+	FReply OnCycleCoordinateSystem();
+
 	/** Camera speed menu construction callback */
 	TSharedRef<SWidget> FillCameraSpeedMenu();
 
diff --git a/Engine/Source/Runtime/Slate/Private/Framework/MultiBox/MultiBox.cpp b/Engine/Source/Runtime/Slate/Private/Framework/MultiBox/MultiBox.cpp
--- a/Engine/Source/Runtime/Slate/Private/Framework/MultiBox/MultiBox.cpp
+++ b/Engine/Source/Runtime/Slate/Private/Framework/MultiBox/MultiBox.cpp
@@ -1098,7 +1098,7 @@ void SMultiBoxWidget::BuildMultiBoxWidget()
 
 TSharedRef<SWidget> SMultiBoxWidget::OnWrapButtonClicked()
 {
-	FMenuBuilder MenuBuilder(true, NULL, TSharedPtr<FExtender>(), false, GetStyleSet());
+	FMenuBuilder MenuBuilder(true, MultiBox->GetLastCommandList(), TSharedPtr<FExtender>(), false, GetStyleSet());
 	{ 
 		const TArray< TSharedRef< const FMultiBlock > >& Blocks = MultiBox->GetBlocks();
 		for (int32 BlockIdx = ClippedHorizontalBox->GetClippedIndex(); BlockIdx < Blocks.Num(); ++BlockIdx)
diff --git a/Engine/Source/Runtime/Slate/Private/Framework/MultiBox/MultiBoxBuilder.cpp b/Engine/Source/Runtime/Slate/Private/Framework/MultiBox/MultiBoxBuilder.cpp
--- a/Engine/Source/Runtime/Slate/Private/Framework/MultiBox/MultiBoxBuilder.cpp
+++ b/Engine/Source/Runtime/Slate/Private/Framework/MultiBox/MultiBoxBuilder.cpp
@@ -381,7 +381,7 @@ void FMenuBarBuilder::ApplyHook(FName InExtensionHook, EExtensionHook::Position
 	}
 }
 
-void FToolBarBuilder::AddToolBarButton(const TSharedPtr< const FUICommandInfo > InCommand, FName InExtensionHook, const TAttribute<FText>& InLabelOverride, const TAttribute<FText>& InToolTipOverride, const TAttribute<FSlateIcon>& InIconOverride, FName InTutorialHighlightName )
+void FToolBarBuilder::AddToolBarButton(const TSharedPtr< const FUICommandInfo > InCommand, FName InExtensionHook, const TAttribute<FText>& InLabelOverride, const TAttribute<FText>& InToolTipOverride, const TAttribute<FSlateIcon>& InIconOverride, FName InTutorialHighlightName, FNewMenuDelegate InCustomMenuDelegate )
 {
 	ApplySectionBeginning();
 
@@ -398,6 +398,7 @@ void FToolBarBuilder::AddToolBarButton(const TSharedPtr< const FUICommandInfo >
 	NewToolBarButtonBlock->SetForceSmallIcons(bForceSmallIcons);
 	NewToolBarButtonBlock->SetTutorialHighlightName(GenerateTutorialIdentfierName(TutorialHighlightName, InTutorialHighlightName, InCommand, MultiBox->GetBlocks().Num()));
 	NewToolBarButtonBlock->SetStyleNameOverride(CurrentStyleOverride);
+	NewToolBarButtonBlock->SetCustomMenuDelegate(InCustomMenuDelegate);
 
 	MultiBox->AddMultiBlock( NewToolBarButtonBlock );
 
@@ -488,7 +489,7 @@ void FToolBarBuilder::AddToolBarWidget( TSharedRef<SWidget> InWidget, const TAtt
 	NewWidgetBlock->SetSearchable(bSearchable);
 }
 
-void FToolBarBuilder::AddWidget( TSharedRef<SWidget> InWidget, FName InTutorialHighlightName, bool bSearchable, EHorizontalAlignment Alignment )
+void FToolBarBuilder::AddWidget( TSharedRef<SWidget> InWidget, FName InTutorialHighlightName, bool bSearchable, EHorizontalAlignment Alignment, FNewMenuDelegate InCustomMenuDelegate )
 {
 	ApplySectionBeginning();
 
@@ -506,6 +507,7 @@ void FToolBarBuilder::AddWidget( TSharedRef<SWidget> InWidget, FName InTutorialH
 	TSharedRef< FWidgetBlock > NewWidgetBlock( new FWidgetBlock( InWidget, FText::GetEmpty(), true, Alignment) );
 	MultiBox->AddMultiBlock( NewWidgetBlock );
 	NewWidgetBlock->SetSearchable(bSearchable);
+	NewWidgetBlock->SetCustomMenuDelegate(InCustomMenuDelegate);
 }
 
 void FToolBarBuilder::AddSeparator(FName InExtensionHook)
diff --git a/Engine/Source/Runtime/Slate/Private/Framework/MultiBox/SToolBarButtonBlock.cpp b/Engine/Source/Runtime/Slate/Private/Framework/MultiBox/SToolBarButtonBlock.cpp
--- a/Engine/Source/Runtime/Slate/Private/Framework/MultiBox/SToolBarButtonBlock.cpp
+++ b/Engine/Source/Runtime/Slate/Private/Framework/MultiBox/SToolBarButtonBlock.cpp
@@ -38,21 +38,40 @@ FToolBarButtonBlock::FToolBarButtonBlock( const TAttribute<FText>& InLabel, cons
 {
 }
 
+void FToolBarButtonBlock::SetCustomMenuDelegate( FNewMenuDelegate& InCustomMenuDelegate )
+{
+	CustomMenuDelegate = InCustomMenuDelegate;	
+}
+
 void FToolBarButtonBlock::CreateMenuEntry(FMenuBuilder& MenuBuilder) const
 {
+	// Setup Command Context 
 	TSharedPtr<const FUICommandInfo> MenuEntryAction = GetAction();
 	TSharedPtr<const FUICommandList> MenuEntryActionList = GetActionList();
-	if (MenuEntryAction.IsValid() && MenuEntryActionList.IsValid())
-	{
+	bool bHasValidCommand = MenuEntryAction.IsValid() && MenuEntryActionList.IsValid();
+	if (bHasValidCommand) 
+	{	
 		MenuBuilder.PushCommandList(MenuEntryActionList.ToSharedRef());
+	}
+
+	if ( CustomMenuDelegate.IsBound() )
+	{
+		CustomMenuDelegate.Execute(MenuBuilder);
+	}
+	else if (bHasValidCommand)
+	{
 		MenuBuilder.AddMenuEntry(MenuEntryAction);
-		MenuBuilder.PopCommandList();
 	}
 	else if ( LabelOverride.IsSet() )
 	{
 		const FUIAction& DirectAction = GetDirectActions();
 		MenuBuilder.AddMenuEntry( LabelOverride.Get(), ToolTipOverride.Get(), IconOverride.Get(), DirectAction );
 	}
+
+	if (bHasValidCommand) 
+	{	
+		MenuBuilder.PopCommandList();
+	}
 }
 
 bool FToolBarButtonBlock::HasIcon() const
diff --git a/Engine/Source/Runtime/Slate/Private/Framework/MultiBox/SWidgetBlock.cpp b/Engine/Source/Runtime/Slate/Private/Framework/MultiBox/SWidgetBlock.cpp
--- a/Engine/Source/Runtime/Slate/Private/Framework/MultiBox/SWidgetBlock.cpp
+++ b/Engine/Source/Runtime/Slate/Private/Framework/MultiBox/SWidgetBlock.cpp
@@ -21,10 +21,23 @@ FWidgetBlock::FWidgetBlock( TSharedRef<SWidget> InContent, const FText& InLabel,
 }
 
 
+void FWidgetBlock::SetCustomMenuDelegate( FNewMenuDelegate& InCustomMenuDelegate )
+{
+	CustomMenuDelegate = InCustomMenuDelegate;	
+}
+
+
 void FWidgetBlock::CreateMenuEntry(FMenuBuilder& MenuBuilder) const
 {
-	FText EntryLabel = (!Label.IsEmpty()) ? Label : NSLOCTEXT("WidgetBlock", "CustomControl", "Custom Control");
-	MenuBuilder.AddWidget(ContentWidget, FText::GetEmpty(), true);
+	if ( CustomMenuDelegate.IsBound() )
+	{
+		CustomMenuDelegate.Execute(MenuBuilder);
+	}
+	else
+	{
+		FText EntryLabel = (!Label.IsEmpty()) ? Label : NSLOCTEXT("WidgetBlock", "CustomControl", "Custom Control");
+		MenuBuilder.AddWidget(ContentWidget, FText::GetEmpty(), true);
+	}
 }
 
 
diff --git a/Engine/Source/Runtime/Slate/Private/Framework/MultiBox/SWidgetBlock.h b/Engine/Source/Runtime/Slate/Private/Framework/MultiBox/SWidgetBlock.h
--- a/Engine/Source/Runtime/Slate/Private/Framework/MultiBox/SWidgetBlock.h
+++ b/Engine/Source/Runtime/Slate/Private/Framework/MultiBox/SWidgetBlock.h
@@ -6,6 +6,8 @@
 #include "Widgets/SWidget.h"
 #include "Widgets/DeclarativeSyntaxSupport.h"
 #include "Framework/MultiBox/MultiBox.h"
+#include "Framework/MultiBox/MultiBoxBuilder.h"
+
 
 /**
  * Arbitrary Widget MultiBlock
@@ -28,6 +30,8 @@ public:
 	/** FMultiBlock interface */
 	virtual void CreateMenuEntry(class FMenuBuilder& MenuBuilder) const override;
 
+	/** Set optional delegate to customize when a menu appears instead of the widget, such as in toolbars */
+	void SetCustomMenuDelegate( FNewMenuDelegate& InOnFillMenuDelegate);
 
 private:
 
@@ -50,6 +54,9 @@ private:
 
 	/** Hortizontal aligment for this widget in its parent container. Note: only applies to toolbars */
 	EHorizontalAlignment HorizontalAlignment;
+
+	/** Optional delegate to customize when a menu appears instead of the widget, such as in toolbars */
+	FNewMenuDelegate CustomMenuDelegate;
 };
 
 
diff --git a/Engine/Source/Runtime/Slate/Public/Framework/MultiBox/MultiBox.h b/Engine/Source/Runtime/Slate/Public/Framework/MultiBox/MultiBox.h
--- a/Engine/Source/Runtime/Slate/Public/Framework/MultiBox/MultiBox.h
+++ b/Engine/Source/Runtime/Slate/Public/Framework/MultiBox/MultiBox.h
@@ -397,6 +397,9 @@ public:
 	/* Whether the MultiBox has a search widget */
 	bool bHasSearchWidget;
 
+	/* Returns the last command list used */
+	const TSharedPtr<const FUICommandList> GetLastCommandList() const { return CommandLists.Last(); }
+
 private:
 	
 	/**
diff --git a/Engine/Source/Runtime/Slate/Public/Framework/MultiBox/MultiBoxBuilder.h b/Engine/Source/Runtime/Slate/Public/Framework/MultiBox/MultiBoxBuilder.h
--- a/Engine/Source/Runtime/Slate/Public/Framework/MultiBox/MultiBoxBuilder.h
+++ b/Engine/Source/Runtime/Slate/Public/Framework/MultiBox/MultiBoxBuilder.h
@@ -426,8 +426,9 @@ public:
 	 * @param	InToolTipOverride		Optional tool tip override.	 If omitted, then the action's label will be used instead.
 	 * @param	InIconOverride			Optional name of the slate brush to use for the tool bar image.  If omitted, then the action's icon will be used instead.
 	 * @param	InTutorialHighlightName	Name to identify this widget and highlight during tutorials
+	 * @param	CustomMenuDelegate  Optional custom menu delegate for cases where the toolbar is compressed into a menu
 	 */
-	void AddToolBarButton(const TSharedPtr< const FUICommandInfo > InCommand, FName InExtensionHook = NAME_None, const TAttribute<FText>& InLabelOverride = TAttribute<FText>(), const TAttribute<FText>& InToolTipOverride = TAttribute<FText>(), const TAttribute<FSlateIcon>& InIconOverride = TAttribute<FSlateIcon>(), FName InTutorialHighlightName = NAME_None );
+	void AddToolBarButton(const TSharedPtr< const FUICommandInfo > InCommand, FName InExtensionHook = NAME_None, const TAttribute<FText>& InLabelOverride = TAttribute<FText>(), const TAttribute<FText>& InToolTipOverride = TAttribute<FText>(), const TAttribute<FSlateIcon>& InIconOverride = TAttribute<FSlateIcon>(), FName InTutorialHighlightName = NAME_None, FNewMenuDelegate InCustomMenuDelegate = FNewMenuDelegate());
 	
 	/**
 	 * Adds a tool bar button
@@ -473,8 +474,9 @@ public:
 	 * @param	InTutorialHighlightName	Name to identify this widget and highlight during tutorials
 	 * @param	bSearchable			If true, widget will be searchable (default == true)
 	 * @param	Alignment			Horizontal alignment for the widget inside the toolbar
+	 * @param	CustomMenuDelegate  Optional custom menu delegate for cases where the toolbar is compressed into a menu
 	 */
-	void AddWidget(TSharedRef<SWidget> InWidget, FName InTutorialHighlightName = NAME_None, bool bSearchable = true, EHorizontalAlignment Alignment = HAlign_Fill);
+	void AddWidget(TSharedRef<SWidget> InWidget, FName InTutorialHighlightName = NAME_None, bool bSearchable = true, EHorizontalAlignment Alignment = HAlign_Fill, FNewMenuDelegate InCustomMenuDelegate = FNewMenuDelegate());
 	
 	/**
 	 * Adds a toolbar separator
diff --git a/Engine/Source/Runtime/Slate/Public/Framework/MultiBox/SToolBarButtonBlock.h b/Engine/Source/Runtime/Slate/Public/Framework/MultiBox/SToolBarButtonBlock.h
--- a/Engine/Source/Runtime/Slate/Public/Framework/MultiBox/SToolBarButtonBlock.h
+++ b/Engine/Source/Runtime/Slate/Public/Framework/MultiBox/SToolBarButtonBlock.h
@@ -11,6 +11,7 @@
 #include "Framework/Commands/UICommandInfo.h"
 #include "Framework/Commands/UICommandList.h"
 #include "Framework/MultiBox/MultiBox.h"
+#include "Framework/MultiBox/MultiBoxBuilder.h"
 
 class SBorder;
 
@@ -56,6 +57,9 @@ public:
 	virtual void CreateMenuEntry(class FMenuBuilder& MenuBuilder) const override;
 	virtual bool HasIcon() const override;
 
+	/** Set optional delegate to customize when a menu appears instead of the widget, such as in toolbars */
+	void SetCustomMenuDelegate( FNewMenuDelegate& InOnFillMenuDelegate);
+
 private:
 
 	/** FMultiBlock private interface */
@@ -86,6 +90,9 @@ private:
 
 	/** Whether this toolbar should always use small icons, regardless of the current settings */
 	bool bForceSmallIcons;
+
+	/** Optional delegate to customize when a menu appears instead of the widget, such as in toolbars */
+	FNewMenuDelegate CustomMenuDelegate;
 };
 
 
