Fixed an issue with deferred deletion and toggling editor modes that caused the non-integrated toolkit palettes to linger while the mode was getting cleaned up on tick.

diff --git a/Engine/Source/Editor/UnrealEd/Private/EditorModeManager.cpp b/Engine/Source/Editor/UnrealEd/Private/EditorModeManager.cpp
--- a/Engine/Source/Editor/UnrealEd/Private/EditorModeManager.cpp
+++ b/Engine/Source/Editor/UnrealEd/Private/EditorModeManager.cpp
@@ -646,12 +646,7 @@ void FEditorModeTools::DeactivateScriptableModeAtIndex(int32 InIndex)
 	UEdMode* Mode = ActiveScriptableModes[InIndex];
 	ActiveScriptableModes.RemoveAt(InIndex);
 
-	Mode->Exit();
-
-	const bool bIsEnteringMode = false;
-	BroadcastEditorModeIDChanged(Mode->GetID(), bIsEnteringMode);
-
-	// Remove the toolbar widget
+	// Remove the toolbar widget to remove any toolkit references that may get removed when the mode exits.
 	ActiveToolBarRows.RemoveAll(
 		[&Mode](FEdModeToolbarRow& Row)
 		{
@@ -661,6 +656,11 @@ void FEditorModeTools::DeactivateScriptableModeAtIndex(int32 InIndex)
 
 	RebuildModeToolBar();
 
+	Mode->Exit();
+
+	const bool bIsEnteringMode = false;
+	BroadcastEditorModeIDChanged(Mode->GetID(), bIsEnteringMode);
+
 	RecycledScriptableModes.Add(Mode->GetID(), Mode);
 }
 
@@ -804,7 +804,7 @@ void FEditorModeTools::RebuildModeToolBar()
 
 void FEditorModeTools::SpawnOrUpdateModeToolbar()
 {
-	if(ShouldShowModeToolbar())
+	if (ShouldShowModeToolbar())
 	{
 		if (ModeToolbarTab.IsValid())
 		{
@@ -1029,12 +1029,11 @@ void FEditorModeTools::ActivateMode(FEditorModeID InID, bool bToggle)
 	// Ask the mode to build the toolbar.
 	TSharedPtr<FUICommandList> CommandList;
 	const TSharedPtr<FModeToolkit> Toolkit = ScriptableMode->GetToolkit().Pin();
-	if (Toolkit.IsValid())
+	if (Toolkit && !Toolkit->HasIntegratedToolPalettes())
 	{
 		CommandList = Toolkit->GetToolkitCommands();
 
 		// Also build the toolkit here 
-		int32 PaletteCount = 0;
 		TArray<FName> PaletteNames;
 		Toolkit->GetToolPaletteNames(PaletteNames);
 		for (auto Palette : PaletteNames)
@@ -1042,17 +1041,12 @@ void FEditorModeTools::ActivateMode(FEditorModeID InID, bool bToggle)
 			FUniformToolBarBuilder ModeToolbarBuilder(CommandList, FMultiBoxCustomization(ScriptableMode->GetModeInfo().ToolbarCustomizationName), TSharedPtr<FExtender>(), false);
 			ModeToolbarBuilder.SetStyle(&FEditorStyle::Get(), "PaletteToolBar");
 			Toolkit->BuildToolPalette(Palette, ModeToolbarBuilder);
-
 			ActiveToolBarRows.Emplace(ScriptableMode->GetID(), Palette, Toolkit->GetToolPaletteDisplayName(Palette), ModeToolbarBuilder.MakeWidget());
-			PaletteCount++;
-		}
-
-		if (!Toolkit->HasIntegratedToolPalettes() && PaletteCount > 0)
-		{
-			SpawnOrUpdateModeToolbar();
 		}
 	}
 
+	SpawnOrUpdateModeToolbar();
+
 	RecycledScriptableModes.Remove(InID);
 
 	// Update the editor UI
diff --git a/Engine/Source/Editor/UnrealEd/Private/Tools/UEdMode.cpp b/Engine/Source/Editor/UnrealEd/Private/Tools/UEdMode.cpp
--- a/Engine/Source/Editor/UnrealEd/Private/Tools/UEdMode.cpp
+++ b/Engine/Source/Editor/UnrealEd/Private/Tools/UEdMode.cpp
@@ -127,7 +127,6 @@ void UEdMode::Exit()
 			ToolsContext->ToolManager->UnregisterToolType(RegisteredTool.Value);
 		}
 
-		FToolkitManager::Get().CloseToolkit(Toolkit.ToSharedRef());
 		Toolkit.Reset();
 	}
 	RegisteredTools.SetNum(0);
@@ -155,6 +154,16 @@ class FEditorModeTools* UEdMode::GetModeManager() const
 	return Owner;
 }
 
+void UEdMode::RequestDeletion()
+{
+	bPendingDeletion = true;
+	
+	if (UsesToolkits() && Toolkit)
+	{
+		FToolkitManager::Get().CloseToolkit(Toolkit.ToSharedRef());
+	}
+}
+
 AActor* UEdMode::GetFirstSelectedActorInstance() const
 {
 	return Owner->GetEditorSelectionSet()->GetTopSelectedObject<AActor>();
diff --git a/Engine/Source/Editor/UnrealEd/Public/Tools/UEdMode.h b/Engine/Source/Editor/UnrealEd/Public/Tools/UEdMode.h
--- a/Engine/Source/Editor/UnrealEd/Public/Tools/UEdMode.h
+++ b/Engine/Source/Editor/UnrealEd/Public/Tools/UEdMode.h
@@ -166,7 +166,7 @@ public:
 public:
 
 	/** Request that this mode be deleted at the next convenient opportunity (FEditorModeTools::Tick) */
-	void RequestDeletion() { bPendingDeletion = true; }
+	void RequestDeletion();
 
 	/** returns true if this mode is to be deleted at the next convenient opportunity (FEditorModeTools::Tick) */
 	bool IsPendingDeletion() const { return bPendingDeletion; }
