Renamed test enums in PropertyEditorTestObject.h to be more clearly test enums.

Added test for UMETA(Hidden).

diff --git a/Engine/Source/Editor/UnrealEd/Classes/Editor/PropertyEditorTestObject.h b/Engine/Source/Editor/UnrealEd/Classes/Editor/PropertyEditorTestObject.h
--- a/Engine/Source/Editor/UnrealEd/Classes/Editor/PropertyEditorTestObject.h
+++ b/Engine/Source/Editor/UnrealEd/Classes/Editor/PropertyEditorTestObject.h
@@ -25,7 +25,7 @@ class UStaticMeshComponent;
 class UTexture;
 
 UENUM()
-enum PropertyEditorTestEnum
+enum EPropertyEditorTestEnum
 {	
 	/** This comment should appear above enum 1 */
 	PropertyEditorTest_Enum1 UMETA(Hidden),
@@ -43,14 +43,14 @@ enum PropertyEditorTestEnum
 };
 
 UENUM(meta=(Bitflags))
-enum class PropertyEditorTestBitflags : uint8
+enum class EPropertyEditorTestBitflags : uint8
 {
 	First,
 	Second,
 	Third,
 	Hidden UMETA(Hidden, ToolTip="This value shouldn't be used or even visible in the editor")
 };
-ENUM_CLASS_FLAGS(PropertyEditorTestBitflags)
+ENUM_CLASS_FLAGS(EPropertyEditorTestBitflags)
 
 UENUM()
 enum ArrayLabelEnum
@@ -64,22 +64,30 @@ enum ArrayLabelEnum
 	ArrayIndex_MAX,
 };
 
-
 UENUM()
-enum class EditColor : uint8
+enum class EPropertyEditorTestEditColor : uint8
 {
 	Red,
 	Orange,
 	Yellow,
 	Green,
 	Blue,
-	Indigo,
+	Indigo UMETA(Hidden),
 	Violet,
 	Pink,
 	Magenta,
 	Cyan
 };
 
+UENUM()
+enum class EPropertyEditorTestUnderscores : uint8
+{
+	_One,
+	_Two,
+	_Three,
+	NotUnderscore
+};
+
 UENUM()
 enum class ETestEnumFlags : uint8
 {
@@ -260,10 +268,13 @@ class UPropertyEditorTestObject : public UObject
 	FColor ColorProperty;
 
 	UPROPERTY(EditAnywhere, Category=BasicProperties)
-	TEnumAsByte<enum PropertyEditorTestEnum> EnumByteProperty;
+	TEnumAsByte<enum EPropertyEditorTestEnum> EnumByteProperty;
+
+	UPROPERTY(EditAnywhere, Category=BasicProperties)
+	EPropertyEditorTestEditColor EnumProperty;
 
 	UPROPERTY(EditAnywhere, Category=BasicProperties)
-	EditColor EnumProperty;
+	EPropertyEditorTestUnderscores EnumUnderscores;
 
 	UPROPERTY(EditAnywhere, Category=BasicProperties)
 	FMatrix MatrixProperty;
@@ -340,7 +351,7 @@ class UPropertyEditorTestObject : public UObject
 	TArray<FColor> ColorPropertyArray;
 
 	UPROPERTY(EditAnywhere, Category=ArraysOfProperties)
-	TArray<TEnumAsByte<enum PropertyEditorTestEnum> > EnumPropertyArray;
+	TArray<TEnumAsByte<enum EPropertyEditorTestEnum> > EnumPropertyArray;
 
 	UPROPERTY(EditAnywhere, Category=ArraysOfProperties)
 	TArray<FPropertyEditorTestBasicStruct> StructPropertyArray;
@@ -382,7 +393,7 @@ class UPropertyEditorTestObject : public UObject
 	UPROPERTY(VisibleAnywhere, Category=AdvancedProperties)
 	TObjectPtr<UPrimitiveComponent> ObjectThatCannotBeChanged;
 
-	UPROPERTY(EditAnywhere, Category=AdvancedProperties, meta=(Bitmask, BitmaskEnum="PropertyEditorTestBitflags"))
+	UPROPERTY(EditAnywhere, Category=AdvancedProperties, meta=(Bitmask, BitmaskEnum="EPropertyEditorTestBitflags"))
 	int32 EnumBitflags=0;
 
 	UPROPERTY(EditAnywhere, Category=AdvancedProperties, meta=(PasswordField=true))
@@ -491,7 +502,7 @@ class UPropertyEditorTestObject : public UObject
 	TSet<TObjectPtr<AActor>> ActorSet;
 
 	UPROPERTY(EditAnywhere, Category="TSet Tests")
-	TSet<EditColor> EditColorSet;
+	TSet<EPropertyEditorTestEditColor> EditColorSet;
 
 	UPROPERTY(EditAnywhere, Category="TSet Tests")
 	TSet<FName> NameSet;
@@ -521,7 +532,7 @@ class UPropertyEditorTestObject : public UObject
 	TMap<TObjectPtr<UObject>, FLinearColor> ObjectToColorMap;
 
 	UPROPERTY(EditAnywhere, Category="TMap Tests")
-	TMap<int32, TEnumAsByte<PropertyEditorTestEnum> > IntToEnumMap;
+	TMap<int32, TEnumAsByte<EPropertyEditorTestEnum> > IntToEnumMap;
 
 	UPROPERTY(EditAnywhere, Category="TMap Tests")
 	TMap<FName, FName> NameToNameMap;
@@ -599,21 +610,21 @@ class UPropertyEditorTestObject : public UObject
 	bool bEnabledByPrevious;
 
 	UPROPERTY(EditAnywhere, Category=EditCondition)
-	EditColor EnumEditCondition;
+	EPropertyEditorTestEditColor EnumEditCondition;
 
-	UPROPERTY(EditAnywhere, Category=EditCondition, meta=(EditCondition="EnumEditCondition == EditColor::Blue"))
+	UPROPERTY(EditAnywhere, Category=EditCondition, meta=(EditCondition="EnumEditCondition == EPropertyEditorTestEditColor::Blue"))
 	bool bEnabledWhenBlue;
 
-	UPROPERTY(EditAnywhere, Category=EditCondition, meta=(EditCondition="EnumEditCondition == EditColor::Pink"))
+	UPROPERTY(EditAnywhere, Category=EditCondition, meta=(EditCondition="EnumEditCondition == EPropertyEditorTestEditColor::Pink"))
 	bool bEnabledWhenPink;
 
 	UPROPERTY(EditAnywhere, Category=EditCondition)
-	TEnumAsByte<PropertyEditorTestEnum> EnumAsByteEditCondition;
+	TEnumAsByte<EPropertyEditorTestEnum> EnumAsByteEditCondition;
 
-	UPROPERTY(EditAnywhere, Category=EditCondition, meta=(EditCondition="EnumAsByteEditCondition == PropertyEditorTestEnum::PropertyEditorTest_Enum2"))
+	UPROPERTY(EditAnywhere, Category=EditCondition, meta=(EditCondition="EnumAsByteEditCondition == EPropertyEditorTestEnum::PropertyEditorTest_Enum2"))
 	bool bEnabledWhenEnumIs2;
 
-	UPROPERTY(EditAnywhere, Category=EditCondition, meta=(EditCondition="EnumAsByteEditCondition == PropertyEditorTestEnum::PropertyEditorTest_Enum4"))
+	UPROPERTY(EditAnywhere, Category=EditCondition, meta=(EditCondition="EnumAsByteEditCondition == EPropertyEditorTestEnum::PropertyEditorTest_Enum4"))
 	bool bEnabledWhenEnumIs4;
 
 	UPROPERTY(EditAnywhere, Category=EditCondition)
@@ -773,4 +784,4 @@ class APropertyEditorTestActor : public AActor
 
 	UFUNCTION()
 	TArray<FString> GetOptionsFunc() const;
-};
\ No newline at end of file
+};
