Adapted ResavePackagesCommandlet to also resave OFPA (actors with external packages).

diff --git a/Engine/Source/Editor/UnrealEd/Classes/Commandlets/ResavePackagesCommandlet.h b/Engine/Source/Editor/UnrealEd/Classes/Commandlets/ResavePackagesCommandlet.h
--- a/Engine/Source/Editor/UnrealEd/Classes/Commandlets/ResavePackagesCommandlet.h
+++ b/Engine/Source/Editor/UnrealEd/Classes/Commandlets/ResavePackagesCommandlet.h
@@ -128,6 +128,7 @@ protected:
 	int32 PackagesConsideredForResave;
 	int32 PackagesResaved;
 	int32 PackagesDeleted;
+	int32 TotalPackagesForResave;
 
 	/** Only collect garbage after N packages */
 	int32 GarbageCollectionFrequency;
diff --git a/Engine/Source/Editor/UnrealEd/Private/Commandlets/ContentCommandlets.cpp b/Engine/Source/Editor/UnrealEd/Private/Commandlets/ContentCommandlets.cpp
--- a/Engine/Source/Editor/UnrealEd/Private/Commandlets/ContentCommandlets.cpp
+++ b/Engine/Source/Editor/UnrealEd/Private/Commandlets/ContentCommandlets.cpp
@@ -43,6 +43,8 @@
 #include "Editor.h"
 #include "FileHelpers.h"
 #include "CommandletSourceControlUtils.h"
+#include "WorldPartition/WorldPartitionSubsystem.h"
+#include "WorldPartition/WorldPartition.h"
 
 #include "PackageHelperFunctions.h"
 #include "PackageTools.h"
@@ -1151,6 +1153,7 @@ int32 UResavePackagesCommandlet::Main( const FString& Params )
 	PackagesConsideredForResave = 0;
 	PackagesResaved = 0;
 	PackagesDeleted = 0;
+	TotalPackagesForResave = PackageNames.Num();
 
 	// allow for an option to restart at a given package name (in case it dies during a run, etc)
 	bool bCanProcessPackage = true;
@@ -1305,7 +1308,7 @@ int32 UResavePackagesCommandlet::Main( const FString& Params )
 		FPlatformMisc::SetEnvironmentVar(TEXT("uebp_UATMutexNoWait"), TEXT("0"));		
 	}
 
-	UE_LOG(LogContentCommandlet, Display, TEXT("[REPORT] %d/%d packages were considered for resaving"), PackagesConsideredForResave, PackageNames.Num());
+	UE_LOG(LogContentCommandlet, Display, TEXT("[REPORT] %d/%d packages were considered for resaving"), PackagesConsideredForResave, TotalPackagesForResave);
 	UE_LOG(LogContentCommandlet, Display, TEXT("[REPORT] %d/%d packages were resaved"), PackagesResaved, PackagesConsideredForResave);
 	UE_LOG(LogContentCommandlet, Display, TEXT("[REPORT] %d/%d packages were deleted"), PackagesDeleted, PackagesConsideredForResave);
 
@@ -1634,31 +1637,76 @@ void UResavePackagesCommandlet::PerformAdditionalOperations(class UWorld* World,
 	// indicates if world and level packages should be checked out only if dirty after building data
 	const bool bShouldCheckoutDirtyPackageOnly = (bShouldBuildHLOD || bShouldBuildNavigationData) && !bBuildingNonHLODData;
 
-	if (bBuildingNonHLODData || bShouldBuildHLOD || bShouldBuildNavigationData)
-	{
-		bool bShouldProceedWithRebuild = true;
+	UWorldPartition* WorldPartition = World->GetWorldPartition();
+	const bool bResaveWorldPartitionExternalActors = !!WorldPartition;
 
-		// Setup the world.
-		World->WorldType = EWorldType::Editor;
+	// Load and Save Level's external packages
+ 	if (!bResaveWorldPartitionExternalActors)
+	{
 		World->AddToRoot();
-		if (!World->bIsWorldInitialized)
+		for (UPackage* Package : World->PersistentLevel->GetPackage()->GetExternalPackages())
 		{
-			UWorld::InitializationValues IVS;
-			IVS.RequiresHitProxies(false);
-			IVS.ShouldSimulatePhysics(false);
-			IVS.EnableTraceCollision(false);
-			IVS.CreateNavigation(bShouldBuildNavigationData);
-			IVS.CreateAISystem(false);
-			IVS.AllowAudioPlayback(false);
-			IVS.CreatePhysicsScene(true);
-
-			World->InitWorld(IVS);
-			World->PersistentLevel->UpdateModelComponents();
-			World->UpdateWorldComponents(true, false);
+			++TotalPackagesForResave;
+			const FString PackageFilename = Package->GetLoadedPath().GetLocalFullPath();
+			check(FLinkerLoad::FindExistingLinkerForPackage(Package));
+			LoadAndSaveOnePackage(PackageFilename);
 		}
-		FWorldContext &WorldContext = GEditor->GetEditorWorldContext(true);
-		WorldContext.SetCurrentWorld(World);
-		GWorld = World;
+		World->RemoveFromRoot();
+	}
+
+	if (!bBuildingNonHLODData && !bShouldBuildHLOD && !bShouldBuildNavigationData && !bResaveWorldPartitionExternalActors)
+	{
+		return;
+	}
+
+	// Setup the world.
+	World->WorldType = EWorldType::Editor;
+	World->AddToRoot();
+	if (!World->bIsWorldInitialized)
+	{
+		UWorld::InitializationValues IVS;
+		IVS.RequiresHitProxies(false);
+		IVS.ShouldSimulatePhysics(false);
+		IVS.EnableTraceCollision(false);
+		IVS.CreateNavigation(bShouldBuildNavigationData);
+		IVS.CreateAISystem(false);
+		IVS.AllowAudioPlayback(false);
+		IVS.CreatePhysicsScene(true);
+
+		World->InitWorld(IVS);
+		World->PersistentLevel->UpdateModelComponents();
+		World->UpdateWorldComponents(true, false);
+	}
+	FWorldContext& WorldContext = GEditor->GetEditorWorldContext(true);
+	WorldContext.SetCurrentWorld(World);
+	GWorld = World;
+
+	// Load and Save world partition actor packages
+	if (bResaveWorldPartitionExternalActors)
+	{
+		UWorldPartitionSubsystem* WorldPartitionSubsystem = World->GetSubsystem<UWorldPartitionSubsystem>();
+		WorldPartitionSubsystem->ForEachActorDesc(AActor::StaticClass(), [this, WorldPartition](const FWorldPartitionActorDesc* ActorDesc)
+		{
+			++TotalPackagesForResave;
+			// Load & Register World Partition Actor
+			FWorldPartitionReference LoadedActor(WorldPartition, ActorDesc->GetGuid());
+			AActor* Actor = LoadedActor->GetActor();
+			UPackage* Package = Actor ? Actor->GetExternalPackage() : nullptr;
+			if (Package == nullptr)
+			{
+				check(bCanIgnoreFails);
+				return true;
+			}
+			const FString PackageFilename = Package->GetLoadedPath().GetLocalFullPath();
+			check(FLinkerLoad::FindExistingLinkerForPackage(Package));
+			LoadAndSaveOnePackage(PackageFilename);
+			return true;
+		});
+	}
+
+	if (bBuildingNonHLODData || bShouldBuildHLOD || bShouldBuildNavigationData)
+	{
+		bool bShouldProceedWithRebuild = true;
 
 		TArray<FString> CheckedOutPackagesFilenames;
 		auto CheckOutLevelFile = [this,&bShouldProceedWithRebuild, &CheckedOutPackagesFilenames](ULevel* InLevel)
@@ -2018,14 +2066,11 @@ void UResavePackagesCommandlet::PerformAdditionalOperations(class UWorld* World,
 				}
 			}
 		}
-
-		
-		World->RemoveFromRoot();
-
-		WorldContext.SetCurrentWorld(nullptr);
-		GWorld = nullptr;
-
 	}
+
+	World->RemoveFromRoot();
+	WorldContext.SetCurrentWorld(nullptr);
+	GWorld = nullptr;
 }
 
 void UResavePackagesCommandlet::PerformAdditionalOperations( class UObject* Object, bool& bSavePackage )
