Change the DerivedDataCacheCommandlet to wait on all objects to return true for IsCachedCookedPlatformDataLoaded before moving on.
Change the loops over objects and packages to be faster and to not add packages as processed if we didn't handle them in the BeginCache loop.
#rnx

diff --git a/Engine/Source/Editor/UnrealEd/Private/Commandlets/DerivedDataCacheCommandlet.cpp b/Engine/Source/Editor/UnrealEd/Private/Commandlets/DerivedDataCacheCommandlet.cpp
--- a/Engine/Source/Editor/UnrealEd/Private/Commandlets/DerivedDataCacheCommandlet.cpp
+++ b/Engine/Source/Editor/UnrealEd/Private/Commandlets/DerivedDataCacheCommandlet.cpp
@@ -24,6 +24,7 @@ DerivedDataCacheCommandlet.cpp: Commandlet for DDC maintenence
 #include "Algo/Transform.h"
 #include "Settings/ProjectPackagingSettings.h"
 #include "Editor.h"
+#include "AssetCompilingManager.h"
 
 DEFINE_LOG_CATEGORY_STATIC(LogDerivedDataCacheCommandlet, Log, All);
 
@@ -42,6 +43,60 @@ void UDerivedDataCacheCommandlet::MaybeMarkPackageAsAlreadyLoaded(UPackage *Pack
 	}
 }
 
+static void WaitForCurrentShaderCompilationToFinish(bool& bInOutHadActivity)
+{
+	int32 CachedShaderCount = GShaderCompilingManager->GetNumRemainingJobs();
+	if (CachedShaderCount == 0)
+	{
+		return;
+	}
+
+	int32 NumCompletedShadersSinceLastLog = 0;
+	bInOutHadActivity = true;
+	UE_LOG(LogDerivedDataCacheCommandlet, Display, TEXT("Waiting for %d shaders to finish."), CachedShaderCount);
+	while (GShaderCompilingManager->IsCompiling())
+	{
+		const int32 CurrentShaderCount = GShaderCompilingManager->GetNumRemainingJobs();
+		NumCompletedShadersSinceLastLog += (CachedShaderCount - CurrentShaderCount);
+		CachedShaderCount = CurrentShaderCount;
+
+		if (NumCompletedShadersSinceLastLog >= 1000)
+		{
+			UE_LOG(LogDerivedDataCacheCommandlet, Display, TEXT("Waiting for %d shaders to finish."), CachedShaderCount);
+			NumCompletedShadersSinceLastLog = 0;
+		}
+
+		// Process any asynchronous shader compile results that are ready, limit execution time
+		GShaderCompilingManager->ProcessAsyncResults(true, false);
+		GDistanceFieldAsyncQueue->ProcessAsyncTasks();
+		GCardRepresentationAsyncQueue->ProcessAsyncTasks();
+	}
+	GShaderCompilingManager->FinishAllCompilation(); // Final blocking check as IsCompiling() may be non-deterministic
+	GDistanceFieldAsyncQueue->BlockUntilAllBuildsComplete();
+	GCardRepresentationAsyncQueue->BlockUntilAllBuildsComplete();
+	UE_LOG(LogDerivedDataCacheCommandlet, Display, TEXT("Done waiting for shaders to finish."));
+};
+
+static void WaitForCurrentTextureBuildingToFinish(bool& bInOutHadActivity)
+{
+	for (TObjectIterator<UTexture> Texture; Texture; ++Texture)
+	{
+		Texture->FinishCachePlatformData();
+	}
+};
+
+static void PumpAsync(bool* bInOutHadActivity = nullptr)
+{
+	bool bHadActivity = false;
+	WaitForCurrentShaderCompilationToFinish(bHadActivity);
+	WaitForCurrentTextureBuildingToFinish(bHadActivity);
+	FAssetCompilingManager::Get().ProcessAsyncTasks(true);
+	if (bInOutHadActivity)
+	{
+		*bInOutHadActivity = *bInOutHadActivity || bHadActivity;
+	}
+}
+
 int32 UDerivedDataCacheCommandlet::Main( const FString& Params )
 {
 	TArray<FString> Tokens, Switches;
@@ -50,51 +105,18 @@ int32 UDerivedDataCacheCommandlet::Main( const FString& Params )
 	bool bFillCache = Switches.Contains("FILL");   // do the equivalent of a "loadpackage -all" to fill the DDC
 	bool bStartupOnly = Switches.Contains("STARTUPONLY");   // regardless of any other flags, do not iterate packages
 
-															// Subsets for parallel processing
+	// Subsets for parallel processing
 	uint32 SubsetMod = 0;
 	uint32 SubsetTarget = MAX_uint32;
 	FParse::Value(*Params, TEXT("SubsetMod="), SubsetMod);
 	FParse::Value(*Params, TEXT("SubsetTarget="), SubsetTarget);
 	bool bDoSubset = SubsetMod > 0 && SubsetTarget < SubsetMod;
+
+	// Timing variables
 	double FindProcessedPackagesTime = 0.0;
 	double GCTime = 0.0;
-
-	auto WaitForCurrentShaderCompilationToFinish = []()
-	{
-		int32 NumCompletedShadersSinceLastLog = 0;
-		int32 CachedShaderCount = GShaderCompilingManager->GetNumRemainingJobs();
-		UE_LOG(LogDerivedDataCacheCommandlet, Display, TEXT("Waiting for %d shaders to finish."), CachedShaderCount);
-		while (GShaderCompilingManager->IsCompiling())
-		{
-			const int32 CurrentShaderCount = GShaderCompilingManager->GetNumRemainingJobs();
-			NumCompletedShadersSinceLastLog += (CachedShaderCount - CurrentShaderCount);
-			CachedShaderCount = CurrentShaderCount;
-
-			if (NumCompletedShadersSinceLastLog >= 1000)
-			{
-				UE_LOG(LogDerivedDataCacheCommandlet, Display, TEXT("Waiting for %d shaders to finish."), CachedShaderCount);
-				NumCompletedShadersSinceLastLog = 0;
-			}
-
-			// Process any asynchronous shader compile results that are ready, limit execution time
-			GShaderCompilingManager->ProcessAsyncResults(true, false);
-			GDistanceFieldAsyncQueue->ProcessAsyncTasks();
-			GCardRepresentationAsyncQueue->ProcessAsyncTasks();
-		}
-		GShaderCompilingManager->FinishAllCompilation(); // Final blocking check as IsCompiling() may be non-deterministic
-		GDistanceFieldAsyncQueue->BlockUntilAllBuildsComplete();
-		GCardRepresentationAsyncQueue->BlockUntilAllBuildsComplete();
-		UE_LOG(LogDerivedDataCacheCommandlet, Display, TEXT("Done waiting for shaders to finish."));
-	};
-
-
-	auto WaitForCurrentTextureBuildingToFinish = []()
-	{
-		for ( TObjectIterator<UTexture> Texture; Texture; ++Texture )
-		{
-			Texture->FinishCachePlatformData();
-		}
-	};
+	double FinishCacheTime = 0.;
+	double BeginCacheTime = 0.;
 
 	if (!bStartupOnly && bFillCache)
 	{
@@ -266,6 +288,9 @@ int32 UDerivedDataCacheCommandlet::Main( const FString& Params )
 			UE_LOG(LogDerivedDataCacheCommandlet, Display, TEXT("%d packages to load..."), PackagePaths.Num());
 		}
 
+		TArray<UObject*> CachingObjects;
+		TArray<UObject*> ObjectBuffer;
+		TArray<UPackage*> NewPackages;
 		for (int32 PackageIndex = PackagePaths.Num() - 1; PackageIndex >= 0; PackageIndex-- )
 		{
 			{
@@ -302,65 +327,113 @@ int32 UDerivedDataCacheCommandlet::Main( const FString& Params )
 			// even if the load failed this could be the first time through the loop so it might have all the startup packages to resolve
 			GRedirectCollector.ResolveAllSoftObjectPaths();
 
-			// cache all the resources for this platform
-			for (TObjectIterator<UObject> It; It; ++It)
+			// Find any new packages and cache all the objects in each package
 			{
-				if ((PackageFilter&NORMALIZE_ExcludeEnginePackages) == 0 || !It->GetOutermost()->GetName().StartsWith(TEXT("/Engine")))
+				const double BeginCacheTimeStart = FPlatformTime::Seconds();
+				CachingObjects.Reset();
+				NewPackages.Reset();
+				for (TObjectIterator<UPackage> PackageIter; PackageIter; ++PackageIter)
 				{
-					if (!ProcessedPackages.Contains(It->GetOutermost()->GetFName()))
+					UPackage* ExistingPackage = *PackageIter;
+					if ((PackageFilter & NORMALIZE_ExcludeEnginePackages) == 0 || !ExistingPackage->GetName().StartsWith(TEXT("/Engine")))
 					{
-						check((It->GetOutermost()->GetPackageFlags() & PKG_ReloadingForCooker) == 0);
-						for (auto Platform : Platforms)
+						FName ExistingPackageName = ExistingPackage->GetFName();
+						if (!ProcessedPackages.Contains(ExistingPackageName))
 						{
-							It->BeginCacheForCookedPlatformData(Platform);
+							ProcessedPackages.Add(ExistingPackageName);
+							NewPackages.Add(ExistingPackage);
+							check((ExistingPackage->GetPackageFlags() & PKG_ReloadingForCooker) == 0);
+
+							ObjectBuffer.Reset();
+							GetObjectsWithOuter(ExistingPackage, ObjectBuffer, true /* bIncludeNestedObjects */, RF_ClassDefaultObject /* ExclusionFlags */);
+							for (UObject* Object : ObjectBuffer)
+							{
+								for (auto Platform : Platforms)
+								{
+									Object->BeginCacheForCookedPlatformData(Platform);
+								}
+								CachingObjects.Add(Object);
+							}
 						}
 					}
 				}
+				PumpAsync();
+				BeginCacheTime += FPlatformTime::Seconds() - BeginCacheTimeStart;
 			}
 
-
-			// Keep track of which packages have already been processed along with the map.
 			{
-				const double FindProcessedPackagesStartTime = FPlatformTime::Seconds();
-				TArray<UObject *> AllPackages;
-				GetObjectsOfClass(UPackage::StaticClass(), AllPackages);
-				for (int32 Index = 0; Index < AllPackages.Num(); Index++)
+				const double FinishCacheTimeStart = FPlatformTime::Seconds();
+				ObjectBuffer.Reset();
+				ObjectBuffer.Append(CachingObjects);
+
+				double LastActivityTime = FinishCacheTimeStart;
+				const double MaxSecondsWithNoActivity = 120.0;
+				const double WaitingForCacheSleepTime = 0.050;
+				while (ObjectBuffer.Num() > 0)
 				{
-					UPackage* Pkg = Cast<UPackage>(AllPackages[Index]);
-					if (!Pkg || Pkg->GetOuter())
+					bool bHadActivity = false;
+					for (int32 ObjectIndex = 0; ObjectIndex < ObjectBuffer.Num();)
 					{
-						continue;
+						UObject* Object = ObjectBuffer[ObjectIndex];
+						bool bIsFinished = true;
+						for (auto Platform : Platforms)
+						{
+							bIsFinished = Object->IsCachedCookedPlatformDataLoaded(Platform) && bIsFinished;
+						}
+						if (bIsFinished)
+						{
+							ObjectBuffer.RemoveAtSwap(ObjectIndex);
+							bHadActivity = true;
+						}
+						else
+						{
+							++ObjectIndex;
+						}
 					}
-					if (!ProcessedPackages.Contains(Pkg->GetFName()))
+
+					double CurrentTime = FPlatformTime::Seconds();
+					if (!bHadActivity)
+					{
+						PumpAsync(&bHadActivity);
+					}
+					if (!bHadActivity)
+					{
+						if (CurrentTime - LastActivityTime >= MaxSecondsWithNoActivity)
 						{
-						ProcessedPackages.Add(Pkg->GetFName());
-							Pkg->SetPackageFlags(PKG_ReloadingForCooker);
-							{
-								TArray<UObject *> ObjectsInPackage;
-								GetObjectsWithOuter(Pkg, ObjectsInPackage, true);
-								for (int32 IndexPackage = 0; IndexPackage < ObjectsInPackage.Num(); IndexPackage++)
-								{
-									for (auto Platform : Platforms)
-									{
-										ObjectsInPackage[IndexPackage]->IsCachedCookedPlatformDataLoaded(Platform);
-									}
-									ObjectsInPackage[IndexPackage]->WillNeverCacheCookedPlatformDataAgain();
-									ObjectsInPackage[IndexPackage]->ClearAllCachedCookedPlatformData();
-								}
-							}
+							UE_LOG(LogDerivedDataCacheCommandlet, Error, TEXT("Timed out for %.2lfs waiting for %d objects to finish caching. First object: %s."),
+								MaxSecondsWithNoActivity, ObjectBuffer.Num(), *ObjectBuffer[0]->GetFullName());
+							ObjectBuffer.Reset();
+						}
+						else
+						{
+							FPlatformProcess::Sleep(WaitingForCacheSleepTime);
 						}
 					}
-				FindProcessedPackagesTime += FPlatformTime::Seconds() - FindProcessedPackagesStartTime;
-			}
+					else
+					{
+						LastActivityTime = CurrentTime;
+					}
+				}
+
+				// Tear down all of the Cached data; we do this only after all objects have finished because we need to not
+				// tear down any object until all objects in its package have finished
+				for (UObject* Object : CachingObjects)
+				{
+					Object->WillNeverCacheCookedPlatformDataAgain();
+					Object->ClearAllCachedCookedPlatformData();
+				}
+				// Mark the packages as processed
+				for (UPackage* NewPackage : NewPackages)
+				{
+					NewPackage->SetPackageFlags(PKG_ReloadingForCooker);
+				}
 
-			// Process any asynchronous shader compile results that are ready, limit execution time
-			GShaderCompilingManager->ProcessAsyncResults(true, false);
+				PumpAsync();
+				FinishCacheTime += FPlatformTime::Seconds() - FinishCacheTimeStart;
+			}
 
 			if (NumProcessedSinceLastGC >= GCInterval || PackageIndex < 0 || bLastPackageWasMap)
 			{
-				WaitForCurrentShaderCompilationToFinish();
-				WaitForCurrentTextureBuildingToFinish();
-
 				const double StartGCTime = FPlatformTime::Seconds();
 				if (NumProcessedSinceLastGC >= GCInterval || PackageIndex < 0)
 				{
@@ -380,11 +453,9 @@ int32 UDerivedDataCacheCommandlet::Main( const FString& Params )
 		}
 	}
 
-	WaitForCurrentShaderCompilationToFinish();
-	WaitForCurrentTextureBuildingToFinish();
 	GetDerivedDataCacheRef().WaitForQuiescence(true);
 
-	UE_LOG(LogDerivedDataCacheCommandlet, Display, TEXT("%.2lfs spent looking for processed packages, %.2lfs spent on GC."), FindProcessedPackagesTime, GCTime);
+	UE_LOG(LogDerivedDataCacheCommandlet, Display, TEXT("BeginCacheTime=%.2lfs, FinishCacheTime=%.2lfs, GCTime=%.2lfs."), BeginCacheTime, FinishCacheTime, GCTime);
 
 	return 0;
 }
