Reversed decision to hide old UE4 thumbnails when loading in UE5. The UE4 thumbnail will remain until the package is loaded and resaved

diff --git a/Engine/Source/Editor/UnrealEd/Private/ObjectTools.cpp b/Engine/Source/Editor/UnrealEd/Private/ObjectTools.cpp
--- a/Engine/Source/Editor/UnrealEd/Private/ObjectTools.cpp
+++ b/Engine/Source/Editor/UnrealEd/Private/ObjectTools.cpp
@@ -4860,12 +4860,6 @@ namespace ThumbnailTools
 			return false;
 		}
 
-		// Hacky way to reset thumbnails for UE5
-		if(FileSummary.SavedByEngineVersion.GetMajor() < 5)
-		{
-			return false;
-		}
-
 		// Seek the the part of the file where the thumbnail table lives
 		FileReader->Seek( FileSummary.ThumbnailTableOffset );
 
