HLOD distributed build - Fixed deleted files not being propagated from the setup step to the submit step

diff --git a/Engine/Source/Editor/UnrealEd/Private/WorldPartition/WorldPartitionHLODsBuilder.cpp b/Engine/Source/Editor/UnrealEd/Private/WorldPartition/WorldPartitionHLODsBuilder.cpp
--- a/Engine/Source/Editor/UnrealEd/Private/WorldPartition/WorldPartitionHLODsBuilder.cpp
+++ b/Engine/Source/Editor/UnrealEd/Private/WorldPartition/WorldPartitionHLODsBuilder.cpp
@@ -314,8 +314,6 @@ bool UWorldPartitionHLODsBuilder::SetupHLODActors(bool bCreateOnly)
 				TArray<TArray<FString>> BuildersFiles;
 				BuildersFiles.SetNum(BuilderCount);
 
-				TArray<FString> FilesToSubmit;
-
 				for (const FString& ModifiedFile : ModifiedFiles)
 				{
 					int32* Idx = FilesToBuilderMap.Find(ModifiedFile);
@@ -325,7 +323,8 @@ bool UWorldPartitionHLODsBuilder::SetupHLODActors(bool bCreateOnly)
 					}
 					else
 					{
-						FilesToSubmit.Add(ModifiedFile);
+						// Add general files to the last builder
+						BuildersFiles.Last().Add(ModifiedFile);
 					}
 				}
 
@@ -341,12 +340,6 @@ bool UWorldPartitionHLODsBuilder::SetupHLODActors(bool bCreateOnly)
 					}
 				}
 
-				// Copy files that won't be handled by builders but must be submitted in the last step
-				if (!CopyFilesToWorkingDir(GetToSubmitFolderName(), FilesToSubmit, BuildProducts))
-				{
-					return false;
-				}
-
 				// The build manifest must also be included as a build product to be available in the next steps
 				BuildProducts.Add(BuildManifest);
 
