Use 'REN_ForceNoResetLoaders' when renaming newly created UDIM textures, avoids extreme delay when importing UDIMs into certain larger projects

diff --git a/Engine/Source/Editor/UnrealEd/Private/Factories/EditorFactories.cpp b/Engine/Source/Editor/UnrealEd/Private/Factories/EditorFactories.cpp
--- a/Engine/Source/Editor/UnrealEd/Private/Factories/EditorFactories.cpp
+++ b/Engine/Source/Editor/UnrealEd/Private/Factories/EditorFactories.cpp
@@ -3984,7 +3984,7 @@ UObject* UTextureFactory::FactoryCreateBinary
 							}
 							else
 							{
-								verify(InParent->Rename(*PackageUDIMName, nullptr, REN_DontCreateRedirectors));
+								verify(InParent->Rename(*PackageUDIMName, nullptr, REN_DontCreateRedirectors | REN_ForceNoResetLoaders));
 							}
 						}
 					}
