ExternalActorsCommandlet: now detects invalid external actor filenames and resave them when possible, delete them if not.

diff --git a/Engine/Source/Editor/UnrealEd/Private/Commandlets/ExternalActorsCommandlet.cpp b/Engine/Source/Editor/UnrealEd/Private/Commandlets/ExternalActorsCommandlet.cpp
--- a/Engine/Source/Editor/UnrealEd/Private/Commandlets/ExternalActorsCommandlet.cpp
+++ b/Engine/Source/Editor/UnrealEd/Private/Commandlets/ExternalActorsCommandlet.cpp
@@ -9,7 +9,6 @@
 #include "UObject/UObjectHash.h"
 #include "AssetRegistryModule.h"
 #include "PackageHelperFunctions.h"
-#include "WorldPartition/ActorDescContainer.h"
 
 DEFINE_LOG_CATEGORY_STATIC(LogExternalActorsCommandlet, All, All);
 
@@ -49,7 +48,6 @@ int32 UExternalActorsCommandlet::Main(const FString& Params)
 	}
 
 	const bool bRepair = Switches.Contains(TEXT("repair"));
-	const bool bLoadPackages = Switches.Contains(TEXT("loadpackages"));
 
 	if (!FPackageName::SearchForPackageOnDisk(Tokens[0], &Tokens[0]))
 	{
@@ -58,17 +56,13 @@ int32 UExternalActorsCommandlet::Main(const FString& Params)
 	}
 
 	// Load world
-	UWorld* MainWorld = bLoadPackages ? LoadWorld(Tokens[0]) : nullptr;
-	if (bLoadPackages && !MainWorld)
+	UWorld* MainWorld = LoadWorld(Tokens[0]);
+	if (!MainWorld)
 	{
 		UE_LOG(LogExternalActorsCommandlet, Error, TEXT("Unknown world '%s'"), *Tokens[0]);
 		return 1;
 	}
 
-	// Initialize actor descriptors
-	UActorDescContainer* ActorDescContainer = NewObject<UActorDescContainer>(GetTransientPackage());
-	ActorDescContainer->Initialize(MainWorld, *Tokens[0], false);
-
 	// Validate external actors
 	FString ExternalActorsPath = ULevel::GetExternalActorsPath(Tokens[0]);
 	FString ExternalActorsFilePath = FPackageName::LongPackageNameToFilename(ExternalActorsPath);
@@ -76,41 +70,62 @@ int32 UExternalActorsCommandlet::Main(const FString& Params)
 	TArray<FString> PackagesToDelete;
 	if (IFileManager::Get().DirectoryExists(*ExternalActorsFilePath))
 	{
-		bool bResult = IFileManager::Get().IterateDirectoryRecursively(*ExternalActorsFilePath, [this, ActorDescContainer, bRepair, bLoadPackages, &PackagesToDelete](const TCHAR* FilenameOrDirectory, bool bIsDirectory)
+		bool bResult = IFileManager::Get().IterateDirectoryRecursively(*ExternalActorsFilePath, [this, bRepair, &PackagesToDelete](const TCHAR* FilenameOrDirectory, bool bIsDirectory)
 		{
 			if (!bIsDirectory)
 			{
 				FString Filename(FilenameOrDirectory);
 				if (Filename.EndsWith(FPackageName::GetAssetPackageExtension()))
 				{
-					bool bIsValidActorPackage = false;
+					AActor* MainPackageActor = nullptr;
+					AActor* PotentialMainPackageActor = nullptr;
+
 					const FString PackageName = FPackageName::FilenameToLongPackageName(*Filename);
 					const FWorldPartitionActorDesc* ActorDesc = nullptr;
 
-					if (bLoadPackages)
+					if (UPackage* Package = LoadPackage(nullptr, *Filename, LOAD_None, nullptr, nullptr))
 					{
-						if (UPackage* Package = LoadPackage(nullptr, *Filename, LOAD_None, nullptr, nullptr))
-						{						
-							ForEachObjectWithPackage(Package, [&bIsValidActorPackage](UObject* Object)
+						ForEachObjectWithPackage(Package, [&MainPackageActor, &PotentialMainPackageActor](UObject* Object)
+						{
+							if (AActor* Actor = Cast<AActor>(Object))
 							{
-								Object->ClearFlags(RF_Public | RF_Standalone);
-								bIsValidActorPackage = Object->IsA<AActor>();
-								return !bIsValidActorPackage;
-							});
-						}
-					}
-					else
-					{
-						ActorDesc = ActorDescContainer->GetActorDesc(PackageName);
-						bIsValidActorPackage = !!ActorDesc;
+								if (Actor->IsMainPackageActor())
+								{
+									MainPackageActor = Actor;
+									PotentialMainPackageActor = nullptr;
+								}
+								else if (!MainPackageActor)
+								{
+									if (!Actor->IsChildActor())
+									{
+										PotentialMainPackageActor = Actor;
+									}
+								}
+							}
+							return true;
+						});
 					}
 
-					if (!bIsValidActorPackage)
+					if (!MainPackageActor)
 					{
-						UE_LOG(LogExternalActorsCommandlet, Error, TEXT("Invalid actor file '%s' (%s)"), *Filename, ActorDesc ? *ActorDesc->ToString() : TEXT("None"));
+						UE_LOG(LogExternalActorsCommandlet, Error, TEXT("Invalid actor file '%s'"), *Filename);
 
 						if (bRepair)
 						{
+							if (PotentialMainPackageActor)
+							{
+								PotentialMainPackageActor->SetPackageExternal(false);
+								PotentialMainPackageActor->SetPackageExternal(true);
+								
+								UPackage* PackageToSave = PotentialMainPackageActor->GetPackage();
+
+								FString PackageFileName = SourceControlHelpers::PackageFilename(PackageToSave);
+								if (UPackage::SavePackage(PackageToSave, nullptr, RF_Standalone, *PackageFileName, GError, nullptr, false, true))
+								{
+									PackageHelper.AddToSourceControl(PackageToSave);
+								}
+							}
+
 							PackagesToDelete.Add(Filename);
 						}
 					}
