Editor: expose DetailsViewArgs to subclasses of FModeToolkit, to give them a chance to customize the values before the ModeToolkit creates the DetailsView
#rnx

diff --git a/Engine/Source/Editor/UnrealEd/Private/Toolkits/BaseToolkit.cpp b/Engine/Source/Editor/UnrealEd/Private/Toolkits/BaseToolkit.cpp
--- a/Engine/Source/Editor/UnrealEd/Private/Toolkits/BaseToolkit.cpp
+++ b/Engine/Source/Editor/UnrealEd/Private/Toolkits/BaseToolkit.cpp
@@ -139,6 +139,8 @@ void FModeToolkit::Init(const TSharedPtr< class IToolkitHost >& InitToolkitHost,
 		ModeDetailsViewArgs.bShowOptions = false;
 		ModeDetailsViewArgs.bAllowMultipleTopLevelObjects = true;
 
+		CustomizeModeDetailsViewArgs(ModeDetailsViewArgs);		// allow subclass to customize arguments
+
 		ModeDetailsView = PropertyEditorModule.CreateDetailView(ModeDetailsViewArgs);
 	}
 
@@ -151,6 +153,8 @@ void FModeToolkit::Init(const TSharedPtr< class IToolkitHost >& InitToolkitHost,
 		DetailsViewArgs.bShowOptions = false;
 		DetailsViewArgs.bAllowMultipleTopLevelObjects = true;
 
+		CustomizeDetailsViewArgs(DetailsViewArgs);		// allow subclass to customize arguments
+
 		DetailsView = PropertyEditorModule.CreateDetailView(DetailsViewArgs);
 	}
 
diff --git a/Engine/Source/Editor/UnrealEd/Public/Toolkits/BaseToolkit.h b/Engine/Source/Editor/UnrealEd/Public/Toolkits/BaseToolkit.h
--- a/Engine/Source/Editor/UnrealEd/Public/Toolkits/BaseToolkit.h
+++ b/Engine/Source/Editor/UnrealEd/Public/Toolkits/BaseToolkit.h
@@ -13,6 +13,7 @@ class UInteractiveTool;
 class UInteractiveToolManager;
 class UEdMode;
 class IDetailsView;
+struct FDetailsViewArgs;
 enum class EToolShutdownType;
 
 /**
@@ -169,6 +170,9 @@ protected:
 	virtual void OnToolStarted(UInteractiveToolManager* Manager, UInteractiveTool* Tool);
 	virtual void OnToolEnded(UInteractiveToolManager* Manager, UInteractiveTool* Tool);
 
+	virtual void CustomizeModeDetailsViewArgs(FDetailsViewArgs& ArgsInOut) {}
+	virtual void CustomizeDetailsViewArgs(FDetailsViewArgs& ArgsInOut) {}
+
 protected:
 	TSharedPtr<SWidget> ToolkitWidget;
 	TSharedPtr<IDetailsView> ModeDetailsView;
