Fix transfering data from old pins after rebuilding
Fix for removing material node pins

diff --git a/Engine/Source/Editor/UnrealEd/Classes/MaterialGraph/MaterialGraphNode_Base.h b/Engine/Source/Editor/UnrealEd/Classes/MaterialGraph/MaterialGraphNode_Base.h
--- a/Engine/Source/Editor/UnrealEd/Classes/MaterialGraph/MaterialGraphNode_Base.h
+++ b/Engine/Source/Editor/UnrealEd/Classes/MaterialGraph/MaterialGraphNode_Base.h
@@ -36,6 +36,7 @@ class UMaterialGraphNode_Base : public UEdGraphNode
 
 	/** The exec input pin */
 	UEdGraphPin* ExecInputPin = nullptr;
+	TArray<UEdGraphPin*> ExecOutputPins;
 
 	/** Create all of the input pins required */
 	virtual void CreateInputPins() {};
diff --git a/Engine/Source/Editor/UnrealEd/Private/MaterialGraphNode_Base.cpp b/Engine/Source/Editor/UnrealEd/Private/MaterialGraphNode_Base.cpp
--- a/Engine/Source/Editor/UnrealEd/Private/MaterialGraphNode_Base.cpp
+++ b/Engine/Source/Editor/UnrealEd/Private/MaterialGraphNode_Base.cpp
@@ -161,11 +161,15 @@ void UMaterialGraphNode_Base::RegisterPin(UEdGraphPin* Pin, EMaterialGraphPinTyp
 
 	if (Type == EMaterialGraphPinType::Exec)
 	{
-		if (Pin->Direction == EGPD_Input)
+		switch (Pin->Direction)
 		{
+		case EGPD_Input:
 			checkf(ExecInputPin == nullptr, TEXT("Only 1 exec input pin allowed"));
 			check(Index == 0);
 			ExecInputPin = Pin;
+			break;
+		case EGPD_Output: verify(ExecOutputPins.Add(Pin) == Index); break;
+		default: checkNoEntry(); break;
 		}
 	}
 	else
@@ -185,6 +189,21 @@ void UMaterialGraphNode_Base::EmptyPins()
 	PinInfoMap.Reset();
 	InputPins.Reset();
 	OutputPins.Reset();
+	ExecOutputPins.Reset();
+	ExecInputPin = nullptr;
+}
+
+static void TransferPinArray(const TArray<UEdGraphPin*>& NewPins, const TArray<UEdGraphPin*>& OldPins)
+{
+	const int32 Num = FMath::Min(NewPins.Num(), OldPins.Num());
+	for (int32 i = 0; i < Num; ++i)
+	{
+		UEdGraphPin* OldPin = OldPins[i];
+		UEdGraphPin* NewPin = NewPins[i];
+		ensure(OldPin->Direction == NewPin->Direction);
+		ensure(OldPin->PinType.PinCategory == NewPin->PinType.PinCategory);
+		NewPin->MovePersistentDataFromOldPin(*OldPin);
+	}
 }
 
 void UMaterialGraphNode_Base::ReconstructNode()
@@ -209,24 +228,27 @@ void UMaterialGraphNode_Base::ReconstructNode()
 
 	// Move the existing pins to a saved array
 	TArray<UEdGraphPin*> OldPins = MoveTemp(Pins);
+	TArray<UEdGraphPin*> OldInputPins = MoveTemp(InputPins);
+	TArray<UEdGraphPin*> OldOutputPins = MoveTemp(OutputPins);
+	TArray<UEdGraphPin*> OldExecOutputPins = MoveTemp(ExecOutputPins);
+	UEdGraphPin* OldExecInputPin = ExecInputPin;
 
 	EmptyPins();
 
 	// Recreate the new pins
 	AllocateDefaultPins();
 
-	for (int32 PinIndex = 0; PinIndex < OldPins.Num(); ++PinIndex)
+	// Transfer data to new pins
+	TransferPinArray(InputPins, OldInputPins);
+	TransferPinArray(OutputPins, OldOutputPins);
+	TransferPinArray(ExecOutputPins, OldExecOutputPins);
+	if (OldExecInputPin && ExecInputPin)
 	{
-		UEdGraphPin* OldPin = OldPins[PinIndex];
-		if (PinIndex < Pins.Num())
-		{
-			// Transfer data to new pin
-			UEdGraphPin* NewPin = Pins[PinIndex];
-			ensure(OldPin->Direction == NewPin->Direction);
-			ensure(OldPin->PinType.PinCategory == NewPin->PinType.PinCategory);
-			NewPin->MovePersistentDataFromOldPin(*OldPin);
-		}
+		ExecInputPin->MovePersistentDataFromOldPin(*OldExecInputPin);
+	}
 
+	for (UEdGraphPin* OldPin : OldPins)
+	{
 		// Throw away the original pins
 		OldPin->Modify();
 		UEdGraphNode::DestroyPin(OldPin);
@@ -237,6 +259,45 @@ void UMaterialGraphNode_Base::ReconstructNode()
 
 void UMaterialGraphNode_Base::RemovePinAt(const int32 PinIndex, const EEdGraphPinDirection PinDirection)
 {
+	UEdGraphPin* Pin = GetPinWithDirectionAt(PinIndex, PinDirection);
+	check(Pin);
+
+	FMaterialGraphPinInfo PinInfo;
+	verify(PinInfoMap.RemoveAndCopyValue(Pin, PinInfo));
+	switch (PinInfo.PinType)
+	{
+	case EMaterialGraphPinType::Data:
+		switch (Pin->Direction)
+		{
+		case EGPD_Input: InputPins.RemoveAt(PinInfo.Index); break;
+		case EGPD_Output: OutputPins.RemoveAt(PinInfo.Index); break;
+		default: checkNoEntry(); break;
+		}
+		break;
+	case EMaterialGraphPinType::Exec:
+		switch (Pin->Direction)
+		{
+		case EGPD_Input: check(ExecInputPin == Pin);  ExecInputPin = nullptr; break;
+		case EGPD_Output: ExecOutputPins.RemoveAt(PinInfo.Index); break;
+		default: checkNoEntry(); break;
+		}
+		break;
+	default:
+		checkNoEntry();
+		break;
+	}
+
+	// Shift down indices to account for the pin we removed
+	for (auto& It : PinInfoMap)
+	{
+		if (It.Value.PinType == PinInfo.PinType &&
+			It.Value.Index > PinInfo.Index &&
+			It.Key->Direction == PinDirection)
+		{
+			It.Value.Index--;
+		}
+	}
+
 	Super::RemovePinAt(PinIndex, PinDirection);
 
 	UMaterialGraph* MaterialGraph = CastChecked<UMaterialGraph>(GetGraph());
diff --git a/Engine/Source/Runtime/Engine/Classes/EdGraph/EdGraphNode.h b/Engine/Source/Runtime/Engine/Classes/EdGraph/EdGraphNode.h
--- a/Engine/Source/Runtime/Engine/Classes/EdGraph/EdGraphNode.h
+++ b/Engine/Source/Runtime/Engine/Classes/EdGraph/EdGraphNode.h
@@ -660,6 +660,9 @@ public:
 	*/
 	UEdGraphPin* GetPinAt(int32 Index) const;
 
+	/** Gets the pin with the given direction, at the given index. Pins of each direction are indexed separately for the purposes of this method */
+	UEdGraphPin* GetPinWithDirectionAt(int32 Index, EEdGraphPinDirection PinDirection) const;
+
 	/** Break all links on this node */
 	void BreakAllNodeLinks();
 
diff --git a/Engine/Source/Runtime/Engine/Private/EdGraph/EdGraphNode.cpp b/Engine/Source/Runtime/Engine/Private/EdGraph/EdGraphNode.cpp
--- a/Engine/Source/Runtime/Engine/Private/EdGraph/EdGraphNode.cpp
+++ b/Engine/Source/Runtime/Engine/Private/EdGraph/EdGraphNode.cpp
@@ -479,25 +479,9 @@ void UEdGraphNode::RemovePinAt(const int32 PinIndex, const EEdGraphPinDirection
 {
 	Modify();
 
-	// Map requested input to actual pin index
-	int32 ActualPinIndex = INDEX_NONE;
-	int32 MatchingPinCount = 0;
-
-	for (int32 Index = 0; Index < Pins.Num(); Index++)
-	{
-		if (Pins[Index]->Direction == PinDirection)
-		{
-			if (PinIndex == MatchingPinCount)
-			{
-				ActualPinIndex = Index;
-			}
-			++MatchingPinCount;
-		}
-	}
+	UEdGraphPin* OldPin = GetPinWithDirectionAt(PinIndex, PinDirection);
+	checkf(OldPin, TEXT("Tried to remove a non-existent pin."));
 
-	checkf(ActualPinIndex != INDEX_NONE && ActualPinIndex < Pins.Num(), TEXT("Tried to remove a non-existent pin."));
-
-	UEdGraphPin* OldPin = Pins[ActualPinIndex];
 	OldPin->BreakAllPinLinks();
 	RemovePin(OldPin);
 
@@ -842,6 +826,25 @@ UEdGraphPin* UEdGraphNode::GetPinAt(int32 index) const
 	return nullptr;
 }
 
+UEdGraphPin* UEdGraphNode::GetPinWithDirectionAt(int32 PinIndex, EEdGraphPinDirection PinDirection) const
+{
+	// Map requested input to actual pin index
+	int32 MatchingPinCount = 0;
+	for (UEdGraphPin* Pin : Pins)
+	{
+		if (Pin->Direction == PinDirection)
+		{
+			if (PinIndex == MatchingPinCount)
+			{
+				return Pin;
+			}
+			++MatchingPinCount;
+		}
+	}
+
+	return nullptr;
+}
+
 void UEdGraphNode::AddSearchMetaDataInfo(TArray<struct FSearchTagDataPair>& OutTaggedMetaData) const
 {
 	// Searchable - Primary label for the item in the search results
