General Transaction Model Clean-up
- Fixes most Undo/Redo issues
- Fix desync between frontend & editor when moving nodes
- Add Metasound Icon & respective logic
- Add logic to bypass waveform rendering on inherited SoundWave types

diff --git a/Engine/Plugins/Runtime/Metasound/Source/MetasoundEditor/Private/MetasoundEditor.cpp b/Engine/Plugins/Runtime/Metasound/Source/MetasoundEditor/Private/MetasoundEditor.cpp
--- a/Engine/Plugins/Runtime/Metasound/Source/MetasoundEditor/Private/MetasoundEditor.cpp
+++ b/Engine/Plugins/Runtime/Metasound/Source/MetasoundEditor/Private/MetasoundEditor.cpp
@@ -247,23 +247,6 @@ namespace Metasound
 			GEditor->UnregisterForUndo(this);
 		}
 
-		Frontend::FGraphHandle FEditor::InitMetasound(UObject& InMetasound)
-		{
-			FMetasoundAssetBase* MetasoundAsset = IMetasoundUObjectRegistry::Get().GetObjectAsAssetBase(&InMetasound);
-			check(MetasoundAsset);
-
-			if (!MetasoundAsset->GetGraph())
-			{
-				UMetasoundEditorGraph* Graph = NewObject<UMetasoundEditorGraph>(&InMetasound);
-				Graph->ParentMetasound = &InMetasound;
-				Graph->Schema = UMetasoundEditorGraphSchema::StaticClass();
-				MetasoundAsset->SetGraph(Graph);
-				FGraphBuilder::ConstructGraph(InMetasound);
-			}
-
-			return MetasoundAsset->GetRootGraphHandle();
-		}
-
 		void FEditor::InitMetasoundEditor(const EToolkitMode::Type Mode, const TSharedPtr<IToolkitHost>& InitToolkitHost, UObject* ObjectToEdit)
 		{
 			check(ObjectToEdit);
@@ -273,7 +256,6 @@ namespace Metasound
 			ObjectToEdit->SetFlags(RF_Transactional);
 
 			Metasound = ObjectToEdit;
-			InitMetasound(*Metasound);
 
 			GEditor->RegisterForUndo(this);
 
@@ -978,7 +960,7 @@ namespace Metasound
 
 				GraphEditorCommands->MapAction(FGenericCommands::Get().Duplicate,
 					FExecuteAction::CreateLambda([this] { CopySelectedNodes(); PasteNodes(); }),
-					FCanExecuteAction::CreateLambda([this]() { return CanCopyNodes(); }));
+					FCanExecuteAction::CreateLambda([this]() { return CanCopyNodes() && CanPasteNodes(); }));
 
 				// Alignment Commands
 				GraphEditorCommands->MapAction(FGraphEditorCommands::Get().AlignNodesTop,
@@ -1162,7 +1144,9 @@ namespace Metasound
 
 				if (!Actions.IsEmpty())
 				{
-					const FScopedTransaction Transaction(TEXT(""), LOCTEXT("MetasoundEditorDeleteSelectedNode", "Delete Metasound Variable"), Metasound);
+					check(Metasound);
+					const FScopedTransaction Transaction(LOCTEXT("MetasoundEditorDeleteSelectedNode", "Delete Metasound Variable"));
+					Metasound->Modify();
 
 					UMetasoundEditorGraph* Graph = CastChecked<UMetasoundEditorGraph>(MetasoundGraphEditor->GetCurrentGraph());
 					Graph->Modify();
@@ -1212,34 +1196,26 @@ namespace Metasound
 		{
 			using namespace Metasound::Frontend;
 
-			const FScopedTransaction Transaction(LOCTEXT("MetasoundEditorDeleteSelectedNode", "Delete Selected Metasound Node(s)"));
-
-			UMetasoundEditorGraph* Graph = CastChecked<UMetasoundEditorGraph>(MetasoundGraphEditor->GetCurrentGraph());
-
 			const FGraphPanelSelectionSet SelectedNodes = MetasoundGraphEditor->GetSelectedNodes();
 			MetasoundGraphEditor->ClearSelectionSet();
 
-			bool bModified = false;
+			const FScopedTransaction Transaction(LOCTEXT("MetasoundEditorDeleteSelectedNode", "Delete Selected Metasound Node(s)"));
+			check(Metasound);
+			Metasound->Modify();
+			UEdGraph* Graph = MetasoundGraphEditor->GetCurrentGraph();
+			check(Graph);
+			Graph->Modify();
 			for (UObject* NodeObj : SelectedNodes)
 			{
 				// Some nodes may not be metasound nodes (ex. comments and perhaps aliases eventually), but can be safely deleted.
 				if (UEdGraphNode* Node = Cast<UEdGraphNode>(NodeObj))
 				{
-					if (FGraphBuilder::DeleteNode(*Node))
-					{
-						bModified = true;
-					}
-					else
+					if (!FGraphBuilder::DeleteNode(*Node))
 					{
 						MetasoundGraphEditor->SetNodeSelection(Node, true /* bSelect */);
 					}
 				}
 			}
-
-			if (bModified)
-			{
-				Graph->Modify();
-			}
 		}
 
 		void FEditor::CutSelectedNodes()
@@ -1360,6 +1336,7 @@ namespace Metasound
 			}
 
 			const FScopedTransaction Transaction(LOCTEXT("MetasoundEditorPaste", "Paste Metasound Node(s)"));
+			Metasound->Modify();
 			Graph->Modify();
 
 			// Clear the selection set (newly pasted stuff will be selected)
@@ -1386,14 +1363,14 @@ namespace Metasound
 				}
 				else if (UMetasoundEditorGraphInputNode* InputNode = Cast<UMetasoundEditorGraphInputNode>(GraphNode))
 				{
-					if (!MetasoundGraph->Inputs.Contains(InputNode->Input))
+					if (!MetasoundGraph->ContainsInput(InputNode->Input))
 					{
 						NodesToRemove.Add(GraphNode);
 					}
 				}
 				else if (UMetasoundEditorGraphOutputNode* OutputNode = Cast<UMetasoundEditorGraphOutputNode>(GraphNode))
 				{
-					if (!MetasoundGraph->Outputs.Contains(OutputNode->Output))
+					if (!MetasoundGraph->ContainsOutput(OutputNode->Output))
 					{
 						NodesToRemove.Add(GraphNode);
 					}
diff --git a/Engine/Plugins/Runtime/Metasound/Source/MetasoundEditor/Private/MetasoundEditor.h b/Engine/Plugins/Runtime/Metasound/Source/MetasoundEditor/Private/MetasoundEditor.h
--- a/Engine/Plugins/Runtime/Metasound/Source/MetasoundEditor/Private/MetasoundEditor.h
+++ b/Engine/Plugins/Runtime/Metasound/Source/MetasoundEditor/Private/MetasoundEditor.h
@@ -171,8 +171,6 @@ namespace Metasound
 			void RedoGraphAction();
 
 		private:
-			static Frontend::FGraphHandle InitMetasound(UObject& InMetasound);
-
 			/** FNotifyHook interface */
 			virtual void NotifyPostChange(const FPropertyChangedEvent& PropertyChangedEvent, FProperty* PropertyThatChanged) override;
 
diff --git a/Engine/Plugins/Runtime/Metasound/Source/MetasoundEditor/Private/MetasoundEditorGraph.cpp b/Engine/Plugins/Runtime/Metasound/Source/MetasoundEditor/Private/MetasoundEditorGraph.cpp
--- a/Engine/Plugins/Runtime/Metasound/Source/MetasoundEditor/Private/MetasoundEditorGraph.cpp
+++ b/Engine/Plugins/Runtime/Metasound/Source/MetasoundEditor/Private/MetasoundEditorGraph.cpp
@@ -45,6 +45,7 @@ void UMetasoundEditorGraphVariable::SetDataType(FName InNewType)
 	}
 
 	const FScopedTransaction Transaction(LOCTEXT("SetVariableDataType", "Set Metasound Variable Type"));
+	Graph->GetMetasoundChecked().Modify();
 	Graph->Modify();
 
 	// 1. Cache current editor input node reference positions & delete nodes.
@@ -115,13 +116,26 @@ Metasound::Frontend::FConstNodeHandle UMetasoundEditorGraphVariable::GetConstNod
 	return GetNodeHandle();
 }
 
-void UMetasoundEditorGraphInput::UpdateDocumentInput() const
+#if WITH_EDITORONLY_DATA
+void UMetasoundEditorGraphInputLiteral::PostEditUndo()
+{
+	Super::PostEditUndo();
+
+	if (UMetasoundEditorGraphInput* Input = Cast<UMetasoundEditorGraphInput>(GetOuter()))
+	{
+		Input->OnLiteralChanged(false /* bPostTransaction */);
+	}
+}
+#endif // WITH_EDITORONLY_DATA
+
+void UMetasoundEditorGraphInput::UpdateDocumentInput(bool bPostTransaction)
 {
 	using namespace Metasound;
 	using namespace Metasound::Frontend;
 
-	UObject* Object = CastChecked<UMetasoundEditorGraph>(GetOuter())->GetMetasound();
-	if (!ensure(Object))
+	UMetasoundEditorGraph* MetasoundGraph = CastChecked<UMetasoundEditorGraph>(GetOuter());
+	UObject* Metasound = MetasoundGraph->GetMetasound();
+	if (!ensure(Metasound))
 	{
 		return;
 	}
@@ -131,7 +145,10 @@ void UMetasoundEditorGraphInput::UpdateDocumentInput() const
 		return;
 	}
 
-	FMetasoundAssetBase* MetasoundAsset = IMetasoundUObjectRegistry::Get().GetObjectAsAssetBase(Object);
+	const FScopedTransaction Transaction(LOCTEXT("Set Input Default", "Set Metasound Input Default"), bPostTransaction);
+	Metasound->Modify();
+
+	FMetasoundAssetBase* MetasoundAsset = IMetasoundUObjectRegistry::Get().GetObjectAsAssetBase(Metasound);
 	check(MetasoundAsset);
 
 	FGraphHandle GraphHandle = MetasoundAsset->GetRootGraphHandle();
@@ -167,9 +184,27 @@ Metasound::Frontend::FNodeHandle UMetasoundEditorGraphInput::AddNodeHandle(const
 	return NewNodeHandle;
 }
 
-void UMetasoundEditorGraphInput::OnLiteralChanged()
+#if WITH_EDITORONLY_DATA
+void UMetasoundEditorGraphInput::PostEditUndo()
+{
+	Super::PostEditUndo();
+
+	if (!Literal)
+	{
+		if (UMetasoundEditorGraph* MetasoundGraph = Cast<UMetasoundEditorGraph>(GetOuter()))
+		{
+			MetasoundGraph->RemoveVariable(*CastChecked<UMetasoundEditorGraphVariable>(this));
+		}
+		return;
+	}
+
+	OnLiteralChanged(false /* bPostTransaction */);
+}
+#endif // WITH_EDITORONLY_DATA
+
+void UMetasoundEditorGraphInput::OnLiteralChanged(bool bPostTransaction)
 {
-	UpdateDocumentInput();
+	UpdateDocumentInput(bPostTransaction);
 
 	if (UMetasoundEditorGraph* MetasoundGraph = Cast<UMetasoundEditorGraph>(GetOuter()))
 	{
@@ -199,7 +234,7 @@ void UMetasoundEditorGraphInput::OnDataTypeChanged()
 	{
 		InputLiteralClass = UMetasoundEditorGraphInputLiteral::StaticClass();
 	}
-	Literal = NewObject<UMetasoundEditorGraphInputLiteral>(this, InputLiteralClass);
+	Literal = NewObject<UMetasoundEditorGraphInputLiteral>(this, InputLiteralClass, FName(), RF_Transactional);
 }
 
 Metasound::Frontend::FNodeHandle UMetasoundEditorGraphOutput::AddNodeHandle(const FString& InNodeName, const FText& InNodeDisplayName, FName InDataType)
@@ -317,7 +352,7 @@ UMetasoundEditorGraphInput* UMetasoundEditorGraph::FindOrAddInput(Metasound::Fro
 		return Input;
 	}
 
-	if (UMetasoundEditorGraphInput* NewInput = NewObject<UMetasoundEditorGraphInput>(this))
+	if (UMetasoundEditorGraphInput* NewInput = NewObject<UMetasoundEditorGraphInput>(this, FName(), RF_Transactional))
 	{
 		NewInput->NodeID = NodeID;
 		NewInput->ClassName = InNodeHandle->GetClassName();
@@ -328,7 +363,7 @@ UMetasoundEditorGraphInput* UMetasoundEditorGraph::FindOrAddInput(Metasound::Fro
 		IMetasoundEditorModule& EditorModule = FModuleManager::GetModuleChecked<IMetasoundEditorModule>("MetasoundEditor");
 		TSubclassOf<UMetasoundEditorGraphInputLiteral> InputLiteralClass = EditorModule.FindInputLiteralClass(LiteralType);
 
-		NewInput->Literal = NewObject<UMetasoundEditorGraphInputLiteral>(this, InputLiteralClass);
+		NewInput->Literal = NewObject<UMetasoundEditorGraphInputLiteral>(NewInput, InputLiteralClass, FName(), RF_Transactional);
 		NewInput->Literal->SetFromLiteral(DefaultLiteral);
 
 		Inputs.Add(NewInput);
@@ -372,7 +407,7 @@ UMetasoundEditorGraphOutput* UMetasoundEditorGraph::FindOrAddOutput(Metasound::F
 		return Output;
 	}
 
-	if (UMetasoundEditorGraphOutput* NewOutput = NewObject<UMetasoundEditorGraphOutput>(this))
+	if (UMetasoundEditorGraphOutput* NewOutput = NewObject<UMetasoundEditorGraphOutput>(this, FName(), RF_Transactional))
 	{
 		NewOutput->NodeID = NodeID;
 		NewOutput->ClassName = InNodeHandle->GetClassName();
@@ -395,6 +430,16 @@ UMetasoundEditorGraphVariable* UMetasoundEditorGraph::FindVariable(FGuid InNodeI
 	return FindInput(InNodeID);
 }
 
+bool UMetasoundEditorGraph::ContainsInput(UMetasoundEditorGraphInput* InInput) const
+{
+	return Inputs.Contains(InInput);
+}
+
+bool UMetasoundEditorGraph::ContainsOutput(UMetasoundEditorGraphOutput* InOutput) const
+{
+	return Outputs.Contains(InOutput);
+}
+
 void UMetasoundEditorGraph::IterateInputs(TUniqueFunction<void(UMetasoundEditorGraphInput&)> InFunction) const
 {
 	for (UMetasoundEditorGraphInput* Input : Inputs)
diff --git a/Engine/Plugins/Runtime/Metasound/Source/MetasoundEditor/Private/MetasoundEditorGraphBuilder.cpp b/Engine/Plugins/Runtime/Metasound/Source/MetasoundEditor/Private/MetasoundEditorGraphBuilder.cpp
--- a/Engine/Plugins/Runtime/Metasound/Source/MetasoundEditor/Private/MetasoundEditorGraphBuilder.cpp
+++ b/Engine/Plugins/Runtime/Metasound/Source/MetasoundEditor/Private/MetasoundEditorGraphBuilder.cpp
@@ -60,8 +60,6 @@ namespace Metasound
 						GraphHandle->RemoveNode(*InNodeHandle);
 					}
 				}
-
-				InMetasound.MarkPackageDirty();
 			}
 		}
 
@@ -130,9 +128,6 @@ namespace Metasound
 			NewGraphNode->SetNodeID(InNodeHandle->GetID());
 
 			RebuildNodePins(*NewGraphNode);
-
-			InMetasound.PostEditChange();
-			InMetasound.MarkPackageDirty();
 		}
 
 		Frontend::FNodeHandle FGraphBuilder::AddExternalNodeHandle(UObject& InMetasound, const Frontend::FNodeClassInfo& InClassInfo)
@@ -359,6 +354,16 @@ namespace Metasound
 			return nullptr;
 		}
 
+		bool FGraphBuilder::IsLiteralInput(Frontend::FNodeHandle InNodeHandle)
+		{
+			if (InNodeHandle->GetClassType() == EMetasoundFrontendClassType::Input)
+			{
+				return InNodeHandle->GetNodeStyle().Display.Visibility == EMetasoundFrontendNodeStyleDisplayVisibility::Hidden;
+			}
+
+			return false;
+		}
+
 		void FGraphBuilder::AddOrUpdateLiteralInput(UObject& InMetasound, Frontend::FNodeHandle InNodeHandle, UEdGraphPin& InInputPin, bool bForcePinValueAsDefault)
 		{
 			using namespace Frontend;
@@ -406,8 +411,7 @@ namespace Metasound
 			else
 			{
 				FNodeHandle InputNode = OutputHandle->GetOwningNode();
-				const FMetasoundFrontendNodeStyle& Style = InputNode->GetNodeStyle();
-				if (InputNode->IsValid() && Style.Display.Visibility == EMetasoundFrontendNodeStyleDisplayVisibility::Hidden)
+				if (IsLiteralInput(InputNode))
 				{
 					Frontend::FGraphHandle GraphHandle = InputNode->GetOwningGraph();
 					const FGuid VertexID = GraphHandle->GetVertexIDForInputVertex(InputNode->GetNodeName());
@@ -716,8 +720,7 @@ namespace Metasound
 			if (ExistingOutput->IsValid())
 			{
 				FNodeHandle NodeHandle = ExistingOutput->GetOwningNode();
-				const FMetasoundFrontendNodeStyle& NodeStyle = NodeHandle->GetNodeStyle();
-				if (NodeHandle->IsValid() && NodeStyle.Display.Visibility == EMetasoundFrontendNodeStyleDisplayVisibility::Hidden)
+				if (IsLiteralInput(NodeHandle))
 				{
 					if (UMetasoundEditorGraph* Graph = Cast<UMetasoundEditorGraph>(InInputPin.GetOwningNode()->GetGraph()))
 					{
@@ -915,7 +918,7 @@ namespace Metasound
 				if (Output->IsValid())
 				{
 					FNodeHandle InputHandle = Output->GetOwningNode();
-					if (InputHandle->GetNodeStyle().Display.Visibility == EMetasoundFrontendNodeStyleDisplayVisibility::Hidden)
+					if (IsLiteralInput(InputHandle))
 					{
 						FGraphHandle Graph = InputHandle->GetOwningGraph();
 						Graph->RemoveNode(*InputHandle);
@@ -1003,10 +1006,11 @@ namespace Metasound
 		{
 			using namespace Frontend;
 		
+			DeleteLiteralInputs(InGraphNode);
+
 			FNodeHandle NodeHandle = InGraphNode.GetNodeHandle();
 			for (int32 i = InGraphNode.Pins.Num() - 1; i >= 0; i--)
 			{
-				DeleteLiteralInputs(InGraphNode);
 				InGraphNode.RemovePin(InGraphNode.Pins[i]);
 			}
 
@@ -1064,15 +1068,13 @@ namespace Metasound
 			using namespace Frontend;
 
 			IMetasoundEditorModule& EditorModule = FModuleManager::GetModuleChecked<IMetasoundEditorModule>("MetasoundEditor");
-			FEdGraphPinType PinType("MetasoundNode", NAME_None, nullptr, EPinContainerType::None, false, FEdGraphTerminalType());
-			UEdGraphPin* NewPin = InEditorNode.CreatePin(EGPD_Input, PinType, FName(*InInputHandle->GetName()));
+			const FEdGraphPinType PinType = EditorModule.FindDataType(InInputHandle->GetDataType()).PinType;
 
+			UEdGraphPin* NewPin = InEditorNode.CreatePin(EGPD_Input, PinType, FName(*InInputHandle->GetName()));
 			if (ensure(NewPin))
 			{
 				NewPin->PinToolTip = InInputHandle->GetTooltip().ToString();
-				NewPin->PinType = EditorModule.FindDataType(InInputHandle->GetDataType()).PinType;
 				NewPin->bAdvancedView = InInputHandle->GetMetadata().bIsAdvancedDisplay;
-
 				if (NewPin->bAdvancedView)
 				{
 					if (InEditorNode.AdvancedPinDisplay == ENodeAdvancedPins::NoPins)
@@ -1092,15 +1094,13 @@ namespace Metasound
 		UEdGraphPin* FGraphBuilder::AddPinToNode(UMetasoundEditorGraphNode& InEditorNode, Frontend::FOutputHandle InOutputHandle)
 		{
 			IMetasoundEditorModule& EditorModule = FModuleManager::GetModuleChecked<IMetasoundEditorModule>("MetasoundEditor");
-			FEdGraphPinType PinType("MetasoundNode", NAME_None, nullptr, EPinContainerType::None, false, FEdGraphTerminalType());
-			UEdGraphPin* NewPin = InEditorNode.CreatePin(EGPD_Output, PinType, FName(*InOutputHandle->GetName()));
+			FEdGraphPinType	PinType = EditorModule.FindDataType(InOutputHandle->GetDataType()).PinType;
 
+			UEdGraphPin* NewPin = InEditorNode.CreatePin(EGPD_Output, PinType, FName(*InOutputHandle->GetName()));
 			if (ensure(NewPin))
 			{
 				NewPin->PinToolTip = InOutputHandle->GetTooltip().ToString();
-				NewPin->PinType = EditorModule.FindDataType(InOutputHandle->GetDataType()).PinType;
 				NewPin->bAdvancedView = InOutputHandle->GetMetadata().bIsAdvancedDisplay;
-
 				if (NewPin->bAdvancedView)
 				{
 					if (InEditorNode.AdvancedPinDisplay == ENodeAdvancedPins::NoPins)
@@ -1238,8 +1238,8 @@ namespace Metasound
 						if (bIsConnectedToLiteralInput)
 						{
 							FNodeHandle InputNode = OutputHandle->GetOwningNode();
-							const FMetasoundFrontendNodeStyle& Style = InputNode->GetNodeStyle();
-							const FGuid VertexID = GraphHandle->GetVertexIDForInputVertex(InputNode->GetNodeName());
+							const FString& InputName = InputNode->GetNodeName();
+							const FGuid VertexID = GraphHandle->GetVertexIDForInputVertex(InputName);
 							const FMetasoundFrontendLiteral DefaultLiteral = GraphHandle->GetDefaultInput(VertexID);
 							const FString DefaultLiteralString = DefaultLiteral.ToString();
 							if (MatchingPin->DefaultValue != DefaultLiteralString)
@@ -1344,39 +1344,47 @@ namespace Metasound
 			// Add missing editor nodes marked as visible.
 			for (FNodeHandle Node : FrontendNodes)
 			{
-				FMetasoundFrontendNodeStyle Style = Node->GetNodeStyle();
-				if (Style.Display.Visibility == EMetasoundFrontendNodeStyleDisplayVisibility::Visible)
+				const FMetasoundFrontendNodeStyle& CurrentStyle = Node->GetNodeStyle();
+				if (CurrentStyle.Display.Locations.IsEmpty())
+				{
+					continue;
+				}
+
+				FMetasoundFrontendNodeStyle NewStyle = CurrentStyle;
+				if (NewStyle.Display.Visibility == EMetasoundFrontendNodeStyleDisplayVisibility::Hidden)
 				{
-					bIsEditorGraphDirty = true;
+					continue;
+				}
+
+				bIsEditorGraphDirty = true;
 
-					TArray<UMetasoundEditorGraphNode*> AddedNodes;
-					for (const TPair<FGuid, FVector2D>& Location : Style.Display.Locations)
+				TArray<UMetasoundEditorGraphNode*> AddedNodes;
+				for (const TPair<FGuid, FVector2D>& Location : NewStyle.Display.Locations)
+				{
+					UMetasoundEditorGraphNode* NewNode = AddNode(InMetasound, Node, Location.Value, false /* bInSelectNewNode */);
+					if (ensure(NewNode))
 					{
-						UMetasoundEditorGraphNode* NewNode = AddNode(InMetasound, Node, Location.Value, false /* bInSelectNewNode */);
-						if (ensure(NewNode))
+						FAssociatedNodes& AssociatedNodeData = AssociatedNodes.FindOrAdd(Node->GetID());
+						if (AssociatedNodeData.Node->IsValid())
 						{
-							FAssociatedNodes& AssociatedNodeData = AssociatedNodes.FindOrAdd(Node->GetID());
-							if (AssociatedNodeData.Node->IsValid())
-							{
-								ensure(AssociatedNodeData.Node == Node);
-							}
-							else
-							{
-								AssociatedNodeData.Node = Node;
-							}
-
-							AddedNodes.Add(NewNode);
-							AssociatedNodeData.EditorNodes.Add(NewNode);
+							ensure(AssociatedNodeData.Node == Node);
+						}
+						else
+						{
+							AssociatedNodeData.Node = Node;
 						}
-					}
 
-					Style.Display.Locations.Reset();
-					for (UMetasoundEditorGraphNode* EditorNode : AddedNodes)
-					{
-						Style.Display.Locations.Add(EditorNode->NodeGuid, FVector2D(EditorNode->NodePosX, EditorNode->NodePosY));
+						AddedNodes.Add(NewNode);
+						AssociatedNodeData.EditorNodes.Add(NewNode);
 					}
-					Node->SetNodeStyle(Style);
 				}
+
+				NewStyle.Display.Locations.Reset();
+				for (UMetasoundEditorGraphNode* EditorNode : AddedNodes)
+				{
+					NewStyle.Display.Locations.Add(EditorNode->NodeGuid, FVector2D(EditorNode->NodePosX, EditorNode->NodePosY));
+				}
+				Node->SetNodeStyle(NewStyle);
 			}
 
 			// Synchronize pins on node associations.
@@ -1489,11 +1497,6 @@ namespace Metasound
 					AddPinToNode(InEditorNode, OutputHandle);
 				}
 			}
-			
-			if (bIsNodeDirty)
-			{
-				InEditorNode.MarkPackageDirty();
-			}
 
 			return bIsNodeDirty;
 		}
@@ -1516,6 +1519,12 @@ namespace Metasound
 			FGraphHandle GraphHandle = MetasoundAsset->GetRootGraphHandle();
 			GraphHandle->IterateNodes([&](FNodeHandle NodeHandle)
 			{
+				// Literal Inputs are not described on EdGraph side
+				if (IsLiteralInput(NodeHandle))
+				{
+					return;
+				}
+
 				if (UMetasoundEditorGraphInput* Input = Graph->FindInput(NodeHandle->GetID()))
 				{
 					Inputs.Add(Input);
diff --git a/Engine/Plugins/Runtime/Metasound/Source/MetasoundEditor/Private/MetasoundEditorGraphBuilder.h b/Engine/Plugins/Runtime/Metasound/Source/MetasoundEditor/Private/MetasoundEditorGraphBuilder.h
--- a/Engine/Plugins/Runtime/Metasound/Source/MetasoundEditor/Private/MetasoundEditorGraphBuilder.h
+++ b/Engine/Plugins/Runtime/Metasound/Source/MetasoundEditor/Private/MetasoundEditorGraphBuilder.h
@@ -121,6 +121,9 @@ namespace Metasound
 			static Frontend::FOutputHandle GetOutputHandleFromPin(const UEdGraphPin* InPin);
 			static Frontend::FConstOutputHandle GetConstOutputHandleFromPin(const UEdGraphPin* InPin);
 
+			// Returns whether or not the given handle refers to a literal input
+			static bool IsLiteralInput(Frontend::FNodeHandle InNodeHandle);
+
 			// Adds or updates a hidden input that is set to the default literal value.  If the literal input is added, value defaults to
 			// 1. what is set on the input vertex of the node's inputer vertex interface, and if not defined there,
 			// 2. the pin's DataType literal default.
diff --git a/Engine/Plugins/Runtime/Metasound/Source/MetasoundEditor/Private/MetasoundEditorGraphNode.cpp b/Engine/Plugins/Runtime/Metasound/Source/MetasoundEditor/Private/MetasoundEditorGraphNode.cpp
--- a/Engine/Plugins/Runtime/Metasound/Source/MetasoundEditor/Private/MetasoundEditorGraphNode.cpp
+++ b/Engine/Plugins/Runtime/Metasound/Source/MetasoundEditor/Private/MetasoundEditorGraphNode.cpp
@@ -224,6 +224,8 @@ void UMetasoundEditorGraphNode::PinDefaultValueChanged(UEdGraphPin* Pin)
 
 	if (Pin && Pin->Direction == EGPD_Input)
 	{
+		GetMetasoundChecked().Modify();
+
 		FNodeHandle NodeHandle = GetNodeHandle();
 		IMetasoundEditorModule& EditorModule = FModuleManager::GetModuleChecked<IMetasoundEditorModule>("MetasoundEditor");
 		FGraphBuilder::AddOrUpdateLiteralInput(GetMetasoundChecked(), NodeHandle, *Pin);
@@ -244,10 +246,7 @@ void UMetasoundEditorGraphNode::PostEditChangeProperty(FPropertyChangedEvent& In
 	if (PropertyName == GET_MEMBER_NAME_CHECKED(UEdGraphNode, NodePosX)
 		|| PropertyName == GET_MEMBER_NAME_CHECKED(UEdGraphNode, NodePosY))
 	{
-		Metasound::Frontend::FNodeHandle NodeHandle = GetNodeHandle();
-		FMetasoundFrontendNodeStyle Style = NodeHandle->GetNodeStyle();
-		Style.Display.Locations.FindOrAdd(NodeGuid) = FVector2D(NodePosX, NodePosY);
-		GetNodeHandle()->SetNodeStyle(Style);
+		UpdatePosition();
 	}
 
 	Super::PostEditChangeProperty(InEvent);
@@ -281,6 +280,16 @@ void UMetasoundEditorGraphNode::PostEditUndo()
 	}
 }
 
+void UMetasoundEditorGraphNode::UpdatePosition()
+{
+	GetMetasoundChecked().Modify();
+
+	Metasound::Frontend::FNodeHandle NodeHandle = GetNodeHandle();
+	FMetasoundFrontendNodeStyle Style = NodeHandle->GetNodeStyle();
+	Style.Display.Locations.FindOrAdd(NodeGuid) = FVector2D(NodePosX, NodePosY);
+	GetNodeHandle()->SetNodeStyle(Style);
+}
+
 void UMetasoundEditorGraphNode::PostDuplicate(bool bDuplicateForPIE)
 {
 	Super::PostDuplicate(bDuplicateForPIE);
diff --git a/Engine/Plugins/Runtime/Metasound/Source/MetasoundEditor/Private/MetasoundEditorGraphSchema.cpp b/Engine/Plugins/Runtime/Metasound/Source/MetasoundEditor/Private/MetasoundEditorGraphSchema.cpp
--- a/Engine/Plugins/Runtime/Metasound/Source/MetasoundEditor/Private/MetasoundEditorGraphSchema.cpp
+++ b/Engine/Plugins/Runtime/Metasound/Source/MetasoundEditor/Private/MetasoundEditorGraphSchema.cpp
@@ -167,11 +167,13 @@ UEdGraphNode* FMetasoundGraphSchemaAction_NewNode::PerformAction(UEdGraph* Paren
 	using namespace Metasound::Frontend;
 
 	const FScopedTransaction Transaction(LOCTEXT("AddNewNode", "Add New Metasound Node"));
+	UObject& ParentMetasound = CastChecked<UMetasoundEditorGraph>(ParentGraph)->GetMetasoundChecked();
+	ParentMetasound.Modify();
 	ParentGraph->Modify();
 
-	UObject& ParentMetasound = CastChecked<UMetasoundEditorGraph>(ParentGraph)->GetMetasoundChecked();
 	if (UMetasoundEditorGraphExternalNode* NewGraphNode = FGraphBuilder::AddExternalNode(ParentMetasound, NodeClassInfo, Location, bSelectNewNode))
 	{
+		NewGraphNode->Modify();
 		SchemaPrivate::TryConnectNewNodeToPin(*NewGraphNode, FromPin);
 		return NewGraphNode;
 	}
@@ -190,26 +192,30 @@ UEdGraphNode* FMetasoundGraphSchemaAction_NewInput::PerformAction(UEdGraph* Pare
 	using namespace Metasound::Editor;
 	using namespace Metasound::Frontend;
 
-	const FScopedTransaction Transaction(LOCTEXT("AddNewInputNode", "Add New Metasound Input Node"));
 
 	UMetasoundEditorGraph* MetasoundGraph = CastChecked<UMetasoundEditorGraph>(ParentGraph);
 	UObject& ParentMetasound = MetasoundGraph->GetMetasoundChecked();
-	ParentMetasound.Modify();
 
 	UMetasoundEditorGraphInput* Input = MetasoundGraph->FindInput(NodeID);
 	if (!ensure(Input))
 	{
 		return nullptr;
 	}
-	
+
 	FNodeHandle NodeHandle = Input->GetNodeHandle();
 	if (!ensure(NodeHandle->IsValid()))
 	{
 		return nullptr;
 	}
 
+	const FScopedTransaction Transaction(LOCTEXT("AddNewInputNode", "Add New Metasound Input Node"));
+	ParentMetasound.Modify();
+	MetasoundGraph->Modify();
+	Input->Modify();
+
 	if (UMetasoundEditorGraphInputNode* NewGraphNode = FGraphBuilder::AddInputNode(ParentMetasound, NodeHandle, InLocation))
 	{
+		NewGraphNode->Modify();
 		UEdGraphNode* EdGraphNode = CastChecked<UEdGraphNode>(NewGraphNode);
 		SchemaPrivate::TryConnectNewNodeToPin(*EdGraphNode, FromPin);
 		return EdGraphNode;
@@ -235,10 +241,9 @@ UEdGraphNode* FMetasoundGraphSchemaAction_PromoteToInput::PerformAction(UEdGraph
 	}
 
 	const FScopedTransaction Transaction(LOCTEXT("PromoteNodeInputToGraphInput", "Promote Metasound Node Input to Graph Input"));
-	ParentGraph->Modify();
-
 	UMetasoundEditorGraph* MetasoundGraph = CastChecked<UMetasoundEditorGraph>(ParentGraph);
 	UObject& ParentMetasound = MetasoundGraph->GetMetasoundChecked();
+	ParentMetasound.Modify();
 
 	const FText& InputName = InputHandle->GetOwningNode()->GetDisplayName();
 	const FText& InputNodeName = InputHandle->GetDisplayName();
@@ -286,12 +291,8 @@ UEdGraphNode* FMetasoundGraphSchemaAction_NewOutput::PerformAction(UEdGraph* Par
 	using namespace Metasound::Editor;
 	using namespace Metasound::Frontend;
 
-	const FScopedTransaction Transaction(LOCTEXT("AddNewOutputNode", "Add New Metasound Output Node"));
-	ParentGraph->Modify();
-
 	UMetasoundEditorGraph* MetasoundGraph = CastChecked<UMetasoundEditorGraph>(ParentGraph);
 	UObject& ParentMetasound = MetasoundGraph->GetMetasoundChecked();
-	ParentMetasound.Modify();
 
 	UMetasoundEditorGraphOutput* Output = MetasoundGraph->FindOutput(NodeID);
 	if (!ensure(Output))
@@ -305,6 +306,10 @@ UEdGraphNode* FMetasoundGraphSchemaAction_NewOutput::PerformAction(UEdGraph* Par
 		return nullptr;
 	}
 
+	const FScopedTransaction Transaction(LOCTEXT("AddNewOutputNode", "Add New Metasound Output Node"));
+	ParentMetasound.Modify();
+	ParentGraph->Modify();
+
 	if (UMetasoundEditorGraphOutputNode* NewGraphNode = FGraphBuilder::AddOutputNode(ParentMetasound, NodeHandle, Location, bSelectNewNode))
 	{
 		SchemaPrivate::TryConnectNewNodeToPin(*NewGraphNode, FromPin);
@@ -324,6 +329,9 @@ UEdGraphNode* FMetasoundGraphSchemaAction_NewComment::PerformAction(UEdGraph* Pa
 {
 	using namespace Metasound::Editor;
 
+	const FScopedTransaction Transaction(LOCTEXT("AddNewOutputNode", "Add Comment to Metasound Graph"));
+	ParentGraph->Modify();
+
 	UEdGraphNode_Comment* CommentTemplate = NewObject<UEdGraphNode_Comment>();
 
 	FSlateRect Bounds;
@@ -582,10 +590,8 @@ bool UMetasoundEditorGraphSchema::TryCreateConnection(UEdGraphPin* PinA, UEdGrap
 	// so that transaction stack observes Frontend changes last if rolled back (i.e. undone).  UEdGraphSchema::TryCreateConnection
 	// intrinsically marks the respective pin EdGraphNodes as modified.
 	UEdGraphNode* PinANode = PinA->GetOwningNode();
-	if (UMetasoundEditorGraph* Graph = CastChecked<UMetasoundEditorGraph>(PinANode->GetGraph()))
-	{
-		Graph->GetMetasoundChecked().Modify();
-	}
+	UMetasoundEditorGraph* Graph = CastChecked<UMetasoundEditorGraph>(PinANode->GetGraph());
+	Graph->GetMetasoundChecked().Modify();
 
 	// This call to parent takes care of marking respective nodes for modification.
 	if (!UEdGraphSchema::TryCreateConnection(PinA, PinB))
@@ -664,15 +670,19 @@ FLinearColor UMetasoundEditorGraphSchema::GetPinTypeColor(const FEdGraphPinType&
 
 void UMetasoundEditorGraphSchema::BreakNodeLinks(UEdGraphNode& TargetNode) const
 {
+	using namespace Metasound::Editor;
 	using namespace Metasound::Frontend;
 
 	const FScopedTransaction Transaction(LOCTEXT("BreakNodeLinks", "Break Node Links"));
+	UMetasoundEditorGraph* Graph = CastChecked<UMetasoundEditorGraph>(TargetNode.GetGraph());
+	Graph->GetMetasoundChecked().Modify();
 	TargetNode.Modify();
 
 	TArray<UEdGraphPin*> Pins = TargetNode.GetAllPins();
 	for (UEdGraphPin* Pin : Pins)
 	{
-		BreakPinLinks(*Pin, false);
+		FGraphBuilder::DisconnectPin(*Pin);
+		Super::BreakPinLinks(*Pin, false /* bSendsNodeNotifcation */);
 	}
 	Super::BreakNodeLinks(TargetNode);
 }
@@ -683,10 +693,11 @@ void UMetasoundEditorGraphSchema::BreakPinLinks(UEdGraphPin& TargetPin, bool bSe
 	using namespace Metasound::Frontend;
 
 	const FScopedTransaction Transaction(LOCTEXT("BreakPinLinks", "Break Pin Links"));
-	TargetPin.GetOwningNode()->Modify();
+	UMetasoundEditorGraph* Graph = CastChecked<UMetasoundEditorGraph>(TargetPin.GetOwningNode()->GetGraph());
+	Graph->GetMetasoundChecked().Modify();
+	TargetPin.Modify();
 
 	FGraphBuilder::DisconnectPin(TargetPin);
-
 	Super::BreakPinLinks(TargetPin, bSendsNodeNotifcation);
 }
 
diff --git a/Engine/Plugins/Runtime/Metasound/Source/MetasoundEditor/Private/MetasoundEditorModule.cpp b/Engine/Plugins/Runtime/Metasound/Source/MetasoundEditor/Private/MetasoundEditorModule.cpp
--- a/Engine/Plugins/Runtime/Metasound/Source/MetasoundEditor/Private/MetasoundEditorModule.cpp
+++ b/Engine/Plugins/Runtime/Metasound/Source/MetasoundEditor/Private/MetasoundEditorModule.cpp
@@ -27,6 +27,7 @@
 #include "PropertyEditorDelegates.h"
 #include "PropertyEditorModule.h"
 #include "Styling/CoreStyle.h"
+#include "Styling/StyleColors.h"
 #include "Styling/SlateStyle.h"
 #include "Styling/SlateStyleMacros.h"
 #include "Styling/SlateStyleRegistry.h"
@@ -83,14 +84,31 @@ namespace Metasound
 					// Graph Editor
 					Set("MetasoundEditor.Graph.Node.Body.Input", new FSlateImageBrush(RootToContentDir(TEXT("/Graph/node_input_body_64x.png")), FVector2D(114.0f, 64.0f)));
 					Set("MetasoundEditor.Graph.Node.Body.Default", new FSlateImageBrush(RootToContentDir(TEXT("/Graph/node_default_body_64x.png")), FVector2D(64.0f, 64.0f)));
+
 					Set("MetasoundEditor.Graph.Node.Math.Add", new FSlateImageBrush(RootToContentDir(TEXT("/Graph/node_math_add_40x.png")), Icon40x40));
 					Set("MetasoundEditor.Graph.Node.Math.Divide", new FSlateImageBrush(RootToContentDir(TEXT("/Graph/node_math_divide_40x.png")), Icon40x40));
+					Set("MetasoundEditor.Graph.Node.Math.Modulo", new FSlateImageBrush(RootToContentDir(TEXT("/Graph/node_math_modulo_40x.png")), Icon40x40));
 					Set("MetasoundEditor.Graph.Node.Math.Multiply", new FSlateImageBrush(RootToContentDir(TEXT("/Graph/node_math_multiply_40x.png")), Icon40x40));
 					Set("MetasoundEditor.Graph.Node.Math.Subtract", new FSlateImageBrush(RootToContentDir(TEXT("/Graph/node_math_subtract_40x.png")), Icon40x40));
-					Set("MetasoundEditor.Graph.Node.Math.Modulo", new FSlateImageBrush(RootToContentDir(TEXT("/Graph/node_math_modulo_40x.png")), Icon40x40));
 
 					// Misc
 					Set("MetasoundEditor.Speaker", new FSlateImageBrush(RootToContentDir(TEXT("/Icons/speaker_144x.png")), FVector2D(144.0f, 144.0f)));
+
+					// Class Icons
+					auto SetClassIcon = [this](const FString& ClassName)
+					{
+						static const FVector2D Icon16 = FVector2D(16.0f, 16.0f);
+						static const FVector2D Icon64 = FVector2D(64.0f, 64.0f);
+
+						const FString IconFileName = FString::Printf(TEXT("Icons/%s"), *ClassName.ToLower());
+						const FSlateColor DefaultForeground(FStyleColors::Foreground);
+
+						Set(*FString::Printf(TEXT("ClassIcon.%s"), *ClassName), new IMAGE_BRUSH_SVG(IconFileName, Icon16));
+						Set(*FString::Printf(TEXT("ClassThumbnail.%s"), *ClassName), new IMAGE_BRUSH_SVG(IconFileName, Icon64));
+					};
+
+					SetClassIcon(TEXT("Metasound"));
+					SetClassIcon(TEXT("MetasoundSource"));
 				}
 
 				FSlateStyleRegistry::RegisterSlateStyle(*this);
diff --git a/Engine/Plugins/Runtime/Metasound/Source/MetasoundEditor/Private/MetasoundFactory.cpp b/Engine/Plugins/Runtime/Metasound/Source/MetasoundEditor/Private/MetasoundFactory.cpp
--- a/Engine/Plugins/Runtime/Metasound/Source/MetasoundEditor/Private/MetasoundFactory.cpp
+++ b/Engine/Plugins/Runtime/Metasound/Source/MetasoundEditor/Private/MetasoundFactory.cpp
@@ -3,9 +3,47 @@
 
 #include "Kismet/KismetSystemLibrary.h"
 #include "Metasound.h"
-#include "MetasoundSource.h"
-#include "MetasoundEditorGraphNode.h"
 #include "MetasoundEditorGraphBuilder.h"
+#include "MetasoundEditorGraphNode.h"
+#include "MetasoundEditorGraphSchema.h"
+#include "MetasoundSource.h"
+#include "MetasoundUObjectRegistry.h"
+
+
+namespace Metasound
+{
+	namespace Editor
+	{
+		namespace FactoryPrivate
+		{
+			UObject* InitMetasound(UMetasoundSource* MetasoundSource, EObjectFlags Flags)
+			{
+				FMetasoundFrontendClassMetadata Metadata;
+
+				Metadata.ClassName = FMetasoundFrontendClassName(FName(), *MetasoundSource->GetName(), FName());
+				Metadata.Version.Major = 1;
+				Metadata.Version.Minor = 0;
+				Metadata.Type = EMetasoundFrontendClassType::Graph;
+				Metadata.Author = FText::FromString(UKismetSystemLibrary::GetPlatformUserName());
+
+				MetasoundSource->SetMetadata(Metadata);
+				MetasoundSource->ConformDocumentToMetasoundArchetype();
+
+				FMetasoundAssetBase* MetasoundAsset = IMetasoundUObjectRegistry::Get().GetObjectAsAssetBase(MetasoundSource);
+				check(MetasoundAsset);
+				check(!MetasoundAsset->GetGraph());
+
+				UObject* MetasoundObject = CastChecked<UObject>(MetasoundSource);
+
+				UMetasoundEditorGraph* Graph = NewObject<UMetasoundEditorGraph>(MetasoundSource, FName(), Flags);
+				Graph->Schema = UMetasoundEditorGraphSchema::StaticClass();
+				MetasoundAsset->SetGraph(Graph);
+
+				return MetasoundSource;
+			}
+		} // namespace FactoryPrivate
+	} // namespace Editor
+} // namespace Metasound
 
 // TODO: Re-enable and potentially rename once composition is supported
 // UMetasoundFactory::UMetasoundFactory(const FObjectInitializer& ObjectInitializer)
@@ -20,21 +58,13 @@
 // 
 // UObject* UMetasoundFactory::FactoryCreateNew(UClass* InClass, UObject* InParent, FName Name, EObjectFlags Flags, UObject* InContext, FFeedbackContext* InFeedbackContext)
 // {
-// 	UMetasound* NewMetasound = NewObject<UMetasound>(InParent, Name, Flags);
-// 
-// 	FMetasoundFrontendClassMetadata Metadata;
-// 
-// 	Metadata.ClassName = FMetasoundFrontendClassName{TEXT(""), *NewMetasound->GetName(), TEXT("")};
-// 	Metadata.Version.Major = 1;
-// 	Metadata.Version.Minor = 0;
-// 	Metadata.Type = EMetasoundFrontendClassType::Graph;
-// 	Metadata.Author = FText::FromString(UKismetSystemLibrary::GetPlatformUserName());
+// 	UMetasound* Metasound = NewObject<UMetasound>(InParent, Name, Flags);
+//	Metasound::Editor::FactoryPrivate::InitMetasound(Metasound, Flags);
 // 
-// 	NewMetasound->SetMetadata(Metadata);
+//	CastChecked<UMetasoundEditorGraph>(&Metasound->GetGraphChecked())->ParentMetasound = MetasoundSource;
 // 
-// 	NewMetasound->ConformDocumentToMetasoundArchetype();
-// 
-// 	return NewMetasound;
+//	FGraphBuilder::ConstructGraph(*Metasound);
+//	FGraphBuilder::SynchronizeGraph(*Metasound);
 // }
 
 UMetasoundSourceFactory::UMetasoundSourceFactory(const FObjectInitializer& ObjectInitializer)
@@ -49,20 +79,17 @@ UMetasoundSourceFactory::UMetasoundSourceFactory(const FObjectInitializer& Objec
 
 UObject* UMetasoundSourceFactory::FactoryCreateNew(UClass* InClass, UObject* InParent, FName Name, EObjectFlags Flags, UObject* InContext, FFeedbackContext* InFeedbackContext)
 {
-	UMetasoundSource* MetasoundSource = NewObject<UMetasoundSource>(InParent, Name, Flags);
-
-	FMetasoundFrontendClassMetadata Metadata;
+	using namespace Metasound;
+	using namespace Metasound::Editor;
+	using namespace Metasound::Frontend;
 
-	Metadata.ClassName = FMetasoundFrontendClassName{TEXT(""), *MetasoundSource->GetName(), TEXT("")};
-	Metadata.Version.Major = 1;
-	Metadata.Version.Minor = 0;
-	Metadata.Type = EMetasoundFrontendClassType::Graph;
-	Metadata.Author = FText::FromString(UKismetSystemLibrary::GetPlatformUserName());
+	UMetasoundSource* MetasoundSource = NewObject<UMetasoundSource>(InParent, Name, Flags);
+	Metasound::Editor::FactoryPrivate::InitMetasound(MetasoundSource, Flags);
 
-	MetasoundSource->SetMetadata(Metadata);
+	CastChecked<UMetasoundEditorGraph>(&MetasoundSource->GetGraphChecked())->ParentMetasound = MetasoundSource;
 
-	MetasoundSource->ConformDocumentToMetasoundArchetype();
+	FGraphBuilder::ConstructGraph(*MetasoundSource);
+	FGraphBuilder::SynchronizeGraph(*MetasoundSource);
 
 	return MetasoundSource;
 }
-
diff --git a/Engine/Plugins/Runtime/Metasound/Source/MetasoundEditor/Private/SMetasoundEnumPin.cpp b/Engine/Plugins/Runtime/Metasound/Source/MetasoundEditor/Private/SMetasoundEnumPin.cpp
--- a/Engine/Plugins/Runtime/Metasound/Source/MetasoundEditor/Private/SMetasoundEnumPin.cpp
+++ b/Engine/Plugins/Runtime/Metasound/Source/MetasoundEditor/Private/SMetasoundEnumPin.cpp
@@ -82,9 +82,6 @@ void SMetasoundEnumPin::ComboBoxSelectionChanged(TSharedPtr<int32> NewSelection,
 		FString EnumValueString = FString::FromInt(EnumValue);
 		if (GraphPinObj->GetDefaultAsString() != EnumValueString)
 		{
-			const FScopedTransaction Transaction(NSLOCTEXT("GraphEditor", "ChangeEnumPinValue", "Change Enum Pin Value"));
-			GraphPinObj->Modify();
-
 			//Set new selection
 			GraphPinObj->GetSchema()->TrySetDefaultValue(*GraphPinObj, EnumValueString);
 		}
diff --git a/Engine/Plugins/Runtime/Metasound/Source/MetasoundEditor/Private/SMetasoundGraphNode.cpp b/Engine/Plugins/Runtime/Metasound/Source/MetasoundEditor/Private/SMetasoundGraphNode.cpp
--- a/Engine/Plugins/Runtime/Metasound/Source/MetasoundEditor/Private/SMetasoundGraphNode.cpp
+++ b/Engine/Plugins/Runtime/Metasound/Source/MetasoundEditor/Private/SMetasoundGraphNode.cpp
@@ -4,36 +4,39 @@
 #include "EditorStyleSet.h"
 #include "GraphEditorSettings.h"
 #include "IDocumentation.h"
-#include "KismetPins/SGraphPinNum.h"
 #include "KismetPins/SGraphPinBool.h"
 #include "KismetPins/SGraphPinExec.h"
 #include "KismetPins/SGraphPinInteger.h"
+#include "KismetPins/SGraphPinNum.h"
 #include "KismetPins/SGraphPinObject.h"
 #include "KismetPins/SGraphPinString.h"
 #include "MetasoundEditorGraphBuilder.h"
 #include "MetasoundEditorGraphNode.h"
 #include "MetasoundEditorGraphSchema.h"
 #include "MetasoundEditorModule.h"
+#include "MetasoundFrontendRegistries.h"
 #include "NodeFactory.h"
+#include "PropertyCustomizationHelpers.h"
 #include "ScopedTransaction.h"
 #include "SCommentBubble.h"
 #include "SGraphNode.h"
 #include "SGraphPin.h"
-#include "Widgets/Layout/SSpacer.h"
 #include "SLevelOfDetailBranchNode.h"
 #include "Styling/SlateStyleRegistry.h"
 #include "Styling/SlateColor.h"
 #include "TutorialMetaData.h"
 #include "Widgets/Images/SImage.h"
+#include "Widgets/Layout/SBox.h"
+#include "Widgets/Layout/SSpacer.h"
 #include "Widgets/SBoxPanel.h"
 #include "Widgets/SOverlay.h"
-#include "MetasoundFrontendRegistries.h"
-#include "PropertyCustomizationHelpers.h"
+#include "Widgets/SWidget.h"
 #include "SGraphPinComboBox.h"
 #include "SMetasoundEnumPin.h"
 
 #define LOCTEXT_NAMESPACE "MetasoundGraphNode"
 
+
 void SMetasoundGraphNode::Construct(const FArguments& InArgs, class UEdGraphNode* InNode)
 {
 	GraphNode = InNode;
@@ -245,6 +248,13 @@ void SMetasoundGraphNode::SetDefaultTitleAreaWidget(TSharedRef<SOverlay> Default
 	}
 }
 
+void SMetasoundGraphNode::MoveTo(const FVector2D& NewPosition, FNodeSet& NodeFilter)
+{
+	SGraphNode::MoveTo(NewPosition, NodeFilter);
+
+	GetMetasoundNode().UpdatePosition();
+}
+
 const FSlateBrush* SMetasoundGraphNode::GetNodeBodyBrush() const
 {
 // 	if (const ISlateStyle* MetasoundStyle = FSlateStyleRegistry::FindSlateStyle("MetasoundStyle"))
@@ -321,15 +331,15 @@ TSharedRef<SWidget> SMetasoundGraphNode::CreateNodeContentArea()
 		{
 			const FSlateBrush* ImageBrush = MetasoundStyle->GetBrush(StyleDisplay.ImageName);
 			ContentBox->AddSlot()
-				.AutoWidth()
-				.HAlign(HAlign_Center)
-				.VAlign(VAlign_Center)
-				[
-					SNew(SImage)
-					.Image(ImageBrush)
-					.ColorAndOpacity(FSlateColor::UseForeground())
-					.DesiredSizeOverride(FVector2D(20, 20))
-				];
+			.AutoWidth()
+			.HAlign(HAlign_Center)
+			.VAlign(VAlign_Center)
+			[
+				SNew(SImage)
+				.Image(ImageBrush)
+				.ColorAndOpacity(FSlateColor::UseForeground())
+				.DesiredSizeOverride(FVector2D(20, 20))
+			];
 		}
 	}
 
diff --git a/Engine/Plugins/Runtime/Metasound/Source/MetasoundEditor/Private/SMetasoundGraphNode.h b/Engine/Plugins/Runtime/Metasound/Source/MetasoundEditor/Private/SMetasoundGraphNode.h
--- a/Engine/Plugins/Runtime/Metasound/Source/MetasoundEditor/Private/SMetasoundGraphNode.h
+++ b/Engine/Plugins/Runtime/Metasound/Source/MetasoundEditor/Private/SMetasoundGraphNode.h
@@ -36,6 +36,7 @@ protected:
 	virtual const FSlateBrush* GetNodeBodyBrush() const override;
 	virtual EVisibility IsAddPinButtonVisible() const override;
 	virtual FReply OnAddPin() override;
+	virtual void MoveTo(const FVector2D& NewPosition, FNodeSet& NodeFilter) override;
 	virtual void SetDefaultTitleAreaWidget(TSharedRef<SOverlay> DefaultTitleAreaWidget) override;
 
 	UMetasoundEditorGraphNode& GetMetasoundNode();
diff --git a/Engine/Plugins/Runtime/Metasound/Source/MetasoundEditor/Public/MetasoundEditorGraph.h b/Engine/Plugins/Runtime/Metasound/Source/MetasoundEditor/Public/MetasoundEditorGraph.h
--- a/Engine/Plugins/Runtime/Metasound/Source/MetasoundEditor/Public/MetasoundEditorGraph.h
+++ b/Engine/Plugins/Runtime/Metasound/Source/MetasoundEditor/Public/MetasoundEditorGraph.h
@@ -74,6 +74,10 @@ public:
 	virtual void SetFromLiteral(const FMetasoundFrontendLiteral& InLiteral)
 	{
 	}
+
+#if WITH_EDITORONLY_DATA
+	virtual void PostEditUndo() override;
+#endif // WITH_EDITORONLY_DATA
 };
 
 UCLASS()
@@ -89,19 +93,14 @@ public:
 	UPROPERTY(VisibleAnywhere, Category = DefaultValue)
 	UMetasoundEditorGraphInputLiteral* Literal;
 
-	void UpdateDocumentInput() const;
+	void UpdateDocumentInput(bool bPostTransaction = true);
 	void UpdatePreviewInstance(const Metasound::FVertexKey& InParameterName, IAudioInstanceTransmitter& InInstanceTransmitter) const;
 
 	void OnDataTypeChanged() override;
-	void OnLiteralChanged();
+	void OnLiteralChanged(bool bPostTransaction = true);
 
 #if WITH_EDITORONLY_DATA
-	virtual void PostEditUndo() override
-	{
-		Super::PostEditUndo();
-
-		UpdateDocumentInput();
-	}
+	virtual void PostEditUndo() override;
 #endif // WITH_EDITORONLY_DATA
 };
 
@@ -128,6 +127,9 @@ public:
 	void IterateInputs(TUniqueFunction<void(UMetasoundEditorGraphInput&)> InFunction) const;
 	void IterateOutputs(TUniqueFunction<void(UMetasoundEditorGraphOutput&)> InFunction) const;
 
+	bool ContainsInput(UMetasoundEditorGraphInput* InInput) const;
+	bool ContainsOutput(UMetasoundEditorGraphOutput* InOutput) const;
+
 	virtual void Synchronize();
 
 	// Sets the transmitter interface for the current metasound preview.
@@ -162,5 +164,7 @@ public:
 
 	TUniquePtr<IAudioInstanceTransmitter> Transmitter;
 
-	friend class Metasound::Editor::FEditor;
+// 	friend class Metasound::Editor::FEditor;
+	friend class UMetasoundFactory;
+	friend class UMetasoundSourceFactory;
 };
diff --git a/Engine/Plugins/Runtime/Metasound/Source/MetasoundEditor/Public/MetasoundEditorGraphInputNodes.h b/Engine/Plugins/Runtime/Metasound/Source/MetasoundEditor/Public/MetasoundEditorGraphInputNodes.h
--- a/Engine/Plugins/Runtime/Metasound/Source/MetasoundEditor/Public/MetasoundEditorGraphInputNodes.h
+++ b/Engine/Plugins/Runtime/Metasound/Source/MetasoundEditor/Public/MetasoundEditorGraphInputNodes.h
@@ -76,7 +76,7 @@ public:
 
 	virtual void UpdatePreviewInstance(const Metasound::FVertexKey& InParameterName, IAudioInstanceTransmitter& InInstanceTransmitter) const
 	{
-		if (ensure(Input))
+		if (Input)
 		{
 			Input->UpdatePreviewInstance(InParameterName, InInstanceTransmitter);
 		}
@@ -84,7 +84,7 @@ public:
 	
 	virtual FGuid GetNodeID() const override
 	{
-		if (ensure(Input))
+		if (Input)
 		{
 			return Input->NodeID;
 		}
@@ -97,9 +97,9 @@ public:
 	{
 		Super::PostEditUndo();
 
-		if (ensure(Input))
+		if (Input)
 		{
-			Input->OnLiteralChanged();
+			Input->OnLiteralChanged(false /* bPostTransaction */);
 		}
 	}
 #endif // WITH_EDITORONLY_DATA
diff --git a/Engine/Plugins/Runtime/Metasound/Source/MetasoundEditor/Public/MetasoundEditorGraphNode.h b/Engine/Plugins/Runtime/Metasound/Source/MetasoundEditor/Public/MetasoundEditorGraphNode.h
--- a/Engine/Plugins/Runtime/Metasound/Source/MetasoundEditor/Public/MetasoundEditorGraphNode.h
+++ b/Engine/Plugins/Runtime/Metasound/Source/MetasoundEditor/Public/MetasoundEditorGraphNode.h
@@ -79,6 +79,8 @@ public:
 	UObject& GetMetasoundChecked();
 	const UObject& GetMetasoundChecked() const;
 
+	void UpdatePosition();
+
 	Metasound::Frontend::FGraphHandle GetRootGraphHandle() const;
 	Metasound::Frontend::FConstGraphHandle GetConstRootGraphHandle() const;
 	Metasound::Frontend::FNodeHandle GetNodeHandle() const;
diff --git a/Engine/Source/Editor/AudioEditor/Private/Factories/ReimportSoundFactory.cpp b/Engine/Source/Editor/AudioEditor/Private/Factories/ReimportSoundFactory.cpp
--- a/Engine/Source/Editor/AudioEditor/Private/Factories/ReimportSoundFactory.cpp
+++ b/Engine/Source/Editor/AudioEditor/Private/Factories/ReimportSoundFactory.cpp
@@ -133,7 +133,7 @@ EReimportResult::Type UReimportSoundFactory::Reimport(UObject* Obj)
 	SoundWave->FreeResources();
 	SoundWave->UpdatePlatformData();
 	SoundWave->MarkPackageDirty();
-	SoundWave->bNeedsThumbnailGeneration = true;
+	SoundWave->SetRedrawThumbnail(true);
 
 	return EReimportResult::Succeeded;
 }
diff --git a/Engine/Source/Editor/AudioEditor/Private/Factories/SoundFactory.cpp b/Engine/Source/Editor/AudioEditor/Private/Factories/SoundFactory.cpp
--- a/Engine/Source/Editor/AudioEditor/Private/Factories/SoundFactory.cpp
+++ b/Engine/Source/Editor/AudioEditor/Private/Factories/SoundFactory.cpp
@@ -559,7 +559,7 @@ UObject* USoundFactory::CreateObject
 			AudioComponent->Play();
 		}
 
-		Sound->bNeedsThumbnailGeneration = true;
+		Sound->SetRedrawThumbnail(true);
 
 		return Sound;
 	}
diff --git a/Engine/Source/Editor/UnrealEd/Private/ThumbnailRendering/SoundWaveThumbnailRenderer.cpp b/Engine/Source/Editor/UnrealEd/Private/ThumbnailRendering/SoundWaveThumbnailRenderer.cpp
--- a/Engine/Source/Editor/UnrealEd/Private/ThumbnailRendering/SoundWaveThumbnailRenderer.cpp
+++ b/Engine/Source/Editor/UnrealEd/Private/ThumbnailRendering/SoundWaveThumbnailRenderer.cpp
@@ -13,166 +13,172 @@ USoundWaveThumbnailRenderer::USoundWaveThumbnailRenderer(const FObjectInitialize
 
 bool USoundWaveThumbnailRenderer::CanVisualizeAsset(UObject* Object)
 {
-	USoundSourceBus* SourceBus = Cast<USoundSourceBus>(Object);
-	if (SourceBus)
+	USoundWave* SoundWave = Cast<USoundWave>(Object);
+	if (!SoundWave)
 	{
 		return false;
 	}
-	return true;
+
+	return SoundWave->CanVisualizeAsset();
 }
 
-void USoundWaveThumbnailRenderer::Draw(UObject* Object, int32 X, int32 Y, uint32 Width, uint32 Height, FRenderTarget*, FCanvas* Canvas, bool bAdditionalViewFamily)
+void USoundWaveThumbnailRenderer::Draw(UObject* Object, int32 X, int32 Y, uint32 Width, uint32 Height, FRenderTarget* RenderTarget, FCanvas* Canvas, bool bAdditionalViewFamily)
 {
+	bool bCanDrawWaveform = true;
+
 	USoundWave* SoundWave = Cast<USoundWave>(Object);
-	if (SoundWave != nullptr && SoundWave->NumChannels > 0)
+	if (!SoundWave)
 	{
-		// check if there is any raw sound data
-		if (SoundWave->RawData.GetBulkDataSize() > 0)
-		{
-			SoundWave->bNeedsThumbnailGeneration = false;
+		return;
+	}
+
+	SoundWave->SetRedrawThumbnail(false);
 
-			// Canvas line item to draw with
-			FCanvasLineItem LineItem;
-			LineItem.SetColor(FLinearColor::White);
+	if (SoundWave->NumChannels == 0 || SoundWave->RawData.GetBulkDataSize() == 0)
+	{
+		return;
+	}
+
+	// Canvas line item to draw with
+	FCanvasLineItem LineItem;
+	LineItem.SetColor(FLinearColor::White);
 
-			uint8* RawWaveData = (uint8*)SoundWave->RawData.Lock(LOCK_READ_ONLY);
-			int32 RawDataSize = SoundWave->RawData.GetBulkDataSize();
+	uint8* RawWaveData = (uint8*)SoundWave->RawData.Lock(LOCK_READ_ONLY);
+	int32 RawDataSize = SoundWave->RawData.GetBulkDataSize();
  
-			// Compute the scaled y-value used to render the channel data
-			const float SampleYScale = Height / (2.f * 32767 * SoundWave->NumChannels);
+	// Compute the scaled y-value used to render the channel data
+	const float SampleYScale = Height / (2.f * 32767 * SoundWave->NumChannels);
+
+	// How many wave files are we going to parse. Note: mono and stereo files are just one file, multi channel files have one mono file per channel.
+	const uint32 NumPackedWaveFiles = SoundWave->NumChannels > 2 ? SoundWave->NumChannels : 1;
+
+	// Loop through each packed wave file and render their data
+	uint32 CurrentRawWaveByteIndex = 0;
+	for (uint32 ChannelWaveFileIndex = 0; ChannelWaveFileIndex < NumPackedWaveFiles; ++ChannelWaveFileIndex)
+	{
+		uint32 RawChannelFileByteSize = 0;
+		if (SoundWave->NumChannels > 2)
+		{
+			// For multi-channel files, we store the packed sound wave data sizes in the sound wave.
+			check((int32)ChannelWaveFileIndex < SoundWave->ChannelSizes.Num());
+			RawChannelFileByteSize = SoundWave->ChannelSizes[ChannelWaveFileIndex];
+		}
+		else
+		{
+			// For mono and stereo files, the sound file size is just the raw data size
+			RawChannelFileByteSize = RawDataSize;
+		}
 
-			// How many wave files are we going to parse. Note: mono and stereo files are just one file, multi channel files have one mono file per channel.
-			const uint32 NumPackedWaveFiles = SoundWave->NumChannels > 2 ? SoundWave->NumChannels : 1;
+		// If we have no data, then nothing to render
+		if (RawChannelFileByteSize == 0)
+		{
+			continue;
+		}
 
-			// Loop through each packed wave file and render their data
-			uint32 CurrentRawWaveByteIndex = 0;
-			for (uint32 ChannelWaveFileIndex = 0; ChannelWaveFileIndex < NumPackedWaveFiles; ++ChannelWaveFileIndex)
+		// Read the packed wave file data
+		FWaveModInfo WaveInfo;				
+		if (WaveInfo.ReadWaveHeader(RawWaveData + CurrentRawWaveByteIndex, RawChannelFileByteSize, 0))
+		{
+			// Make sure we have a mono file here
+			if (SoundWave->NumChannels > 2)
 			{
-				uint32 RawChannelFileByteSize = 0;
-				if (SoundWave->NumChannels > 2)
-				{
-					// For multi-channel files, we store the packed sound wave data sizes in the sound wave.
-					check((int32)ChannelWaveFileIndex < SoundWave->ChannelSizes.Num());
-					RawChannelFileByteSize = SoundWave->ChannelSizes[ChannelWaveFileIndex];
-				}
-				else
+				bool bLooksSane = (WaveInfo.pChannels && *WaveInfo.pChannels == 1);
+				if (!bLooksSane)
 				{
-					// For mono and stereo files, the sound file size is just the raw data size
-					RawChannelFileByteSize = RawDataSize;
-				}
+					UE_LOG(LogAudio, Warning, TEXT("USoundWaveThumbnailRenderer::Draw: Not rendering thumbnail for Wave '%s', as its supposed to be mono and has '%d' channels."), 
+						*GetNameSafe(SoundWave), WaveInfo.pChannels ? *WaveInfo.pChannels : -1);
 
-				// If we have no data, then nothing to render
-				if (RawChannelFileByteSize == 0)
+					continue;
+				}
+			}
+			else
+			{
+				bool bLooksSane = WaveInfo.pChannels && *WaveInfo.pChannels == SoundWave->NumChannels;
+				if (!bLooksSane)
 				{
+					UE_LOG(LogAudio, Warning, TEXT("USoundWaveThumbnailRenderer::Draw: Not rendering thumbnail for Wave '%s', as pChannels('%d') != SoundWave->NumChannels('%d')"),
+						*GetNameSafe(SoundWave), WaveInfo.pChannels ? *WaveInfo.pChannels : -1, SoundWave->NumChannels);
+							
 					continue;
 				}
+			}
+
+			// Sample count
+			const uint32 TotalSampleCount = WaveInfo.SampleDataSize / sizeof(int16);
+			const uint32 TotalFrameCount = SoundWave->NumChannels == 2 ? TotalSampleCount / 2 : TotalSampleCount;
+			const uint32 NumChannelsInWaveFile = *WaveInfo.pChannels;
+			const uint32 FramesPerPixel = TotalFrameCount / Width;
+
+			// Get the sample data of this file
+			const int16* SamplePtr = reinterpret_cast<const int16*>(WaveInfo.SampleDataStart);
 
-				// Read the packed wave file data
-				FWaveModInfo WaveInfo;				
-				if (WaveInfo.ReadWaveHeader(RawWaveData + CurrentRawWaveByteIndex, RawChannelFileByteSize, 0))
+			// Render each channel separately so outer loop is the sound wave channel index.
+			// Note: for multi-channel files this will always be 1-channel (mono).
+			for (uint32 ChannelIndex = 0; ChannelIndex < NumChannelsInWaveFile; ++ChannelIndex)
+			{
+				// Reset the current frame count as we're starting from the beginning of the file to
+				// render the channel data
+				uint32 CurrentFrameCount = 0;
+
+				// Loop through each pixel (in x direction)
+				for (uint32 PixelIndex = 0; PixelIndex < Width; ++PixelIndex)
 				{
-					// Make sure we have a mono file here
-					if (SoundWave->NumChannels > 2)
-					{
-						bool bLooksSane = (WaveInfo.pChannels && *WaveInfo.pChannels == 1);
-						if (!bLooksSane)
-						{
-							UE_LOG(LogAudio, Warning, TEXT("USoundWaveThumbnailRenderer::Draw: Not rendering thumbnail for Wave '%s', as its supposed to be mono and has '%d' channels."), 
-								*GetNameSafe(SoundWave), WaveInfo.pChannels ? *WaveInfo.pChannels : -1);
+					// reset the sample sum and num samples in pixel for each pixel
+					int64 SampleSum = 0;
+					int32 NumSamplesInPixel = 0;
 
-							continue;
-						}
-					}
-					else
+					// Loop through all pixels in this x-frame, sum all audio data. Track total frames rendered to avoid writing past buffer boundary
+					for (uint32 PixelFrameIndex = 0; PixelFrameIndex < FramesPerPixel && CurrentFrameCount < TotalFrameCount; ++PixelFrameIndex)
 					{
-						bool bLooksSane = WaveInfo.pChannels && *WaveInfo.pChannels == SoundWave->NumChannels;
-						if (!bLooksSane)
-						{
-							UE_LOG(LogAudio, Warning, TEXT("USoundWaveThumbnailRenderer::Draw: Not rendering thumbnail for Wave '%s', as pChannels('%d') != SoundWave->NumChannels('%d')"),
-								*GetNameSafe(SoundWave), WaveInfo.pChannels ? *WaveInfo.pChannels : -1, SoundWave->NumChannels);
-							
-							continue;
-						}
-					}
+						// Get the sample value of the wave file
+						const uint32 SampleIndex = CurrentFrameCount * NumChannelsInWaveFile + ChannelIndex;
+						check(SampleIndex < TotalSampleCount);
+						int16 SampleValue = SamplePtr[SampleIndex];
+
+						// Sum the sample value with the running sum
+						SampleSum += FMath::Abs(SampleValue);
 
-					// Sample count
-					const uint32 TotalSampleCount = WaveInfo.SampleDataSize / sizeof(int16);
-					const uint32 TotalFrameCount = SoundWave->NumChannels == 2 ? TotalSampleCount / 2 : TotalSampleCount;
-					const uint32 NumChannelsInWaveFile = *WaveInfo.pChannels;
-					const uint32 FramesPerPixel = TotalFrameCount / Width;
+						// Track the number of samples we're actually summing to get an accurate average
+						++NumSamplesInPixel;
 
-					// Get the sample data of this file
-					const int16* SamplePtr = reinterpret_cast<const int16*>(WaveInfo.SampleDataStart);
+						// Increment the frame after processing channels
+						++CurrentFrameCount;
+					}
 
-					// Render each channel separately so outer loop is the sound wave channel index.
-					// Note: for multi-channel files this will always be 1-channel (mono).
-					for (uint32 ChannelIndex = 0; ChannelIndex < NumChannelsInWaveFile; ++ChannelIndex)
+					// If we actually added any audio data in this pixel
+					if (NumSamplesInPixel > 0)
 					{
-						// Reset the current frame count as we're starting from the beginning of the file to
-						// render the channel data
-						uint32 CurrentFrameCount = 0;
+						const float AverageSampleValue = (float)SampleSum / NumSamplesInPixel;
+						const float AverageSampleValueScaled = AverageSampleValue * SampleYScale;
 
-						// Loop through each pixel (in x direction)
-						for (uint32 PixelIndex = 0; PixelIndex < Width; ++PixelIndex)
+						// Don't try to draw anything if the audio data was too quiet
+						if (AverageSampleValueScaled > 0.001f)
 						{
-							// reset the sample sum and num samples in pixel for each pixel
-							int64 SampleSum = 0;
-							int32 NumSamplesInPixel = 0;
-
-							// Loop through all pixels in this x-frame, sum all audio data. Track total frames rendered to avoid writing past buffer boundary
-							for (uint32 PixelFrameIndex = 0; PixelFrameIndex < FramesPerPixel && CurrentFrameCount < TotalFrameCount; ++PixelFrameIndex)
-							{
-								// Get the sample value of the wave file
-								const uint32 SampleIndex = CurrentFrameCount * NumChannelsInWaveFile + ChannelIndex;
-								check(SampleIndex < TotalSampleCount);
-								int16 SampleValue = SamplePtr[SampleIndex];
-
-								// Sum the sample value with the running sum
-								SampleSum += FMath::Abs(SampleValue);
-
-								// Track the number of samples we're actually summing to get an accurate average
-								++NumSamplesInPixel;
-
-								// Increment the frame after processing channels
-								++CurrentFrameCount;
-							}
-
-							// If we actually added any audio data in this pixel
-							if (NumSamplesInPixel > 0)
-							{
-								const float AverageSampleValue = (float)SampleSum / NumSamplesInPixel;
-								const float AverageSampleValueScaled = AverageSampleValue * SampleYScale;
-
-								// Don't try to draw anything if the audio data was too quiet
-								if (AverageSampleValueScaled > 0.001f)
-								{
-									// What channel we're rendering is going to be one of the interleaved channels (ChannelIndex) in the case
-									// of a stereo file, or the channel wave file index for multi-channel (or mono) files.
-									uint32 Channel = SoundWave->NumChannels == 2 ? ChannelIndex : ChannelWaveFileIndex;
-
-									// Draw vertical line mirrored around x-axis for channel equal to average sample value height
-									const float YCenter = Y + ((2 * Channel) + 1) * Height / (2.f * SoundWave->NumChannels);
-									LineItem.Draw(Canvas, FVector2D(X + PixelIndex, YCenter - AverageSampleValueScaled), FVector2D(X + PixelIndex, YCenter + AverageSampleValueScaled));
-								}
-							}
+							// What channel we're rendering is going to be one of the interleaved channels (ChannelIndex) in the case
+							// of a stereo file, or the channel wave file index for multi-channel (or mono) files.
+							uint32 Channel = SoundWave->NumChannels == 2 ? ChannelIndex : ChannelWaveFileIndex;
+
+							// Draw vertical line mirrored around x-axis for channel equal to average sample value height
+							const float YCenter = Y + ((2 * Channel) + 1) * Height / (2.f * SoundWave->NumChannels);
+							LineItem.Draw(Canvas, FVector2D(X + PixelIndex, YCenter - AverageSampleValueScaled), FVector2D(X + PixelIndex, YCenter + AverageSampleValueScaled));
 						}
 					}
 				}
-				// Offset the raw wave data byte index by the current mono file byte size
-				CurrentRawWaveByteIndex += RawChannelFileByteSize;
 			}
-
-			SoundWave->RawData.Unlock();
 		}
+		// Offset the raw wave data byte index by the current mono file byte size
+		CurrentRawWaveByteIndex += RawChannelFileByteSize;
 	}
+
+	SoundWave->RawData.Unlock();
 }
 
 bool USoundWaveThumbnailRenderer::AllowsRealtimeThumbnails(UObject* Object) const
 {
-	USoundWave* SoundWave = Cast<USoundWave>(Object);
-	if (SoundWave)
+	if (USoundWave* SoundWave = Cast<USoundWave>(Object))
 	{
-		return SoundWave->bNeedsThumbnailGeneration;
+		return SoundWave->GetRedrawThumbnail();
 	}
+
 	return false;
 }
diff --git a/Engine/Source/Runtime/Engine/Classes/Sound/SoundSourceBus.h b/Engine/Source/Runtime/Engine/Classes/Sound/SoundSourceBus.h
--- a/Engine/Source/Runtime/Engine/Classes/Sound/SoundSourceBus.h
+++ b/Engine/Source/Runtime/Engine/Classes/Sound/SoundSourceBus.h
@@ -48,6 +48,11 @@ public:
 	//~ Begin UObject Interface.
 	virtual void PostLoad() override;
 #if WITH_EDITOR
+	virtual bool CanVisualizeAsset() const override
+	{
+		return false;
+	}
+
 	virtual void PostEditChangeProperty(FPropertyChangedEvent& PropertyChangedEvent) override;
 #endif // WITH_EDITOR
 	//~ End UObject Interface.
diff --git a/Engine/Source/Runtime/Engine/Classes/Sound/SoundWave.h b/Engine/Source/Runtime/Engine/Classes/Sound/SoundWave.h
--- a/Engine/Source/Runtime/Engine/Classes/Sound/SoundWave.h
+++ b/Engine/Source/Runtime/Engine/Classes/Sound/SoundWave.h
@@ -502,6 +502,9 @@ private:
 	uint8 bSampleRateManuallyReset : 1;
 
 #if WITH_EDITOR
+	// Whether or not the thumbnail supports generation
+	uint8 bNeedsThumbnailGeneration : 1;
+
 	// Whether this was previously cooked with stream caching enabled.
 	uint8 bWasStreamCachingEnabledOnLastCook : 1;
 	// Whether this asset is loaded from cooked data.
@@ -1098,7 +1101,10 @@ public:
 
 	virtual void WillNeverCacheCookedPlatformDataAgain() override;
 
-	uint32 bNeedsThumbnailGeneration : 1;
+	virtual bool GetRedrawThumbnail() const;
+	virtual void SetRedrawThumbnail(bool bInRedraw);
+	virtual bool CanVisualizeAsset() const;
+
 #endif // WITH_EDITOR
 
 	/**
diff --git a/Engine/Source/Runtime/Engine/Private/SoundWave.cpp b/Engine/Source/Runtime/Engine/Private/SoundWave.cpp
--- a/Engine/Source/Runtime/Engine/Private/SoundWave.cpp
+++ b/Engine/Source/Runtime/Engine/Private/SoundWave.cpp
@@ -1164,9 +1164,12 @@ void USoundWave::PostLoad()
 		AssetImportData->SourceData = MoveTemp(Info);
 	}
 
-	bNeedsThumbnailGeneration = true;
 #endif // #if WITH_EDITORONLY_DATA
 
+#if WITH_EDITOR
+	bNeedsThumbnailGeneration = true;
+#endif // WITH_EDITOR
+
 	INC_FLOAT_STAT_BY(STAT_AudioBufferTime, Duration);
 	INC_FLOAT_STAT_BY(STAT_AudioBufferTimeChannels, NumChannels * Duration);
 
@@ -3060,6 +3063,23 @@ void USoundWave::CacheInheritedLoadingBehavior() const
 	}
 }
 
+#if WITH_EDITOR
+bool USoundWave::GetRedrawThumbnail() const
+{
+	return bNeedsThumbnailGeneration != 0;
+}
+
+void USoundWave::SetRedrawThumbnail(bool bInRedraw)
+{
+	bNeedsThumbnailGeneration = static_cast<uint32>(bInRedraw);
+}
+
+bool USoundWave::CanVisualizeAsset() const
+{
+	return true;
+}
+#endif // WITH_EDITOR
+
 ESoundWaveLoadingBehavior USoundWave::GetLoadingBehavior(bool bCheckSoundClasses /*= true*/) const
 {
 	checkf(!bCheckSoundClasses || CachedSoundWaveLoadingBehavior != ESoundWaveLoadingBehavior::Uninitialized,
@@ -3200,7 +3220,6 @@ bool FSoundWaveProxy::GetChunkData(int32 ChunkIndex, uint8** OutChunkData, bool
 	return true;
 }
 
-
 bool FSoundWaveProxy::IsZerothChunkDataLoaded() const
 {
 	return ZerothChunkData->GetView().Num() > 0;
