Adding bConsiderForActorPlacementWhenHidden flag to PrimitiveComponent, which when enabled, will trace against a component that is not rendered when dragging and placing actors in the editor, such as in the case of proxy collision meshes

diff --git a/Engine/Source/Editor/UnrealEd/Private/Editor/ActorPositioning.cpp b/Engine/Source/Editor/UnrealEd/Private/Editor/ActorPositioning.cpp
--- a/Engine/Source/Editor/UnrealEd/Private/Editor/ActorPositioning.cpp
+++ b/Engine/Source/Editor/UnrealEd/Private/Editor/ActorPositioning.cpp
@@ -89,6 +89,7 @@ bool IsHitIgnored(const FHitResult& InHit, const FSceneView& InSceneView)
 	// Only use this component if it is visible in the specified scene views
 	bool bIsRenderedOnScreen = false;
 	bool bIgnoreTranslucentPrimitive = false;
+	const bool bConsiderInvisibleComponentForPlacement = PrimitiveComponent->bConsiderForActorPlacementWhenHidden;
 	{				
 		if (PrimitiveComponent && PrimitiveComponent->SceneProxy)
 		{
@@ -99,7 +100,7 @@ bool IsHitIgnored(const FHitResult& InHit, const FSceneView& InSceneView)
 		}
 	}
 
-	return !bIsRenderedOnScreen || bIgnoreTranslucentPrimitive;
+	return (!bIsRenderedOnScreen && !bConsiderInvisibleComponentForPlacement) || bIgnoreTranslucentPrimitive;
 }
 
 FActorPositionTraceResult FActorPositioning::TraceWorldForPosition(const UWorld& InWorld, const FSceneView& InSceneView, const FVector& RayStart, const FVector& RayEnd, const TArray<AActor*>* IgnoreActors)
diff --git a/Engine/Source/Runtime/Engine/Classes/Components/PrimitiveComponent.h b/Engine/Source/Runtime/Engine/Classes/Components/PrimitiveComponent.h
--- a/Engine/Source/Runtime/Engine/Classes/Components/PrimitiveComponent.h
+++ b/Engine/Source/Runtime/Engine/Classes/Components/PrimitiveComponent.h
@@ -382,6 +382,12 @@ public:
 	UPROPERTY()
 	uint8 bSelectable:1;
 
+#if WITH_EDITORONLY_DATA
+	/** If true, this component will be considered for placement when dragging and placing items in the editor even if it is not visible, such as in the case of hidden collision meshes */
+	UPROPERTY(EditAnywhere, AdvancedDisplay, Category = Collision)
+	uint8 bConsiderForActorPlacementWhenHidden:1;
+#endif //WITH_EDITORONLY_DATA
+
 	/** If true, forces mips for textures used by this component to be resident when this component's level is loaded. */
 	UPROPERTY(EditAnywhere, AdvancedDisplay, BlueprintReadOnly, Category=TextureStreaming)
 	uint8 bForceMipStreaming:1;
diff --git a/Engine/Source/Runtime/Engine/Private/Components/PrimitiveComponent.cpp b/Engine/Source/Runtime/Engine/Private/Components/PrimitiveComponent.cpp
--- a/Engine/Source/Runtime/Engine/Private/Components/PrimitiveComponent.cpp
+++ b/Engine/Source/Runtime/Engine/Private/Components/PrimitiveComponent.cpp
@@ -332,6 +332,9 @@ UPrimitiveComponent::UPrimitiveComponent(const FObjectInitializer& ObjectInitial
 	bCastContactShadow = true;
 	IndirectLightingCacheQuality = ILCQ_Point;
 	bSelectable = true;
+#if WITH_EDITORONLY_DATA
+	bConsiderForActorPlacementWhenHidden = false;
+#endif // WITH_EDITORONLY_DATA
 	bFillCollisionUnderneathForNavmesh = false;
 	AlwaysLoadOnClient = true;
 	AlwaysLoadOnServer = true;
