- Adapted wrangle commandlet for WorldPartition.
- Fixed saving of unreferenced content crashing because of bad package naming.

diff --git a/Engine/Source/Editor/UnrealEd/Private/Commandlets/ContentCommandlets.cpp b/Engine/Source/Editor/UnrealEd/Private/Commandlets/ContentCommandlets.cpp
--- a/Engine/Source/Editor/UnrealEd/Private/Commandlets/ContentCommandlets.cpp
+++ b/Engine/Source/Editor/UnrealEd/Private/Commandlets/ContentCommandlets.cpp
@@ -45,6 +45,7 @@
 #include "CommandletSourceControlUtils.h"
 #include "WorldPartition/WorldPartitionSubsystem.h"
 #include "WorldPartition/WorldPartition.h"
+#include "LevelInstance/LevelInstanceActor.h"
 
 #include "PackageHelperFunctions.h"
 #include "PackageTools.h"
@@ -2372,6 +2373,7 @@ int32 UWrangleContentCommandlet::Main( const FString& Params )
 
 	// store all referenced objects
 	TMap<FString, FPackageObjects> AllReferencedPublicObjects;
+	TSet<FString> ExternalPackages;
 
 	if (bShouldRestoreFromPreviousRun)
 	{
@@ -2482,10 +2484,17 @@ int32 UWrangleContentCommandlet::Main( const FString& Params )
 			// start off just loading this package (more may be added in the loop)
 			PackagesToLoad.Add(*PackageIt.Value().GetValue());
 
+			TArray<UObject*> ObjectsAddedToRoot;
+			TSet<UObject*> AlreadyVisited;
 			for (int32 PackageIndex = 0; PackageIndex < PackagesToLoad.Num(); PackageIndex++)
 			{
 				// save a copy of the packagename (not a reference in case the PackgesToLoad array gets realloced)
 				FString PackageName = PackagesToLoad[PackageIndex];
+				const bool bIsExternalPackage = ExternalPackages.Contains(PackageName);
+				const bool bForceSkipSavePackage = bIsExternalPackage;
+				// Avoid GC'ing between each external package
+				const int32 GarbageCollectionFrequency = bIsExternalPackage ? 2048 : 0;
+
 				FPackagePath PackagePath;
 				if (!FPackagePath::TryFromMountedName(PackageName, PackagePath) ||
 					!FPackageName::DoesPackageExist(PackagePath, &PackagePath))
@@ -2495,7 +2504,7 @@ int32 UWrangleContentCommandlet::Main( const FString& Params )
 
 				FString PackageFilename = PackagePath.GetLocalFullPath();
 				SET_WARN_COLOR(COLOR_WHITE);
-				UE_LOG(LogContentCommandlet, Warning, TEXT("Fully loading %s..."), *PackageFilename);
+				UE_LOG(LogContentCommandlet, Warning, TEXT("Fully loading %s... [%d/%d]"), *PackageFilename, PackageIndex + 1, PackagesToLoad.Num());
 				CLEAR_WARN_COLOR();
 
 	// @todo josh: track redirects in this package and then save the package instead of copy it if there were redirects
@@ -2506,6 +2515,22 @@ int32 UWrangleContentCommandlet::Main( const FString& Params )
 
 				FLinkerLoad* Linker = LoadPackageLinker(nullptr, PackagePath, LOAD_Quiet | LOAD_NoWarn | LOAD_NoVerify);
 
+				UWorld* World = UWorld::FindWorldInPackage(Package);
+				if (World && World->GetWorldPartition())
+				{
+					TArray<FString> ExternalActorPackages = World->PersistentLevel->GetOnDiskExternalActorPackages();
+					if (ExternalActorPackages.Num())
+					{
+						PackagesToLoad.Append(ExternalActorPackages);
+						// Exclude external actors packages from save
+						ExternalPackages.Append(ExternalActorPackages);
+					
+						// Keep partition world around to avoid reloading it for every loaded external actor
+						World->AddToRoot();
+						ObjectsAddedToRoot.Add(World);
+					}
+				}
+
 				// look for special package types
 				bool bIsMap = Linker->ContainsMap();
 				bool bIsScriptPackage = Linker->ContainsCode();
@@ -2514,6 +2539,25 @@ int32 UWrangleContentCommandlet::Main( const FString& Params )
 				for(FThreadSafeObjectIterator ObjectIt; ObjectIt; ++ObjectIt)
 				{
 					UObject* Object = *ObjectIt;
+					
+					// Because we don't GC between each load of external packages, we want to skip already visited objects
+					bool bAlreadyVisited = false;
+					AlreadyVisited.Add(Object, &bAlreadyVisited);
+					if (bAlreadyVisited)
+					{
+						continue;
+					}
+
+					// Add any levels referenced by Level Instance actors to the list of levels to load
+					ALevelInstance* LevelInstance = Cast<ALevelInstance>(Object);
+					if (LevelInstance && !LevelInstance->IsTemplate())
+					{
+						FString LevelPackageName = LevelInstance->GetWorldAssetPackage();
+						if (!LevelPackageName.IsEmpty() && PackagesToFullyLoad.FindKey(LevelPackageName) == nullptr)
+						{
+							PackagesToLoad.AddUnique(LevelPackageName);
+						}
+					}
 
 					// record all public referenced objects (skipping over top level packages)
 					if (/*Object->HasAnyFlags(RF_Public) &&*/ Object->GetOuter() != NULL)
@@ -2539,9 +2583,8 @@ int32 UWrangleContentCommandlet::Main( const FString& Params )
 				// add any sublevels of this world to the list of levels to load
 				for (TObjectIterator<UWorld> WorldIt; WorldIt; ++WorldIt)
 				{
-					UWorld*	World = *WorldIt;
 					// iterate over streaming level objects loading the levels.
-					for (ULevelStreaming* StreamingLevel : World->GetStreamingLevels())
+					for (ULevelStreaming* StreamingLevel : (*WorldIt)->GetStreamingLevels())
 					{
 						if (StreamingLevel)
 						{
@@ -2557,7 +2600,7 @@ int32 UWrangleContentCommandlet::Main( const FString& Params )
 
 				// save/copy the package if desired, and only if it's not a script package (script code is
 				// not cutdown, so we always use original script code)
-				if (bShouldSavePackages && !bIsScriptPackage)
+				if (bShouldSavePackages && !bIsScriptPackage && !bForceSkipSavePackage)
 				{
 					// make the name of the location to put the package
 					FString CutdownPackageName = MakeCutdownFilename(PackageFilename);
@@ -2583,7 +2626,21 @@ int32 UWrangleContentCommandlet::Main( const FString& Params )
 					}
 				}
 
-				// close this package
+				static int32 Counter = 0;
+				if (!GarbageCollectionFrequency || (Counter++ % GarbageCollectionFrequency) == 0)
+				{
+					AlreadyVisited.Empty();
+					CollectGarbage(RF_NoFlags);
+				}
+			}
+
+			// Get rid of rooted objects if any
+			if (ObjectsAddedToRoot.Num())
+			{
+				for (UObject* Object : ObjectsAddedToRoot)
+				{
+					Object->RemoveFromRoot();
+				}
 				CollectGarbage(RF_NoFlags);
 			}
 		}
@@ -2757,6 +2814,15 @@ int32 UWrangleContentCommandlet::Main( const FString& Params )
 		int32 PackageIndex = 0;
 		for (TMap<FString, FPackageObjects>::TIterator It(AllReferencedPublicObjects); It; ++It, PackageIndex++ )
 		{
+			const FString& PackageName = It.Key();
+
+			// Skip save of external packages
+			const bool bForceSkipSavePackage = ExternalPackages.Contains(PackageName);
+			if (bForceSkipSavePackage)
+			{
+				continue;
+			}
+
 			// if the package was a fully loaded package, than we already saved it
 			if (It.Value().bIsFullyLoadedPackage)
 			{
@@ -2815,9 +2881,9 @@ int32 UWrangleContentCommandlet::Main( const FString& Params )
 				//UE_LOG(LogContentCommandlet, Warning, TEXT( "*It.Key(): %s" ), *It.Key() );
 
 				// we need to be able to find the original package
-				if( FPackageName::DoesPackageExist(It.Key(), NULL, &OriginalPackageFilename) == false )
+				if( FPackageName::DoesPackageExist(PackageName, NULL, &OriginalPackageFilename) == false )
 				{
-					UE_LOG(LogContentCommandlet, Fatal, TEXT( "Could not find file in file cache: %s"), *It.Key() );
+					UE_LOG(LogContentCommandlet, Fatal, TEXT( "Could not find file in file cache: %s"), *PackageName);
 				}
 
 				// any maps need to be fully referenced
@@ -2880,13 +2946,15 @@ int32 UWrangleContentCommandlet::Main( const FString& Params )
 		int32 PackageIndex = 0;
 		for (TMap<FString, FPackageObjects>::TIterator PackageIt(UnnecessaryObjectsByPackage); PackageIt; ++PackageIt, PackageIndex++)
 		{
-			//UE_LOG(LogContentCommandlet, Warning, TEXT("Processing %s"), *PackageIt.Key());
+			const FString& PackageWithUnnecessaryObjects = PackageIt.Key();
+			UE_LOG(LogContentCommandlet, Warning, TEXT("Processing %s"), *PackageWithUnnecessaryObjects);
+
 			UPackage* FullyLoadedPackage = NULL;
 			// fully load unnecessary packages with no objects, 
 			if (PackageIt.Value().ReferencedObjects.Num() == 0)
 			{
 				// just load it, and don't need a reference to it
-				FullyLoadedPackage = LoadPackage(NULL, FPackagePath::FromLocalPath(PackageIt.Key()), LOAD_None);
+				FullyLoadedPackage = LoadPackage(NULL, FPackagePath::FromLocalPath(PackageWithUnnecessaryObjects), LOAD_None);
 			}
 			else
 			{
@@ -2921,40 +2989,53 @@ int32 UWrangleContentCommandlet::Main( const FString& Params )
 				// if was unnecessary...
 				if (UnnecessaryObjects.Find(*It->GetFullName()))
 				{
+					UObject* UnnecessaryObject = *It;
 					// ... then rename it (its outer needs to be a package, everything else will have to be
 					// moved by its outer getting moved)
-					if (!It->IsA(UPackage::StaticClass()) &&
-						It->GetOuter() &&
-						It->GetOuter()->IsA(UPackage::StaticClass()) &&
-						It->GetOutermost()->GetName().Left(4) != TEXT("NFS_"))
+					if (!UnnecessaryObject->IsA(UPackage::StaticClass()) &&
+						UnnecessaryObject->GetOuter() &&
+						UnnecessaryObject->GetOuter()->IsA(UPackage::StaticClass()))
 					{
-						UPackage* NewPackage = CreatePackage( *(FString(TEXT("NFS_")) + It->GetOuter()->GetPathName()));
-						//UE_LOG(LogContentCommandlet, Warning, TEXT("Renaming object from %s to %s.%s"), *It->GetPathName(), *NewPackage->GetPathName(), *It->GetName());
-
-						// move the object if we can. IF the rename fails, then the object was already renamed to this spot, but not GC'd.
-						// that's okay.
-						if (It->Rename(*It->GetName(), NewPackage, REN_Test))
+						FString PackageFullName = UnnecessaryObject->GetOuter()->GetPathName();
+						FString Path = FPackageName::GetLongPackagePath(PackageFullName);
+						FString AssetName = FPackageName::GetLongPackageAssetName(PackageFullName);
+						if (AssetName.Left(4) != TEXT("NFS_"))
 						{
-							It->Rename(*It->GetName(), NewPackage, REN_None);
+							FString NewPackageName = FString::Printf(TEXT("%s/NFS_%s"), *Path, *AssetName);
+							UPackage* NewPackage = CreatePackage(*NewPackageName);
+							//UE_LOG(LogContentCommandlet, Warning, TEXT("Renaming object from %s to %s.%s"), *UnnecessaryObject->GetPathName(), *NewPackage->GetPathName(), *UnnecessaryObject->GetName());
+
+							// move the object if we can. IF the rename fails, then the object was already renamed to this spot, but not GC'd.
+							// that's okay.
+							if (UnnecessaryObject->Rename(*UnnecessaryObject->GetName(), NewPackage, REN_Test))
+							{
+								UnnecessaryObject->Rename(*UnnecessaryObject->GetName(), NewPackage, REN_None);
+							}
 						}
 					}
-
 				}
 			}
 
 			// find the one we moved this packages objects to
-			FPackagePath PackagePath = FPackagePath::FromLocalPath(PackageIt.Key());
+			FPackagePath PackagePath = FPackagePath::FromLocalPath(PackageWithUnnecessaryObjects);
 			FString PackageFilePath = PackagePath.GetLocalFullPath();
 			FString PackageName = PackagePath.GetPackageName();
-			UPackage* MovedPackage = FindPackage(NULL, *FString::Printf(TEXT("%s/NFS_%s"), *FPackageName::GetLongPackagePath(PackageName), *FPackageName::GetLongPackageAssetName(PackageName)));
-			check(MovedPackage);
-
+			FString FindPackageName = FString::Printf(TEXT("%s/NFS_%s"), *FPackageName::GetLongPackagePath(PackageName), *FPackageName::GetLongPackageAssetName(PackageName));
 			// convert the new name to a a NFS directory directory
 			FString MovedFilename = MakeCutdownFilename(FString::Printf(TEXT("%s/NFS_%s"), *FPaths::GetPath(PackageFilePath), *FPaths::GetCleanFilename(PackageFilePath)), TEXT("NFSContent"));
-			UE_LOG(LogContentCommandlet, Warning, TEXT("Saving package %s [%d/%d]"), *MovedFilename, PackageIndex, NumPackages);
-			// finally save it out
-			SavePackageHelper(MovedPackage, *MovedFilename);
 
+			// finally save it out
+			UPackage* MovedPackage = FindPackage(NULL, *FindPackageName);
+			if (ensure(MovedPackage))
+			{
+				UE_LOG(LogContentCommandlet, Warning, TEXT("Saving package %s [%d/%d]"), *MovedFilename, PackageIndex, NumPackages);
+				SavePackageHelper(MovedPackage, *MovedFilename);
+			}
+			else
+			{
+				UE_LOG(LogContentCommandlet, Warning, TEXT("Moved package %s not found."), *FindPackageName);
+				UE_LOG(LogContentCommandlet, Error, TEXT("Can't save package %s [%d/%d]"), *MovedFilename, PackageIndex, NumPackages);
+			}
 			CollectGarbage(RF_NoFlags);
 		}
 	}
@@ -3203,4 +3284,4 @@ int32 UListStaticMeshesImportedFromSpeedTreesCommandlet::Main(const FString& Par
 		UE_LOG(LogContentCommandlet, Display, TEXT("No static meshes were imported from speedtrees in this project."));
 	}
 	return 0;
-}
+}
\ No newline at end of file
