Fix ExternalActorsCommandlet to properly GC actor packages before opening the invalid ones for delete.

diff --git a/Engine/Source/Editor/UnrealEd/Private/Commandlets/ExternalActorsCommandlet.cpp b/Engine/Source/Editor/UnrealEd/Private/Commandlets/ExternalActorsCommandlet.cpp
--- a/Engine/Source/Editor/UnrealEd/Private/Commandlets/ExternalActorsCommandlet.cpp
+++ b/Engine/Source/Editor/UnrealEd/Private/Commandlets/ExternalActorsCommandlet.cpp
@@ -135,7 +135,7 @@ int32 UExternalActorsCommandlet::Main(const FString& Params)
 		});
 	}
 
-	CollectGarbage(GARBAGE_COLLECTION_KEEPFLAGS);
+	CollectGarbage(RF_NoFlags);
 
 	for (const FString& PackageToDelete : PackagesToDelete)
 	{
