DerivedDataCacheCommandlet: Bump DDCCommandletMaxWaitSeconds to 10 minutes and make it configurable.
[REDACTED_HANDLE] Arciel.Rekman, Devin.Doucette
#rnx

diff --git a/Engine/Source/Editor/UnrealEd/Private/Commandlets/DerivedDataCacheCommandlet.cpp b/Engine/Source/Editor/UnrealEd/Private/Commandlets/DerivedDataCacheCommandlet.cpp
--- a/Engine/Source/Editor/UnrealEd/Private/Commandlets/DerivedDataCacheCommandlet.cpp
+++ b/Engine/Source/Editor/UnrealEd/Private/Commandlets/DerivedDataCacheCommandlet.cpp
@@ -113,6 +113,9 @@ int32 UDerivedDataCacheCommandlet::Main( const FString& Params )
 	bool bDoSubset = SubsetMod > 0 && SubsetTarget < SubsetMod;
 
 	// Timing variables
+	double DDCCommandletMaxWaitSeconds = 60. * 10.;
+	GConfig->GetDouble(TEXT("CookSettings"), TEXT("DDCCommandletMaxWaitSeconds"), DDCCommandletMaxWaitSeconds, GEditorIni);
+
 	double FindProcessedPackagesTime = 0.0;
 	double GCTime = 0.0;
 	double FinishCacheTime = 0.;
@@ -367,7 +370,6 @@ int32 UDerivedDataCacheCommandlet::Main( const FString& Params )
 				ObjectBuffer.Append(CachingObjects);
 
 				double LastActivityTime = FinishCacheTimeStart;
-				const double MaxSecondsWithNoActivity = 120.0;
 				const double WaitingForCacheSleepTime = 0.050;
 				while (ObjectBuffer.Num() > 0)
 				{
@@ -398,10 +400,10 @@ int32 UDerivedDataCacheCommandlet::Main( const FString& Params )
 					}
 					if (!bHadActivity)
 					{
-						if (CurrentTime - LastActivityTime >= MaxSecondsWithNoActivity)
+						if (CurrentTime - LastActivityTime >= DDCCommandletMaxWaitSeconds)
 						{
 							UE_LOG(LogDerivedDataCacheCommandlet, Error, TEXT("Timed out for %.2lfs waiting for %d objects to finish caching. First object: %s."),
-								MaxSecondsWithNoActivity, ObjectBuffer.Num(), *ObjectBuffer[0]->GetFullName());
+								DDCCommandletMaxWaitSeconds, ObjectBuffer.Num(), *ObjectBuffer[0]->GetFullName());
 							ObjectBuffer.Reset();
 						}
 						else
