Fix DerivedDataCache commandlet getting stuck.

- DDC commandlet used number of outstanding jobs as a criterium to pump the compiler. This is not correct as there may be no outstanding jobs but the shadermaps aren't yet finalized.
- Another issue was FShaderCompilingManager::IsCompiling() also not being watertight and able to return false when there is known outstanding work. Thi can possibly fix more problems than the DDC commandlet (cooking as well).

diff --git a/Engine/Source/Editor/UnrealEd/Private/Commandlets/DerivedDataCacheCommandlet.cpp b/Engine/Source/Editor/UnrealEd/Private/Commandlets/DerivedDataCacheCommandlet.cpp
--- a/Engine/Source/Editor/UnrealEd/Private/Commandlets/DerivedDataCacheCommandlet.cpp
+++ b/Engine/Source/Editor/UnrealEd/Private/Commandlets/DerivedDataCacheCommandlet.cpp
@@ -45,37 +45,40 @@ void UDerivedDataCacheCommandlet::MaybeMarkPackageAsAlreadyLoaded(UPackage *Pack
 
 static void WaitForCurrentShaderCompilationToFinish(bool& bInOutHadActivity)
 {
-	int32 CachedShaderCount = GShaderCompilingManager->GetNumRemainingJobs();
-	if (CachedShaderCount == 0)
+	if (GShaderCompilingManager->IsCompiling())
 	{
-		return;
-	}
-
-	int32 NumCompletedShadersSinceLastLog = 0;
-	bInOutHadActivity = true;
-	UE_LOG(LogDerivedDataCacheCommandlet, Display, TEXT("Waiting for %d shaders to finish."), CachedShaderCount);
-	while (GShaderCompilingManager->IsCompiling())
-	{
-		const int32 CurrentShaderCount = GShaderCompilingManager->GetNumRemainingJobs();
-		NumCompletedShadersSinceLastLog += (CachedShaderCount - CurrentShaderCount);
-		CachedShaderCount = CurrentShaderCount;
-
-		if (NumCompletedShadersSinceLastLog >= 1000)
+		bInOutHadActivity = true;
+		int32 CachedShaderCount = GShaderCompilingManager->GetNumRemainingJobs();
+		if (CachedShaderCount > 0)
 		{
 			UE_LOG(LogDerivedDataCacheCommandlet, Display, TEXT("Waiting for %d shaders to finish."), CachedShaderCount);
-			NumCompletedShadersSinceLastLog = 0;
 		}
+		int32 NumCompletedShadersSinceLastLog = 0;
+		while (GShaderCompilingManager->IsCompiling())
+		{
+			const int32 CurrentShaderCount = GShaderCompilingManager->GetNumRemainingJobs();
+			NumCompletedShadersSinceLastLog += (CachedShaderCount - CurrentShaderCount);
+			CachedShaderCount = CurrentShaderCount;
+
+			if (NumCompletedShadersSinceLastLog >= 1000)
+			{
+				UE_LOG(LogDerivedDataCacheCommandlet, Display, TEXT("Waiting for %d shaders to finish."), CachedShaderCount);
+				NumCompletedShadersSinceLastLog = 0;
+			}
 
-		// Process any asynchronous shader compile results that are ready, limit execution time
-		GShaderCompilingManager->ProcessAsyncResults(true, false);
-		GDistanceFieldAsyncQueue->ProcessAsyncTasks();
-		GCardRepresentationAsyncQueue->ProcessAsyncTasks();
+			// Process any asynchronous shader compile results that are ready, limit execution time
+			GShaderCompilingManager->ProcessAsyncResults(true, false);
+			GDistanceFieldAsyncQueue->ProcessAsyncTasks();
+			GCardRepresentationAsyncQueue->ProcessAsyncTasks();
+		}
+		GShaderCompilingManager->FinishAllCompilation(); // Final blocking check as IsCompiling() may be non-deterministic
+		UE_LOG(LogDerivedDataCacheCommandlet, Display, TEXT("Done waiting for shaders to finish."));
 	}
-	GShaderCompilingManager->FinishAllCompilation(); // Final blocking check as IsCompiling() may be non-deterministic
+
+	// these shouldn't be predicated on whether the shaders were being compiled
 	GDistanceFieldAsyncQueue->BlockUntilAllBuildsComplete();
 	GCardRepresentationAsyncQueue->BlockUntilAllBuildsComplete();
-	UE_LOG(LogDerivedDataCacheCommandlet, Display, TEXT("Done waiting for shaders to finish."));
-};
+}
 
 static void WaitForCurrentTextureBuildingToFinish(bool& bInOutHadActivity)
 {
diff --git a/Engine/Source/Runtime/Engine/Public/ShaderCompiler.h b/Engine/Source/Runtime/Engine/Public/ShaderCompiler.h
--- a/Engine/Source/Runtime/Engine/Public/ShaderCompiler.h
+++ b/Engine/Source/Runtime/Engine/Public/ShaderCompiler.h
@@ -1074,7 +1074,7 @@ public:
 	 */
 	bool IsCompiling() const
 	{
-		return GetNumOutstandingJobs() > 0 || PendingFinalizeShaderMaps.Num() > 0 || GetNumPendingJobs() > 0 || NumExternalJobs > 0;
+		return GetNumOutstandingJobs() > 0 || HasShaderJobs() || GetNumPendingJobs() > 0 || NumExternalJobs > 0;
 	}
 
 	/**
