WrangleContentCommandlet: report missing classes to log file.

diff --git a/Engine/Source/Editor/UnrealEd/Private/Commandlets/ContentCommandlets.cpp b/Engine/Source/Editor/UnrealEd/Private/Commandlets/ContentCommandlets.cpp
--- a/Engine/Source/Editor/UnrealEd/Private/Commandlets/ContentCommandlets.cpp
+++ b/Engine/Source/Editor/UnrealEd/Private/Commandlets/ContentCommandlets.cpp
@@ -2771,16 +2771,12 @@ int32 UWrangleContentCommandlet::Main( const FString& Params )
 					UClass* Class = StaticLoadClass(UObject::StaticClass(), NULL, *ClassName, NULL, LOAD_None, NULL);
 					// When wrangling content, you often are loading packages that have not been saved in ages and have a reference to a class
 					// that no longer exists.  Instead of asserting, we will just continue
-					if( bShouldSkipMissingClasses == true )
-					{
-						if( Class == NULL )
-						{
-							continue;
-						}
-					}
-					else
+
+					if ( !Class )
 					{
-						check(Class);
+						UE_LOG(LogContentCommandlet, Warning, TEXT("Missing class %s"), *ClassName);
+						check(bShouldSkipMissingClasses);
+						continue;
 					}
 
 					// make sure it doesn't get GC'd
