Skip populating the cooker with files from disk when that population was already done on a previous cook in the same process.
#rnx

diff --git a/Engine/Source/Editor/UnrealEd/Private/CookOnTheFlyServer.cpp b/Engine/Source/Editor/UnrealEd/Private/CookOnTheFlyServer.cpp
--- a/Engine/Source/Editor/UnrealEd/Private/CookOnTheFlyServer.cpp
+++ b/Engine/Source/Editor/UnrealEd/Private/CookOnTheFlyServer.cpp
@@ -6936,6 +6936,7 @@ void UCookOnTheFlyServer::InitializeSandbox(const TArrayView<const ITargetPlatfo
 			bIsInitializingSandbox = false;
 		};
 
+		TSet<const ITargetPlatform*> AlreadyInitializedPlatforms;
 		TArray<const ITargetPlatform*, TInlineAllocator<ExpectedMaxNumPlatforms>> RefreshPlatforms;
 		const bool bIsDiffOnly = FParse::Param(FCommandLine::Get(), TEXT("DIFFONLY"));
 		const bool bIsIterativeCook = IsCookFlagSet(ECookInitializationFlags::Iterative);
@@ -6987,6 +6988,10 @@ void UCookOnTheFlyServer::InitializeSandbox(const TArrayView<const ITargetPlatfo
 			else
 			{
 				RefreshPlatforms.Add(Target);
+				if (PlatformData->bIsSandboxInitialized)
+				{
+					AlreadyInitializedPlatforms.Add(Target);
+				}
 			}
 
 			PlatformData->bIsSandboxInitialized = true;
@@ -7002,6 +7007,10 @@ void UCookOnTheFlyServer::InitializeSandbox(const TArrayView<const ITargetPlatfo
 			// Don't populate however if we are looking for deterministic cooking differences; start from an empty list of cooked packages
 			if (!bIsDiffOnly) 
 			{
+				// Don't populate platforms that were already initialized; we already populated them when we first initialized them
+				RefreshPlatforms.RemoveAllSwap([&AlreadyInitializedPlatforms](const ITargetPlatform* TargetPlatform) {
+					return AlreadyInitializedPlatforms.Contains(TargetPlatform);
+				});
 				PopulateCookedPackagesFromDisk(RefreshPlatforms);
 			}
 		}
