Fix generic thumbnail images for content browser assets not scaling based on user thumbnail size

diff --git a/Engine/Source/Editor/ContentBrowser/Private/AssetViewWidgets.cpp b/Engine/Source/Editor/ContentBrowser/Private/AssetViewWidgets.cpp
--- a/Engine/Source/Editor/ContentBrowser/Private/AssetViewWidgets.cpp
+++ b/Engine/Source/Editor/ContentBrowser/Private/AssetViewWidgets.cpp
@@ -42,6 +42,7 @@
 
 #define LOCTEXT_NAMESPACE "ContentBrowser"
 
+static float GenericThumbnailSizes[(int32)EThumbnailSize::MAX] = { 24, 32, 64, 128, 200 };
 
 ///////////////////////////////
 // FAssetViewModeUtils
@@ -1664,6 +1665,8 @@ void SAssetTileItem::Construct( const FArguments& InArgs )
 		ThumbnailConfig.HighlightedText = InArgs._HighlightText;
 		ThumbnailConfig.HintColorAndOpacity = InArgs._ThumbnailHintColorAndOpacity;
 		ThumbnailConfig.Padding= FMargin(2.0f);
+		ThumbnailConfig.GenericThumbnailSize = MakeAttributeSP(this, &SAssetTileItem::GetGenericThumbnailSize);
+
 		{
 			FContentBrowserItemDataAttributeValue ColorAttributeValue = AssetItem->GetItem().GetItemAttribute(ContentBrowserItemAttributes::ItemColor);
 			if (ColorAttributeValue.IsValid())
@@ -1733,7 +1736,7 @@ void SAssetTileItem::Construct( const FArguments& InArgs )
 							.VAlign(VAlign_Top)
 							[
 								SNew(SBox)
-								.MaxDesiredHeight(AssetNameHeights[(int32)CurrentThumbnailSize])
+								.MaxDesiredHeight(this, &SAssetTileItem::GetNameAreaMaxDesiredHeight)
 								[
 									SAssignNew(InlineRenameWidget, SInlineEditableTextBlock)
 									.Font(this, &SAssetTileItem::GetThumbnailFont)
@@ -1928,6 +1931,16 @@ const FSlateBrush* SAssetTileItem::GetFolderBackgroundShadowImage() const
 	return FStyleDefaults::GetNoBrush();
 }
 
+FOptionalSize SAssetTileItem::GetNameAreaMaxDesiredHeight() const
+{
+	return AssetNameHeights[(int32)CurrentThumbnailSize.Get()];
+}
+
+int32 SAssetTileItem::GetGenericThumbnailSize() const
+{
+	return GenericThumbnailSizes[(int32)CurrentThumbnailSize.Get()];
+}
+
 void SAssetTileItem::InitializeAssetNameHeights()
 {
 	// The height of the asset name field for each thumbnail size
diff --git a/Engine/Source/Editor/ContentBrowser/Private/AssetViewWidgets.h b/Engine/Source/Editor/ContentBrowser/Private/AssetViewWidgets.h
--- a/Engine/Source/Editor/ContentBrowser/Private/AssetViewWidgets.h
+++ b/Engine/Source/Editor/ContentBrowser/Private/AssetViewWidgets.h
@@ -496,7 +496,7 @@ public:
 		SLATE_ARGUMENT( bool, AllowThumbnailHintLabel )
 
 		/** Current size of the thumbnail that was generated */
-		SLATE_ARGUMENT(EThumbnailSize, CurrentThumbnailSize)
+		SLATE_ATTRIBUTE(EThumbnailSize, CurrentThumbnailSize)
 
 		/** The width of the item */
 		SLATE_ATTRIBUTE( float, ItemWidth )
@@ -590,6 +590,10 @@ protected:
 	const FSlateBrush* GetFolderBackgroundImage() const;
 	const FSlateBrush* GetFolderBackgroundShadowImage() const;
 
+	FOptionalSize GetNameAreaMaxDesiredHeight() const;
+
+	int32 GetGenericThumbnailSize() const;
+
 	static void InitializeAssetNameHeights();
 
 private:
@@ -606,7 +610,7 @@ private:
 	float ThumbnailPadding;
 
 	/** Current thumbnail size when this widget was generated */
-	EThumbnailSize CurrentThumbnailSize;
+	TAttribute<EThumbnailSize> CurrentThumbnailSize;
 };
 
 /** An item in the asset column view */
diff --git a/Engine/Source/Editor/ContentBrowser/Private/SAssetView.cpp b/Engine/Source/Editor/ContentBrowser/Private/SAssetView.cpp
--- a/Engine/Source/Editor/ContentBrowser/Private/SAssetView.cpp
+++ b/Engine/Source/Editor/ContentBrowser/Private/SAssetView.cpp
@@ -3234,7 +3234,7 @@ TSharedRef<ITableRow> SAssetView::MakeTileViewWidget(TSharedPtr<FAssetViewItem>
 			.AssetThumbnail(AssetThumbnail)
 			.AssetItem(AssetItem)
 			.ThumbnailPadding(TileViewThumbnailPadding)
-			.CurrentThumbnailSize(ThumbnailSize)
+			.CurrentThumbnailSize(this, &SAssetView::GetThumbnailSize)
 			.ItemWidth(this, &SAssetView::GetTileViewItemWidth)
 			.OnRenameBegin(this, &SAssetView::AssetRenameBegin)
 			.OnRenameCommit(this, &SAssetView::AssetRenameCommit)
diff --git a/Engine/Source/Editor/ContentBrowser/Private/SAssetView.h b/Engine/Source/Editor/ContentBrowser/Private/SAssetView.h
--- a/Engine/Source/Editor/ContentBrowser/Private/SAssetView.h
+++ b/Engine/Source/Editor/ContentBrowser/Private/SAssetView.h
@@ -677,6 +677,9 @@ private:
 	/** Gets the current thumbnail scale */
 	float GetThumbnailScale() const;
 
+	/** Gets the current thumbnail size enum */
+	EThumbnailSize GetThumbnailSize() const { return ThumbnailSize; }
+
 	/** Gets the scaled item height for the list view */
 	float GetListViewItemHeight() const;
 	
diff --git a/Engine/Source/Editor/UnrealEd/Private/AssetThumbnail.cpp b/Engine/Source/Editor/UnrealEd/Private/AssetThumbnail.cpp
--- a/Engine/Source/Editor/UnrealEd/Private/AssetThumbnail.cpp
+++ b/Engine/Source/Editor/UnrealEd/Private/AssetThumbnail.cpp
@@ -56,6 +56,7 @@ public:
 		, _ClassThumbnailBrushOverride(NAME_None)
 		, _AssetTypeColorOverride()
 		, _Padding(0)
+		, _GenericThumbnailSize(64)
 		{}
 
 		SLATE_ARGUMENT(FName, Style)
@@ -72,6 +73,7 @@ public:
 		SLATE_ARGUMENT(FName, ClassThumbnailBrushOverride)
 		SLATE_ARGUMENT(TOptional<FLinearColor>, AssetTypeColorOverride)
 		SLATE_ARGUMENT(FMargin, Padding)
+		SLATE_ATTRIBUTE(int32, GenericThumbnailSize)
 
 	SLATE_END_ARGS()
 
@@ -90,6 +92,7 @@ public:
 		bHasRenderedThumbnail = false;
 		WidthLastFrame = 0;
 		GenericThumbnailBorderPadding = 2.f;
+		GenericThumbnailSize = InArgs._GenericThumbnailSize;
 
 		AssetThumbnail->OnAssetDataChanged().AddSP(this, &SAssetThumbnail::OnAssetDataChanged);
 
@@ -148,6 +151,7 @@ public:
 				+SOverlay::Slot()
 				[
 					SAssignNew(GenericThumbnailImage, SImage)
+					.DesiredSizeOverride(this, &SAssetThumbnail::GetGenericThumbnailDesiredSize)
 					.Image(this, &SAssetThumbnail::GetClassThumbnailBrush)
 				]
 			]
@@ -673,6 +677,13 @@ private:
 		}
 	}
 
+	TOptional<FVector2D> GetGenericThumbnailDesiredSize() const
+	{
+		const int32 Size = GenericThumbnailSize.Get();
+
+		return FVector2D(Size, Size);
+	}
+
 private:
 	TSharedPtr<STextBlock> GenericLabelTextBlock;
 	TSharedPtr<STextBlock> HintTextBlock;
@@ -696,6 +707,8 @@ private:
 	EThumbnailLabel::Type Label;
 
 	TAttribute< FLinearColor > HintColorAndOpacity;
+	TAttribute<int32> GenericThumbnailSize;
+
 	bool bAllowHintText;
 	bool bAllowRealTimeOnHovered;
 
@@ -831,7 +844,8 @@ TSharedRef<SWidget> FAssetThumbnail::MakeThumbnailWidget( const FAssetThumbnailC
 		.ClassThumbnailBrushOverride(InConfig.ClassThumbnailBrushOverride)
 		.AllowAssetSpecificThumbnailOverlay(InConfig.bAllowAssetSpecificThumbnailOverlay)
 		.AssetTypeColorOverride(InConfig.AssetTypeColorOverride)
-		.Padding(InConfig.Padding);
+		.Padding(InConfig.Padding)
+		.GenericThumbnailSize(InConfig.GenericThumbnailSize);
 }
 
 void FAssetThumbnail::RefreshThumbnail()
diff --git a/Engine/Source/Editor/UnrealEd/Public/AssetThumbnail.h b/Engine/Source/Editor/UnrealEd/Public/AssetThumbnail.h
--- a/Engine/Source/Editor/UnrealEd/Public/AssetThumbnail.h
+++ b/Engine/Source/Editor/UnrealEd/Public/AssetThumbnail.h
@@ -57,6 +57,7 @@ struct FAssetThumbnailConfig
 	TAttribute< FLinearColor > HintColorAndOpacity;
 	TOptional< FLinearColor > AssetTypeColorOverride;
 	FMargin Padding;
+	TAttribute<int32> GenericThumbnailSize = 64;
 };
 
 
