About Dialog

diff --git a/Engine/Source/Editor/EditorStyle/Private/StarshipStyle.cpp b/Engine/Source/Editor/EditorStyle/Private/StarshipStyle.cpp
--- a/Engine/Source/Editor/EditorStyle/Private/StarshipStyle.cpp
+++ b/Engine/Source/Editor/EditorStyle/Private/StarshipStyle.cpp
@@ -1356,7 +1356,7 @@ void FStarshipEditorStyle::FStyle::SetupGeneralStyles()
 	// About screen
 	if (IncludeEditorSpecificStyles())
 	{
-		Set( "AboutScreen.Background", new IMAGE_BRUSH( "About/Background", FVector2D(688,317), FLinearColor::White, ESlateBrushTileType::Both) );
+		Set( "AboutScreen.Background", new IMAGE_BRUSH( "About/Background", FVector2D(688,317) ) );
 		Set( "AboutScreen.UnrealLogo", new IMAGE_BRUSH_SVG( "About/UnrealLogo", FVector2D(40,40) ) );
 		Set( "AboutScreen.EpicGamesLogo", new IMAGE_BRUSH_SVG( "About/EpicGamesLogo", FVector2D(40,40) ) );
 		Set( "AboutScreen.TitleFont", DEFAULT_FONT( "Bold", 13) );
diff --git a/Engine/Source/Editor/MainFrame/Private/Frame/MainFrameActions.cpp b/Engine/Source/Editor/MainFrame/Private/Frame/MainFrameActions.cpp
--- a/Engine/Source/Editor/MainFrame/Private/Frame/MainFrameActions.cpp
+++ b/Engine/Source/Editor/MainFrame/Private/Frame/MainFrameActions.cpp
@@ -120,7 +120,7 @@ void FMainFrameCommands::RegisterCommands()
 	UI_COMMAND( OpenIDE, "Open IDE", "Opens your C++ code in an integrated development environment.", EUserInterfaceActionType::Button, FInputChord() );
 	ActionList->MapAction( OpenIDE, FExecuteAction::CreateStatic( &FMainFrameActionCallbacks::OpenIDE ), FCanExecuteAction::CreateStatic( &FMainFrameActionCallbacks::IsCodeProject ), FGetActionCheckState(), FIsActionButtonVisible::CreateStatic( &FMainFrameActionCallbacks::CanOpenIDE ) );
 
-	UI_COMMAND( ZipUpProject, "Zip Up Project", "Zips up the project into a zip file.", EUserInterfaceActionType::Button, FInputChord() );
+	UI_COMMAND( ZipUpProject, "Zip Project", "Zips the project into a zip file.", EUserInterfaceActionType::Button, FInputChord() );
 	ActionList->MapAction(ZipUpProject, FExecuteAction::CreateStatic( &FMainFrameActionCallbacks::ZipUpProject ), DefaultExecuteAction);
 
 	//UI_COMMAND( LocalizeProject, "Localize Project...", "Opens the dashboard for managing project localization data.", EUserInterfaceActionType::Button, FInputChord() );
@@ -629,7 +629,7 @@ void FMainFrameActionCallbacks::AboutUnrealEd_Execute()
 	TSharedPtr<SWindow> AboutWindow = 
 		SNew(SWindow)
 		.Title( AboutWindowTitle )
-		.ClientSize(FVector2D(600.f, 200.f))
+		.ClientSize(FVector2D(720.f, 538.f))
 		.SupportsMaximize(false) .SupportsMinimize(false)
 		.SizingRule( ESizingRule::FixedSize )
 		[
diff --git a/Engine/Source/Editor/UnrealEd/Private/AboutScreen.cpp b/Engine/Source/Editor/UnrealEd/Private/AboutScreen.cpp
--- a/Engine/Source/Editor/UnrealEd/Private/AboutScreen.cpp
+++ b/Engine/Source/Editor/UnrealEd/Private/AboutScreen.cpp
@@ -17,6 +17,7 @@
 #include "Widgets/Views/STableRow.h"
 #include "Widgets/Views/SListView.h"
 #include "Styling/CoreStyle.h"
+#include "Styling/StyleColors.h"
 #include "EditorStyleSet.h"
 #include "UnrealEdMisc.h"
 #include "IDocumentation.h"
@@ -29,106 +30,129 @@ void SAboutScreen::Construct(const FArguments& InArgs)
 #pragma warning(push)
 #pragma warning(disable : 4428)	// universal-character-name encountered in source
 #endif
-	AboutLines.Add(MakeShareable(new FLineDefinition(LOCTEXT("Copyright1", "Copyright Epic Games, Inc. All rights reserved"), 11, FLinearColor(1.f, 1.f, 1.f), FMargin(0.f) )));
-	AboutLines.Add(MakeShareable(new FLineDefinition(LOCTEXT("Copyright2", "Epic, Epic Games, Unreal, and their respective logos are trademarks or registered trademarks of Epic Games, Inc.\nin the United States of America and elsewhere."), 8, FLinearColor(1.f, 1.f, 1.f), FMargin(0.0f,2.0f) )));
+	AboutLines.Add(MakeShareable(new FLineDefinition(LOCTEXT("Copyright1", "Copyright Epic Games, Inc. All rights reserved.  Epic, Epic Games, Unreal, and their respective logos are trademarks or registered trademarks of Epic Games, Inc. in the United States of America and elsewhere."), 9, FLinearColor(1.f, 1.f, 1.f), FMargin(0.f, 2.f) )));
 
 #ifdef _MSC_VER
 #pragma warning(pop)
 #endif
 
 	FText Version = FText::Format( LOCTEXT("VersionLabel", "Version: {0}"), FText::FromString( FEngineVersion::Current().ToString( ) ) );
-
+	FText Title = FText::FromString(FApp::GetName());
+	
 	ChildSlot
+	[
+		SNew(SBorder)
+		.BorderImage(FAppStyle::Get().GetBrush("Brushes.Panel"))
+		.Padding(16.f)
 		[
-			SNew(SOverlay)
-			+SOverlay::Slot()
+			SNew(SVerticalBox)
+
+			+SVerticalBox::Slot()
+			.AutoHeight()
 			[
-				SNew(SImage).Image(FEditorStyle::GetBrush(TEXT("AboutScreen.Background")))
+				SNew(SImage)
+				.Image(FEditorStyle::GetBrush(TEXT("AboutScreen.Background")))
 			]
 
-			+SOverlay::Slot()
+			+SVerticalBox::Slot()
+			.AutoHeight()
+			.Padding(0.f, 16.f, 0.f, 0.f)
 			[
-				SNew(SVerticalBox)
-				+SVerticalBox::Slot()
-				.AutoHeight()
+
+				SNew(SHorizontalBox)
+
+				+SHorizontalBox::Slot()
+				.FillWidth(1.0)
 				[
-					SNew(SHorizontalBox)
-					+SHorizontalBox::Slot()
-					.AutoWidth()
-					.VAlign(VAlign_Top)
-					.Padding(FMargin(10.f, 10.f, 0.f, 0.f))
-					[
-						SAssignNew(UE4Button, SButton)
-						.ButtonStyle( FEditorStyle::Get(), "NoBorder" )
-						.OnClicked(this, &SAboutScreen::OnUE4ButtonClicked)
-						[
-							SNew(SImage).Image(this, &SAboutScreen::GetUE4ButtonBrush)
-						]
-					]
-					+SHorizontalBox::Slot()
-					.AutoWidth()
-					.VAlign(VAlign_Top)
-					.Padding(FMargin(10.f, 10.f, 0.f, 0.f))
+
+					SNew(SVerticalBox)
+
+					+SVerticalBox::Slot()
+					.Padding(0.f, 4.f)
 					[
-						SAssignNew(EpicGamesButton, SButton)
-						.ButtonStyle( FEditorStyle::Get(), "NoBorder" )
-						.OnClicked(this, &SAboutScreen::OnEpicGamesButtonClicked)
-						[
-							SNew(SImage).Image(this, &SAboutScreen::GetEpicGamesButtonBrush)
-						]
+						SNew(STextBlock)
+						.ColorAndOpacity(FStyleColors::ForegroundHover)
+						.Font(FAppStyle::Get().GetFontStyle("AboutScreen.TitleFont"))
+						.Text( Title )
 					]
-					+SHorizontalBox::Slot()
-					.FillWidth(1.f)
-					.HAlign(HAlign_Right)
-					.Padding(FMargin(0.f, 52.f, 7.f, 0.f))
+
+					+SVerticalBox::Slot()
+					.Padding(0.f, 4.f)
 					[
-						SNew(SEditableText)
-						.ColorAndOpacity(FLinearColor(0.7f, 0.7f, 0.7f))
-						.IsReadOnly(true)
+						SNew(STextBlock)
+						.ColorAndOpacity(FStyleColors::ForegroundHover)
 						.Text( Version )
 					]
 				]
-				+SVerticalBox::Slot()
-				.Padding(FMargin(5.f, 5.f, 5.f, 5.f))
-				.VAlign(VAlign_Top)
-				[
-					SNew(SListView<TSharedRef<FLineDefinition>>)
-					.ListItemsSource(&AboutLines)
-					.OnGenerateRow(this, &SAboutScreen::MakeAboutTextItemWidget)
-					.SelectionMode( ESelectionMode::None )
-				] 
-				+SVerticalBox::Slot()
-				.AutoHeight()
+
+				+SHorizontalBox::Slot()
+				.AutoWidth()
+				.VAlign(VAlign_Center)
+				.HAlign(HAlign_Right)
+				.Padding(0.0f, 0.0f, 8.f, 0.0f)
 				[
-					SNew(SHorizontalBox)
-					+SHorizontalBox::Slot()
-					.HAlign(HAlign_Left)
-					.Padding(FMargin(5.f, 0.f, 5.f, 5.f)) 
+					SAssignNew(UEButton, SButton)
+					.ButtonStyle( FAppStyle::Get(), "SimpleButton" )
+					.OnClicked(this, &SAboutScreen::OnUEButtonClicked)
+					.ContentPadding(0.f)
 					[
-						SAssignNew(FacebookButton, SButton)
-						.ButtonStyle( FEditorStyle::Get(), "NoBorder" )
-						.ToolTipText(LOCTEXT("FacebookToolTip", "Unreal Engine on Facebook"))
-						.OnClicked(this, &SAboutScreen::OnFacebookButtonClicked)
-						[
-							SNew(SImage).Image(this, &SAboutScreen::GetFacebookButtonBrush)
-						]
+						SNew(SImage)
+						.Image(FAppStyle::Get().GetBrush("AboutScreen.UnrealLogo"))
+						.ColorAndOpacity(FSlateColor::UseForeground())
 					]
-					+SHorizontalBox::Slot()
-					.AutoWidth()
-					.HAlign(HAlign_Right)
-					.VAlign(VAlign_Bottom)
-					.Padding(FMargin(5.f, 0.f, 5.f,5.f))
+				]
+
+				+SHorizontalBox::Slot()
+				.AutoWidth()
+				.VAlign(VAlign_Center)
+				.HAlign(HAlign_Right)
+				.Padding(0.f)
+				[
+					SAssignNew(EpicGamesButton, SButton)
+					.ButtonStyle( FAppStyle::Get(), "SimpleButton" )
+					.OnClicked(this, &SAboutScreen::OnEpicGamesButtonClicked)
+					.ContentPadding(0.f)
 					[
-						SNew(SButton)
-						.HAlign(HAlign_Center)
-						.VAlign(VAlign_Center)
-						.Text(LOCTEXT("Close", "Close"))
-						.ButtonColorAndOpacity(FLinearColor(0.6f, 0.6f, 0.6f))
-						.OnClicked(this, &SAboutScreen::OnClose)
+						SNew(SImage)
+						.Image(FAppStyle::Get().GetBrush("AboutScreen.EpicGamesLogo"))
+						.ColorAndOpacity(FSlateColor::UseForeground())
 					]
 				]
 			]
-		];
+
+
+			+SVerticalBox::Slot()
+			.Padding(FMargin(0.f, 16.f))
+			.AutoHeight()
+			[
+				SNew(SListView<TSharedRef<FLineDefinition>>)
+				.ListViewStyle(&FAppStyle::Get().GetWidgetStyle<FTableViewStyle>("SimpleListView"))
+				.ListItemsSource(&AboutLines)
+				.OnGenerateRow(this, &SAboutScreen::MakeAboutTextItemWidget)
+				.SelectionMode( ESelectionMode::None )
+			] 
+
+			+SVerticalBox::Slot()
+			.Padding(FMargin(0.f, 16.f, 0.0f, 0.0f))
+			.AutoHeight()
+			[
+				SNew(SHorizontalBox)
+
+				+SHorizontalBox::Slot()
+				.HAlign(HAlign_Right)
+				.VAlign(VAlign_Bottom)
+				[
+					SNew(SButton)
+					.ButtonStyle(FAppStyle::Get(), "PrimaryButton")
+					.TextStyle( FAppStyle::Get(), "DialogButtonText" )
+					.HAlign(HAlign_Center)
+					.VAlign(VAlign_Center)
+					.Text(LOCTEXT("Close", "Close"))
+					.OnClicked(this, &SAboutScreen::OnClose)
+				]
+			]
+		]
+	];
 }
 
 TSharedRef<ITableRow> SAboutScreen::MakeAboutTextItemWidget(TSharedRef<FLineDefinition> Item, const TSharedRef<STableViewBase>& OwnerTable)
@@ -137,6 +161,7 @@ TSharedRef<ITableRow> SAboutScreen::MakeAboutTextItemWidget(TSharedRef<FLineDefi
 	{
 		return
 			SNew(STableRow< TSharedPtr<FString> >, OwnerTable)
+			.Style( &FAppStyle::Get().GetWidgetStyle<FTableRowStyle>("SimpleTableView.Row") )
 			.Padding(6.0f)
 			[
 				SNew(SSpacer)
@@ -146,9 +171,12 @@ TSharedRef<ITableRow> SAboutScreen::MakeAboutTextItemWidget(TSharedRef<FLineDefi
 	{
 		return
 			SNew(STableRow< TSharedPtr<FString> >, OwnerTable)
+			.Style( &FAppStyle::Get().GetWidgetStyle<FTableRowStyle>("SimpleTableView.Row") )
 			.Padding( Item->Margin )
 			[
 				SNew(STextBlock)
+				.LineHeightPercentage(1.3f)
+				.AutoWrapText(true)
 				.ColorAndOpacity( Item->TextColor )
 				.Font( FCoreStyle::GetDefaultFontStyle("Regular", Item->FontSize) )
 				.Text( Item->Text )
@@ -156,22 +184,7 @@ TSharedRef<ITableRow> SAboutScreen::MakeAboutTextItemWidget(TSharedRef<FLineDefi
 	}
 }
 
-const FSlateBrush* SAboutScreen::GetUE4ButtonBrush() const
-{
-	return FEditorStyle::GetBrush(!UE4Button->IsHovered() ? TEXT("AboutScreen.UE4") : TEXT("AboutScreen.UE4Hovered"));
-}
-
-const FSlateBrush* SAboutScreen::GetEpicGamesButtonBrush() const
-{
-	return FEditorStyle::GetBrush(!EpicGamesButton->IsHovered() ? TEXT("AboutScreen.EpicGames") : TEXT("AboutScreen.EpicGamesHovered"));
-}
-
-const FSlateBrush* SAboutScreen::GetFacebookButtonBrush() const
-{
-	return FEditorStyle::GetBrush(!FacebookButton->IsHovered() ? TEXT("AboutScreen.Facebook") : TEXT("AboutScreen.FacebookHovered"));
-}
-
-FReply SAboutScreen::OnUE4ButtonClicked()
+FReply SAboutScreen::OnUEButtonClicked()
 {
 	IDocumentation::Get()->OpenHome(FDocumentationSourceInfo(TEXT("logo_docs")));
 	return FReply::Handled();
@@ -187,16 +200,6 @@ FReply SAboutScreen::OnEpicGamesButtonClicked()
 	return FReply::Handled();
 }
 
-FReply SAboutScreen::OnFacebookButtonClicked()
-{
-	FString FacebookURL;
-	if(FUnrealEdMisc::Get().GetURL( TEXT("FacebookURL"), FacebookURL ))
-	{
-		FPlatformProcess::LaunchURL( *FacebookURL, NULL, NULL );
-	}
-	return FReply::Handled();
-}
-
 FReply SAboutScreen::OnClose()
 {
 	TSharedRef<SWindow> ParentWindow = FSlateApplication::Get().FindWidgetWindow( AsShared() ).ToSharedRef();
diff --git a/Engine/Source/Editor/UnrealEd/Public/AboutScreen.h b/Engine/Source/Editor/UnrealEd/Public/AboutScreen.h
--- a/Engine/Source/Editor/UnrealEd/Public/AboutScreen.h
+++ b/Engine/Source/Editor/UnrealEd/Public/AboutScreen.h
@@ -58,22 +58,16 @@ private:
 	};
 
 	TArray< TSharedRef< FLineDefinition > > AboutLines;
-	TSharedPtr<SButton> UE4Button;
+	TSharedPtr<SButton> UEButton;
 	TSharedPtr<SButton> EpicGamesButton;
-	TSharedPtr<SButton> FacebookButton;
 
 	/** 
 	 * Makes the widget for the checkbox items in the list view 
 	 */
 	TSharedRef<ITableRow> MakeAboutTextItemWidget(TSharedRef<FLineDefinition> Item, const TSharedRef<STableViewBase>& OwnerTable);
 
-	const FSlateBrush* GetUE4ButtonBrush() const;
-	const FSlateBrush* GetEpicGamesButtonBrush() const;
-	const FSlateBrush* GetFacebookButtonBrush() const;
-
-	FReply OnUE4ButtonClicked();
+	FReply OnUEButtonClicked();
 	FReply OnEpicGamesButtonClicked();
-	FReply OnFacebookButtonClicked();
 	FReply OnClose();
 };
 
diff --git a/Engine/Source/Runtime/SlateCore/Private/Styling/StarshipCoreStyle.cpp b/Engine/Source/Runtime/SlateCore/Private/Styling/StarshipCoreStyle.cpp
--- a/Engine/Source/Runtime/SlateCore/Private/Styling/StarshipCoreStyle.cpp
+++ b/Engine/Source/Runtime/SlateCore/Private/Styling/StarshipCoreStyle.cpp
@@ -1712,6 +1712,8 @@ void FStarshipCoreStyle::SetupTableViewStyles(TSharedRef<FStyle>& Style)
 	const FTableViewStyle DefaultTreeViewStyle = FTableViewStyle();
 	Style->Set("ListView", DefaultTreeViewStyle);
 
+	Style->Set("SimpleListView", FTableViewStyle().SetBackgroundBrush(FSlateNoResource()));
+
 	const FTableViewStyle DefaultTableViewStyle = FTableViewStyle();
 	Style->Set("TreeView", DefaultTableViewStyle);
 
@@ -1742,6 +1744,34 @@ void FStarshipCoreStyle::SetupTableViewStyles(TSharedRef<FStyle>& Style)
 
 	Style->Set("TableView.Row", DefaultTableRowStyle);
 
+
+	const FTableRowStyle SimpleTableRowStyle = FTableRowStyle()
+		.SetEvenRowBackgroundBrush(FSlateNoResource())
+		.SetEvenRowBackgroundHoveredBrush(FSlateColorBrush(FStyleColors::SelectHover))
+
+		.SetOddRowBackgroundBrush(FSlateNoResource())
+		.SetOddRowBackgroundHoveredBrush(FSlateColorBrush(FStyleColors::SelectHover))
+
+		.SetSelectorFocusedBrush(BORDER_BRUSH("Common/Selector", FMargin(4.f / 16.f), SelectorColor))
+
+		.SetActiveBrush(FSlateColorBrush(FStyleColors::Select))
+		.SetActiveHoveredBrush(FSlateColorBrush(FStyleColors::Select))
+
+		.SetInactiveBrush(FSlateColorBrush(FStyleColors::SelectInactive))
+		.SetInactiveHoveredBrush(FSlateColorBrush(FStyleColors::SelectInactive))
+
+		.SetActiveHighlightedBrush(FSlateColorBrush(FStyleColors::SelectParent)) // This is the parent hightlight
+		.SetInactiveHighlightedBrush(FSlateColorBrush(FStyleColors::SelectParent))// This is the parent highlight
+
+		.SetTextColor(FStyleColors::Foreground)
+		.SetSelectedTextColor(FStyleColors::ForegroundInverted)
+
+		.SetDropIndicator_Above(BOX_BRUSH("Common/DropZoneIndicator_Above", FMargin(10.0f / 16.0f, 10.0f / 16.0f, 0, 0), SelectionColor))
+		.SetDropIndicator_Onto(BOX_BRUSH("Common/DropZoneIndicator_Onto", FMargin(4.0f / 16.0f), SelectionColor))
+		.SetDropIndicator_Below(BOX_BRUSH("Common/DropZoneIndicator_Below", FMargin(10.0f / 16.0f, 0, 0, 10.0f / 16.0f), SelectionColor));
+
+	Style->Set("SimpleTableView.Row", SimpleTableRowStyle);
+
 	const FTableRowStyle DarkTableRowStyle = FTableRowStyle(DefaultTableRowStyle)
 		.SetEvenRowBackgroundBrush(IMAGE_BRUSH("Common/Selection", Icon8x8, FLinearColor(0.0f, 0.0f, 0.0f, 0.1f)))
 		.SetOddRowBackgroundBrush(IMAGE_BRUSH("Common/Selection", Icon8x8, FLinearColor(0.0f, 0.0f, 0.0f, 0.1f)));
