Remove SuspendRenderingThread when clicking in a viewport

diff --git a/Engine/Source/Editor/UnrealEd/Private/Editor/ActorPositioning.cpp b/Engine/Source/Editor/UnrealEd/Private/Editor/ActorPositioning.cpp
--- a/Engine/Source/Editor/UnrealEd/Private/Editor/ActorPositioning.cpp
+++ b/Engine/Source/Editor/UnrealEd/Private/Editor/ActorPositioning.cpp
@@ -56,6 +56,9 @@ FActorPositionTraceResult FActorPositioning::TraceWorldForPosition(const FViewpo
 /** Check to see if the specified hit result should be ignored from actor positioning calculations for the specified scene view */
 bool IsHitIgnored(const FHitResult& InHit, const FSceneView& InSceneView)
 {
+	// We're using the SceneProxy and ViewRelevance here, we should execute from the render thread but since
+	// we're also accessing UObjects, the game thread should be blocked during this call.
+	check(IsInParallelRenderingThread());
 	const FActorInstanceHandle& HitObjHandle = InHit.HitObjectHandle;
 	AActor* Actor = HitObjHandle.FetchActor();
 	
@@ -118,11 +121,21 @@ FActorPositionTraceResult FActorPositioning::TraceWorldForPosition(const UWorld&
 	if ( InWorld.LineTraceMultiByObjectType(Hits, RayStart, RayEnd, FCollisionObjectQueryParams(FCollisionObjectQueryParams::InitType::AllObjects), Param) )
 	{
 		{
-			// Filter out anything that should be ignored
-			FSuspendRenderingThread SuspendRendering(false);
-			Hits.RemoveAll([&](const FHitResult& Hit){
-				return IsHitIgnored(Hit, InSceneView);
-			});
+			// Send IsHitIgnored on the render thread since we're accessing view relevance
+			ENQUEUE_RENDER_COMMAND(TraceWorldForPosition_FilterHitsByViewRelevance)(
+				[&Hits, &InSceneView](FRHICommandListImmediate& RHICmdList)
+				{
+					// Filter out anything that should be ignored
+					Hits.RemoveAll([&InSceneView](const FHitResult& Hit){
+						return IsHitIgnored(Hit, InSceneView);
+					});
+				}
+			);
+			
+			// We need the result to come back before continuing
+			FRenderCommandFence Fence;
+			Fence.BeginFence();
+			Fence.Wait();
 		}
 
 		// Go through all hits and find closest
