Re-add support for deleting folders in the world outliner

diff --git a/Engine/Source/Editor/SceneOutliner/Private/ActorBrowsingMode.cpp b/Engine/Source/Editor/SceneOutliner/Private/ActorBrowsingMode.cpp
--- a/Engine/Source/Editor/SceneOutliner/Private/ActorBrowsingMode.cpp
+++ b/Engine/Source/Editor/SceneOutliner/Private/ActorBrowsingMode.cpp
@@ -717,6 +717,13 @@ FReply FActorBrowsingMode::OnKeyDown(const FKeyEvent& InKeyEvent)
 	return FReply::Unhandled();
 }
 
+bool FActorBrowsingMode::CanDelete() const
+{
+	const FSceneOutlinerItemSelection ItemSelection = SceneOutliner->GetSelection();
+	const uint32 NumberOfFolders = ItemSelection.Num<FFolderTreeItem>();
+	return (NumberOfFolders > 0 && NumberOfFolders == ItemSelection.Num());
+}
+
 bool FActorBrowsingMode::CanRename() const
 {
 	const FSceneOutlinerItemSelection ItemSelection = SceneOutliner->GetSelection();
diff --git a/Engine/Source/Editor/SceneOutliner/Public/ActorBrowsingMode.h b/Engine/Source/Editor/SceneOutliner/Public/ActorBrowsingMode.h
--- a/Engine/Source/Editor/SceneOutliner/Public/ActorBrowsingMode.h
+++ b/Engine/Source/Editor/SceneOutliner/Public/ActorBrowsingMode.h
@@ -39,6 +39,7 @@ public:
 	virtual bool ShowStatusBar() const override { return true; }
 	virtual bool ShowViewButton() const override { return true; }
 	virtual bool ShowFilterOptions() const override { return true; }
+	virtual bool CanDelete() const override;
 	virtual bool CanRename() const override;
 	virtual bool CanCut() const override;
 	virtual bool CanCopy() const override;
diff --git a/Engine/Source/Editor/UnrealEd/Private/UnrealEdSrv.cpp b/Engine/Source/Editor/UnrealEd/Private/UnrealEdSrv.cpp
--- a/Engine/Source/Editor/UnrealEd/Private/UnrealEdSrv.cpp
+++ b/Engine/Source/Editor/UnrealEd/Private/UnrealEdSrv.cpp
@@ -2497,7 +2497,17 @@ bool UUnrealEdEngine::Exec_Actor( UWorld* InWorld, const TCHAR* Str, FOutputDevi
 				UTypedElementSelectionSet* SelectionSet = LevelEditor->GetMutableElementSelectionSet();
 				
 				const FScopedTransaction Transaction(NSLOCTEXT("UnrealEd", "DeleteElements", "Delete Elements"));
-				CommonActions->DeleteSelectedElements(SelectionSet, InWorld, FTypedElementDeletionOptions());
+				if (SelectionSet->GetNumSelectedElements() == 0)
+				{
+					// HACK: Call these directly for an empty selection so that folder deletion in the outliner still works
+					// TODO: Move this logic into FLevelEditorActionCallbacks and have it call into the outliner directly
+					FEditorDelegates::OnDeleteActorsBegin.Broadcast();
+					FEditorDelegates::OnDeleteActorsEnd.Broadcast();
+				}
+				else
+				{
+					CommonActions->DeleteSelectedElements(SelectionSet, InWorld, FTypedElementDeletionOptions());
+				}
 			}
 		}
 		return true;
