Prevent crash when pressing alt+0 with r.AllowStaticLighting disabled.

The lightmap density shaders aren't compiled when r.AllowStaticLighting is set to 0, so trying to turn on this view mode crashes. The menu code already hides the lightmap density option when static lighting is off, but the hotkey invokes the command directly, so it's easier to skip binding the command altogether.

#rnx

diff --git a/Engine/Source/Editor/UnrealEd/Private/SEditorViewport.cpp b/Engine/Source/Editor/UnrealEd/Private/SEditorViewport.cpp
--- a/Engine/Source/Editor/UnrealEd/Private/SEditorViewport.cpp
+++ b/Engine/Source/Editor/UnrealEd/Private/SEditorViewport.cpp
@@ -414,7 +414,12 @@ void SEditorViewport::BindCommands()
 	MAP_VIEWMODE_ACTION( Commands.TexStreamAccMaterialTextureScaleMode, VMI_MaterialTextureScaleAccuracy );
 	MAP_VIEWMODE_ACTION( Commands.RequiredTextureResolutionMode, VMI_RequiredTextureResolution );
 	MAP_VIEWMODE_ACTION( Commands.StationaryLightOverlapMode, VMI_StationaryLightOverlap );
-	MAP_VIEWMODE_ACTION( Commands.LightmapDensityMode, VMI_LightmapDensity );
+
+	if (IConsoleManager::Get().FindTConsoleVariableDataInt(TEXT("r.AllowStaticLighting"))->GetValueOnAnyThread() != 0)
+	{
+		MAP_VIEWMODE_ACTION(Commands.LightmapDensityMode, VMI_LightmapDensity);
+	}
+
 	MAP_VIEWMODE_ACTION( Commands.ReflectionOverrideMode, VMI_ReflectionOverride );
 	MAP_VIEWMODE_ACTION( Commands.GroupLODColorationMode, VMI_GroupLODColoration);
 	MAP_VIEWMODE_ACTION( Commands.LODColorationMode, VMI_LODColoration );
