Fixed crash in cooking when CookOnTheFlyServer exceeds maximum allowed memory and garbage collects while a generator package is in the middle of splitting the package.

diff --git a/Engine/Source/Editor/UnrealEd/Private/CookOnTheFlyServer.cpp b/Engine/Source/Editor/UnrealEd/Private/CookOnTheFlyServer.cpp
--- a/Engine/Source/Editor/UnrealEd/Private/CookOnTheFlyServer.cpp
+++ b/Engine/Source/Editor/UnrealEd/Private/CookOnTheFlyServer.cpp
@@ -3601,7 +3601,19 @@ void UCookOnTheFlyServer::PostGarbageCollect()
 	{
 		if (PackageData->IsSaveInvalidated())
 		{
-			Demotes.Add(PackageData);
+			if (PackageData->GetGeneratorPackage())
+			{
+				// Generator packages cannot be demoted as that would break their generation data which must be preserved
+				// We're relying on BeginCacheForCookedPlatformData not having started; this is currently guaranteed when PackageData has a valid GeneratorPackage.
+				check(PackageData->GetCookedPlatformDataNextIndex() == -1);
+				// Recalculate object cache
+				PackageData->ClearObjectCache();
+				PackageData->CreateObjectCache();
+			}
+			else
+			{
+				Demotes.Add(PackageData);
+			}
 		}
 	}
 	for (FPackageData* PackageData : Demotes)
