Fixed missing delete/duplicate support

 - Duplicating folders in the world outliner was broken (same issue as delete from UE-112611).
 - Deletion/Duplication of non-element based selections in ed modes was broken (eg, landscape spline control points).

diff --git a/Engine/Source/Editor/UnrealEd/Private/UnrealEdSrv.cpp b/Engine/Source/Editor/UnrealEd/Private/UnrealEdSrv.cpp
--- a/Engine/Source/Editor/UnrealEd/Private/UnrealEdSrv.cpp
+++ b/Engine/Source/Editor/UnrealEd/Private/UnrealEdSrv.cpp
@@ -2499,10 +2499,15 @@ bool UUnrealEdEngine::Exec_Actor( UWorld* InWorld, const TCHAR* Str, FOutputDevi
 				const FScopedTransaction Transaction(NSLOCTEXT("UnrealEd", "DeleteElements", "Delete Elements"));
 				if (SelectionSet->GetNumSelectedElements() == 0)
 				{
-					// HACK: Call these directly for an empty selection so that folder deletion in the outliner still works
-					// TODO: Move this logic into FLevelEditorActionCallbacks and have it call into the outliner directly
-					FEditorDelegates::OnDeleteActorsBegin.Broadcast();
-					FEditorDelegates::OnDeleteActorsEnd.Broadcast();
+					// HACK: Not all modes will select elements, so allow them a shot at deletion if we don't think anything else is selected
+					// TODO: Move this logic into FLevelEditorActionCallbacks and have it call into the mode directly
+					if (!GLevelEditorModeTools().ProcessEditDelete())
+					{
+						// HACK: Call these directly for an empty selection so that folder deletion in the outliner still works
+						// TODO: Move this logic into FLevelEditorActionCallbacks and have it call into the outliner directly
+						FEditorDelegates::OnDeleteActorsBegin.Broadcast();
+						FEditorDelegates::OnDeleteActorsEnd.Broadcast();
+					}
 				}
 				else
 				{
@@ -2685,11 +2690,26 @@ bool UUnrealEdEngine::Exec_Actor( UWorld* InWorld, const TCHAR* Str, FOutputDevi
 				UTypedElementSelectionSet* SelectionSet = LevelEditor->GetMutableElementSelectionSet();
 
 				const FScopedTransaction Transaction(NSLOCTEXT("UnrealEd", "DuplicateElements", "Duplicate Elements"));
-				const TArray<FTypedElementHandle> DuplicatedElements = CommonActions->DuplicateSelectedElements(SelectionSet, InWorld, GEditor->GetGridLocationOffset(/*bUniformOffset*/false));
-				if (DuplicatedElements.Num() > 0)
+				if (SelectionSet->GetNumSelectedElements() == 0)
+				{
+					// HACK: Not all modes will select elements, so allow them a shot at duplication if we don't think anything else is selected
+					// TODO: Move this logic into FLevelEditorActionCallbacks and have it call into the mode directly
+					if (!GLevelEditorModeTools().ProcessEditDuplicate())
+					{
+						// HACK: Call these directly for an empty selection so that folder duplication in the outliner still works
+						// TODO: Move this logic into FLevelEditorActionCallbacks and have it call into the outliner directly
+						FEditorDelegates::OnDuplicateActorsBegin.Broadcast();
+						FEditorDelegates::OnDuplicateActorsEnd.Broadcast();
+					}
+				}
+				else
 				{
-					SelectionSet->SetSelection(DuplicatedElements, FTypedElementSelectionOptions());
-					SelectionSet->NotifyPendingChanges();
+					const TArray<FTypedElementHandle> DuplicatedElements = CommonActions->DuplicateSelectedElements(SelectionSet, InWorld, GEditor->GetGridLocationOffset(/*bUniformOffset*/false));
+					if (DuplicatedElements.Num() > 0)
+					{
+						SelectionSet->SetSelection(DuplicatedElements, FTypedElementSelectionOptions());
+						SelectionSet->NotifyPendingChanges();
+					}
 				}
 			}
 		}
