Lists platform when CookGlobalShader commandlet has not found a matching platform

diff --git a/Engine/Source/Editor/UnrealEd/Private/Commandlets/CookGlobalShaderCommandlet.cpp b/Engine/Source/Editor/UnrealEd/Private/Commandlets/CookGlobalShaderCommandlet.cpp
--- a/Engine/Source/Editor/UnrealEd/Private/Commandlets/CookGlobalShaderCommandlet.cpp
+++ b/Engine/Source/Editor/UnrealEd/Private/Commandlets/CookGlobalShaderCommandlet.cpp
@@ -66,16 +66,27 @@ int32 UCookGlobalShadersCommandlet::Main(const FString& Params)
 	FString PlatformName;
 	ITargetPlatform* TargetPlatform = nullptr;
 	{
-		if ( !FParse::Value(*Params, TEXT("platform="), PlatformName, true) )
+		ITargetPlatformManagerModule& TPM = GetTargetPlatformManagerRef();
+
+		if (!ParamVals.Contains(TEXT("platform")))
 		{
 			UE_LOG(LogCookGlobalShaders, Warning, TEXT("You must include a target platform with -platform=xxx"));
+			for (ITargetPlatform* TP : TPM.GetTargetPlatforms())
+			{
+				UE_LOG(LogCookGlobalShaders, Display, TEXT("   %s"), *TP->PlatformName());
+			}
 			return 1;
 		}
-		ITargetPlatformManagerModule& TPM = GetTargetPlatformManagerRef();
+		PlatformName = ParamVals.FindRef(TEXT("platform"));
+
 		TargetPlatform = TPM.FindTargetPlatform(PlatformName);
 		if (TargetPlatform == nullptr)
 		{
 			UE_LOG(LogCookGlobalShaders, Warning, TEXT("Target platform '%s' was not found"), *PlatformName);
+			for (ITargetPlatform* TP : TPM.GetTargetPlatforms())
+			{
+				UE_LOG(LogCookGlobalShaders, Display, TEXT("   %s"), *TP->PlatformName());
+			}
 			return 1;
 		}
 
