Fix a crash when using the create menu on an actor that doesn't have a factory.

diff --git a/Engine/Source/Editor/UnrealEd/Private/Editor/ActorPositioning.cpp b/Engine/Source/Editor/UnrealEd/Private/Editor/ActorPositioning.cpp
--- a/Engine/Source/Editor/UnrealEd/Private/Editor/ActorPositioning.cpp
+++ b/Engine/Source/Editor/UnrealEd/Private/Editor/ActorPositioning.cpp
@@ -160,41 +160,39 @@ FActorPositionTraceResult FActorPositioning::TraceWorldForPosition(const UWorld&
 
 FTransform FActorPositioning::GetCurrentViewportPlacementTransform(const AActor& Actor, bool bSnap, const FViewportCursorLocation* InCursor)
 {
-	FVector Collision = Actor.GetPlacementExtent();
-	const UActorFactory* Factory = GEditor->FindActorFactoryForActorClass(Actor.GetClass());
-
-	// Get cursor origin and direction in world space.
-	FViewportCursorLocation CursorLocation = InCursor ? *InCursor : GCurrentLevelEditingViewportClient->GetCursorWorldLocationFromMousePos();
-	const auto CursorPos = CursorLocation.GetCursorPos();
-
 	FTransform ActorTransform = FTransform::Identity;
-
-	if (CursorLocation.GetViewportType() == LVT_Perspective && !GCurrentLevelEditingViewportClient->Viewport->GetHitProxy( CursorPos.X, CursorPos.Y ))
+	if (GCurrentLevelEditingViewportClient)
 	{
-		ActorTransform.SetTranslation(GetActorPositionInFrontOfCamera(Actor, CursorLocation.GetOrigin(), CursorLocation.GetDirection()));
+		// Get cursor origin and direction in world space.
+		FViewportCursorLocation CursorLocation = InCursor ? *InCursor : GCurrentLevelEditingViewportClient->GetCursorWorldLocationFromMousePos();
+		const auto CursorPos = CursorLocation.GetCursorPos();
+
+		if (CursorLocation.GetViewportType() == LVT_Perspective && !GCurrentLevelEditingViewportClient->Viewport->GetHitProxy(CursorPos.X, CursorPos.Y))
+		{
+			ActorTransform.SetTranslation(GetActorPositionInFrontOfCamera(Actor, CursorLocation.GetOrigin(), CursorLocation.GetDirection()));
+			return ActorTransform;
+		}
 	}
-	else
-	{
-		const FSnappedPositioningData PositioningData = FSnappedPositioningData(GCurrentLevelEditingViewportClient, GEditor->ClickLocation, GEditor->ClickPlane)
-			.DrawSnapHelpers(true)
-			.UseFactory(Factory)
-			.UsePlacementExtent(Actor.GetPlacementExtent());
 
-		ActorTransform = bSnap ? GetSnappedSurfaceAlignedTransform(PositioningData) : GetSurfaceAlignedTransform(PositioningData);
+	const FSnappedPositioningData PositioningData = FSnappedPositioningData(nullptr, GEditor->ClickLocation, GEditor->ClickPlane)
+		.DrawSnapHelpers(true)
+		.UseFactory(GEditor->FindActorFactoryForActorClass(Actor.GetClass()))
+		.UsePlacementExtent(Actor.GetPlacementExtent());
 
-		if (GetDefault<ULevelEditorViewportSettings>()->SnapToSurface.bEnabled)
-		{
-			// HACK: If we are aligning rotation to surfaces, we have to factor in the inverse of the actor's rotation and translation so that the resulting transform after SpawnActor is correct.
+	ActorTransform = bSnap ? GetSnappedSurfaceAlignedTransform(PositioningData) : GetSurfaceAlignedTransform(PositioningData);
 
-			if (auto* RootComponent = Actor.GetRootComponent())
-			{
-				RootComponent->UpdateComponentToWorld();
-			}
+	if (GetDefault<ULevelEditorViewportSettings>()->SnapToSurface.bEnabled)
+	{
+		// HACK: If we are aligning rotation to surfaces, we have to factor in the inverse of the actor's rotation and translation so that the resulting transform after SpawnActor is correct.
 
-			FVector OrigActorScale3D = ActorTransform.GetScale3D();
-			ActorTransform = Actor.GetTransform().Inverse() * ActorTransform;
-			ActorTransform.SetScale3D(OrigActorScale3D);
+		if (auto* RootComponent = Actor.GetRootComponent())
+		{
+			RootComponent->UpdateComponentToWorld();
 		}
+
+		FVector OrigActorScale3D = ActorTransform.GetScale3D();
+		ActorTransform = Actor.GetTransform().Inverse() * ActorTransform;
+		ActorTransform.SetScale3D(OrigActorScale3D);
 	}
 
 	return ActorTransform;
diff --git a/Engine/Source/Editor/UnrealEd/Private/EditorEngine.cpp b/Engine/Source/Editor/UnrealEd/Private/EditorEngine.cpp
--- a/Engine/Source/Editor/UnrealEd/Private/EditorEngine.cpp
+++ b/Engine/Source/Editor/UnrealEd/Private/EditorEngine.cpp
@@ -230,6 +230,7 @@
 #include "ToolMenus.h"
 #include "IToolMenusEditorModule.h"
 #include "Subsystems/AssetEditorSubsystem.h"
+#include "LevelEditorSubsystem.h"
 #include "StudioAnalytics.h"
 #include "Engine/LevelScriptActor.h"
 #include "UObject/UnrealType.h"
@@ -4713,99 +4714,28 @@ AActor* UEditorEngine::UseActorFactoryOnCurrentSelection( UActorFactory* Factory
 
 AActor* UEditorEngine::UseActorFactory( UActorFactory* Factory, const FAssetData& AssetData, const FTransform* InActorTransform, EObjectFlags InObjectFlags )
 {
-	check( Factory );
+	AActor* NewActor = nullptr;
 
-	bool bIsAllowedToCreateActor = true;
-
-	FText ActorErrorMsg;
-	if( !Factory->CanCreateActorFrom( AssetData, ActorErrorMsg ) )
-	{
-		bIsAllowedToCreateActor = false;
-		if(!ActorErrorMsg.IsEmpty())
-		{
-			FMessageLog EditorErrors("EditorErrors");
-			EditorErrors.Warning(ActorErrorMsg);
-			EditorErrors.Notify();
-		}
-	}
-
-	//Load Asset
-	UObject* Asset = AssetData.GetAsset();
-
-	UWorld* OldWorld = nullptr;
-
-	// The play world needs to be selected if it exists
-	if (GIsEditor && PlayWorld && !GIsPlayInEditorWorld)
+	if (ULevelEditorSubsystem* LevelEditorSubsystem = GEditor->GetEditorSubsystem<ULevelEditorSubsystem>())
 	{
-		OldWorld = SetPlayInEditorWorld(PlayWorld);
-	}
-
-	AActor* Actor = NULL;
-	if ( bIsAllowedToCreateActor )
-	{
-		AActor* NewActorTemplate = Factory->GetDefaultActor( AssetData );
-
-		if ( !NewActorTemplate )
-		{
-			return NULL;
-		}
-
-		const FTransform ActorTransform = InActorTransform ? *InActorTransform : FActorPositioning::GetCurrentViewportPlacementTransform(*NewActorTemplate);
-
-		ULevel* DesiredLevel = GWorld->GetCurrentLevel();
-
-		// Don't spawn the actor if the current level is locked.
-		if( !FLevelUtils::IsLevelLocked( DesiredLevel ) )
+		if (ULevel* DesiredLevel = LevelEditorSubsystem->GetCurrentLevel())
 		{
-			// Check to see if the level it's being added to is hidden and ask the user if they want to proceed
-			const bool bLevelVisible = FLevelUtils::IsLevelVisible( DesiredLevel );
-			if ( bLevelVisible || EAppReturnType::Ok == FMessageDialog::Open( EAppMsgType::OkCancel, FText::Format( LOCTEXT("CurrentLevelHiddenActorWillAlsoBeHidden", "Current level [{0}] is hidden, actor will also be hidden until level is visible"), FText::FromString( DesiredLevel->GetOutermost()->GetName() ) ) ) )
+			if (UObject* LoadedAsset = AssetData.GetAsset())
 			{
-				const FScopedTransaction Transaction( NSLOCTEXT("UnrealEd", "CreateActor", "Create Actor") );
-
-				// Create the actor.
-				Actor = Factory->CreateActor( Asset, DesiredLevel, ActorTransform, InObjectFlags );
-				if(Actor != NULL)
+				TArray<AActor*> Actors = FLevelEditorViewportClient::TryPlacingActorFromObject(DesiredLevel, LoadedAsset, true, RF_Transactional, Factory);
+				if (Actors.Num() && (Actors[0] != nullptr))
 				{
-					SelectNone( false, true );
-					SelectActor( Actor, true, true );
-					Actor->InvalidateLightingCache();
-					Actor->PostEditMove( true );
-
-					// Make sure the actors visibility reflects that of the level it's in
-					if ( !bLevelVisible )
+					NewActor = Actors[0];
+					if (InActorTransform)
 					{
-						Actor->bHiddenEdLevel = true;
-						// We update components, so things like draw scale take effect.
-						Actor->ReregisterAllComponents(); // @todo UE4 insist on a property update callback
+						NewActor->SetActorTransform(*InActorTransform, false, nullptr, ETeleportType::TeleportPhysics);
 					}
 				}
-
-				RedrawLevelEditingViewports();
-
-
-				if ( Actor )
-				{
-					Actor->MarkPackageDirty();
-					ULevel::LevelDirtiedEvent.Broadcast();
-				}
 			}
 		}
-		else
-		{
-			FNotificationInfo Info( NSLOCTEXT("UnrealEd", "Error_OperationDisallowedOnLockedLevel", "The requested operation could not be completed because the level is locked.") );
-			Info.ExpireDuration = 3.0f;
-			FSlateNotificationManager::Get().AddNotification(Info);
-		}
 	}
 
-	// Restore the old world if there was one
-	if (OldWorld)
-	{
-		RestoreEditorWorld(OldWorld);
-	}
-
-	return Actor;
+	return NewActor;
 }
 
 namespace ReattachActorsHelper
