Test and failure code for cook on the fly of a partitioned world.

diff --git a/Engine/Source/Editor/UnrealEd/Private/CookOnTheFlyServer.cpp b/Engine/Source/Editor/UnrealEd/Private/CookOnTheFlyServer.cpp
--- a/Engine/Source/Editor/UnrealEd/Private/CookOnTheFlyServer.cpp
+++ b/Engine/Source/Editor/UnrealEd/Private/CookOnTheFlyServer.cpp
@@ -2086,6 +2086,17 @@ UE::Cook::FGeneratorPackage* UCookOnTheFlyServer::CreateGeneratorPackage(UE::Coo
 		return nullptr;
 	}
 
+	// TODO: Add support for cook on the fly
+	if (IsCookOnTheFlyMode())
+	{
+		UWorld* World = Cast<UWorld>(SplitDataObject);
+		if (World && World->GetWorldPartition())
+		{
+			UE_LOG(LogCook, Error, TEXT("Cook on the fly doesn't support World Partition."));
+			return nullptr;
+		}
+	}
+
 	UE_LOG(LogCook, Display, TEXT("Splitting Package %s with class %s acting on object %s."), *PackageData.GetPackageName().ToString(), *Splitter->GetSplitDataClass()->GetName(), *SplitDataObject->GetFullName());
 
 	// Prepare necessary paths
