WorldPartitionResaveActorsBuilder

diff --git a/Engine/Source/Editor/UnrealEd/Private/WorldPartition/WorldPartitionResaveActorsBuilder.cpp b/Engine/Source/Editor/UnrealEd/Private/WorldPartition/WorldPartitionResaveActorsBuilder.cpp
new file mode 100644
--- /dev/null
+++ b/Engine/Source/Editor/UnrealEd/Private/WorldPartition/WorldPartitionResaveActorsBuilder.cpp
@@ -0,0 +1,120 @@
+// Copyright Epic Games, Inc. All Rights Reserved.
+
+#include "WorldPartition/WorldPartitionResaveActorsBuilder.h"
+
+#include "WorldPartition/WorldPartition.h"
+#include "PackageSourceControlHelper.h"
+#include "SourceControlHelpers.h"
+#include "Engine/World.h"
+#include "GameFramework/Actor.h"
+
+DEFINE_LOG_CATEGORY_STATIC(LogWorldPartitionResaveActorsBuilder, All, All);
+
+UWorldPartitionResaveActorsBuilder::UWorldPartitionResaveActorsBuilder(const FObjectInitializer& ObjectInitializer)
+	: Super(ObjectInitializer)
+{
+	FParse::Value(FCommandLine::Get(), TEXT("ActorClass="), ActorClassName);
+}
+
+bool UWorldPartitionResaveActorsBuilder::Run(UWorld* World, FPackageSourceControlHelper& PackageHelper)
+{
+	UWorldPartition* WorldPartition = World->GetWorldPartition();
+	check(WorldPartition);
+
+	TMap<FGuid, FWorldPartitionReference> ActorReferences;
+	TArray<FGuid> ActorsToSave;
+
+	// Recursive loading of references 
+	TFunction<void(const FGuid&, TMap<FGuid, FWorldPartitionReference>&)> LoadReferences = [WorldPartition, &LoadReferences](const FGuid& ActorGuid, TMap<FGuid, FWorldPartitionReference>& InOutActorReferences)
+	{
+		if (InOutActorReferences.Contains(ActorGuid))
+		{
+			return;
+		}
+
+		if (FWorldPartitionActorDesc* ActorDesc = WorldPartition->GetActorDesc(ActorGuid))
+		{
+			for (FGuid ReferenceGuid : ActorDesc->GetReferences())
+			{
+				LoadReferences(ReferenceGuid, InOutActorReferences);
+			}
+
+			InOutActorReferences.Add(ActorGuid, FWorldPartitionReference(WorldPartition, ActorGuid));
+		}
+	};
+
+	// Actor Class Filter
+	UClass* ActorClass = nullptr;
+	// @todo_ow: support BP Classes when the ActorDesc have that information
+	if (!ActorClassName.IsEmpty())
+	{
+		ActorClass = FindObject<UClass>(ANY_PACKAGE, *ActorClassName);
+		if (!ActorClass)
+		{
+			UE_LOG(LogWorldPartitionResaveActorsBuilder, Error, TEXT("Failed to find Actor Class: %s."), *ActorClassName);
+			return false;
+		}
+	}
+		
+	UE_LOG(LogWorldPartitionResaveActorsBuilder, Display, TEXT("Gathering actors to save"));
+	// Get list of actors to save
+	for (UActorDescContainer::TIterator<AActor> ActorIterator(WorldPartition); ActorIterator; ++ActorIterator)
+	{
+		FWorldPartitionActorDesc* ActorDesc = *ActorIterator;	
+		if (!ActorClass || ActorDesc->GetActorClass()->IsChildOf(ActorClass))
+		{
+			ActorsToSave.Add(ActorDesc->GetGuid());
+		}
+	}
+	UE_LOG(LogWorldPartitionResaveActorsBuilder, Display, TEXT("Found %d actors to save"), ActorsToSave.Num());
+	FPackageSourceControlHelper SCCHelper;
+
+	int32 SaveCount = 0;
+	int32 FailCount = 0;
+
+	for (const FGuid& ActorGuid : ActorsToSave)
+	{
+		LoadReferences(ActorGuid, ActorReferences);
+		FWorldPartitionReference ActorReference(WorldPartition, ActorGuid);
+		AActor* Actor = ActorReference.Get()->GetActor();
+		if (UPackage* Package = Actor->GetExternalPackage())
+		{
+			if (SCCHelper.Checkout(Package))
+			{				
+				// Save package
+				FString PackageFileName = SourceControlHelpers::PackageFilename(Package);
+				if (!UPackage::SavePackage(Package, nullptr, RF_Standalone, *PackageFileName))
+				{
+					UE_LOG(LogWorldPartitionResaveActorsBuilder, Error, TEXT("Error saving package %s."), *Package->GetName());
+					++FailCount;
+					continue;
+				}
+
+				// It is possible the resave can't checkout everything. Continue processing.
+				UE_LOG(LogWorldPartitionResaveActorsBuilder, Display, TEXT("Saved package %s."), *Package->GetName());
+				++SaveCount;
+			}
+			else
+			{
+				// It is possible the resave can't checkout everything. Continue processing.
+				UE_LOG(LogWorldPartitionResaveActorsBuilder, Warning, TEXT("Error checking out package %s."), *Package->GetName());
+				++FailCount;
+				continue;
+			}
+
+			UE_LOG(LogWorldPartitionResaveActorsBuilder, Display, TEXT("Processed %d / %d packages (%d Saved / %d Failed)"), SaveCount + FailCount, ActorsToSave.Num(), SaveCount, FailCount);
+		}
+
+		if (HasExceededMaxMemory())
+		{
+			UE_LOG(LogWorldPartitionResaveActorsBuilder, Log, TEXT("Freeing some memory"));
+			ActorReferences.Empty();
+			DoCollectGarbage();
+		}
+	}
+
+	UPackage::WaitForAsyncFileWrites();
+
+	return true;
+}
+
diff --git a/Engine/Source/Editor/UnrealEd/Public/WorldPartition/WorldPartitionResaveActorsBuilder.h b/Engine/Source/Editor/UnrealEd/Public/WorldPartition/WorldPartitionResaveActorsBuilder.h
new file mode 100644
--- /dev/null
+++ b/Engine/Source/Editor/UnrealEd/Public/WorldPartition/WorldPartitionResaveActorsBuilder.h
@@ -0,0 +1,23 @@
+// Copyright Epic Games, Inc. All Rights Reserved.
+
+#pragma once
+
+#include "WorldPartition/WorldPartitionBuilder.h"
+#include "WorldPartitionResaveActorsBuilder.generated.h"
+
+// Example Command Line: ProjectName MapName -run=WorldPartitionBuilderCommandlet -SCCProvider=Perforce -Builder=WorldPartitionResaveActorsBuilder -ActorClass=PackedLevelInstance
+UCLASS()
+class UWorldPartitionResaveActorsBuilder : public UWorldPartitionBuilder
+{
+	GENERATED_UCLASS_BODY()
+
+public:
+	// UWorldPartitionBuilder interface begin
+	virtual bool RequiresCommandletRendering() const override { return false; }
+	virtual bool RequiresEntireWorldLoading() const override { return false; }
+	virtual bool Run(UWorld* World, FPackageSourceControlHelper& PackageHelper) override;
+	// UWorldPartitionBuilder interface end
+
+private:
+	FString ActorClassName;
+};
\ No newline at end of file
