Reject old "stable PSO" files gracefully (UE-108045)

diff --git a/Engine/Source/Editor/UnrealEd/Private/Commandlets/ShaderPipelineCacheToolsCommandlet.cpp b/Engine/Source/Editor/UnrealEd/Private/Commandlets/ShaderPipelineCacheToolsCommandlet.cpp
--- a/Engine/Source/Editor/UnrealEd/Private/Commandlets/ShaderPipelineCacheToolsCommandlet.cpp
+++ b/Engine/Source/Editor/UnrealEd/Private/Commandlets/ShaderPipelineCacheToolsCommandlet.cpp
@@ -22,7 +22,7 @@ const TCHAR* STABLE_CSV_EXT = TEXT("stablepc.csv");
 const TCHAR* STABLE_CSV_COMPRESSED_EXT = TEXT("stablepc.csv.compressed");
 const TCHAR* STABLE_COMPRESSED_EXT = TEXT(".compressed");
 const int32  STABLE_COMPRESSED_EXT_LEN = 11; // len of ".compressed";
-const int32  STABLE_COMPRESSED_VER = 2;
+const int32  STABLE_COMPRESSED_VER = 3;
 const int64  STABLE_MAX_CHUNK_SIZE = MAX_int32 - 100 * 1024 * 1024;
 
 struct FSCDataChunk
@@ -146,39 +146,43 @@ static bool LoadAndDecompressStableCSV(const FString& Filename, TArray<FString>&
 			int32 NumChunks = 1;
 
 			Ar->Serialize(&CompressedVersion, sizeof(int32));
-			if (CompressedVersion > 1)
+			if (CompressedVersion >= STABLE_COMPRESSED_VER)
 			{
 				Ar->Serialize(&NumChunks, sizeof(int32));
-			}
 
-			for (int32 Index = 0; Index < NumChunks; ++Index)
-			{
-				int32 UncompressedSize = 0;
-				int32 CompressedSize = 0;
+				for (int32 Index = 0; Index < NumChunks; ++Index)
+				{
+					int32 UncompressedSize = 0;
+					int32 CompressedSize = 0;
 
-				Ar->Serialize(&UncompressedSize, sizeof(int32));
-				Ar->Serialize(&CompressedSize, sizeof(int32));
+					Ar->Serialize(&UncompressedSize, sizeof(int32));
+					Ar->Serialize(&CompressedSize, sizeof(int32));
 
-				TArray<uint8> CompressedData;
-				CompressedData.SetNumUninitialized(CompressedSize);
-				Ar->Serialize(CompressedData.GetData(), CompressedSize);
+					TArray<uint8> CompressedData;
+					CompressedData.SetNumUninitialized(CompressedSize);
+					Ar->Serialize(CompressedData.GetData(), CompressedSize);
 
-				TArray<uint8> UncompressedData;
-				UncompressedData.SetNumUninitialized(UncompressedSize);
-				bResult = FCompression::UncompressMemory(NAME_Zlib, UncompressedData.GetData(), UncompressedSize, CompressedData.GetData(), CompressedSize);
-				if (!bResult)
-				{
-					UE_LOG(LogShaderPipelineCacheTools, Display, TEXT("Failed to decompress file %s"), *Filename);
-				}
+					TArray<uint8> UncompressedData;
+					UncompressedData.SetNumUninitialized(UncompressedSize);
+					bResult = FCompression::UncompressMemory(NAME_Zlib, UncompressedData.GetData(), UncompressedSize, CompressedData.GetData(), CompressedSize);
+					if (!bResult)
+					{
+						UE_LOG(LogShaderPipelineCacheTools, Display, TEXT("Failed to decompress file %s"), *Filename);
+					}
 
-				FMemoryReader MemArchive(UncompressedData);
-				FString LineCSV;
-				while (!MemArchive.AtEnd())
-				{
-					MemArchive << LineCSV;
-					OutputLines.Add(LineCSV);
+					FMemoryReader MemArchive(UncompressedData);
+					FString LineCSV;
+					while (!MemArchive.AtEnd())
+					{
+						MemArchive << LineCSV;
+						OutputLines.Add(LineCSV);
+					}
 				}
 			}
+			else
+			{
+				UE_LOG(LogShaderPipelineCacheTools, Warning, TEXT("File %s is too old (version %d, we need at least %d), rejecting."), *Filename, CompressedVersion, STABLE_COMPRESSED_VER);			
+			}
 		}
 		else
 		{
@@ -215,7 +219,7 @@ static bool LoadStableCSV(const FString& Filename, TArray<FString>& OutputLines)
 	}
 	else
 	{
-		bResult = FFileHelper::LoadFileToStringArray(OutputLines, *Filename);
+		UE_LOG(LogShaderPipelineCacheTools, Warning, TEXT("Uncompressed CSV files are no longer supported, rejecting %s."), *Filename);
 	}
 
 	return bResult;
@@ -278,28 +282,13 @@ static int64 SaveStableCSV(const FString& Filename, const FSCDataChunk* DataChun
 	}
 	else
 	{
-		if (NumChunks > 1)
-		{
-			UE_LOG(LogShaderPipelineCacheTools, Fatal, TEXT("SaveStableCSV does not support saving uncompressed files larger than 2GB."));
-		}
-
-		FMemoryReader MemArchive(DataChunks[0].UncomressedOutputLines);
-		FString CombinedCSV;
-		FString LineCSV;
-		while (!MemArchive.AtEnd())
-		{
-			MemArchive << LineCSV;
-			CombinedCSV.Append(LineCSV);
-			CombinedCSV.Append(LINE_TERMINATOR);
-		}
-
-		FFileHelper::SaveStringToFile(CombinedCSV, *Filename);
+		UE_LOG(LogShaderPipelineCacheTools, Warning, TEXT("SaveStableCSV does not support saving uncompressed files."));
 	}
 
 	int64 Size = IFileManager::Get().FileSize(*Filename);
 	if (Size < 1)
 	{
-		UE_LOG(LogShaderPipelineCacheTools, Fatal, TEXT("Failed to write %s"), *Filename);
+		UE_LOG(LogShaderPipelineCacheTools, Error, TEXT("Failed to write %s"), *Filename);
 	}
 
 	return Size;
@@ -1898,7 +1887,7 @@ int32 BuildPSOSC(const TArray<FString>& Tokens)
 		{
 			if (!LoadStableCSV(FileName, StableCSV))
 			{
-				UE_LOG(LogShaderPipelineCacheTools, Fatal, TEXT("Could not load %s"), *FileName);
+				UE_LOG(LogShaderPipelineCacheTools, Error, TEXT("Could not load %s"), *FileName);
 			}
 		}, TStatId()));
 	}
