Fix a few broken notifications

diff --git a/Engine/Source/Developer/MessageLog/Private/Presentation/MessageLogListingViewModel.cpp b/Engine/Source/Developer/MessageLog/Private/Presentation/MessageLogListingViewModel.cpp
--- a/Engine/Source/Developer/MessageLog/Private/Presentation/MessageLogListingViewModel.cpp
+++ b/Engine/Source/Developer/MessageLog/Private/Presentation/MessageLogListingViewModel.cpp
@@ -310,7 +310,6 @@ void FMessageLogListingViewModel::NotifyIfAnyMessages( const FText& Message, EMe
 		}
 
 		FNotificationInfo ErrorNotification(NotificationMessage);
-		ErrorNotification.Image = FEditorStyle::GetBrush(FTokenizedMessage::GetSeverityIconName(HighestSeverityPresent(0)));
 
 		if (NumMessagesPresent(0, EMessageSeverity::Error) > 0)
 		{
@@ -326,7 +325,7 @@ void FMessageLogListingViewModel::NotifyIfAnyMessages( const FText& Message, EMe
 			TSharedPtr<SNotificationItem> NewNotificationItem = FSlateNotificationManager::Get().AddNotification(ErrorNotification);
 			if (NewNotificationItem.IsValid())
 			{
-				NewNotificationItem->SetCompletionState(SNotificationItem::CS_Pending);
+				NewNotificationItem->SetCompletionState(SNotificationItem::CS_Fail);
 				OpenNotifications.Emplace(NotificationId, NewNotificationItem);
 			}
 		}
diff --git a/Engine/Source/Editor/UnrealEd/Private/AutoReimport/ReimportFeedbackContext.cpp b/Engine/Source/Editor/UnrealEd/Private/AutoReimport/ReimportFeedbackContext.cpp
--- a/Engine/Source/Editor/UnrealEd/Private/AutoReimport/ReimportFeedbackContext.cpp
+++ b/Engine/Source/Editor/UnrealEd/Private/AutoReimport/ReimportFeedbackContext.cpp
@@ -16,216 +16,48 @@
 #include "FileCacheUtilities.h"
 #include "MessageLogModule.h"
 #include "Widgets/Input/SHyperlink.h"
+#include "Widgets/Layout/SScrollBox.h"
 
 #define LOCTEXT_NAMESPACE "ReimportContext"
 
-class SWidgetStack : public SVerticalBox
+class SWidgetStack : public SCompoundWidget
 {
 	SLATE_BEGIN_ARGS(SWidgetStack){}
 	SLATE_END_ARGS()
 
-	void Construct(const FArguments& InArgs, int32 InMaxNumVisible)
-	{
-		MaxNumVisible = InMaxNumVisible;
-		SlideCurve = FCurveSequence(0.f, .5f, ECurveEaseFunction::QuadOut);
-		SizeCurve = FCurveSequence(0.f, .5f, ECurveEaseFunction::QuadOut);
-
-		StartSlideOffset = 0;
-		StartSizeOffset = FVector2D(ForceInitToZero);
-	}
-
-	FVector2D ComputeTotalSize() const
-	{
-		FVector2D Size(ForceInitToZero);
-		for (int32 Index = 0; Index < FMath::Min(NumSlots(), MaxNumVisible); ++Index)
-		{
-			const FVector2D& ChildSize = Children[Index].GetWidget()->GetDesiredSize();
-			if (ChildSize.X > Size.X)
-			{
-				Size.X = ChildSize.X;
-			}
-			Size.Y += ChildSize.Y + Children[Index].SlotPadding.Get().GetTotalSpaceAlong<Orient_Vertical>();
-		}
-		return Size;
-	}
-
-	virtual FVector2D ComputeDesiredSize(float) const override
-	{
-		const float Lerp = SizeCurve.GetLerp();
-		return ComputeTotalSize() * Lerp + StartSizeOffset * (1.f-Lerp);
-	}
-
-	virtual void OnArrangeChildren( const FGeometry& AllottedGeometry, FArrangedChildren& ArrangedChildren ) const override
+	void Construct(const FArguments& InArgs)
 	{
-		if (Children.Num() == 0)
-		{
-			return;
-		}
-
-		const float Alpha = 1.f - SlideCurve.GetLerp();
-		float PositionSoFar = AllottedGeometry.GetLocalSize().Y + StartSlideOffset*Alpha;
-
-		for (int32 Index = 0; Index < NumSlots(); ++Index)
-		{
-			const SBoxPanel::FSlot& CurChild = Children[Index];
-			const EVisibility ChildVisibility = CurChild.GetWidget()->GetVisibility();
-
-			if (ChildVisibility != EVisibility::Collapsed)
-			{
-				const FVector2D ChildDesiredSize = CurChild.GetWidget()->GetDesiredSize();
-
-				const FMargin SlotPadding(CurChild.SlotPadding.Get());
-				const FVector2D SlotSize(AllottedGeometry.Size.X, ChildDesiredSize.Y + SlotPadding.GetTotalSpaceAlong<Orient_Vertical>());
-
-				const AlignmentArrangeResult XAlignmentResult = AlignChild<Orient_Horizontal>( SlotSize.X, CurChild, SlotPadding );
-				const AlignmentArrangeResult YAlignmentResult = AlignChild<Orient_Vertical>( SlotSize.Y, CurChild, SlotPadding );
-
-				ArrangedChildren.AddWidget( ChildVisibility, AllottedGeometry.MakeChild(
-					CurChild.GetWidget(),
-					FVector2D( XAlignmentResult.Offset, PositionSoFar - SlotSize.Y + YAlignmentResult.Offset ),
-					FVector2D( XAlignmentResult.Size, YAlignmentResult.Size )
-					));
-
-				PositionSoFar -= SlotSize.Y;
-			}
-		}
+		ChildSlot
+		[
+			SNew(SBox)
+			.HeightOverride(45.0f)
+			[
+				SAssignNew(ScrollBox, SScrollBox)
+				.Style(FAppStyle::Get(), "ScrollBoxNoShadow")
+				.ScrollBarVisibility(EVisibility::Collapsed)
+				.AllowOverscroll(EAllowOverscroll::No)
+			]
+		];
 	}
 
 	void Add(const TSharedRef<SWidget>& InWidget)
 	{
-		TSharedPtr<SWidgetStackItem> NewItem;
-
-		InsertSlot(0)
-		.AutoHeight()
+		ScrollBox->AddSlot()
 		[
-			SAssignNew(NewItem, SWidgetStackItem)
-			[
-				InWidget
-			]
+			InWidget
 		];
-		
-		{
-			auto Widget = Children[0].GetWidget();
-			Widget->SlatePrepass();
-
-			const float WidgetHeight = Widget->GetDesiredSize().Y;
-			StartSlideOffset += WidgetHeight;
-			// Fade in time is 1 second x the proportion of the slide amount that this widget takes up
-			NewItem->FadeIn(WidgetHeight / StartSlideOffset);
-
-			if (!SlideCurve.IsPlaying())
-			{
-				SlideCurve.Play(AsShared());
-			}
-		}
 
-		const FVector2D NewSize = ComputeTotalSize();
-		if (NewSize != StartSizeOffset)
-		{
-			StartSizeOffset = NewSize;
+		++NumSlots;
 
-			if (!SizeCurve.IsPlaying())
-			{
-				SizeCurve.Play(AsShared());
-			}
-		}
+		ScrollBox->ScrollToEnd();
 	}
 	
-	virtual void Tick(const FGeometry& AllottedGeometry, const double InCurrentTime, const float InDeltaTime) override
-	{
-		if (!SlideCurve.IsPlaying())
-		{
-			StartSlideOffset = 0;
-		}
+	int32 GetNumSlots() const { return NumSlots; }
 
-		// Delete any widgets that are now offscreen
-		if (Children.Num() != 0)
-		{
-			const float Alpha = 1.f - SlideCurve.GetLerp();
-			float PositionSoFar = AllottedGeometry.GetLocalSize().Y + Alpha * StartSlideOffset;
-
-			for (int32 Index = 0; PositionSoFar > 0 && Index < NumSlots(); ++Index)
-			{
-				const SBoxPanel::FSlot& CurChild = Children[Index];
-				const EVisibility ChildVisibility = CurChild.GetWidget()->GetVisibility();
-
-				if (ChildVisibility != EVisibility::Collapsed)
-				{
-					const FVector2D ChildDesiredSize = CurChild.GetWidget()->GetDesiredSize();
-					PositionSoFar -= ChildDesiredSize.Y + CurChild.SlotPadding.Get().GetTotalSpaceAlong<Orient_Vertical>();
-				}
-			}
-
-			for (int32 Index = MaxNumVisible; Index < Children.Num(); )
-			{
-				if (StaticCastSharedRef<SWidgetStackItem>(Children[Index].GetWidget())->bIsFinished)
-				{
-					Children.RemoveAt(Index);
-				}
-				else
-				{
-					++Index;
-				}
-			}
-		}
-	}
-
-	class SWidgetStackItem : public SCompoundWidget
-	{
-		SLATE_BEGIN_ARGS(SWidgetStackItem){}
-			SLATE_DEFAULT_SLOT(FArguments, Content)
-		SLATE_END_ARGS()
-
-		void Construct(const FArguments& InArgs)
-		{
-			bIsFinished = false;
-
-			ChildSlot
-			[
-				SNew(SBorder)
-				.BorderImage(FEditorStyle::GetBrush("NoBorder"))
-				.ColorAndOpacity(this, &SWidgetStackItem::GetColorAndOpacity)
-				.Padding(0)
-				[
-					InArgs._Content.Widget
-				]
-			];
-		}
-		
-		void FadeIn(float Time)
-		{
-			OpacityCurve = FCurveSequence(0.f, Time, ECurveEaseFunction::QuadOut);
-			OpacityCurve.Play(AsShared());
-		}
-
-		virtual void Tick(const FGeometry& AllottedGeometry, const double InCurrentTime, const float InDeltaTime)
-		{
-			if (!bIsFinished && OpacityCurve.IsAtStart() && OpacityCurve.IsInReverse())
-			{
-				bIsFinished = true;
-			}
-		}
-
-		FLinearColor GetColorAndOpacity() const
-		{
-			if (OpacityCurve.IsPlaying() || OpacityCurve.IsAtEnd())
-			{
-				return FLinearColor(1.f, 1.f, 1.f, OpacityCurve.GetLerp());
-			}
+	TSharedPtr<SScrollBox> ScrollBox;
 
-			return FLinearColor(1.f, 1.f, 1.f, 1.f);
-		}
-
-		bool bIsFinished;
-		FCurveSequence OpacityCurve;
-	};
-
-	FCurveSequence SlideCurve, SizeCurve;
-
-	float StartSlideOffset;
-	FVector2D StartSizeOffset;
-
-	int32 MaxNumVisible;
+	int32 NumSlots = 0;
+	
 };
 
 class SWidgetStack;
@@ -270,88 +102,83 @@ public:
 
 		ChildSlot
 		[
-			SNew(SBorder)
-			.Padding(FMargin(10))
-			.BorderImage(FCoreStyle::Get().GetBrush("NotificationList.ItemBackground"))
+			SNew(SVerticalBox)
+
+			+ SVerticalBox::Slot()
+			.AutoHeight()
 			[
-				SNew(SVerticalBox)
+				SAssignNew(TopRow, SHorizontalBox)
 
-				+ SVerticalBox::Slot()
-				.AutoHeight()
+				+ SHorizontalBox::Slot()
+				.VAlign(VAlign_Center)
 				[
-					SAssignNew(TopRow, SHorizontalBox)
-
-					+ SHorizontalBox::Slot()
-					.VAlign(VAlign_Center)
-					[
-						SNew(STextBlock)
-						.Text(LOCTEXT("ProcessingChanges", "Processing source file changes..."))
-						.Font(FCoreStyle::Get().GetFontStyle(TEXT("NotificationList.FontLight")))
-					]
-
-					+ SHorizontalBox::Slot()
-					.AutoWidth()
-					.VAlign(VAlign_Center)
-					.Padding(FMargin(4,0,4,0))
-					[
-						SAssignNew(PauseButton, SButton)
-						.ButtonStyle(FEditorStyle::Get(), "HoverHintOnly")
-						.ToolTipText(LOCTEXT("PauseTooltip", "Temporarily pause processing of these source content files"))
-						.OnClicked(this, &SReimportFeedback::OnPauseClicked, InArgs._OnPauseClicked)
-						[
-							SNew(SImage)
-							.ColorAndOpacity(FLinearColor(.8f,.8f,.8f,1.f))
-							.Image(this, &SReimportFeedback::GetPlayPauseBrush)
-						]
-					]
+					SNew(STextBlock)
+					.Text(LOCTEXT("ProcessingChanges", "Processing source file changes..."))
+					.Font(FCoreStyle::Get().GetFontStyle(TEXT("NotificationList.FontLight")))
+				]
 
-					+ SHorizontalBox::Slot()
-					.AutoWidth()
-					.VAlign(VAlign_Center)
+				+ SHorizontalBox::Slot()
+				.AutoWidth()
+				.VAlign(VAlign_Center)
+				.Padding(FMargin(4,0,4,0))
+				[
+					SAssignNew(PauseButton, SButton)
+					.ButtonStyle(FEditorStyle::Get(), "HoverHintOnly")
+					.ToolTipText(LOCTEXT("PauseTooltip", "Temporarily pause processing of these source content files"))
+					.OnClicked(this, &SReimportFeedback::OnPauseClicked, InArgs._OnPauseClicked)
 					[
-						SAssignNew(AbortButton, SButton)
-						.ButtonStyle(FEditorStyle::Get(), "HoverHintOnly")
-						.ToolTipText(LOCTEXT("AbortTooltip", "Permanently abort processing of these source content files"))
-						.OnClicked(this, &SReimportFeedback::OnAbortClicked, InArgs._OnAbortClicked)
-						[
-							SNew(SImage)
-							.ColorAndOpacity(FLinearColor(0.8f,0.8f,0.8f,1.f))
-							.Image(FEditorStyle::GetBrush("GenericStop"))
-						]
+						SNew(SImage)
+						.ColorAndOpacity(FLinearColor(.8f,.8f,.8f,1.f))
+						.Image(this, &SReimportFeedback::GetPlayPauseBrush)
 					]
 				]
 
-				+ SVerticalBox::Slot()
-				.Padding(FMargin(0, 1))
-				.AutoHeight()
+				+ SHorizontalBox::Slot()
+				.AutoWidth()
+				.VAlign(VAlign_Center)
 				[
-					SNew(SBox)
-					.HeightOverride(2)
+					SAssignNew(AbortButton, SButton)
+					.ButtonStyle(FEditorStyle::Get(), "HoverHintOnly")
+					.ToolTipText(LOCTEXT("AbortTooltip", "Permanently abort processing of these source content files"))
+					.OnClicked(this, &SReimportFeedback::OnAbortClicked, InArgs._OnAbortClicked)
 					[
-						SAssignNew(ProgressBar, SProgressBar)
-						.BorderPadding(FVector2D::ZeroVector)
-						.Percent( this, &SReimportFeedback::GetProgressFraction )
+						SNew(SImage)
+						.ColorAndOpacity(FLinearColor(0.8f,0.8f,0.8f,1.f))
+						.Image(FEditorStyle::GetBrush("GenericStop"))
 					]
 				]
+			]
 
-				+ SVerticalBox::Slot()
-				.Padding(FMargin(0, 5, 0, 0))
-				.AutoHeight()
+			+ SVerticalBox::Slot()
+			.Padding(FMargin(0, 1))
+			.AutoHeight()
+			[
+				SNew(SBox)
+				.HeightOverride(2)
 				[
-					SAssignNew(WidgetStack, SWidgetStack, 3)
+					SAssignNew(ProgressBar, SProgressBar)
+					.BorderPadding(FVector2D::ZeroVector)
+					.Percent( this, &SReimportFeedback::GetProgressFraction )
 				]
+			]
 
-				+ SVerticalBox::Slot()
-				.Padding(FMargin(0, 5, 0, 0))
-				.AutoHeight()
-				.HAlign(HAlign_Right)
-				[
-					SNew(SHyperlink)
-					.Visibility(this, &SReimportFeedback::GetHyperlinkVisibility)
-					.Text(LOCTEXT("OpenMessageLog", "Open message log"))
-					.TextStyle(FCoreStyle::Get(), "SmallText")
-					.OnNavigate_Lambda(OpenMessageLog)
-				]
+			+ SVerticalBox::Slot()
+			.Padding(FMargin(0, 5, 0, 0))
+			.AutoHeight()
+			[
+				SAssignNew(WidgetStack, SWidgetStack)
+			]
+
+			+ SVerticalBox::Slot()
+			.Padding(FMargin(0, 5, 0, 0))
+			.AutoHeight()
+			.HAlign(HAlign_Right)
+			[
+				SNew(SHyperlink)
+				.Visibility(this, &SReimportFeedback::GetHyperlinkVisibility)
+				.Text(LOCTEXT("OpenMessageLog", "Open message log"))
+				.TextStyle(FCoreStyle::Get(), "SmallText")
+				.OnNavigate_Lambda(OpenMessageLog)
 			]
 		];
 	}
@@ -440,7 +267,7 @@ private:
 	/** Get the visibility of the hyperlink to open the message log */
 	EVisibility GetHyperlinkVisibility() const
 	{
-		return WidgetStack->NumSlots() != 0 ? EVisibility::Visible : EVisibility::Collapsed;
+		return WidgetStack->GetNumSlots() != 0 ? EVisibility::Visible : EVisibility::Collapsed;
 	}
 
 private:
diff --git a/Engine/Source/Editor/UnrealEd/Private/AutoReimport/ReimportFeedbackContext.h b/Engine/Source/Editor/UnrealEd/Private/AutoReimport/ReimportFeedbackContext.h
--- a/Engine/Source/Editor/UnrealEd/Private/AutoReimport/ReimportFeedbackContext.h
+++ b/Engine/Source/Editor/UnrealEd/Private/AutoReimport/ReimportFeedbackContext.h
@@ -42,6 +42,7 @@ public:
 	virtual void StartSlowTask(const FText& Task, bool bShowCancelButton=false) override;
 	virtual void ProgressReported(const float TotalProgressInterp, FText DisplayMessage) override;
 	virtual void FinalizeSlowTask() override;
+	virtual bool UseNotificationBackground() const override { return true; }
 
 	/** True if we will suppress slow task messages from being shown on the UI, false if otherwise */
 	bool bSuppressSlowTaskMessages;
diff --git a/Engine/Source/Runtime/SlateCore/Private/Styling/StarshipCoreStyle.cpp b/Engine/Source/Runtime/SlateCore/Private/Styling/StarshipCoreStyle.cpp
--- a/Engine/Source/Runtime/SlateCore/Private/Styling/StarshipCoreStyle.cpp
+++ b/Engine/Source/Runtime/SlateCore/Private/Styling/StarshipCoreStyle.cpp
@@ -656,6 +656,14 @@ TSharedRef<ISlateStyle> FStarshipCoreStyle::Create()
 			.SetRightShadowBrush(BOX_BRUSH("Common/ScrollBoxShadowRight", FVector2D(8.f, 16.f), FMargin(0.f, 0.5f, 1.f, 0.5f)))
 			.SetBarThickness(8.0)
 		);
+
+		Style->Set("ScrollBoxNoShadow", FScrollBoxStyle()
+			.SetTopShadowBrush(FSlateNoResource())
+			.SetBottomShadowBrush(FSlateNoResource())
+			.SetLeftShadowBrush(FSlateNoResource())
+			.SetRightShadowBrush(FSlateNoResource())
+			.SetBarThickness(8.0)
+		);
 	}
 
 	// SScrollBorder defaults...
