COTFS: Always append ExternalActors to NeverCookDirectories for Cook By [REDACTED_PERSON].
Prevents an ensure while cooking from the editor.

diff --git a/Engine/Source/Editor/UnrealEd/Private/CookOnTheFlyServer.cpp b/Engine/Source/Editor/UnrealEd/Private/CookOnTheFlyServer.cpp
--- a/Engine/Source/Editor/UnrealEd/Private/CookOnTheFlyServer.cpp
+++ b/Engine/Source/Editor/UnrealEd/Private/CookOnTheFlyServer.cpp
@@ -7655,6 +7655,8 @@ TArray<FName> UCookOnTheFlyServer::GetNeverCookPackageFileNames(TArrayView<const
 	const UProjectPackagingSettings* const PackagingSettings = GetDefault<UProjectPackagingSettings>();
 	AddDirectoryPathArray(PackagingSettings->DirectoriesToNeverCook, TEXT("ProjectSettings -> Project -> Packaging -> Directories to never cook"));
 	AddDirectoryPathArray(PackagingSettings->TestDirectoriesToNotSearch, TEXT("ProjectSettings -> Project -> Packaging -> Test directories to not search"));
+	// Never cook External Actors
+	NeverCookDirectories.Add(FString::Printf(TEXT("/Game/%s"), ULevel::GetExternalActorsFolderName()));
 
 	TArray<FString> NeverCookPackagesPaths;
 	FPackageName::FindPackagesInDirectories(NeverCookPackagesPaths, NeverCookDirectories);
