Fixed a crash with stat enable/disable in editor viewports.

diff --git a/Engine/Source/Editor/UnrealEd/Private/EditorViewportClient.cpp b/Engine/Source/Editor/UnrealEd/Private/EditorViewportClient.cpp
--- a/Engine/Source/Editor/UnrealEd/Private/EditorViewportClient.cpp
+++ b/Engine/Source/Editor/UnrealEd/Private/EditorViewportClient.cpp
@@ -6043,7 +6043,14 @@ const TArray<FString>* FEditorViewportClient::GetEnabledStats() const
 
 void FEditorViewportClient::SetEnabledStats(const TArray<FString>& InEnabledStats)
 {
+	HandleViewportStatDisableAll(true);
+
 	EnabledStats = InEnabledStats;
+	if (EnabledStats.Num())
+	{
+		SetShowStats(true);
+		AddRealtimeOverride(true, LOCTEXT("RealtimeOverrideMessage_Stats", "Stats Display"));
+	}
 
 #if ENABLE_AUDIO_DEBUG
 	if (GEngine)
