Fixed an ensure from the selection changing while applying gizmo manipulation

Construction scripts can trigger this, so we now defer actually clearing the cached list of elements to manipulate until the next cache request, or until the typed element registry is about to process pending deletions.

diff --git a/Engine/Source/Editor/UnrealEd/Private/LevelEditorViewport.cpp b/Engine/Source/Editor/UnrealEd/Private/LevelEditorViewport.cpp
--- a/Engine/Source/Editor/UnrealEd/Private/LevelEditorViewport.cpp
+++ b/Engine/Source/Editor/UnrealEd/Private/LevelEditorViewport.cpp
@@ -1773,6 +1773,10 @@ FLevelEditorViewportClient::~FLevelEditorViewportClient()
 	{
 		SelectionSet->OnPreChange().RemoveAll(this);
 	}
+	if (UTypedElementRegistry* Registry = UTypedElementRegistry::GetInstance())
+	{
+		Registry->OnProcessingDeferredElementsToDestroy().RemoveAll(this);
+	}
 	ResetElementsToManipulate();
 
 	// Unregister for all global callbacks to this object
@@ -1833,6 +1837,9 @@ void FLevelEditorViewportClient::InitializeViewportInteraction()
 	{
 		UTypedElementSelectionSet* SelectionSet = GetMutableSelectionSet();
 		SelectionSet->OnPreChange().AddRaw(this, &FLevelEditorViewportClient::ResetElementsToManipulateFromSelectionChange);
+
+		UTypedElementRegistry* Registry = UTypedElementRegistry::GetInstance();
+		Registry->OnProcessingDeferredElementsToDestroy().AddRaw(this, &FLevelEditorViewportClient::ResetElementsToManipulateFromProcessingDeferredElementsToDestroy);
 	}
 
 	{
@@ -3515,19 +3522,31 @@ void FLevelEditorViewportClient::CacheElementsToManipulate(const bool bForceRefr
 	}
 }
 
-void FLevelEditorViewportClient::ResetElementsToManipulate()
+void FLevelEditorViewportClient::ResetElementsToManipulate(const bool bClearList)
 {
-	if (bHasCachedElementsToManipulate)
+	if (bClearList)
 	{
 		CachedElementsToManipulate->Reset();
-		bHasCachedElementsToManipulate = false;
 	}
+	bHasCachedElementsToManipulate = false;
 }
 
 void FLevelEditorViewportClient::ResetElementsToManipulateFromSelectionChange(const UTypedElementSelectionSet* InSelectionSet)
 {
 	check(InSelectionSet == GetSelectionSet());
-	ResetElementsToManipulate();
+
+	// Don't clear the list immediately, as the selection may change from a construction script running (while we're still iterating the list!)
+	// We'll process the clear on the next cache request, or when the typed element registry actually processes its pending deletion
+	ResetElementsToManipulate(/*bClearList*/false);
+}
+
+void FLevelEditorViewportClient::ResetElementsToManipulateFromProcessingDeferredElementsToDestroy()
+{
+	if (!bHasCachedElementsToManipulate)
+	{
+		// If we have no cache, make sure the cached list is definitely empty now to ensure it doesn't contain any lingering references to things that are about to be deleted
+		CachedElementsToManipulate->Reset();
+	}
 }
 
 const UTypedElementSelectionSet* FLevelEditorViewportClient::GetSelectionSet() const
diff --git a/Engine/Source/Editor/UnrealEd/Public/LevelEditorViewport.h b/Engine/Source/Editor/UnrealEd/Public/LevelEditorViewport.h
--- a/Engine/Source/Editor/UnrealEd/Public/LevelEditorViewport.h
+++ b/Engine/Source/Editor/UnrealEd/Public/LevelEditorViewport.h
@@ -718,11 +718,14 @@ private:
 	void CacheElementsToManipulate(const bool bForceRefresh = false);
 
 	/** Reset the list of elements to manipulate */
-	void ResetElementsToManipulate();
+	void ResetElementsToManipulate(const bool bClearList = true);
 
 	/** Reset the list of elements to manipulate, because the selection set they were cached from has changed */
 	void ResetElementsToManipulateFromSelectionChange(const UTypedElementSelectionSet* InSelectionSet);
 
+	/** Reset the list of elements to manipulate, because the typed element registry is about to process deferred deletion */
+	void ResetElementsToManipulateFromProcessingDeferredElementsToDestroy();
+
 	/** Get the selection set that associated with our level editor. */
 	const UTypedElementSelectionSet* GetSelectionSet() const;
 	UTypedElementSelectionSet* GetMutableSelectionSet() const;
diff --git a/Engine/Source/Runtime/TypedElementFramework/Private/Elements/Framework/TypedElementRegistry.cpp b/Engine/Source/Runtime/TypedElementFramework/Private/Elements/Framework/TypedElementRegistry.cpp
--- a/Engine/Source/Runtime/TypedElementFramework/Private/Elements/Framework/TypedElementRegistry.cpp
+++ b/Engine/Source/Runtime/TypedElementFramework/Private/Elements/Framework/TypedElementRegistry.cpp
@@ -142,6 +142,8 @@ UTypedElementInterface* UTypedElementRegistry::GetElementInterfaceImpl(const FTy
 
 void UTypedElementRegistry::ProcessDeferredElementsToDestroy()
 {
+	OnProcessingDeferredElementsToDestroyDelegate.Broadcast();
+
 	FReadScopeLock RegisteredElementTypesLock(RegisteredElementTypesRW);
 
 	for (TUniquePtr<FRegisteredElementType>& RegisteredElementType : RegisteredElementTypes)
diff --git a/Engine/Source/Runtime/TypedElementFramework/Public/Elements/Framework/TypedElementRegistry.h b/Engine/Source/Runtime/TypedElementFramework/Public/Elements/Framework/TypedElementRegistry.h
--- a/Engine/Source/Runtime/TypedElementFramework/Public/Elements/Framework/TypedElementRegistry.h
+++ b/Engine/Source/Runtime/TypedElementFramework/Public/Elements/Framework/TypedElementRegistry.h
@@ -60,6 +60,15 @@ public:
 		return OnElementUpdatedDelegate;
 	}
 
+	/**
+	 * Event fired prior to processing any elements that were previously marked for deferred destruction.
+	 * @see ProcessDeferredElementsToDestroy.
+	 */
+	FSimpleMulticastDelegate& OnProcessingDeferredElementsToDestroy()
+	{
+		return OnProcessingDeferredElementsToDestroyDelegate;
+	}
+
 	/**
 	 * Get the element type ID for the associated element type name, if any.
 	 * @return The element type ID, or 0 if the given name wasn't registered.
@@ -509,4 +518,5 @@ private:
 
 	FOnElementReplaced OnElementReplacedDelegate;
 	FOnElementUpdated OnElementUpdatedDelegate;
+	FSimpleMulticastDelegate OnProcessingDeferredElementsToDestroyDelegate;
 };
