Distributed HLOD builds - Ensure we don't hold on to packages of always loaded actors
When running distributed builds, we wanna leave the machine clean, so added files are deleted, check'd out files are reverted and deleted files are restored.

diff --git a/Engine/Source/Editor/UnrealEd/Private/WorldPartition/WorldPartitionHLODsBuilder.cpp b/Engine/Source/Editor/UnrealEd/Private/WorldPartition/WorldPartitionHLODsBuilder.cpp
--- a/Engine/Source/Editor/UnrealEd/Private/WorldPartition/WorldPartitionHLODsBuilder.cpp
+++ b/Engine/Source/Editor/UnrealEd/Private/WorldPartition/WorldPartitionHLODsBuilder.cpp
@@ -309,6 +309,12 @@ bool UWorldPartitionHLODsBuilder::SetupHLODActors(bool bCreateOnly)
 			// When performing a distributed build, move modified files to the temporary working dir, to be submitted later in the last "submit" step
 			if (IsDistributedBuild())
 			{
+				// Ensure we don't hold on to packages of always loaded actors
+				// When running distributed builds, we wanna leave the machine clean, so added files are deleted, check'd out files are reverted
+				// and deleted files are restored.
+				WorldPartition->Uninitialize();
+				DoCollectGarbage();
+
 				ModifiedFiles.Append(SourceControlHelper->GetModifiedFiles());
 
 				TArray<TArray<FString>> BuildersFiles;
@@ -410,6 +416,12 @@ bool UWorldPartitionHLODsBuilder::BuildHLODActors()
 	// Move modified files to the temporary working dir, to be submitted later in the final "submit" pass, from a single machine.
 	if (IsDistributedBuild())
 	{
+		// Ensure we don't hold on to packages of always loaded actors
+		// When running distributed builds, we wanna leave the machine clean, so added files are deleted, check'd out files are reverted
+		// and deleted files are restored.
+		WorldPartition->Uninitialize();
+		DoCollectGarbage();
+
 		ModifiedFiles.Append(SourceControlHelper->GetModifiedFiles());
 
 		TArray<FString> BuildProducts;
