When importing / building Nanite static meshes, always disable removal of degenerate tris.  For clarity, this option is also disabled in the import UI and the Build Settings UI in the Static Mesh Editor, when Nanite is enabled.

diff --git a/Engine/Source/Developer/MeshBuilder/Private/MeshDescriptionHelper.cpp b/Engine/Source/Developer/MeshBuilder/Private/MeshDescriptionHelper.cpp
--- a/Engine/Source/Developer/MeshBuilder/Private/MeshDescriptionHelper.cpp
+++ b/Engine/Source/Developer/MeshBuilder/Private/MeshDescriptionHelper.cpp
@@ -37,7 +37,8 @@ void FMeshDescriptionHelper::SetupRenderMeshDescription(UObject* Owner, FMeshDes
 	UStaticMesh* StaticMesh = Cast<UStaticMesh>(Owner);
 	check(StaticMesh);
 
-	float ComparisonThreshold = BuildSettings->bRemoveDegenerates ? THRESH_POINTS_ARE_SAME : 0.0f;
+	const bool bNaniteBuildEnabled = StaticMesh->NaniteSettings.bEnabled;
+	float ComparisonThreshold = (BuildSettings->bRemoveDegenerates && !bNaniteBuildEnabled) ? THRESH_POINTS_ARE_SAME : 0.0f;
 	
 	//This function make sure the Polygon Normals Tangents Binormals are computed and also remove degenerated triangle from the render mesh description.
 	FStaticMeshOperations::ComputeTriangleTangentsAndNormals(RenderMeshDescription, ComparisonThreshold);
@@ -55,12 +56,17 @@ void FMeshDescriptionHelper::SetupRenderMeshDescription(UObject* Owner, FMeshDes
 
 	// Static meshes always blend normals of overlapping corners.
 	EComputeNTBsFlags ComputeNTBsOptions = EComputeNTBsFlags::BlendOverlappingNormals;
-	ComputeNTBsOptions |= BuildSettings->bRemoveDegenerates ? EComputeNTBsFlags::IgnoreDegenerateTriangles : EComputeNTBsFlags::None;
 	ComputeNTBsOptions |= BuildSettings->bComputeWeightedNormals ? EComputeNTBsFlags::WeightedNTBs : EComputeNTBsFlags::None;
 	ComputeNTBsOptions |= BuildSettings->bRecomputeNormals ? EComputeNTBsFlags::Normals : EComputeNTBsFlags::None;
-	ComputeNTBsOptions |= BuildSettings->bRecomputeTangents && !StaticMesh->NaniteSettings.bEnabled ? EComputeNTBsFlags::Tangents : EComputeNTBsFlags::None;	// Nanite doesn't need tangents
 	ComputeNTBsOptions |= BuildSettings->bUseMikkTSpace ? EComputeNTBsFlags::UseMikkTSpace : EComputeNTBsFlags::None;
 
+	// Set extra options for non-Nanite meshes
+	if (!bNaniteBuildEnabled)
+	{
+		ComputeNTBsOptions |= BuildSettings->bRemoveDegenerates ? EComputeNTBsFlags::IgnoreDegenerateTriangles : EComputeNTBsFlags::None;
+		ComputeNTBsOptions |= BuildSettings->bRecomputeTangents ? EComputeNTBsFlags::Tangents : EComputeNTBsFlags::None;
+	}
+
 	// Compute any missing normals or tangents.
 	FStaticMeshOperations::ComputeTangentsAndNormals(RenderMeshDescription, ComputeNTBsOptions);
 
diff --git a/Engine/Source/Editor/StaticMeshEditor/Private/StaticMeshEditorTools.cpp b/Engine/Source/Editor/StaticMeshEditor/Private/StaticMeshEditorTools.cpp
--- a/Engine/Source/Editor/StaticMeshEditor/Private/StaticMeshEditorTools.cpp
+++ b/Engine/Source/Editor/StaticMeshEditor/Private/StaticMeshEditorTools.cpp
@@ -561,6 +561,7 @@ void FMeshBuildSettingsLayout::GenerateChildContent( IDetailChildrenBuilder& Chi
 		.ValueContent()
 		[
 			SNew(SCheckBox)
+			.IsEnabled(this, &FMeshBuildSettingsLayout::IsRemoveDegeneratesDisabled)
 			.IsChecked(this, &FMeshBuildSettingsLayout::ShouldRemoveDegenerates)
 			.OnCheckStateChanged(this, &FMeshBuildSettingsLayout::OnRemoveDegeneratesChanged)
 		];
@@ -850,6 +851,16 @@ ECheckBoxState FMeshBuildSettingsLayout::ShouldGenerateDistanceFieldAsIfTwoSided
 	return BuildSettings.bGenerateDistanceFieldAsIfTwoSided ? ECheckBoxState::Checked : ECheckBoxState::Unchecked;
 }
 
+bool FMeshBuildSettingsLayout::IsRemoveDegeneratesDisabled() const
+{
+	if (TSharedPtr<FLevelOfDetailSettingsLayout> LODSettingsLayout = ParentLODSettings.Pin())
+	{
+		return !LODSettingsLayout->IsNaniteEnabled();
+	}
+
+	return true;
+}
+
 int32 FMeshBuildSettingsLayout::GetMinLightmapResolution() const
 {
 	return BuildSettings.MinLightmapResolution;
@@ -4524,6 +4535,13 @@ FText FLevelOfDetailSettingsLayout::GetCurrentLodTooltip() const
 	return FText::GetEmpty();
 }
 
+bool FLevelOfDetailSettingsLayout::IsNaniteEnabled() const
+{
+	UStaticMesh* StaticMesh = StaticMeshEditor.GetStaticMesh();
+	check(StaticMesh != nullptr);
+	return StaticMesh->NaniteSettings.bEnabled;
+}
+
 /////////////////////////////////
 // FNaniteSettingsLayout
 /////////////////////////////////
diff --git a/Engine/Source/Editor/StaticMeshEditor/Private/StaticMeshEditorTools.h b/Engine/Source/Editor/StaticMeshEditor/Private/StaticMeshEditorTools.h
--- a/Engine/Source/Editor/StaticMeshEditor/Private/StaticMeshEditorTools.h
+++ b/Engine/Source/Editor/StaticMeshEditor/Private/StaticMeshEditorTools.h
@@ -197,6 +197,7 @@ private:
 	ECheckBoxState ShouldUseFullPrecisionUVs() const;
 	ECheckBoxState ShouldGenerateLightmapUVs() const;
 	ECheckBoxState ShouldGenerateDistanceFieldAsIfTwoSided() const;
+	bool IsRemoveDegeneratesDisabled() const;
 	int32 GetMinLightmapResolution() const;
 	int32 GetSrcLightmapIndex() const;
 	int32 GetDstLightmapIndex() const;
@@ -524,6 +525,9 @@ public:
 	/** Apply current LOD settings to the mesh. */
 	void ApplyChanges();
 
+	/** Returns true if the LOD's static mesh has Nanite enabled */
+	bool IsNaniteEnabled() const;
+
 private:
 
 	/** Creates the UI for Current LOD panel */
diff --git a/Engine/Source/Editor/UnrealEd/Classes/Factories/FbxStaticMeshImportData.h b/Engine/Source/Editor/UnrealEd/Classes/Factories/FbxStaticMeshImportData.h
--- a/Engine/Source/Editor/UnrealEd/Classes/Factories/FbxStaticMeshImportData.h
+++ b/Engine/Source/Editor/UnrealEd/Classes/Factories/FbxStaticMeshImportData.h
@@ -31,7 +31,7 @@ class UFbxStaticMeshImportData : public UFbxMeshImportData
 	FColor VertexOverrideColor;
 
 	/** Disabling this option will keep degenerate triangles found.  In general you should leave this option on. */
-	UPROPERTY(EditAnywhere, BlueprintReadWrite, AdvancedDisplay, config, Category = Mesh, meta = (ImportType = "StaticMesh", ReimportRestrict = "true"))
+	UPROPERTY(EditAnywhere, BlueprintReadWrite, AdvancedDisplay, config, Category = Mesh, meta = (ImportType = "StaticMesh", ReimportRestrict = "true", EditCondition = "!bBuildNanite"))
 	uint32 bRemoveDegenerates:1;
 
 	UPROPERTY(EditAnywhere, BlueprintReadWrite, AdvancedDisplay, config, Category = Mesh, meta = (ImportType = "StaticMesh", ReimportRestrict = "true"))
diff --git a/Engine/Source/Editor/UnrealEd/Private/Fbx/FbxStaticMeshImport.cpp b/Engine/Source/Editor/UnrealEd/Private/Fbx/FbxStaticMeshImport.cpp
--- a/Engine/Source/Editor/UnrealEd/Private/Fbx/FbxStaticMeshImport.cpp
+++ b/Engine/Source/Editor/UnrealEd/Private/Fbx/FbxStaticMeshImport.cpp
@@ -360,6 +360,17 @@ struct FFBXUVs
 };
 
 
+float UnFbx::FFbxImporter::GetPointComparisonThreshold() const
+{
+	return (ImportOptions->bRemoveDegenerates && !ImportOptions->bBuildNanite) ? THRESH_POINTS_ARE_SAME : 0.0f;
+}
+
+float UnFbx::FFbxImporter::GetTriangleAreaThreshold() const
+{
+	return (ImportOptions->bRemoveDegenerates && !ImportOptions->bBuildNanite) ? SMALL_NUMBER : 0.0f;
+}
+
+
 bool UnFbx::FFbxImporter::BuildStaticMeshFromGeometry(FbxNode* Node, UStaticMesh* StaticMesh, TArray<FFbxMaterial>& MeshMaterials, int32 LODIndex,
 	EVertexColorImportOption::Type VertexColorImportOption, const TMap<FVector, FColor>& ExistingVertexColorData, const FColor& VertexOverrideColor)
 {
@@ -811,7 +822,7 @@ bool UnFbx::FFbxImporter::BuildStaticMeshFromGeometry(FbxNode* Node, UStaticMesh
 				int32 PolygonVertexCount = Mesh->GetPolygonSize(PolygonIndex);
 				//Verify if the polygon is degenerate, in this case do not add them
 				{
-					float ComparisonThreshold = ImportOptions->bRemoveDegenerates ? SMALL_NUMBER : 0.0f;
+					float ComparisonThreshold = GetTriangleAreaThreshold();
 					P.Reset();
 					P.AddUninitialized(PolygonVertexCount);
 					for (int32 CornerIndex = 0; CornerIndex < PolygonVertexCount; CornerIndex++)
@@ -956,14 +967,14 @@ bool UnFbx::FFbxImporter::BuildStaticMeshFromGeometry(FbxNode* Node, UStaticMesh
 				//TODO check all polygon vertex, not just the first 3 vertex
 				if (!bHasNonDegeneratePolygons)
 				{
-					float TriangleComparisonThreshold = ImportOptions->bRemoveDegenerates ? THRESH_POINTS_ARE_SAME : 0.0f;
+					float PointComparisonThreshold = GetPointComparisonThreshold();
 					FVector VertexPosition[3];
 					VertexPosition[0] = VertexPositions[CornerVerticesIDs[0]];
 					VertexPosition[1] = VertexPositions[CornerVerticesIDs[1]];
 					VertexPosition[2] = VertexPositions[CornerVerticesIDs[2]];
-					if (!( VertexPosition[0].Equals(VertexPosition[1], TriangleComparisonThreshold)
-						|| VertexPosition[0].Equals(VertexPosition[2], TriangleComparisonThreshold)
-						|| VertexPosition[1].Equals(VertexPosition[2], TriangleComparisonThreshold)))
+					if (!( VertexPosition[0].Equals(VertexPosition[1], PointComparisonThreshold)
+						|| VertexPosition[0].Equals(VertexPosition[2], PointComparisonThreshold)
+						|| VertexPosition[1].Equals(VertexPosition[2], PointComparisonThreshold)))
 					{
 						bHasNonDegeneratePolygons = true;
 					}
