Added missing file

diff --git a/Engine/Source/Editor/UnrealEd/Public/FbxImporter.h b/Engine/Source/Editor/UnrealEd/Public/FbxImporter.h
--- a/Engine/Source/Editor/UnrealEd/Public/FbxImporter.h
+++ b/Engine/Source/Editor/UnrealEd/Public/FbxImporter.h
@@ -1210,6 +1210,16 @@ private:
 	*/
 	void VerifyGeometry(UStaticMesh* StaticMesh);
 
+	/**
+	 * Get the epsilon value, according to the import settings, below which two vertex positions should be considered coincident.
+	 */
+	float GetPointComparisonThreshold() const;
+
+	/**
+	 * Get the epsilon value, according to the import settings, below which a triangle should be considered to have "zero area" and therefore be degenerate.
+	 */
+	float GetTriangleAreaThreshold() const;
+
 	/**
 	* When there is some materials with the same name we add a clash suffixe _ncl1_x.
 	* Example, if we have 3 materials name shader we will get (shader, shader_ncl1_1, shader_ncl1_2).
