WorldPartition builder cmdlet - Added missing call to CleanupWorld()

diff --git a/Engine/Source/Editor/UnrealEd/Private/Commandlets/WorldPartitionBuilderCommandlet.cpp b/Engine/Source/Editor/UnrealEd/Private/Commandlets/WorldPartitionBuilderCommandlet.cpp
--- a/Engine/Source/Editor/UnrealEd/Private/Commandlets/WorldPartitionBuilderCommandlet.cpp
+++ b/Engine/Source/Editor/UnrealEd/Private/Commandlets/WorldPartitionBuilderCommandlet.cpp
@@ -188,6 +188,9 @@ int32 UWorldPartitionBuilderCommandlet::Main(const FString& Params)
 
 	// Cleanup
 	World->RemoveFromRoot();
+	World->ClearWorldComponents();
+	World->CleanupWorld();
+
 	WorldContext.SetCurrentWorld(nullptr);
 	GWorld = nullptr;
 
