Fixed ResavePackagesCommanlet to perform proper world cleanup and restore previously assigned GWorld

diff --git a/Engine/Source/Editor/UnrealEd/Private/Commandlets/ContentCommandlets.cpp b/Engine/Source/Editor/UnrealEd/Private/Commandlets/ContentCommandlets.cpp
--- a/Engine/Source/Editor/UnrealEd/Private/Commandlets/ContentCommandlets.cpp
+++ b/Engine/Source/Editor/UnrealEd/Private/Commandlets/ContentCommandlets.cpp
@@ -1680,6 +1680,7 @@ void UResavePackagesCommandlet::PerformAdditionalOperations(class UWorld* World,
 	}
 	FWorldContext& WorldContext = GEditor->GetEditorWorldContext(true);
 	WorldContext.SetCurrentWorld(World);
+	UWorld* PrevGWorld = GWorld;
 	GWorld = World;
 
 	// Load and Save world partition actor packages
@@ -2069,9 +2070,14 @@ void UResavePackagesCommandlet::PerformAdditionalOperations(class UWorld* World,
 		}
 	}
 
+	// Cleanup
 	World->RemoveFromRoot();
-	WorldContext.SetCurrentWorld(nullptr);
-	GWorld = nullptr;
+	World->ClearWorldComponents();
+	World->CleanupWorld();
+	
+	// Restore previous world
+	WorldContext.SetCurrentWorld(PrevGWorld);
+	GWorld = PrevGWorld;
 }
 
 void UResavePackagesCommandlet::PerformAdditionalOperations( class UObject* Object, bool& bSavePackage )
