Fixed WorldPartitionBuilderCommandlet to perform proper cleanup even in case of failure
Also restore previous GWorld

diff --git a/Engine/Source/Editor/UnrealEd/Private/Commandlets/WorldPartitionBuilderCommandlet.cpp b/Engine/Source/Editor/UnrealEd/Private/Commandlets/WorldPartitionBuilderCommandlet.cpp
--- a/Engine/Source/Editor/UnrealEd/Private/Commandlets/WorldPartitionBuilderCommandlet.cpp
+++ b/Engine/Source/Editor/UnrealEd/Private/Commandlets/WorldPartitionBuilderCommandlet.cpp
@@ -163,6 +163,7 @@ int32 UWorldPartitionBuilderCommandlet::Main(const FString& Params)
 
 	FWorldContext& WorldContext = GEditor->GetEditorWorldContext(true /*bEnsureIsGWorld*/);
 	WorldContext.SetCurrentWorld(World);
+	UWorld* PrevGWorld = GWorld;
 	GWorld = World;
 
 	if (Builder->RequiresEntireWorldLoading())
@@ -172,18 +173,18 @@ int32 UWorldPartitionBuilderCommandlet::Main(const FString& Params)
 	}
 
 	// Run builder
-	if (!Builder->Run(World, *this))
-	{
-		return 1;
-	}
-
+	bool bResult = Builder->Run(World, *this);
+	
 	// Save default configuration
-	if (!FPlatformFileManager::Get().GetPlatformFile().FileExists(*WorldConfigFilename) ||
-		!FPlatformFileManager::Get().GetPlatformFile().IsReadOnly(*WorldConfigFilename))
+	if (bResult)
 	{
-		SaveConfig(CPF_Config, *WorldConfigFilename);
-
-		Builder->SaveConfig(CPF_Config, *WorldConfigFilename);
+	    if (!FPlatformFileManager::Get().GetPlatformFile().FileExists(*WorldConfigFilename) ||
+		    !FPlatformFileManager::Get().GetPlatformFile().IsReadOnly(*WorldConfigFilename))
+	    {
+		    SaveConfig(CPF_Config, *WorldConfigFilename);
+    
+		    Builder->SaveConfig(CPF_Config, *WorldConfigFilename);
+	    }
 	}
 
 	// Cleanup
@@ -191,8 +192,9 @@ int32 UWorldPartitionBuilderCommandlet::Main(const FString& Params)
 	World->ClearWorldComponents();
 	World->CleanupWorld();
 
-	WorldContext.SetCurrentWorld(nullptr);
-	GWorld = nullptr;
+	// Restore previous world
+	WorldContext.SetCurrentWorld(PrevGWorld);
+	GWorld = PrevGWorld;
 
-	return 0;
-}
+	return bResult ? 0 : 1;
+}
\ No newline at end of file
