Removal of UE4 references in Archive.h

diff --git a/Engine/Plugins/2D/Paper2D/Source/Paper2D/Private/PaperTileMap.cpp b/Engine/Plugins/2D/Paper2D/Source/Paper2D/Private/PaperTileMap.cpp
--- a/Engine/Plugins/2D/Paper2D/Source/Paper2D/Private/PaperTileMap.cpp
+++ b/Engine/Plugins/2D/Paper2D/Source/Paper2D/Private/PaperTileMap.cpp
@@ -86,7 +86,7 @@ void UPaperTileMap::Serialize(FArchive& Ar)
 	Ar.UsingCustomVersion(FPaperCustomVersion::GUID);
 
 #if WITH_EDITORONLY_DATA
-	if (Ar.IsLoading() && (Ar.UE4Ver() < VER_UE4_ASSET_IMPORT_DATA_AS_JSON) && (AssetImportData == nullptr))
+	if (Ar.IsLoading() && (Ar.UEVer() < VER_UE4_ASSET_IMPORT_DATA_AS_JSON) && (AssetImportData == nullptr))
 	{
 		// AssetImportData should always be valid
 		AssetImportData = NewObject<UAssetImportData>(this, TEXT("AssetImportData"));
diff --git a/Engine/Plugins/2D/Paper2D/Source/PaperSpriteSheetImporter/Private/PaperSpriteSheet.cpp b/Engine/Plugins/2D/Paper2D/Source/PaperSpriteSheetImporter/Private/PaperSpriteSheet.cpp
--- a/Engine/Plugins/2D/Paper2D/Source/PaperSpriteSheetImporter/Private/PaperSpriteSheet.cpp
+++ b/Engine/Plugins/2D/Paper2D/Source/PaperSpriteSheetImporter/Private/PaperSpriteSheet.cpp
@@ -33,7 +33,7 @@ void UPaperSpriteSheet::Serialize(FArchive& Ar)
 {
 	Super::Serialize(Ar);
 
-	if (Ar.IsLoading() && Ar.UE4Ver() < VER_UE4_ASSET_IMPORT_DATA_AS_JSON && !AssetImportData)
+	if (Ar.IsLoading() && Ar.UEVer() < VER_UE4_ASSET_IMPORT_DATA_AS_JSON && !AssetImportData)
 	{
 		// AssetImportData should always be valid
 		AssetImportData = NewObject<UAssetImportData>(this, TEXT("AssetImportData"));
diff --git a/Engine/Plugins/Developer/Concert/ConcertSync/ConcertSyncCore/Source/ConcertSyncCore/Private/ConcertSyncArchives.cpp b/Engine/Plugins/Developer/Concert/ConcertSync/ConcertSyncCore/Source/ConcertSyncCore/Private/ConcertSyncArchives.cpp
--- a/Engine/Plugins/Developer/Concert/ConcertSync/ConcertSyncCore/Source/ConcertSyncCore/Private/ConcertSyncArchives.cpp
+++ b/Engine/Plugins/Developer/Concert/ConcertSync/ConcertSyncCore/Source/ConcertSyncCore/Private/ConcertSyncArchives.cpp
@@ -180,8 +180,8 @@ FConcertSyncObjectReader::FConcertSyncObjectReader(const FConcertLocalIdentifier
 
 	if (InVersionInfo)
 	{
-		SetUE4Ver(InVersionInfo->FileVersion.FileVersion);
-		SetLicenseeUE4Ver(InVersionInfo->FileVersion.FileVersionLicensee);
+		SetUEVer(InVersionInfo->FileVersion.FileVersion);
+		SetLicenseeUEVer(InVersionInfo->FileVersion.FileVersionLicensee);
 		SetEngineVer(FEngineVersionBase(InVersionInfo->EngineVersion.Major, InVersionInfo->EngineVersion.Minor, InVersionInfo->EngineVersion.Patch, InVersionInfo->EngineVersion.Changelist));
 
 		FCustomVersionContainer EngineCustomVersions;
diff --git a/Engine/Plugins/Runtime/ApexDestruction/Source/ApexDestruction/Private/DestructibleMesh.cpp b/Engine/Plugins/Runtime/ApexDestruction/Source/ApexDestruction/Private/DestructibleMesh.cpp
--- a/Engine/Plugins/Runtime/ApexDestruction/Source/ApexDestruction/Private/DestructibleMesh.cpp
+++ b/Engine/Plugins/Runtime/ApexDestruction/Source/ApexDestruction/Private/DestructibleMesh.cpp
@@ -291,7 +291,7 @@ void UDestructibleMesh::Serialize(FArchive& Ar)
 #endif
 	}
 
-	if (Ar.UE4Ver() < VER_UE4_CLEAN_DESTRUCTIBLE_SETTINGS)
+	if (Ar.UEVer() < VER_UE4_CLEAN_DESTRUCTIBLE_SETTINGS)
 	{
 PRAGMA_DISABLE_DEPRECATION_WARNINGS
 		APEXToImpactResistance(DefaultDestructibleParameters.DamageParameters.bCustomImpactResistance, DefaultDestructibleParameters.DamageParameters.ImpactResistance);
diff --git a/Engine/Plugins/Runtime/GameplayAbilities/Source/GameplayAbilitiesEditor/Private/K2Node_GameplayCueEvent.cpp b/Engine/Plugins/Runtime/GameplayAbilities/Source/GameplayAbilitiesEditor/Private/K2Node_GameplayCueEvent.cpp
--- a/Engine/Plugins/Runtime/GameplayAbilities/Source/GameplayAbilitiesEditor/Private/K2Node_GameplayCueEvent.cpp
+++ b/Engine/Plugins/Runtime/GameplayAbilities/Source/GameplayAbilitiesEditor/Private/K2Node_GameplayCueEvent.cpp
@@ -27,7 +27,7 @@ void UK2Node_GameplayCueEvent::Serialize(FArchive& Ar)
 
 	if(Ar.IsLoading())
 	{
-		if(Ar.UE4Ver() < VER_UE4_K2NODE_EVENT_MEMBER_REFERENCE && EventSignatureName_DEPRECATED.IsNone() && EventSignatureClass_DEPRECATED == nullptr)
+		if(Ar.UEVer() < VER_UE4_K2NODE_EVENT_MEMBER_REFERENCE && EventSignatureName_DEPRECATED.IsNone() && EventSignatureClass_DEPRECATED == nullptr)
 		{
 			EventReference.SetExternalMember(GAMEPLAYABILITIES_BlueprintCustomHandler, UGameplayCueInterface::StaticClass());
 		}
diff --git a/Engine/Plugins/Runtime/PhysXVehicles/Source/PhysXVehicles/Private/WheeledVehicleMovementComponent4W.cpp b/Engine/Plugins/Runtime/PhysXVehicles/Source/PhysXVehicles/Private/WheeledVehicleMovementComponent4W.cpp
--- a/Engine/Plugins/Runtime/PhysXVehicles/Source/PhysXVehicles/Private/WheeledVehicleMovementComponent4W.cpp
+++ b/Engine/Plugins/Runtime/PhysXVehicles/Source/PhysXVehicles/Private/WheeledVehicleMovementComponent4W.cpp
@@ -408,7 +408,7 @@ void UWheeledVehicleMovementComponent4W::Serialize(FArchive & Ar)
 {
 	Super::Serialize(Ar);
 #if PHYSICS_INTERFACE_PHYSX
-	if (Ar.IsLoading() && Ar.UE4Ver() < VER_UE4_VEHICLES_UNIT_CHANGE)
+	if (Ar.IsLoading() && Ar.UEVer() < VER_UE4_VEHICLES_UNIT_CHANGE)
 	{
 		PxVehicleEngineData DefEngineData;
 		float DefaultRPM = OmegaToRPM(DefEngineData.mMaxOmega);
@@ -417,7 +417,7 @@ void UWheeledVehicleMovementComponent4W::Serialize(FArchive & Ar)
 		EngineSetup.MaxRPM = EngineSetup.MaxRPM != DefaultRPM ? OmegaToRPM(EngineSetup.MaxRPM) : DefaultRPM;	//need to convert from rad/s to RPM
 	}
 
-	if (Ar.IsLoading() && Ar.UE4Ver() < VER_UE4_VEHICLES_UNIT_CHANGE2)
+	if (Ar.IsLoading() && Ar.UEVer() < VER_UE4_VEHICLES_UNIT_CHANGE2)
 	{
 		PxVehicleEngineData DefEngineData;
 		PxVehicleClutchData DefClutchData;
diff --git a/Engine/Plugins/VirtualProduction/DMX/DMXEngine/Source/DMXBlueprintGraph/Private/K2Node_CastPatchToType.cpp b/Engine/Plugins/VirtualProduction/DMX/DMXEngine/Source/DMXBlueprintGraph/Private/K2Node_CastPatchToType.cpp
--- a/Engine/Plugins/VirtualProduction/DMX/DMXEngine/Source/DMXBlueprintGraph/Private/K2Node_CastPatchToType.cpp
+++ b/Engine/Plugins/VirtualProduction/DMX/DMXEngine/Source/DMXBlueprintGraph/Private/K2Node_CastPatchToType.cpp
@@ -55,7 +55,7 @@ FArchive& operator<<(FArchive& Ar, FUserPinInfo& Info)
 		Info.PinName = *PinNameStr;
 	}
 
-	if (Ar.UE4Ver() >= VER_UE4_SERIALIZE_PINTYPE_CONST)
+	if (Ar.UEVer() >= VER_UE4_SERIALIZE_PINTYPE_CONST)
 	{
 		Info.PinType.Serialize(Ar);
 		Ar << Info.DesiredPinDirection;
diff --git a/Engine/Source/Editor/BlueprintGraph/Private/K2Node.cpp b/Engine/Source/Editor/BlueprintGraph/Private/K2Node.cpp
--- a/Engine/Source/Editor/BlueprintGraph/Private/K2Node.cpp
+++ b/Engine/Source/Editor/BlueprintGraph/Private/K2Node.cpp
@@ -715,7 +715,7 @@ void UK2Node::ReconstructNode()
 				Pin->LinkedTo.Remove(OtherPin);
 			}
 
-			if (Blueprint->bIsRegeneratingOnLoad && Linker->UE4Ver() < VER_UE4_INJECT_BLUEPRINT_STRUCT_PIN_CONVERSION_NODES)
+			if (Blueprint->bIsRegeneratingOnLoad && Linker->UEVer() < VER_UE4_INJECT_BLUEPRINT_STRUCT_PIN_CONVERSION_NODES)
 			{
 				if (OtherPin == nullptr || (Pin->PinType.PinCategory != UEdGraphSchema_K2::PC_Struct))
 				{
diff --git a/Engine/Source/Editor/BlueprintGraph/Private/K2Node_ActorBoundEvent.cpp b/Engine/Source/Editor/BlueprintGraph/Private/K2Node_ActorBoundEvent.cpp
--- a/Engine/Source/Editor/BlueprintGraph/Private/K2Node_ActorBoundEvent.cpp
+++ b/Engine/Source/Editor/BlueprintGraph/Private/K2Node_ActorBoundEvent.cpp
@@ -252,7 +252,7 @@ void UK2Node_ActorBoundEvent::Serialize(FArchive& Ar)
 	// Fix up legacy nodes that may not yet have a delegate pin
 	if (Ar.IsLoading())
 	{
-		if (Ar.UE4Ver() < VER_UE4_K2NODE_EVENT_MEMBER_REFERENCE)
+		if (Ar.UEVer() < VER_UE4_K2NODE_EVENT_MEMBER_REFERENCE)
 		{
 			DelegateOwnerClass = EventSignatureClass_DEPRECATED;	
 		}
diff --git a/Engine/Source/Editor/BlueprintGraph/Private/K2Node_CallFunction.cpp b/Engine/Source/Editor/BlueprintGraph/Private/K2Node_CallFunction.cpp
--- a/Engine/Source/Editor/BlueprintGraph/Private/K2Node_CallFunction.cpp
+++ b/Engine/Source/Editor/BlueprintGraph/Private/K2Node_CallFunction.cpp
@@ -2174,7 +2174,7 @@ void UK2Node_CallFunction::Serialize(FArchive& Ar)
 
 	if (Ar.IsLoading())
 	{
-		if (Ar.UE4Ver() < VER_UE4_SWITCH_CALL_NODE_TO_USE_MEMBER_REFERENCE)
+		if (Ar.UEVer() < VER_UE4_SWITCH_CALL_NODE_TO_USE_MEMBER_REFERENCE)
 		{
 			UFunction* Function = FindUField<UFunction>(CallFunctionClass_DEPRECATED, CallFunctionName_DEPRECATED);
 			const bool bProbablySelfCall = (CallFunctionClass_DEPRECATED == NULL) || ((Function != NULL) && (Function->GetOuterUClass()->ClassGeneratedBy == GetBlueprint()));
@@ -2182,7 +2182,7 @@ void UK2Node_CallFunction::Serialize(FArchive& Ar)
 			FunctionReference.SetDirect(CallFunctionName_DEPRECATED, FGuid(), CallFunctionClass_DEPRECATED, bProbablySelfCall);
 		}
 
-		if(Ar.UE4Ver() < VER_UE4_K2NODE_REFERENCEGUIDS)
+		if(Ar.UEVer() < VER_UE4_K2NODE_REFERENCEGUIDS)
 		{
 			FGuid FunctionGuid;
 
diff --git a/Engine/Source/Editor/BlueprintGraph/Private/K2Node_ComponentBoundEvent.cpp b/Engine/Source/Editor/BlueprintGraph/Private/K2Node_ComponentBoundEvent.cpp
--- a/Engine/Source/Editor/BlueprintGraph/Private/K2Node_ComponentBoundEvent.cpp
+++ b/Engine/Source/Editor/BlueprintGraph/Private/K2Node_ComponentBoundEvent.cpp
@@ -219,7 +219,7 @@ void UK2Node_ComponentBoundEvent::Serialize(FArchive& Ar)
 	// Fix up legacy nodes that may not yet have a delegate pin
 	if (Ar.IsLoading())
 	{
-		if(Ar.UE4Ver() < VER_UE4_K2NODE_EVENT_MEMBER_REFERENCE)
+		if(Ar.UEVer() < VER_UE4_K2NODE_EVENT_MEMBER_REFERENCE)
 		{
 			DelegateOwnerClass = EventSignatureClass_DEPRECATED;
 		}
diff --git a/Engine/Source/Editor/BlueprintGraph/Private/K2Node_EditablePinBase.cpp b/Engine/Source/Editor/BlueprintGraph/Private/K2Node_EditablePinBase.cpp
--- a/Engine/Source/Editor/BlueprintGraph/Private/K2Node_EditablePinBase.cpp
+++ b/Engine/Source/Editor/BlueprintGraph/Private/K2Node_EditablePinBase.cpp
@@ -31,7 +31,7 @@ FArchive& operator<<(FArchive& Ar, FUserPinInfo& Info)
 		Info.PinName = *PinNameStr;
 	}
 
-	if (Ar.UE4Ver() >= VER_UE4_SERIALIZE_PINTYPE_CONST)
+	if (Ar.UEVer() >= VER_UE4_SERIALIZE_PINTYPE_CONST)
 	{
 		Info.PinType.Serialize(Ar);
 		Ar << Info.DesiredPinDirection;
diff --git a/Engine/Source/Editor/BlueprintGraph/Private/K2Node_Event.cpp b/Engine/Source/Editor/BlueprintGraph/Private/K2Node_Event.cpp
--- a/Engine/Source/Editor/BlueprintGraph/Private/K2Node_Event.cpp
+++ b/Engine/Source/Editor/BlueprintGraph/Private/K2Node_Event.cpp
@@ -58,7 +58,7 @@ void UK2Node_Event::Serialize(FArchive& Ar)
 	// Fix up legacy nodes that may not yet have a delegate pin
 	if(Ar.IsLoading())
 	{
-		if(Ar.UE4Ver() < VER_UE4_K2NODE_EVENT_MEMBER_REFERENCE)
+		if(Ar.UEVer() < VER_UE4_K2NODE_EVENT_MEMBER_REFERENCE)
 		{
 			EventReference.SetExternalMember(EventSignatureName_DEPRECATED, EventSignatureClass_DEPRECATED);
 		}
diff --git a/Engine/Source/Editor/BlueprintGraph/Private/K2Node_FunctionEntry.cpp b/Engine/Source/Editor/BlueprintGraph/Private/K2Node_FunctionEntry.cpp
--- a/Engine/Source/Editor/BlueprintGraph/Private/K2Node_FunctionEntry.cpp
+++ b/Engine/Source/Editor/BlueprintGraph/Private/K2Node_FunctionEntry.cpp
@@ -267,7 +267,7 @@ void UK2Node_FunctionEntry::Serialize(FArchive& Ar)
 			}
 		}
 
-		if (Ar.UE4Ver() < VER_UE4_BLUEPRINT_ENFORCE_CONST_IN_FUNCTION_OVERRIDES
+		if (Ar.UEVer() < VER_UE4_BLUEPRINT_ENFORCE_CONST_IN_FUNCTION_OVERRIDES
 			|| ((Ar.CustomVer(FFrameworkObjectVersion::GUID) < FFrameworkObjectVersion::EnforceConstInAnimBlueprintFunctionGraphs) && GetBlueprint()->IsA<UAnimBlueprint>()))
 		{
 			// Allow legacy implementations to violate const-correctness
diff --git a/Engine/Source/Editor/BlueprintGraph/Private/K2Node_InputAxisEvent.cpp b/Engine/Source/Editor/BlueprintGraph/Private/K2Node_InputAxisEvent.cpp
--- a/Engine/Source/Editor/BlueprintGraph/Private/K2Node_InputAxisEvent.cpp
+++ b/Engine/Source/Editor/BlueprintGraph/Private/K2Node_InputAxisEvent.cpp
@@ -30,7 +30,7 @@ void UK2Node_InputAxisEvent::Serialize(FArchive& Ar)
 
 	if(Ar.IsLoading())
 	{
-		if(Ar.UE4Ver() < VER_UE4_K2NODE_EVENT_MEMBER_REFERENCE && EventSignatureName_DEPRECATED.IsNone() && EventSignatureClass_DEPRECATED == nullptr)
+		if(Ar.UEVer() < VER_UE4_K2NODE_EVENT_MEMBER_REFERENCE && EventSignatureName_DEPRECATED.IsNone() && EventSignatureClass_DEPRECATED == nullptr)
 		{
 			EventReference.SetExternalDelegateMember(TEXT("InputAxisHandlerDynamicSignature__DelegateSignature"));
 		}
diff --git a/Engine/Source/Editor/BlueprintGraph/Private/K2Node_InputAxisKeyEvent.cpp b/Engine/Source/Editor/BlueprintGraph/Private/K2Node_InputAxisKeyEvent.cpp
--- a/Engine/Source/Editor/BlueprintGraph/Private/K2Node_InputAxisKeyEvent.cpp
+++ b/Engine/Source/Editor/BlueprintGraph/Private/K2Node_InputAxisKeyEvent.cpp
@@ -27,7 +27,7 @@ void UK2Node_InputAxisKeyEvent::Serialize(FArchive& Ar)
 
 	if(Ar.IsLoading())
 	{
-		if(Ar.UE4Ver() < VER_UE4_K2NODE_EVENT_MEMBER_REFERENCE && EventSignatureName_DEPRECATED.IsNone() && EventSignatureClass_DEPRECATED == nullptr)
+		if(Ar.UEVer() < VER_UE4_K2NODE_EVENT_MEMBER_REFERENCE && EventSignatureName_DEPRECATED.IsNone() && EventSignatureClass_DEPRECATED == nullptr)
 		{
 			EventReference.SetExternalDelegateMember(FName(TEXT("InputAxisHandlerDynamicSignature__DelegateSignature")));
 		}
diff --git a/Engine/Source/Editor/BlueprintGraph/Private/K2Node_InputVectorAxisEvent.cpp b/Engine/Source/Editor/BlueprintGraph/Private/K2Node_InputVectorAxisEvent.cpp
--- a/Engine/Source/Editor/BlueprintGraph/Private/K2Node_InputVectorAxisEvent.cpp
+++ b/Engine/Source/Editor/BlueprintGraph/Private/K2Node_InputVectorAxisEvent.cpp
@@ -19,7 +19,7 @@ void UK2Node_InputVectorAxisEvent::Serialize(FArchive& Ar)
 
 	if(Ar.IsLoading())
 	{
-		if(Ar.UE4Ver() < VER_UE4_K2NODE_EVENT_MEMBER_REFERENCE && EventSignatureName_DEPRECATED.IsNone() && EventSignatureClass_DEPRECATED == nullptr)
+		if(Ar.UEVer() < VER_UE4_K2NODE_EVENT_MEMBER_REFERENCE && EventSignatureName_DEPRECATED.IsNone() && EventSignatureClass_DEPRECATED == nullptr)
 		{
 			EventReference.SetExternalDelegateMember(FName(TEXT("InputVectorAxisHandlerDynamicSignature__DelegateSignature")));
 		}
diff --git a/Engine/Source/Editor/BlueprintGraph/Private/K2Node_MacroInstance.cpp b/Engine/Source/Editor/BlueprintGraph/Private/K2Node_MacroInstance.cpp
--- a/Engine/Source/Editor/BlueprintGraph/Private/K2Node_MacroInstance.cpp
+++ b/Engine/Source/Editor/BlueprintGraph/Private/K2Node_MacroInstance.cpp
@@ -24,7 +24,7 @@ void UK2Node_MacroInstance::Serialize(FArchive& Ar)
 {
 	Super::Serialize(Ar);
 
-	if (Ar.UE4Ver() < VER_UE4_K2NODE_REFERENCEGUIDS)
+	if (Ar.UEVer() < VER_UE4_K2NODE_REFERENCEGUIDS)
 	{
 		MacroGraphReference.SetGraph(MacroGraph_DEPRECATED);
 	}
diff --git a/Engine/Source/Editor/BlueprintGraph/Private/K2Node_Variable.cpp b/Engine/Source/Editor/BlueprintGraph/Private/K2Node_Variable.cpp
--- a/Engine/Source/Editor/BlueprintGraph/Private/K2Node_Variable.cpp
+++ b/Engine/Source/Editor/BlueprintGraph/Private/K2Node_Variable.cpp
@@ -30,13 +30,13 @@ void UK2Node_Variable::Serialize(FArchive& Ar)
 	// Fix old content 
 	if(Ar.IsLoading())
 	{
-		if(Ar.UE4Ver() < VER_UE4_VARK2NODE_USE_MEMBERREFSTRUCT)
+		if(Ar.UEVer() < VER_UE4_VARK2NODE_USE_MEMBERREFSTRUCT)
 		{
 			// Copy info into new struct
 			VariableReference.SetDirect(VariableName_DEPRECATED, FGuid(), VariableSourceClass_DEPRECATED, bSelfContext_DEPRECATED);
 		}
 
-		if(Ar.UE4Ver() < VER_UE4_K2NODE_VAR_REFERENCEGUIDS)
+		if(Ar.UEVer() < VER_UE4_K2NODE_VAR_REFERENCEGUIDS)
 		{
 			FGuid VarGuid;
 			
diff --git a/Engine/Source/Editor/Kismet/Private/BlueprintCompilationManager.cpp b/Engine/Source/Editor/Kismet/Private/BlueprintCompilationManager.cpp
--- a/Engine/Source/Editor/Kismet/Private/BlueprintCompilationManager.cpp
+++ b/Engine/Source/Editor/Kismet/Private/BlueprintCompilationManager.cpp
@@ -878,7 +878,7 @@ void FBlueprintCompilationManagerImpl::FlushCompilationQueueImpl(bool bSuppressB
 				BP->ConformNativeComponents();
 				if (FLinkerLoad* Linker = BP->GetLinker())
 				{
-					if (Linker->UE4Ver() < VER_UE4_EDITORONLY_BLUEPRINTS)
+					if (Linker->UEVer() < VER_UE4_EDITORONLY_BLUEPRINTS)
 					{
 						BP->ChangeOwnerOfTemplates();
 					}
diff --git a/Engine/Source/Editor/UnrealEd/Classes/Commandlets/ResavePackagesCommandlet.h b/Engine/Source/Editor/UnrealEd/Classes/Commandlets/ResavePackagesCommandlet.h
--- a/Engine/Source/Editor/UnrealEd/Classes/Commandlets/ResavePackagesCommandlet.h
+++ b/Engine/Source/Editor/UnrealEd/Classes/Commandlets/ResavePackagesCommandlet.h
@@ -31,19 +31,19 @@ protected:
 	EBrevity Verbosity;
 
 	/** only packages that have this version or higher will be resaved; a value of IGNORE_PACKAGE_VERSION indicates that there is no minimum package version */
-	int32 MinResaveUE4Version;
+	int32 MinResaveUEVersion;
 
 	/**
 	 * Limits resaving to packages with this UE4 package version or lower.
 	 * A value of IGNORE_PACKAGE_VERSION (default) removes this limitation.
 	 */
-	int32 MaxResaveUE4Version;
+	int32 MaxResaveUEVersion;
 
 	/**
 	 * Limits resaving to packages with this licensee package version or lower.
 	 * A value of IGNORE_PACKAGE_VERSION (default) removes this limitation.
 	 */
-	int32 MaxResaveLicenseeUE4Version;
+	int32 MaxResaveLicenseeUEVersion;
 
 	/** 
 	 * Maximum number of packages to resave to avoid having a massive sync
diff --git a/Engine/Source/Editor/UnrealEd/Private/Commandlets/ContentCommandlets.cpp b/Engine/Source/Editor/UnrealEd/Private/Commandlets/ContentCommandlets.cpp
--- a/Engine/Source/Editor/UnrealEd/Private/Commandlets/ContentCommandlets.cpp
+++ b/Engine/Source/Editor/UnrealEd/Private/Commandlets/ContentCommandlets.cpp
@@ -1,4 +1,5 @@
 // Copyright Epic Games, Inc. All Rights Reserved.
+// Copyright Epic Games, Inc. All Rights Reserved.
 
 /*=============================================================================
 	UCContentCommandlets.cpp: Various commmandlets.
@@ -414,23 +415,23 @@ int32 UResavePackagesCommandlet::InitializeResaveParameters( const TArray<FStrin
 	}
 
 	// Check for the min and max versions
-	MinResaveUE4Version = IGNORE_PACKAGE_VERSION;
-	MaxResaveUE4Version = IGNORE_PACKAGE_VERSION;
-	MaxResaveLicenseeUE4Version = IGNORE_PACKAGE_VERSION;
+	MinResaveUEVersion = IGNORE_PACKAGE_VERSION;
+	MaxResaveUEVersion = IGNORE_PACKAGE_VERSION;
+	MaxResaveLicenseeUEVersion = IGNORE_PACKAGE_VERSION;
 	if ( Switches.Contains(TEXT("CHECKLICENSEEVER")) )
 	{
 		// Limits resaving to packages with this licensee package version or lower.
-		MaxResaveLicenseeUE4Version = FMath::Max<int32>(GPackageFileLicenseeUEVersion - 1, 0);
+		MaxResaveLicenseeUEVersion = FMath::Max<int32>(GPackageFileLicenseeUEVersion - 1, 0);
 	}
 	if ( Switches.Contains(TEXT("CHECKUE4VER")) )
 	{
 		// Limits resaving to packages with this ue4 package version or lower.
-		MaxResaveUE4Version = FMath::Max<int32>(GPackageFileUEVersion - 1, 0);
+		MaxResaveUEVersion = FMath::Max<int32>(GPackageFileUEVersion - 1, 0);
 	}
 	else if ( Switches.Contains(TEXT("RESAVEDEPRECATED")) )
 	{
 		// Limits resaving to packages with this package version or lower.
-		MaxResaveUE4Version = FMath::Max<int32>(VER_UE4_DEPRECATED_PACKAGE - 1, 0);
+		MaxResaveUEVersion = FMath::Max<int32>(VER_UE4_DEPRECATED_PACKAGE - 1, 0);
 	}
 	else
 	{
@@ -438,19 +439,19 @@ int32 UResavePackagesCommandlet::InitializeResaveParameters( const TArray<FStrin
 		for ( int32 SwitchIdx = 0; SwitchIdx < Switches.Num(); SwitchIdx++ )
 		{
 			const FString& CurrentSwitch = Switches[SwitchIdx];
-			if ( MinResaveUE4Version == IGNORE_PACKAGE_VERSION && FParse::Value(*CurrentSwitch,TEXT("MINVER="),MinResaveUE4Version) )
+			if ( MinResaveUEVersion == IGNORE_PACKAGE_VERSION && FParse::Value(*CurrentSwitch,TEXT("MINVER="),MinResaveUEVersion) )
 			{
-				if ( MinResaveUE4Version == CURRENT_PACKAGE_VERSION )
+				if ( MinResaveUEVersion == CURRENT_PACKAGE_VERSION )
 				{
-					MinResaveUE4Version = GPackageFileUEVersion;
+					MinResaveUEVersion = GPackageFileUEVersion;
 				}
 			}
 
-			if ( MaxResaveUE4Version == IGNORE_PACKAGE_VERSION && FParse::Value(*CurrentSwitch,TEXT("MAXVER="),MaxResaveUE4Version) )
+			if ( MaxResaveUEVersion == IGNORE_PACKAGE_VERSION && FParse::Value(*CurrentSwitch,TEXT("MAXVER="),MaxResaveUEVersion) )
 			{
-				if ( MaxResaveUE4Version == CURRENT_PACKAGE_VERSION )
+				if ( MaxResaveUEVersion == CURRENT_PACKAGE_VERSION )
 				{
-					MaxResaveUE4Version = GPackageFileUEVersion;
+					MaxResaveUEVersion = GPackageFileUEVersion;
 				}
 			}
 		}
@@ -1367,17 +1368,17 @@ void UResavePackagesCommandlet::PerformPreloadOperations( FLinkerLoad* PackageLi
 
 
 	// validate that this package meets the minimum requirement
-	if ( MinResaveUE4Version != IGNORE_PACKAGE_VERSION && UE4PackageVersion < MinResaveUE4Version )
+	if ( MinResaveUEVersion != IGNORE_PACKAGE_VERSION && UE4PackageVersion < MinResaveUEVersion )
 	{
 		bSavePackage = false;
 		return;
 	}
 
 	// Check if this package meets the maximum requirements.
-	const bool bNoLimitation = MaxResaveUE4Version == IGNORE_PACKAGE_VERSION && MaxResaveLicenseeUE4Version == IGNORE_PACKAGE_VERSION;
+	const bool bNoLimitation = MaxResaveUEVersion == IGNORE_PACKAGE_VERSION && MaxResaveLicenseeUEVersion == IGNORE_PACKAGE_VERSION;
 	const bool bAllowResave = bNoLimitation ||
-						 (MaxResaveUE4Version != IGNORE_PACKAGE_VERSION && UE4PackageVersion <= MaxResaveUE4Version) ||
-						 (MaxResaveLicenseeUE4Version != IGNORE_PACKAGE_VERSION && LicenseeUE4PackageVersion <= MaxResaveLicenseeUE4Version);
+						 (MaxResaveUEVersion != IGNORE_PACKAGE_VERSION && UE4PackageVersion <= MaxResaveUEVersion) ||
+						 (MaxResaveLicenseeUEVersion != IGNORE_PACKAGE_VERSION && LicenseeUE4PackageVersion <= MaxResaveLicenseeUEVersion);
 	// If not, don't resave it.
 	if ( !bAllowResave )
 	{
diff --git a/Engine/Source/Editor/UnrealEd/Private/Commandlets/PackageUtilities.cpp b/Engine/Source/Editor/UnrealEd/Private/Commandlets/PackageUtilities.cpp
--- a/Engine/Source/Editor/UnrealEd/Private/Commandlets/PackageUtilities.cpp
+++ b/Engine/Source/Editor/UnrealEd/Private/Commandlets/PackageUtilities.cpp
@@ -944,7 +944,7 @@ void FPkgInfoReporter_Log::GeneratePackageReport( FLinkerLoad* InLinker /*=nullp
 	Out.Logf(ELogVerbosity::Display, TEXT("--------------------------------------------") );
 
 	Out.Logf(ELogVerbosity::Display, TEXT("\t         Filename: %s"), *Linker->GetPackagePath().GetLocalFullPath());
-	Out.Logf(ELogVerbosity::Display, TEXT("\t     File Version: %i"), Linker->UE4Ver() );
+	Out.Logf(ELogVerbosity::Display, TEXT("\t     File Version: %i"), Linker->UEVer() );
 	Out.Logf(ELogVerbosity::Display, TEXT("\t   Engine Version: %s"), *Linker->Summary.SavedByEngineVersion.ToString());
 	Out.Logf(ELogVerbosity::Display, TEXT("\t   Compat Version: %s"), *Linker->Summary.CompatibleWithEngineVersion.ToString());
 	Out.Logf(ELogVerbosity::Display, TEXT("\t     PackageFlags: %X"), Linker->Summary.GetPackageFlags() );
diff --git a/Engine/Source/Editor/UnrealEd/Private/CookOnTheFlyServer.cpp b/Engine/Source/Editor/UnrealEd/Private/CookOnTheFlyServer.cpp
--- a/Engine/Source/Editor/UnrealEd/Private/CookOnTheFlyServer.cpp
+++ b/Engine/Source/Editor/UnrealEd/Private/CookOnTheFlyServer.cpp
@@ -4608,9 +4608,9 @@ void GetAdditionalCurrentIniVersionStrings( const ITargetPlatform* TargetPlatfor
 		IniVersionMap.Add(CustomVersionString, CustomVersionValue);
 	}
 
-	FString UE4Ver = FString::Printf(TEXT("PackageFileVersions:%d"), GPackageFileUEVersion);
+	FString UEVer = FString::Printf(TEXT("PackageFileVersions:%d"), GPackageFileUEVersion);
 	FString UE4Value = FString::Printf(TEXT("%d"), GPackageFileLicenseeUEVersion);
-	IniVersionMap.Add(UE4Ver, UE4Value);
+	IniVersionMap.Add(UEVer, UE4Value);
 
 	/*FString UE4EngineVersionCompatibleName = TEXT("EngineVersionCompatibleWith");
 	FString UE4EngineVersionCompatible = FEngineVersion::CompatibleWith().ToString();
diff --git a/Engine/Source/Editor/UnrealEd/Private/EditorServer.cpp b/Engine/Source/Editor/UnrealEd/Private/EditorServer.cpp
--- a/Engine/Source/Editor/UnrealEd/Private/EditorServer.cpp
+++ b/Engine/Source/Editor/UnrealEd/Private/EditorServer.cpp
@@ -2416,14 +2416,14 @@ bool UEditorEngine::PackageIsAMapFile( const TCHAR* PackageFilename, FText& OutN
 			return false;
 		}
 
-		const int32 UE4Version = Summary.GetFileVersionUE4();
+		const int32 UEVersion = Summary.GetFileVersionUE4();
 
 		// Validate the summary.
-		if( UE4Version < VER_UE4_OLDEST_LOADABLE_PACKAGE )
+		if( UEVersion < VER_UE4_OLDEST_LOADABLE_PACKAGE )
 		{
 			FFormatNamedArguments Arguments;
 			Arguments.Add(TEXT("File"), FText::FromString( FString( PackageFilename ) ));
-			Arguments.Add(TEXT("Version"), UE4Version);
+			Arguments.Add(TEXT("Version"), UEVersion);
 			Arguments.Add(TEXT("First"), VER_UE4_OLDEST_LOADABLE_PACKAGE);
 			OutNotMapReason = FText::Format( LOCTEXT( "UE4FileIsOlder", "{File} is an UE4 map [File:v{Version}], from an engine release no longer supported [Min:v{First}]." ), 
 				Arguments);
@@ -2433,11 +2433,11 @@ bool UEditorEngine::PackageIsAMapFile( const TCHAR* PackageFilename, FText& OutN
 		const int32 UE4LicenseeVersion = Summary.GetFileVersionLicenseeUE4();
 
 		// Don't load packages that were saved with an engine version newer than the current one.
-		if( UE4Version > GPackageFileUEVersion )
+		if( UEVersion > GPackageFileUEVersion )
 		{
 			FFormatNamedArguments Arguments;
 			Arguments.Add(TEXT("File"), FText::FromString( FString( PackageFilename ) ));
-			Arguments.Add(TEXT("Version"), UE4Version);
+			Arguments.Add(TEXT("Version"), UEVersion);
 			Arguments.Add(TEXT("Last"), GPackageFileUEVersion);
 			OutNotMapReason = FText::Format( LOCTEXT( "UE4FileIsNewer", "{File} is a UE4 map [File:v{Version}], from an engine release newer than this [Cur:v{Last}]." ), 
 				Arguments);
diff --git a/Engine/Source/Editor/UnrealEd/Private/Fbx/FbxAnimSequenceImportData.cpp b/Engine/Source/Editor/UnrealEd/Private/Fbx/FbxAnimSequenceImportData.cpp
--- a/Engine/Source/Editor/UnrealEd/Private/Fbx/FbxAnimSequenceImportData.cpp
+++ b/Engine/Source/Editor/UnrealEd/Private/Fbx/FbxAnimSequenceImportData.cpp
@@ -55,7 +55,7 @@ void UFbxAnimSequenceImportData::Serialize(FArchive& Ar)
 {
 	Super::Serialize(Ar);
 
-	if(Ar.UE4Ver() < VER_UE4_FBX_IMPORT_DATA_RANGE_ENCAPSULATION)
+	if(Ar.UEVer() < VER_UE4_FBX_IMPORT_DATA_RANGE_ENCAPSULATION)
 	{
 		FrameImportRange.Min = StartFrame_DEPRECATED;
 		FrameImportRange.Max = EndFrame_DEPRECATED;
diff --git a/Engine/Source/Editor/UnrealEd/Private/UserDefinedStructEditorData.cpp b/Engine/Source/Editor/UnrealEd/Private/UserDefinedStructEditorData.cpp
--- a/Engine/Source/Editor/UnrealEd/Private/UserDefinedStructEditorData.cpp
+++ b/Engine/Source/Editor/UnrealEd/Private/UserDefinedStructEditorData.cpp
@@ -18,7 +18,7 @@ void FStructVariableDescription::PostSerialize(const FArchive& Ar)
 		ContainerType = FEdGraphPinType::ToPinContainerType(bIsArray_DEPRECATED, bIsSet_DEPRECATED, bIsMap_DEPRECATED);
 	}
 
-	if (Ar.UE4Ver() < VER_UE4_ADDED_SOFT_OBJECT_PATH)
+	if (Ar.UEVer() < VER_UE4_ADDED_SOFT_OBJECT_PATH)
 	{
 		// Fix up renamed categories
 		if (Category == TEXT("asset"))
diff --git a/Engine/Source/Runtime/AssetRegistry/Private/AssetDataGatherer.cpp b/Engine/Source/Runtime/AssetRegistry/Private/AssetDataGatherer.cpp
--- a/Engine/Source/Runtime/AssetRegistry/Private/AssetDataGatherer.cpp
+++ b/Engine/Source/Runtime/AssetRegistry/Private/AssetDataGatherer.cpp
@@ -2685,7 +2685,7 @@ bool FAssetDataGatherer::ReadAssetFile(FPackageReader& PackageReader, TArray<FAs
 			return false;
 		}
 
-		// DEPRECATION_TODO: Remove this fixup-on-load once we bump EUnrealEngineObjectUE4Version VER_UE4_ASSETREGISTRY_DEPENDENCYFLAGS and therefore all projects will resave each ObjectRedirector
+		// DEPRECATION_TODO: Remove this fixup-on-load once we bump EUnrealEngineObjectUEVersion VER_UE4_ASSETREGISTRY_DEPENDENCYFLAGS and therefore all projects will resave each ObjectRedirector
 		// UObjectRedirectors were originally incorrectly marked as having editor-only imports, since UObjectRedirector is an editor-only class. But UObjectRedirectors are followed during cooking
 		// and so their imports should be considered used-in-game. Mark all dependencies in the package as used in game if the package has a UObjectRedirector object
 		FName RedirectorClassName = UObjectRedirector::StaticClass()->GetFName();
diff --git a/Engine/Source/Runtime/AssetRegistry/Private/PackageReader.cpp b/Engine/Source/Runtime/AssetRegistry/Private/PackageReader.cpp
--- a/Engine/Source/Runtime/AssetRegistry/Private/PackageReader.cpp
+++ b/Engine/Source/Runtime/AssetRegistry/Private/PackageReader.cpp
@@ -119,8 +119,8 @@ bool FPackageReader::OpenPackageFile(EOpenPackageResult* OutErrorCode)
 	}
 
 	//make sure the filereader gets the correct version number (it defaults to latest version)
-	SetUE4Ver(PackageFileSummary.GetFileVersionUE4());
-	SetLicenseeUE4Ver(PackageFileSummary.GetFileVersionLicenseeUE4());
+	SetUEVer(PackageFileSummary.GetFileVersionUE4());
+	SetLicenseeUEVer(PackageFileSummary.GetFileVersionLicenseeUE4());
 	SetEngineVer(PackageFileSummary.SavedByEngineVersion);
 
 	const FCustomVersionContainer& PackageFileSummaryVersions = PackageFileSummary.GetCustomVersionContainer();
@@ -286,7 +286,7 @@ bool FPackageReader::ReadAssetRegistryDataIfCookedPackage(TArray<FAssetData*>& A
 		// which of the objects in the export map as the asset.
 		// Otherwise we need to store a temp minimal data and then force load the asset
 		// to re-generate its registry data
-		if (UE4Ver() >= VER_UE4_COOKED_ASSETS_IN_EDITOR_SUPPORT)
+		if (UEVer() >= VER_UE4_COOKED_ASSETS_IN_EDITOR_SUPPORT)
 		{
 			const FString PackagePath = FPackageName::GetLongPackagePath(PackageName);
 
@@ -483,7 +483,7 @@ bool FPackageReader::SerializeExportMap(TArray<FObjectExport>& OutExportMap)
 
 bool FPackageReader::SerializeSoftPackageReferenceList(TArray<FName>& OutSoftPackageReferenceList)
 {
-	if (UE4Ver() >= VER_UE4_ADD_STRING_ASSET_REFERENCES_MAP && PackageFileSummary.SoftPackageReferencesOffset > 0 && PackageFileSummary.SoftPackageReferencesCount > 0)
+	if (UEVer() >= VER_UE4_ADD_STRING_ASSET_REFERENCES_MAP && PackageFileSummary.SoftPackageReferencesOffset > 0 && PackageFileSummary.SoftPackageReferencesCount > 0)
 	{
 		if (!StartSerializeSection(PackageFileSummary.SoftPackageReferencesOffset))
 		{
@@ -497,7 +497,7 @@ bool FPackageReader::SerializeSoftPackageReferenceList(TArray<FName>& OutSoftPac
 			UE_PACKAGEREADER_CORRUPTPACKAGE_WARNING("SerializeSoftPackageReferenceListInvalidReferencesCount", PackageFilename);
 			return false;
 		}
-		if (UE4Ver() < VER_UE4_ADDED_SOFT_OBJECT_PATH)
+		if (UEVer() < VER_UE4_ADDED_SOFT_OBJECT_PATH)
 		{
 			for (int32 ReferenceIdx = 0; ReferenceIdx < PackageFileSummary.SoftPackageReferencesCount; ++ReferenceIdx)
 			{
@@ -509,7 +509,7 @@ bool FPackageReader::SerializeSoftPackageReferenceList(TArray<FName>& OutSoftPac
 					return false;
 				}
 
-				if (UE4Ver() < VER_UE4_KEEP_ONLY_PACKAGE_NAMES_IN_STRING_ASSET_REFERENCES_MAP)
+				if (UEVer() < VER_UE4_KEEP_ONLY_PACKAGE_NAMES_IN_STRING_ASSET_REFERENCES_MAP)
 				{
 					PackageName = FPackageName::GetNormalizedObjectPath(PackageName);
 					if (!PackageName.IsEmpty())
@@ -543,7 +543,7 @@ bool FPackageReader::SerializeSoftPackageReferenceList(TArray<FName>& OutSoftPac
 
 bool FPackageReader::SerializeSearchableNamesMap(FPackageDependencyData& OutDependencyData)
 {
-	if (UE4Ver() >= VER_UE4_ADDED_SEARCHABLE_NAMES && PackageFileSummary.SearchableNamesOffset > 0)
+	if (UEVer() >= VER_UE4_ADDED_SEARCHABLE_NAMES && PackageFileSummary.SearchableNamesOffset > 0)
 	{
 		if (!StartSerializeSection(PackageFileSummary.SearchableNamesOffset))
 		{
diff --git a/Engine/Source/Runtime/Core/Private/Internationalization/Text.cpp b/Engine/Source/Runtime/Core/Private/Internationalization/Text.cpp
--- a/Engine/Source/Runtime/Core/Private/Internationalization/Text.cpp
+++ b/Engine/Source/Runtime/Core/Private/Internationalization/Text.cpp
@@ -813,7 +813,7 @@ void FText::SerializeText(FStructuredArchive::FSlot Slot, FText& Value)
 	//This would result in all subsequently duplicated objects stamping over formerly duplicated ones.
 
 	// Older FText's stored their "SourceString", that is now stored in a history class so move it there
-	if (UnderlyingArchive.IsLoading() && UnderlyingArchive.UE4Ver() < VER_UE4_FTEXT_HISTORY)
+	if (UnderlyingArchive.IsLoading() && UnderlyingArchive.UEVer() < VER_UE4_FTEXT_HISTORY)
 	{
 		FString SourceStringToImplantIntoHistory;
 		Record << SA_VALUE(TEXT("SourceStringToImplantIntoHistory"), SourceStringToImplantIntoHistory);
@@ -821,7 +821,7 @@ void FText::SerializeText(FStructuredArchive::FSlot Slot, FText& Value)
 		FTextDisplayStringPtr DisplayString;
 
 		// Namespaces and keys are no longer stored in the FText, we need to read them in and discard
-		if (UnderlyingArchive.UE4Ver() >= VER_UE4_ADDED_NAMESPACE_AND_KEY_DATA_TO_FTEXT)
+		if (UnderlyingArchive.UEVer() >= VER_UE4_ADDED_NAMESPACE_AND_KEY_DATA_TO_FTEXT)
 		{
 			FString Namespace;
 			FString Key;
@@ -877,7 +877,7 @@ void FText::SerializeText(FStructuredArchive::FSlot Slot, FText& Value)
 		Value.Flags = OriginalFlags;
 	}
 
-	if (UnderlyingArchive.UE4Ver() >= VER_UE4_FTEXT_HISTORY)
+	if (UnderlyingArchive.UEVer() >= VER_UE4_FTEXT_HISTORY)
 	{
 		bool bSerializeHistory = true;
 
@@ -1460,7 +1460,7 @@ void operator<<(FStructuredArchive::FSlot Slot, FFormatArgumentData& Value)
 	if (UnderlyingArchive.IsLoading())
 	{
 		// ArgumentName was changed to be FString rather than FText, so we need to convert older data to ensure serialization stays happy outside of UStruct::SerializeTaggedProperties.
-		if (UnderlyingArchive.UE4Ver() >= VER_UE4_K2NODE_VAR_REFERENCEGUIDS) // There was no version bump for this change, but VER_UE4_K2NODE_VAR_REFERENCEGUIDS was made at almost the same time.
+		if (UnderlyingArchive.UEVer() >= VER_UE4_K2NODE_VAR_REFERENCEGUIDS) // There was no version bump for this change, but VER_UE4_K2NODE_VAR_REFERENCEGUIDS was made at almost the same time.
 		{
 			Record << SA_VALUE(TEXT("ArgumentName"), Value.ArgumentName);
 		}
diff --git a/Engine/Source/Runtime/Core/Private/Internationalization/TextHistory.cpp b/Engine/Source/Runtime/Core/Private/Internationalization/TextHistory.cpp
--- a/Engine/Source/Runtime/Core/Private/Internationalization/TextHistory.cpp
+++ b/Engine/Source/Runtime/Core/Private/Internationalization/TextHistory.cpp
@@ -1870,7 +1870,7 @@ void FTextHistory_AsCurrency::Serialize(FStructuredArchive::FRecord Record)
 		Record << SA_VALUE(TEXT("HistoryType"), HistoryType);
 	}
 
-	if (BaseArchive.UE4Ver() >= VER_UE4_ADDED_CURRENCY_CODE_TO_FTEXT)
+	if (BaseArchive.UEVer() >= VER_UE4_ADDED_CURRENCY_CODE_TO_FTEXT)
 	{
 		Record << SA_VALUE(TEXT("CurrencyCode"), CurrencyCode);
 	}
@@ -2047,7 +2047,7 @@ void FTextHistory_AsDate::Serialize(FStructuredArchive::FRecord Record)
 	Record << SA_VALUE(TEXT("DateStyleInt8"), DateStyleInt8);
 	DateStyle = (EDateTimeStyle::Type)DateStyleInt8;
 
-	if( BaseArchive.UE4Ver() >= VER_UE4_FTEXT_HISTORY_DATE_TIMEZONE )
+	if( BaseArchive.UEVer() >= VER_UE4_FTEXT_HISTORY_DATE_TIMEZONE )
 	{
 		Record << SA_VALUE(TEXT("TimeZone"), TimeZone);
 	}
diff --git a/Engine/Source/Runtime/Core/Private/Serialization/Archive.cpp b/Engine/Source/Runtime/Core/Private/Serialization/Archive.cpp
--- a/Engine/Source/Runtime/Core/Private/Serialization/Archive.cpp
+++ b/Engine/Source/Runtime/Core/Private/Serialization/Archive.cpp
@@ -150,8 +150,8 @@ void FArchiveState::Reset()
 #if DEVIRTUALIZE_FLinkerLoad_Serialize
 	ActiveFPLB->Reset();
 #endif
-	ArUE4Ver							= GPackageFileUEVersion;
-	ArLicenseeUE4Ver					= GPackageFileLicenseeUEVersion;
+	ArUEVer								= GPackageFileUEVersion;
+	ArLicenseeUEVer						= GPackageFileLicenseeUEVersion;
 	ArEngineVer							= FEngineVersion::Current();
 	ArEngineNetVer						= FNetworkVersion::GetEngineNetworkProtocolVersion();
 	ArGameNetVer						= FNetworkVersion::GetGameNetworkProtocolVersion();
@@ -209,8 +209,8 @@ void FArchiveState::Reset()
 
 void FArchiveState::CopyTrivialFArchiveStatusMembers(const FArchiveState& ArchiveToCopy)
 {
-	ArUE4Ver                             = ArchiveToCopy.ArUE4Ver;
-	ArLicenseeUE4Ver                     = ArchiveToCopy.ArLicenseeUE4Ver;
+	ArUEVer                              = ArchiveToCopy.ArUEVer;
+	ArLicenseeUEVer                      = ArchiveToCopy.ArLicenseeUEVer;
 	ArEngineVer                          = ArchiveToCopy.ArEngineVer;
 	ArEngineNetVer                       = ArchiveToCopy.ArEngineNetVer;
 	ArGameNetVer                         = ArchiveToCopy.ArGameNetVer;
@@ -733,8 +733,8 @@ void FArchive::SerializeCompressed(void* V, int64 Length, FName CompressionForma
 		if (!bHeaderWasValid)
 		{
 			UE_LOG(LogSerialization, Log, TEXT("ArchiveName: %s"), *GetArchiveName());
-			UE_LOG(LogSerialization, Log, TEXT("Archive UE4 Version: %d"), UE4Ver());
-			UE_LOG(LogSerialization, Log, TEXT("Archive Licensee Version: %d"), LicenseeUE4Ver());
+			UE_LOG(LogSerialization, Log, TEXT("Archive UE4 Version: %d"), UEVer());
+			UE_LOG(LogSerialization, Log, TEXT("Archive Licensee Version: %d"), LicenseeUEVer());
 			UE_LOG(LogSerialization, Log, TEXT("Position: %lld"), Tell());
 			UE_LOG(LogSerialization, Log, TEXT("Read Size: %lld"), Length);
 			UE_LOG(LogSerialization, Fatal, TEXT("BulkData compressed header read error. This package may be corrupt!"));
@@ -1211,14 +1211,14 @@ void FArchive::LogfImpl(const TCHAR* Fmt, ...)
 	FMemory::SystemFree( Buffer );
 }
 
-void FArchiveState::SetUE4Ver(int32 InVer)
+void FArchiveState::SetUEVer(int32 InVer)
 {
-	ArUE4Ver = InVer;
+	ArUEVer = InVer;
 }
 
-void FArchiveState::SetLicenseeUE4Ver(int32 InVer)
+void FArchiveState::SetLicenseeUEVer(int32 InVer)
 {
-	ArLicenseeUE4Ver = InVer;
+	ArLicenseeUEVer = InVer;
 }
 
 void FArchiveState::SetEngineVer(const FEngineVersionBase& InVer)
diff --git a/Engine/Source/Runtime/Core/Private/UObject/UnrealNames.cpp b/Engine/Source/Runtime/Core/Private/UObject/UnrealNames.cpp
--- a/Engine/Source/Runtime/Core/Private/UObject/UnrealNames.cpp
+++ b/Engine/Source/Runtime/Core/Private/UObject/UnrealNames.cpp
@@ -2721,7 +2721,7 @@ FArchive& operator<<(FArchive& Ar, FNameEntrySerialized& E)
 		}
 
 		uint16 DummyHashes[2];
-		uint32 SkipPastHashBytes = (Ar.UE4Ver() >= VER_UE4_NAME_HASHES_SERIALIZED) * sizeof(DummyHashes);
+		uint32 SkipPastHashBytes = (Ar.UEVer() >= VER_UE4_NAME_HASHES_SERIALIZED) * sizeof(DummyHashes);
 		Ar.Serialize(&DummyHashes, SkipPastHashBytes);
 	}
 	else
diff --git a/Engine/Source/Runtime/Core/Public/Math/InterpCurve.h b/Engine/Source/Runtime/Core/Public/Math/InterpCurve.h
--- a/Engine/Source/Runtime/Core/Public/Math/InterpCurve.h
+++ b/Engine/Source/Runtime/Core/Public/Math/InterpCurve.h
@@ -138,7 +138,7 @@ public:
 		//   as inline struct properties in UnClass.cpp!
 
 		Ar << Curve.Points;
-		if (Ar.UE4Ver() >= VER_UE4_INTERPCURVE_SUPPORTS_LOOPING)
+		if (Ar.UEVer() >= VER_UE4_INTERPCURVE_SUPPORTS_LOOPING)
 		{
 			Ar << Curve.bIsLooped;
 			Ar << Curve.LoopKeyOffset;
diff --git a/Engine/Source/Runtime/Core/Public/Math/Matrix.h b/Engine/Source/Runtime/Core/Public/Math/Matrix.h
--- a/Engine/Source/Runtime/Core/Public/Math/Matrix.h
+++ b/Engine/Source/Runtime/Core/Public/Math/Matrix.h
@@ -343,7 +343,7 @@ public:
 
 	bool Serialize( FArchive& Ar )
 	{
-		if (Ar.UE4Ver() >= VER_UE4_ADDED_NATIVE_SERIALIZATION_FOR_IMMUTABLE_STRUCTURES)
+		if (Ar.UEVer() >= VER_UE4_ADDED_NATIVE_SERIALIZATION_FOR_IMMUTABLE_STRUCTURES)
 		{
 			Ar << *this;
 			return true;
diff --git a/Engine/Source/Runtime/Core/Public/Math/Plane.h b/Engine/Source/Runtime/Core/Public/Math/Plane.h
--- a/Engine/Source/Runtime/Core/Public/Math/Plane.h
+++ b/Engine/Source/Runtime/Core/Public/Math/Plane.h
@@ -270,7 +270,7 @@ public:
 
 	bool Serialize(FArchive& Ar)
 	{
-		if (Ar.UE4Ver() >= VER_UE4_ADDED_NATIVE_SERIALIZATION_FOR_IMMUTABLE_STRUCTURES)
+		if (Ar.UEVer() >= VER_UE4_ADDED_NATIVE_SERIALIZATION_FOR_IMMUTABLE_STRUCTURES)
 		{
 			Ar << *this;
 			return true;
diff --git a/Engine/Source/Runtime/Core/Public/Math/VectorLWC.h b/Engine/Source/Runtime/Core/Public/Math/VectorLWC.h
--- a/Engine/Source/Runtime/Core/Public/Math/VectorLWC.h
+++ b/Engine/Source/Runtime/Core/Public/Math/VectorLWC.h
@@ -1060,7 +1060,7 @@ public:
 		// @warning BulkSerialize: TVector<TReal> is serialized as memory dump
 		// See TArray::BulkSerialize for detailed description of implied limitations.
 		FStructuredArchive::FRecord Record = Slot.EnterRecord();
-		//if (Record.GetUnderlyingArchive().UE4Ver() >= VER_UE5_LARGE_WORLD_COORDINATES)
+		//if (Record.GetUnderlyingArchive().UEVer() >= VER_UE5_LARGE_WORLD_COORDINATES)
 		//{
 		//	// Stored as doubles, so serialize double and copy.
 		//	double SX = V.X, SY = V.Y, SZ = V.Z;
diff --git a/Engine/Source/Runtime/Core/Public/Serialization/Archive.h b/Engine/Source/Runtime/Core/Public/Serialization/Archive.h
--- a/Engine/Source/Runtime/Core/Public/Serialization/Archive.h
+++ b/Engine/Source/Runtime/Core/Public/Serialization/Archive.h
@@ -186,15 +186,29 @@ public:
 	}
 
 	/** Returns the global engine serialization version used for this archive. */
+	FORCEINLINE int32 UEVer() const
+	{
+		return ArUEVer;
+	}
+
+	/** Returns the licensee-specific version used for this archive, will be 0 by default. */
+	FORCEINLINE int32 LicenseeUEVer() const
+	{
+		return ArLicenseeUEVer;
+	}
+
+	/** Returns the global engine serialization version used for this archive. */
+	UE_DEPRECATED(5.0, "Use UEVer instead")
 	FORCEINLINE int32 UE4Ver() const
 	{
-		return ArUE4Ver;
+		return UEVer();
 	}
 
 	/** Returns the licensee-specific version used for this archive, will be 0 by default. */
+	UE_DEPRECATED(5.0, "Use LicenseeUEVer instead")
 	FORCEINLINE int32 LicenseeUE4Ver() const
 	{
-		return ArLicenseeUE4Ver;
+		return LicenseeUEVer();
 	}
 
 	/** Returns the compiled engine version used for this archive. */
@@ -861,9 +875,15 @@ public:
 	 * Sets the archive version number. Used by the code that makes sure that FLinkerLoad's 
 	 * internal archive versions match the file reader it creates.
 	 *
-	 * @param UE4Ver	new version number
+	 * @param UEVer	new version number
 	 */
-	virtual void SetUE4Ver(int32 InVer);
+	virtual void SetUEVer(int32 InVer);
+
+	UE_DEPRECATED(5.0, "Use SetUEVer instead")
+	FORCEINLINE void SetUE4Ver(int32 InVer)
+	{
+		SetUEVer(InVer);
+	}
 
 	/**
 	 * Sets the archive licensee version number. Used by the code that makes sure that FLinkerLoad's 
@@ -871,7 +891,13 @@ public:
 	 *
 	 * @param Ver	new version number
 	 */
-	virtual void SetLicenseeUE4Ver(int32 InVer);
+	virtual void SetLicenseeUEVer(int32 InVer);
+
+	UE_DEPRECATED(5.0, "Use SetLicenseeUEVer instead")
+	FORCEINLINE void SetLicenseeUE4Ver(int32 InVer)
+	{
+		SetLicenseeUEVer(InVer);
+	}
 
 	/**
 	 * Sets the archive engine version. Used by the code that makes sure that FLinkerLoad's
@@ -894,10 +920,10 @@ public:
 // These will be private in FArchive
 protected:
 	/** Holds the archive version. */
-	int32 ArUE4Ver;
+	int32 ArUEVer;
 
 	/** Holds the archive version for licensees. */
-	int32 ArLicenseeUE4Ver;
+	int32 ArLicenseeUEVer;
 
 	/** Holds the engine version. */
 	FEngineVersionBase ArEngineVer;
@@ -1737,8 +1763,8 @@ public:
 		LogfImpl(Fmt, Args...);
 	}
 
-	using FArchiveState::UE4Ver;
-	using FArchiveState::LicenseeUE4Ver;
+	using FArchiveState::UEVer;
+	using FArchiveState::LicenseeUEVer;
 	using FArchiveState::EngineVer;
 	using FArchiveState::EngineNetVer;
 	using FArchiveState::GameNetVer;
@@ -1983,16 +2009,16 @@ public:
 	using FArchiveState::SetUseUnversionedPropertySerialization;
 	using FArchiveState::SetForceUnicode;
 	using FArchiveState::SetIsPersistent;
-	using FArchiveState::SetUE4Ver;
-	using FArchiveState::SetLicenseeUE4Ver;
+	using FArchiveState::SetUEVer;
+	using FArchiveState::SetLicenseeUEVer;
 	using FArchiveState::SetEngineVer;
 	using FArchiveState::SetEngineNetVer;
 	using FArchiveState::SetGameNetVer;
 	using FArchiveState::ShouldSkipCompilingAssets;
 	
 private:
-	using FArchiveState::ArUE4Ver;
-	using FArchiveState::ArLicenseeUE4Ver;
+	using FArchiveState::ArUEVer;
+	using FArchiveState::ArLicenseeUEVer;
 	using FArchiveState::ArEngineVer;
 	using FArchiveState::ArEngineNetVer;
 	using FArchiveState::ArGameNetVer;
diff --git a/Engine/Source/Runtime/Core/Public/UObject/ObjectVersion.h b/Engine/Source/Runtime/Core/Public/UObject/ObjectVersion.h
--- a/Engine/Source/Runtime/Core/Public/UObject/ObjectVersion.h
+++ b/Engine/Source/Runtime/Core/Public/UObject/ObjectVersion.h
@@ -13,7 +13,7 @@
 // This is useful as enum vales cannot be seen by the preprocessor.
 #define PREPROCESSOR_ENUM_PROTECT(a) ((unsigned int)(a))
 
-enum EUnrealEngineObjectUE4Version
+enum EUnrealEngineObjectUEVersion
 {
 	VER_UE4_OLDEST_LOADABLE_PACKAGE = 214,
 
@@ -648,7 +648,7 @@ enum EUnrealEngineObjectUE4Version
 	VER_UE4_AUTOMATIC_VERSION = VER_UE4_AUTOMATIC_VERSION_PLUS_ONE - 1
 };
 
-enum EUnrealEngineObjectLicenseeUE4Version
+enum EUnrealEngineObjectLicenseeUEVersion
 {
 	VER_LIC_NONE = 0,
 	// - this needs to be the last line (see note below)
diff --git a/Engine/Source/Runtime/CoreUObject/Private/Misc/WorldCompositionUtility.cpp b/Engine/Source/Runtime/CoreUObject/Private/Misc/WorldCompositionUtility.cpp
--- a/Engine/Source/Runtime/CoreUObject/Private/Misc/WorldCompositionUtility.cpp
+++ b/Engine/Source/Runtime/CoreUObject/Private/Misc/WorldCompositionUtility.cpp
@@ -16,12 +16,12 @@ FArchive& operator<<( FArchive& Ar, FWorldTileLayer& D )
 	// Serialized with FPackageFileSummary
 	Ar << D.Name << D.Reserved0 << D.Reserved1;
 		
-	if (Ar.UE4Ver() >= VER_UE4_WORLD_LEVEL_INFO_UPDATED)
+	if (Ar.UEVer() >= VER_UE4_WORLD_LEVEL_INFO_UPDATED)
 	{
 		Ar << D.StreamingDistance;
 	}
 
-	if (Ar.UE4Ver() >= VER_UE4_WORLD_LAYER_ENABLE_DISTANCE_STREAMING)
+	if (Ar.UEVer() >= VER_UE4_WORLD_LAYER_ENABLE_DISTANCE_STREAMING)
 	{
 		Ar << D.DistanceStreamingEnabled;
 	}
@@ -32,7 +32,7 @@ FArchive& operator<<( FArchive& Ar, FWorldTileLayer& D )
 void operator<<(FStructuredArchive::FSlot Slot, FWorldTileLayer& D)
 {
 	FStructuredArchive::FRecord Record = Slot.EnterRecord();
-	int32 Version = Slot.GetUnderlyingArchive().UE4Ver();
+	int32 Version = Slot.GetUnderlyingArchive().UEVer();
 
 	// Serialized with FPackageFileSummary
 	Record << SA_VALUE(TEXT("Name"), D.Name) << SA_VALUE(TEXT("Reserved0"), D.Reserved0) << SA_VALUE(TEXT("Reserved1"), D.Reserved1);
@@ -90,17 +90,17 @@ FArchive& operator<<( FArchive& Ar, FWorldTileInfo& D )
 	Ar << D.Bounds;
 	Ar << D.Layer;
 	
-	if (Ar.UE4Ver() >= VER_UE4_WORLD_LEVEL_INFO_UPDATED)
+	if (Ar.UEVer() >= VER_UE4_WORLD_LEVEL_INFO_UPDATED)
 	{
 		Ar << D.bHideInTileView << D.ParentTilePackageName;
 	}
 	
-	if (Ar.UE4Ver() >= VER_UE4_WORLD_LEVEL_INFO_LOD_LIST)
+	if (Ar.UEVer() >= VER_UE4_WORLD_LEVEL_INFO_LOD_LIST)
 	{
 		Ar << D.LODList;
 	}
 	
-	if (Ar.UE4Ver() >= VER_UE4_WORLD_LEVEL_INFO_ZORDER)
+	if (Ar.UEVer() >= VER_UE4_WORLD_LEVEL_INFO_ZORDER)
 	{
 		Ar << D.ZOrder;
 	}
@@ -116,7 +116,7 @@ FArchive& operator<<( FArchive& Ar, FWorldTileInfo& D )
 void operator<<(FStructuredArchive::FSlot Slot, FWorldTileInfo& D)
 {
 	FStructuredArchive::FRecord Record = Slot.EnterRecord();
-	int32 ArchiveVersion = Slot.GetUnderlyingArchive().UE4Ver();
+	int32 ArchiveVersion = Slot.GetUnderlyingArchive().UEVer();
 
 	// Serialized with FPackageFileSummary
 	Record << SA_VALUE(TEXT("Position"), D.Position) << SA_VALUE(TEXT("Bounds"), D.Bounds) << SA_VALUE(TEXT("Layer"), D.Layer);
@@ -173,9 +173,9 @@ bool FWorldTileInfo::Read(const FString& InPackageFileName, FWorldTileInfo& OutI
 		FileReader->Seek(FileSummary.WorldTileInfoDataOffset);
 
 		//make sure the filereader gets the correct version number (it defaults to latest version)
-		FileReader->SetUE4Ver(FileSummary.GetFileVersionUE4());
+		FileReader->SetUEVer(FileSummary.GetFileVersionUE4());
 		FileReader->SetEngineVer(FileSummary.SavedByEngineVersion);
-		FileReader->SetLicenseeUE4Ver(FileSummary.GetFileVersionLicenseeUE4());
+		FileReader->SetLicenseeUEVer(FileSummary.GetFileVersionLicenseeUE4());
 		FileReader->SetCustomVersions(FileSummary.GetCustomVersionContainer());
 		
 		// Load the structure
diff --git a/Engine/Source/Runtime/CoreUObject/Private/Serialization/AsyncLoading2.cpp b/Engine/Source/Runtime/CoreUObject/Private/Serialization/AsyncLoading2.cpp
--- a/Engine/Source/Runtime/CoreUObject/Private/Serialization/AsyncLoading2.cpp
+++ b/Engine/Source/Runtime/CoreUObject/Private/Serialization/AsyncLoading2.cpp
@@ -3309,8 +3309,8 @@ EAsyncPackageState::Type FAsyncPackage2::Event_ProcessExportBundle(FAsyncLoading
 		const uint64 AllExportDataSize = Package->IoBuffer.DataSize() - (Package->AllExportDataPtr - Package->IoBuffer.Data());
 		FExportArchive Ar(Package->AllExportDataPtr, Package->CurrentExportDataPtr, AllExportDataSize);
 		{
-			Ar.SetUE4Ver(Package->LinkerRoot->LinkerPackageVersion);
-			Ar.SetLicenseeUE4Ver(Package->LinkerRoot->LinkerLicenseeVersion);
+			Ar.SetUEVer(Package->LinkerRoot->LinkerPackageVersion);
+			Ar.SetLicenseeUEVer(Package->LinkerRoot->LinkerLicenseeVersion);
 			// Ar.SetEngineVer(Summary.SavedByEngineVersion); // very old versioning scheme
 			// Ar.SetCustomVersions(LinkerRoot->LinkerCustomVersion); // only if not cooking with -unversioned
 			Ar.SetUseUnversionedPropertySerialization((Package->LinkerRoot->GetPackageFlags() & PKG_UnversionedProperties) != 0);
diff --git a/Engine/Source/Runtime/CoreUObject/Private/UObject/Class.cpp b/Engine/Source/Runtime/CoreUObject/Private/UObject/Class.cpp
--- a/Engine/Source/Runtime/CoreUObject/Private/UObject/Class.cpp
+++ b/Engine/Source/Runtime/CoreUObject/Private/UObject/Class.cpp
@@ -1292,7 +1292,7 @@ void UStruct::SerializeVersionedTaggedProperties(FStructuredArchive::FSlot Slot,
 	//SCOPED_LOADTIMER(SerializeTaggedPropertiesTime);
 
 	// Determine if this struct supports optional property guid's (UBlueprintGeneratedClasses Only)
-	const bool bArePropertyGuidsAvailable = (UnderlyingArchive.UE4Ver() >= VER_UE4_PROPERTY_GUID_IN_PROPERTY_TAG) && !FPlatformProperties::RequiresCookedData() && ArePropertyGuidsAvailable();
+	const bool bArePropertyGuidsAvailable = (UnderlyingArchive.UEVer() >= VER_UE4_PROPERTY_GUID_IN_PROPERTY_TAG) && !FPlatformProperties::RequiresCookedData() && ArePropertyGuidsAvailable();
 	const bool bUseRedirects = (!FPlatformProperties::RequiresCookedData() || UnderlyingArchive.IsSaveGame()) && !UnderlyingArchive.IsUsingEventDrivenLoader();
 
 	if (UnderlyingArchive.IsLoading())
@@ -4431,7 +4431,7 @@ void UClass::Serialize( FArchive& Ar )
 	{
 		Ar << (uint32&)ClassFlags;
 	}
-	if (Ar.UE4Ver() < VER_UE4_CLASS_NOTPLACEABLE_ADDED)
+	if (Ar.UEVer() < VER_UE4_CLASS_NOTPLACEABLE_ADDED)
 	{
 		// We need to invert the CLASS_NotPlaceable flag here because it used to mean CLASS_Placeable
 		ClassFlags ^= CLASS_NotPlaceable;
@@ -4460,7 +4460,7 @@ void UClass::Serialize( FArchive& Ar )
 		Interfaces.Empty();
 
 		// In older versions, interface classes were serialized before linking. In case of cyclic dependencies, we need to skip over the serialized array and defer the load until after Link() is called below.
-		if(Ar.UE4Ver() < VER_UE4_UCLASS_SERIALIZE_INTERFACES_AFTER_LINKING && !GIsDuplicatingClassForReinstancing)
+		if(Ar.UEVer() < VER_UE4_UCLASS_SERIALIZE_INTERFACES_AFTER_LINKING && !GIsDuplicatingClassForReinstancing)
 		{
 			// Get our current position
 			InterfacesStart = Ar.Tell();
@@ -4501,7 +4501,7 @@ void UClass::Serialize( FArchive& Ar )
 		int64 CurrentOffset = Ar.Tell();
 
 		// In older versions, we need to seek backwards to the start of the interfaces array
-		if(Ar.UE4Ver() < VER_UE4_UCLASS_SERIALIZE_INTERFACES_AFTER_LINKING && !GIsDuplicatingClassForReinstancing)
+		if(Ar.UEVer() < VER_UE4_UCLASS_SERIALIZE_INTERFACES_AFTER_LINKING && !GIsDuplicatingClassForReinstancing)
 		{
 			Ar.Seek(InterfacesStart);
 		}
@@ -4517,7 +4517,7 @@ void UClass::Serialize( FArchive& Ar )
 		}
 
 		// In older versions, seek back to our current position after linking
-		if(Ar.UE4Ver() < VER_UE4_UCLASS_SERIALIZE_INTERFACES_AFTER_LINKING && !GIsDuplicatingClassForReinstancing)
+		if(Ar.UEVer() < VER_UE4_UCLASS_SERIALIZE_INTERFACES_AFTER_LINKING && !GIsDuplicatingClassForReinstancing)
 		{
 			Ar.Seek(CurrentOffset);
 		}
@@ -4533,7 +4533,7 @@ void UClass::Serialize( FArchive& Ar )
 	FName Dummy = NAME_None;
 	Ar << Dummy;
 
-	if (Ar.UE4Ver() >= VER_UE4_ADD_COOKED_TO_UCLASS)
+	if (Ar.UEVer() >= VER_UE4_ADD_COOKED_TO_UCLASS)
 	{
 		if (Ar.IsSaving())
 		{
@@ -5692,7 +5692,7 @@ void UFunction::Serialize( FArchive& Ar )
 	int32 EventGraphCallOffset = 0;
 #endif
 
-	if (Ar.UE4Ver() >= VER_UE4_SERIALIZE_BLUEPRINT_EVENTGRAPH_FASTCALLS_IN_UFUNCTION)
+	if (Ar.UEVer() >= VER_UE4_SERIALIZE_BLUEPRINT_EVENTGRAPH_FASTCALLS_IN_UFUNCTION)
 	{
 		Ar << EventGraphFunction;
 		Ar << EventGraphCallOffset;
diff --git a/Engine/Source/Runtime/CoreUObject/Private/UObject/Enum.cpp b/Engine/Source/Runtime/CoreUObject/Private/UObject/Enum.cpp
--- a/Engine/Source/Runtime/CoreUObject/Private/UObject/Enum.cpp
+++ b/Engine/Source/Runtime/CoreUObject/Private/UObject/Enum.cpp
@@ -37,7 +37,7 @@ void UEnum::Serialize( FArchive& Ar )
 	Super::Serialize(Ar);
 	if (Ar.IsLoading())
 	{
-		if (Ar.UE4Ver() < VER_UE4_TIGHTLY_PACKED_ENUMS)
+		if (Ar.UEVer() < VER_UE4_TIGHTLY_PACKED_ENUMS)
 		{
 			TArray<FName> TempNames;
 			Ar << TempNames;
@@ -67,7 +67,7 @@ void UEnum::Serialize( FArchive& Ar )
 		Ar << Names;
 	}
 
-	if (Ar.UE4Ver() < VER_UE4_ENUM_CLASS_SUPPORT)
+	if (Ar.UEVer() < VER_UE4_ENUM_CLASS_SUPPORT)
 	{
 		bool bIsNamespace;
 		Ar << bIsNamespace;
diff --git a/Engine/Source/Runtime/CoreUObject/Private/UObject/LinkerLoad.cpp b/Engine/Source/Runtime/CoreUObject/Private/UObject/LinkerLoad.cpp
--- a/Engine/Source/Runtime/CoreUObject/Private/UObject/LinkerLoad.cpp
+++ b/Engine/Source/Runtime/CoreUObject/Private/UObject/LinkerLoad.cpp
@@ -991,8 +991,8 @@ bool FLinkerLoad::IsTimeLimitExceeded( const TCHAR* CurrentTask, int32 Granulari
 void FLinkerLoad::ResetStatusInfo()
 {
 	// Set status info.
-	this->SetUE4Ver(GPackageFileUEVersion);
-	this->SetLicenseeUE4Ver(GPackageFileLicenseeUEVersion);
+	this->SetUEVer(GPackageFileUEVersion);
+	this->SetLicenseeUEVer(GPackageFileLicenseeUEVersion);
 	this->SetEngineVer(FEngineVersion::Current());
 	this->SetIsLoading(true);
 	this->SetIsPersistent(true);
@@ -1443,8 +1443,8 @@ FLinkerLoad::ELinkerStatus FLinkerLoad::UpdateFromPackageFileSummary()
 
 	const FCustomVersionContainer& SummaryVersions = Summary.GetCustomVersionContainer();
 
-	SetUE4Ver(Summary.GetFileVersionUE4());
-	SetLicenseeUE4Ver(Summary.GetFileVersionLicenseeUE4());
+	SetUEVer(Summary.GetFileVersionUE4());
+	SetLicenseeUEVer(Summary.GetFileVersionLicenseeUE4());
 	SetEngineVer(Summary.SavedByEngineVersion);
 	SetCustomVersions(SummaryVersions);
 
@@ -1466,8 +1466,8 @@ FLinkerLoad::ELinkerStatus FLinkerLoad::UpdateFromPackageFileSummary()
 	// Loader needs to be the same version.
 	if (Loader)
 	{
-		Loader->SetUE4Ver(Summary.GetFileVersionUE4());
-		Loader->SetLicenseeUE4Ver(Summary.GetFileVersionLicenseeUE4());
+		Loader->SetUEVer(Summary.GetFileVersionUE4());
+		Loader->SetLicenseeUEVer(Summary.GetFileVersionLicenseeUE4());
 		Loader->SetEngineVer(Summary.SavedByEngineVersion);
 		Loader->SetCustomVersions(SummaryVersions);
 	}
@@ -6000,7 +6000,7 @@ FLinkerLoad::ELinkerStatus FLinkerLoad::FixupExportMap()
 	DECLARE_SCOPE_CYCLE_COUNTER( TEXT( "FLinkerLoad::FixupExportMap" ), STAT_LinkerLoad_FixupExportMap, STATGROUP_LinkerLoad );
 
 #if WITH_EDITOR
-	if (UE4Ver() < VER_UE4_SKIP_DUPLICATE_EXPORTS_ON_SAVE_PACKAGE && !bExportsDuplicatesFixed)
+	if (UEVer() < VER_UE4_SKIP_DUPLICATE_EXPORTS_ON_SAVE_PACKAGE && !bExportsDuplicatesFixed)
 	{
 		FixupDuplicateExports();
 		bExportsDuplicatesFixed = true;
@@ -6029,7 +6029,7 @@ FLinkerLoad::ELinkerStatus FLinkerLoad::FixupExportMap()
 
 			// ActorComponents outered to a BlueprintGeneratedClass (or even older ones that are outered to Blueprint) need to be marked RF_Public, but older content was 
 			// not created as such.  This updates the ExportTable such that they are correctly flagged when created and when other packages validate their imports.
-			if (UE4Ver() < VER_UE4_BLUEPRINT_GENERATED_CLASS_COMPONENT_TEMPLATES_PUBLIC)
+			if (UEVer() < VER_UE4_BLUEPRINT_GENERATED_CLASS_COMPONENT_TEMPLATES_PUBLIC)
 			{
 				if ((Export.ObjectFlags & RF_Public) == 0)
 				{
diff --git a/Engine/Source/Runtime/CoreUObject/Private/UObject/ObjectBaseUtility.cpp b/Engine/Source/Runtime/CoreUObject/Private/UObject/ObjectBaseUtility.cpp
--- a/Engine/Source/Runtime/CoreUObject/Private/UObject/ObjectBaseUtility.cpp
+++ b/Engine/Source/Runtime/CoreUObject/Private/UObject/ObjectBaseUtility.cpp
@@ -34,7 +34,7 @@ int32 UObjectBaseUtility::GetLinkerUEVersion() const
 	if (Loader != nullptr)
 	{
 		// We have a linker so we can return its version.
-		return Loader->UE4Ver();
+		return Loader->UEVer();
 
 	}
 	else if (GetOutermost())
@@ -108,7 +108,7 @@ int32 UObjectBaseUtility::GetLinkerLicenseeUEVersion() const
 	if ( Loader != NULL )
 	{
 		// We have a linker so we can return its version.
-		return Loader->LicenseeUE4Ver();
+		return Loader->LicenseeUEVer();
 	}
 	else if (GetOutermost())
 	{
diff --git a/Engine/Source/Runtime/CoreUObject/Private/UObject/ObjectResource.cpp b/Engine/Source/Runtime/CoreUObject/Private/UObject/ObjectResource.cpp
--- a/Engine/Source/Runtime/CoreUObject/Private/UObject/ObjectResource.cpp
+++ b/Engine/Source/Runtime/CoreUObject/Private/UObject/ObjectResource.cpp
@@ -110,7 +110,7 @@ void operator<<(FStructuredArchive::FSlot Slot, FObjectExport& E)
 	Record << SA_VALUE(TEXT("ClassIndex"), E.ClassIndex);
 	Record << SA_VALUE(TEXT("SuperIndex"), E.SuperIndex);
 
-	if (BaseArchive.UE4Ver() >= VER_UE4_TemplateIndex_IN_COOKED_EXPORTS)
+	if (BaseArchive.UEVer() >= VER_UE4_TemplateIndex_IN_COOKED_EXPORTS)
 	{
 		Record << SA_VALUE(TEXT("TemplateIndex"), E.TemplateIndex);
 	}
@@ -126,7 +126,7 @@ void operator<<(FStructuredArchive::FSlot Slot, FObjectExport& E)
 		E.ObjectFlags = EObjectFlags(Save & RF_Load);
 	}
 
-	if (BaseArchive.UE4Ver() < VER_UE4_64BIT_EXPORTMAP_SERIALSIZES)
+	if (BaseArchive.UEVer() < VER_UE4_64BIT_EXPORTMAP_SERIALSIZES)
 	{
 		int32 SerialSize = E.SerialSize;
 		Record << SA_VALUE(TEXT("SerialSize"), SerialSize);
@@ -151,17 +151,17 @@ void operator<<(FStructuredArchive::FSlot Slot, FObjectExport& E)
 	PRAGMA_ENABLE_DEPRECATION_WARNINGS
 	Record << SA_VALUE(TEXT("PackageFlags"), E.PackageFlags);
 
-	if (BaseArchive.UE4Ver() >= VER_UE4_LOAD_FOR_EDITOR_GAME)
+	if (BaseArchive.UEVer() >= VER_UE4_LOAD_FOR_EDITOR_GAME)
 	{
 		Record << SA_VALUE(TEXT("bNotAlwaysLoadedForEditorGame"), E.bNotAlwaysLoadedForEditorGame);
 	}
 
-	if (BaseArchive.UE4Ver() >= VER_UE4_COOKED_ASSETS_IN_EDITOR_SUPPORT)
+	if (BaseArchive.UEVer() >= VER_UE4_COOKED_ASSETS_IN_EDITOR_SUPPORT)
 	{
 		Record << SA_VALUE(TEXT("bIsAsset"), E.bIsAsset);
 	}
 
-	if (BaseArchive.UE4Ver() >= VER_UE4_PRELOAD_DEPENDENCIES_IN_COOKED_EXPORTS)
+	if (BaseArchive.UEVer() >= VER_UE4_PRELOAD_DEPENDENCIES_IN_COOKED_EXPORTS)
 	{
 		Record << SA_VALUE(TEXT("FirstExportDependency"), E.FirstExportDependency);
 		Record << SA_VALUE(TEXT("SerializationBeforeSerializationDependencies"), E.SerializationBeforeSerializationDependencies);
@@ -291,7 +291,7 @@ void operator<<(FStructuredArchive::FSlot Slot, FObjectImport& I)
 
 	//@todo: re-enable package override at runtime when ready
 #if WITH_EDITORONLY_DATA
-	if (Slot.GetUnderlyingArchive().UE4Ver() >= VER_UE4_NON_OUTER_PACKAGE_IMPORT && !Slot.GetUnderlyingArchive().IsFilterEditorOnly())
+	if (Slot.GetUnderlyingArchive().UEVer() >= VER_UE4_NON_OUTER_PACKAGE_IMPORT && !Slot.GetUnderlyingArchive().IsFilterEditorOnly())
 	{
 		Record << SA_VALUE(TEXT("PackageName"), I.PackageName);
 	}
diff --git a/Engine/Source/Runtime/CoreUObject/Private/UObject/PropertyArray.cpp b/Engine/Source/Runtime/CoreUObject/Private/UObject/PropertyArray.cpp
--- a/Engine/Source/Runtime/CoreUObject/Private/UObject/PropertyArray.cpp
+++ b/Engine/Source/Runtime/CoreUObject/Private/UObject/PropertyArray.cpp
@@ -202,7 +202,7 @@ void FArrayProperty::SerializeItem(FStructuredArchive::FSlot Slot, void* Value,
 	ArrayHelper.CountBytes( UnderlyingArchive );
 
 	// Serialize a PropertyTag for the inner property of this array, allows us to validate the inner struct to see if it has changed
-	if (UnderlyingArchive.UE4Ver() >= VER_UE4_INNER_ARRAY_TAG_INFO && Inner->IsA<FStructProperty>())
+	if (UnderlyingArchive.UEVer() >= VER_UE4_INNER_ARRAY_TAG_INFO && Inner->IsA<FStructProperty>())
 	{
 		if (!MaybeInnerTag)
 		{
diff --git a/Engine/Source/Runtime/CoreUObject/Private/UObject/PropertyStruct.cpp b/Engine/Source/Runtime/CoreUObject/Private/UObject/PropertyStruct.cpp
--- a/Engine/Source/Runtime/CoreUObject/Private/UObject/PropertyStruct.cpp
+++ b/Engine/Source/Runtime/CoreUObject/Private/UObject/PropertyStruct.cpp
@@ -343,7 +343,7 @@ EConvertFromTypeResult FStructProperty::ConvertFromType(const FPropertyTag& Tag,
 
 	auto CanSerializeFromStructWithDifferentName = [](const FArchive& InAr, const FPropertyTag& PropertyTag, const FStructProperty* StructProperty)
 	{
-		if (InAr.UE4Ver() < VER_UE4_STRUCT_GUID_IN_PROPERTY_TAG)
+		if (InAr.UEVer() < VER_UE4_STRUCT_GUID_IN_PROPERTY_TAG)
 		{
 			// Old Implementation
 			return StructProperty && !StructProperty->UseBinaryOrNativeSerialization(InAr);
diff --git a/Engine/Source/Runtime/CoreUObject/Private/UObject/PropertyTag.cpp b/Engine/Source/Runtime/CoreUObject/Private/UObject/PropertyTag.cpp
--- a/Engine/Source/Runtime/CoreUObject/Private/UObject/PropertyTag.cpp
+++ b/Engine/Source/Runtime/CoreUObject/Private/UObject/PropertyTag.cpp
@@ -83,7 +83,7 @@ void operator<<(FStructuredArchive::FSlot Slot, FPropertyTag& Tag)
 	FArchive& UnderlyingArchive = Slot.GetUnderlyingArchive();
 	bool bIsTextFormat = UnderlyingArchive.IsTextFormat();
 
-	int32 Version = UnderlyingArchive.UE4Ver();
+	int32 Version = UnderlyingArchive.UEVer();
 
 	check(!UnderlyingArchive.GetArchiveState().UseUnversionedPropertySerialization());
 	checkf(!UnderlyingArchive.IsSaving() || Tag.Prop, TEXT("FPropertyTag must be constructed with a valid property when used for saving data!"));
diff --git a/Engine/Source/Runtime/CoreUObject/Private/UObject/SavePackage.cpp b/Engine/Source/Runtime/CoreUObject/Private/UObject/SavePackage.cpp
--- a/Engine/Source/Runtime/CoreUObject/Private/UObject/SavePackage.cpp
+++ b/Engine/Source/Runtime/CoreUObject/Private/UObject/SavePackage.cpp
@@ -2511,8 +2511,8 @@ FSavePackageResultStruct UPackage::Save(UPackage* InOuter, UObject* Base, EObjec
 				Linker->SetUseUnversionedPropertySerialization(bUseUnversionedProperties);
 
 				// Make sure the package has the same version as the linker
-				InOuter->LinkerPackageVersion = Linker->UE4Ver();
-				InOuter->LinkerLicenseeVersion = Linker->LicenseeUE4Ver();
+				InOuter->LinkerPackageVersion = Linker->UEVer();
+				InOuter->LinkerLicenseeVersion = Linker->LicenseeUEVer();
 				InOuter->LinkerCustomVersion = Linker->GetCustomVersions();
 
 				if (EndSavingIfCancelled())
diff --git a/Engine/Source/Runtime/CoreUObject/Private/UObject/SavePackage2.cpp b/Engine/Source/Runtime/CoreUObject/Private/UObject/SavePackage2.cpp
--- a/Engine/Source/Runtime/CoreUObject/Private/UObject/SavePackage2.cpp
+++ b/Engine/Source/Runtime/CoreUObject/Private/UObject/SavePackage2.cpp
@@ -785,8 +785,8 @@ ESavePackageResult BuildLinker(FSaveContext& SaveContext)
 		}
 #endif
 		// Make sure the package has the same version as the linker
-		SaveContext.GetPackage()->LinkerPackageVersion = SaveContext.Linker->UE4Ver();
-		SaveContext.GetPackage()->LinkerLicenseeVersion = SaveContext.Linker->LicenseeUE4Ver();
+		SaveContext.GetPackage()->LinkerPackageVersion = SaveContext.Linker->UEVer();
+		SaveContext.GetPackage()->LinkerLicenseeVersion = SaveContext.Linker->LicenseeUEVer();
 		SaveContext.GetPackage()->LinkerCustomVersion = SaveContext.Linker->GetCustomVersions();
 	}
 	
diff --git a/Engine/Source/Runtime/CoreUObject/Private/UObject/SoftObjectPath.cpp b/Engine/Source/Runtime/CoreUObject/Private/UObject/SoftObjectPath.cpp
--- a/Engine/Source/Runtime/CoreUObject/Private/UObject/SoftObjectPath.cpp
+++ b/Engine/Source/Runtime/CoreUObject/Private/UObject/SoftObjectPath.cpp
@@ -227,12 +227,12 @@ void FSoftObjectPath::SerializePath(FArchive& Ar)
 
 	if (bSerializeInternals)
 	{
-		if (Ar.IsLoading() && Ar.UE4Ver() < VER_UE4_ADDED_SOFT_OBJECT_PATH)
+		if (Ar.IsLoading() && Ar.UEVer() < VER_UE4_ADDED_SOFT_OBJECT_PATH)
 		{
 			FString Path;
 			Ar << Path;
 
-			if (Ar.UE4Ver() < VER_UE4_KEEP_ONLY_PACKAGE_NAMES_IN_STRING_ASSET_REFERENCES_MAP)
+			if (Ar.UEVer() < VER_UE4_KEEP_ONLY_PACKAGE_NAMES_IN_STRING_ASSET_REFERENCES_MAP)
 			{
 				Path = FPackageName::GetNormalizedObjectPath(Path);
 			}
diff --git a/Engine/Source/Runtime/CoreUObject/Private/UObject/StructScriptLoader.cpp b/Engine/Source/Runtime/CoreUObject/Private/UObject/StructScriptLoader.cpp
--- a/Engine/Source/Runtime/CoreUObject/Private/UObject/StructScriptLoader.cpp
+++ b/Engine/Source/Runtime/CoreUObject/Private/UObject/StructScriptLoader.cpp
@@ -205,7 +205,7 @@ bool FStructScriptLoader::LoadStructWithScript(UStruct* DestScriptContainer, FAr
 
 	// to help us move development forward (and not have to support ancient 
 	// script code), we define a minimum script version
-	bool bSkipScriptSerialization = (Ar.UE4Ver() < VER_MIN_SCRIPTVM_UE4) || (Ar.LicenseeUE4Ver() < VER_MIN_SCRIPTVM_LICENSEEUE4);
+	bool bSkipScriptSerialization = (Ar.UEVer() < VER_MIN_SCRIPTVM_UE4) || (Ar.LicenseeUEVer() < VER_MIN_SCRIPTVM_LICENSEEUE4);
 #if WITH_EDITOR
 	if (!bSkipScriptSerialization)
 	{
diff --git a/Engine/Source/Runtime/CoreUObject/Public/UObject/ScriptSerialization.h b/Engine/Source/Runtime/CoreUObject/Public/UObject/ScriptSerialization.h
--- a/Engine/Source/Runtime/CoreUObject/Public/UObject/ScriptSerialization.h
+++ b/Engine/Source/Runtime/CoreUObject/Public/UObject/ScriptSerialization.h
@@ -426,7 +426,7 @@
 		case EX_SetArray:
 		{
 			// If not loading, or its a newer version
-			if((!GetLinker()) || !Ar.IsLoading() || (Ar.UE4Ver() >= VER_UE4_CHANGE_SETARRAY_BYTECODE))
+			if((!GetLinker()) || !Ar.IsLoading() || (Ar.UEVer() >= VER_UE4_CHANGE_SETARRAY_BYTECODE))
 			{
 				// Array property to assign to
 				EExprToken TargetToken = SerializeExpr( iCode, Ar );
diff --git a/Engine/Source/Runtime/CoreUObject/Public/UObject/UObjectBaseUtility.h b/Engine/Source/Runtime/CoreUObject/Public/UObject/UObjectBaseUtility.h
--- a/Engine/Source/Runtime/CoreUObject/Public/UObject/UObjectBaseUtility.h
+++ b/Engine/Source/Runtime/CoreUObject/Public/UObject/UObjectBaseUtility.h
@@ -608,9 +608,9 @@ public:
 	int32 GetLinkerIndex() const;
 
 	/**
-	 * Returns the UE4 version of the linker for this object.
+	 * Returns the UE version of the linker for this object.
 	 *
-	 * @return	the UE4 version of the engine's package file when this object
+	 * @return	the UE version of the engine's package file when this object
 	 *			was last saved, or GPackageFileUEVersion (current version) if
 	 *			this object does not have a linker, which indicates that
 	 *			a) this object is a native only class, or
diff --git a/Engine/Source/Runtime/Engine/Classes/Animation/MorphTarget.h b/Engine/Source/Runtime/Engine/Classes/Animation/MorphTarget.h
--- a/Engine/Source/Runtime/Engine/Classes/Animation/MorphTarget.h
+++ b/Engine/Source/Runtime/Engine/Classes/Animation/MorphTarget.h
@@ -29,7 +29,7 @@ struct FMorphTargetDelta
 	/** pipe operator */
 	friend FArchive& operator<<(FArchive& Ar, FMorphTargetDelta& V)
 	{
-		if ((Ar.UE4Ver() < VER_UE4_MORPHTARGET_CPU_TANGENTZDELTA_FORMATCHANGE) && Ar.IsLoading())
+		if ((Ar.UEVer() < VER_UE4_MORPHTARGET_CPU_TANGENTZDELTA_FORMATCHANGE) && Ar.IsLoading())
 		{
 			/** old format of change in tangent basis normal */
 			FDeprecatedSerializedPackedNormal TangentZDelta_DEPRECATED;
diff --git a/Engine/Source/Runtime/Engine/Private/AI/Navigation/NavigationModifier.cpp b/Engine/Source/Runtime/Engine/Private/AI/Navigation/NavigationModifier.cpp
--- a/Engine/Source/Runtime/Engine/Private/AI/Navigation/NavigationModifier.cpp
+++ b/Engine/Source/Runtime/Engine/Private/AI/Navigation/NavigationModifier.cpp
@@ -55,7 +55,7 @@ bool FNavigationLinkBase::HasMetaArea() const
 #if WITH_EDITORONLY_DATA
 void FNavigationLinkBase::PostSerialize(const FArchive& Ar)
 {
-	if (Ar.IsLoading() && Ar.UE4Ver() < VER_UE4_NAVIGATION_AGENT_SELECTOR)
+	if (Ar.IsLoading() && Ar.UEVer() < VER_UE4_NAVIGATION_AGENT_SELECTOR)
 	{
 		SupportedAgents.bSupportsAgent0 = bSupportsAgent0;
 		SupportedAgents.bSupportsAgent1 = bSupportsAgent1;
diff --git a/Engine/Source/Runtime/Engine/Private/Animation/AnimCurveTypes.cpp b/Engine/Source/Runtime/Engine/Private/Animation/AnimCurveTypes.cpp
--- a/Engine/Source/Runtime/Engine/Private/Animation/AnimCurveTypes.cpp
+++ b/Engine/Source/Runtime/Engine/Private/Animation/AnimCurveTypes.cpp
@@ -21,7 +21,7 @@ void FAnimCurveBase::PostSerialize(FArchive& Ar)
 	{
 		if (Ar.CustomVer(FFrameworkObjectVersion::GUID) < FFrameworkObjectVersion::SmartNameRefactor)
 		{
-			if (Ar.UE4Ver() >= VER_UE4_SKELETON_ADD_SMARTNAMES)
+			if (Ar.UEVer() >= VER_UE4_SKELETON_ADD_SMARTNAMES)
 			{
 				Ar << CurveUid;
 
@@ -652,7 +652,7 @@ void FRawCurveTracks::PostSerialize(FArchive& Ar)
 #if WITH_EDITORONLY_DATA
 	if( !Ar.IsCooking() )
 	{
-		if( Ar.UE4Ver() >= VER_UE4_ANIMATION_ADD_TRACKCURVES )
+		if( Ar.UEVer() >= VER_UE4_ANIMATION_ADD_TRACKCURVES )
 		{
 			for( FTransformCurve& Curve : TransformCurves )
 			{
diff --git a/Engine/Source/Runtime/Engine/Private/Animation/AnimMontage.cpp b/Engine/Source/Runtime/Engine/Private/Animation/AnimMontage.cpp
--- a/Engine/Source/Runtime/Engine/Private/Animation/AnimMontage.cpp
+++ b/Engine/Source/Runtime/Engine/Private/Animation/AnimMontage.cpp
@@ -476,7 +476,7 @@ void UAnimMontage::PostLoad()
 	}
 
 	// Convert BranchingPoints to AnimNotifies.
-	if (GetLinker() && (GetLinker()->UE4Ver() < VER_UE4_MONTAGE_BRANCHING_POINT_REMOVAL) )
+	if (GetLinker() && (GetLinker()->UEVer() < VER_UE4_MONTAGE_BRANCHING_POINT_REMOVAL) )
 	{
 		ConvertBranchingPointsToAnimNotifies();
 	}
diff --git a/Engine/Source/Runtime/Engine/Private/Animation/AnimSequence.cpp b/Engine/Source/Runtime/Engine/Private/Animation/AnimSequence.cpp
--- a/Engine/Source/Runtime/Engine/Private/Animation/AnimSequence.cpp
+++ b/Engine/Source/Runtime/Engine/Private/Animation/AnimSequence.cpp
@@ -592,7 +592,7 @@ void UAnimSequence::Serialize(FArchive& Ar)
 		PRAGMA_ENABLE_DEPRECATION_WARNINGS
 		if (!Ar.IsCooking())
 		{
-			if (Ar.UE4Ver() >= VER_UE4_ANIMATION_ADD_TRACKCURVES)
+			if (Ar.UEVer() >= VER_UE4_ANIMATION_ADD_TRACKCURVES)
 			{
 				if (Ar.CustomVer(FUE5MainStreamObjectVersion::GUID) < FUE5MainStreamObjectVersion::RemovingSourceAnimationData)
 				{
@@ -694,14 +694,14 @@ void UAnimSequence::Serialize(FArchive& Ar)
 		PRAGMA_ENABLE_DEPRECATION_WARNINGS
 	}
 
-	if ( Ar.IsLoading() && Ar.UE4Ver() < VER_UE4_ASSET_IMPORT_DATA_AS_JSON && !AssetImportData)
+	if ( Ar.IsLoading() && Ar.UEVer() < VER_UE4_ASSET_IMPORT_DATA_AS_JSON && !AssetImportData)
 	{
 		// AssetImportData should always be valid
 		AssetImportData = NewObject<UAssetImportData>(this, TEXT("AssetImportData"));
 	}
 
 	// SourceFilePath and SourceFileTimestamp were moved into a subobject
-	if ( Ar.IsLoading() && Ar.UE4Ver() < VER_UE4_ADDED_FBX_ASSET_IMPORT_DATA && AssetImportData)
+	if ( Ar.IsLoading() && Ar.UEVer() < VER_UE4_ADDED_FBX_ASSET_IMPORT_DATA && AssetImportData)
 	{
 		// AssetImportData should always have been set up in the constructor where this is relevant
 		FAssetImportInfo Info;
diff --git a/Engine/Source/Runtime/Engine/Private/Animation/ReferenceSkeleton.cpp b/Engine/Source/Runtime/Engine/Private/Animation/ReferenceSkeleton.cpp
--- a/Engine/Source/Runtime/Engine/Private/Animation/ReferenceSkeleton.cpp
+++ b/Engine/Source/Runtime/Engine/Private/Animation/ReferenceSkeleton.cpp
@@ -33,14 +33,14 @@ FArchive &operator<<(FArchive& Ar, FMeshBoneInfo& F)
 {
 	Ar << F.Name << F.ParentIndex;
 
-	if (Ar.IsLoading() && (Ar.UE4Ver() < VER_UE4_REFERENCE_SKELETON_REFACTOR))
+	if (Ar.IsLoading() && (Ar.UEVer() < VER_UE4_REFERENCE_SKELETON_REFACTOR))
 	{
 		FColor DummyColor = FColor::White;
 		Ar << DummyColor;
 	}
 
 #if WITH_EDITORONLY_DATA
-	if (Ar.UE4Ver() >= VER_UE4_STORE_BONE_EXPORT_NAMES)
+	if (Ar.UEVer() >= VER_UE4_STORE_BONE_EXPORT_NAMES)
 	{
 		if (!Ar.IsCooking() && !Ar.IsFilterEditorOnly())
 		{
@@ -345,13 +345,13 @@ FArchive & operator<<(FArchive & Ar, FReferenceSkeleton & F)
 	Ar << F.RawRefBoneInfo;
 	Ar << F.RawRefBonePose;
 
-	if (Ar.UE4Ver() >= VER_UE4_REFERENCE_SKELETON_REFACTOR)
+	if (Ar.UEVer() >= VER_UE4_REFERENCE_SKELETON_REFACTOR)
 	{
 		Ar << F.RawNameToIndexMap;
 	}
 
 	// Fix up any assets that don't have an INDEX_NONE parent for Bone[0]
-	if (Ar.IsLoading() && Ar.UE4Ver() < VER_UE4_FIXUP_ROOTBONE_PARENT)
+	if (Ar.IsLoading() && Ar.UEVer() < VER_UE4_FIXUP_ROOTBONE_PARENT)
 	{
 		if ((F.RawRefBoneInfo.Num() > 0) && (F.RawRefBoneInfo[0].ParentIndex != INDEX_NONE))
 		{
diff --git a/Engine/Source/Runtime/Engine/Private/Animation/Skeleton.cpp b/Engine/Source/Runtime/Engine/Private/Animation/Skeleton.cpp
--- a/Engine/Source/Runtime/Engine/Private/Animation/Skeleton.cpp
+++ b/Engine/Source/Runtime/Engine/Private/Animation/Skeleton.cpp
@@ -505,7 +505,7 @@ void USkeleton::PostLoad()
 {
 	Super::PostLoad();
 
-	if( GetLinker() && (GetLinker()->UE4Ver() < VER_UE4_REFERENCE_SKELETON_REFACTOR) )
+	if( GetLinker() && (GetLinker()->UEVer() < VER_UE4_REFERENCE_SKELETON_REFACTOR) )
 	{
 		// Convert RefLocalPoses & BoneTree to FReferenceSkeleton
 		ConvertToFReferenceSkeleton();
@@ -664,12 +664,12 @@ void USkeleton::Serialize( FArchive& Ar )
 
 	Super::Serialize(Ar);
 
-	if( Ar.UE4Ver() >= VER_UE4_REFERENCE_SKELETON_REFACTOR )
+	if( Ar.UEVer() >= VER_UE4_REFERENCE_SKELETON_REFACTOR )
 	{
 		Ar << ReferenceSkeleton;
 	}
 
-	if (Ar.UE4Ver() >= VER_UE4_FIX_ANIMATIONBASEPOSE_SERIALIZATION)
+	if (Ar.UEVer() >= VER_UE4_FIX_ANIMATIONBASEPOSE_SERIALIZATION)
 	{
 		// Load Animation RetargetSources
 		if (Ar.IsLoading())
@@ -710,7 +710,7 @@ void USkeleton::Serialize( FArchive& Ar )
 		}
 	}
 
-	if (Ar.UE4Ver() < VER_UE4_SKELETON_GUID_SERIALIZATION)
+	if (Ar.UEVer() < VER_UE4_SKELETON_GUID_SERIALIZATION)
 	{
 		UE_LOG(LogAnimation, Warning, TEXT("Skeleton '%s' has not been saved since version 'VER_UE4_SKELETON_GUID_SERIALIZATION' This asset will not cook deterministically until it is resaved."), *GetPathName());
 		RegenerateGuid();
@@ -721,7 +721,7 @@ void USkeleton::Serialize( FArchive& Ar )
 	}
 
 	// If we should be using smartnames, serialize the mappings
-	if(Ar.UE4Ver() >= VER_UE4_SKELETON_ADD_SMARTNAMES)
+	if(Ar.UEVer() >= VER_UE4_SKELETON_ADD_SMARTNAMES)
 	{
 		SmartNames.Serialize(Ar, IsTemplate());
 
@@ -729,7 +729,7 @@ void USkeleton::Serialize( FArchive& Ar )
 	}
 
 	// Build look up table between Slot nodes and their Group.
-	if(Ar.UE4Ver() < VER_UE4_FIX_SLOT_NAME_DUPLICATION)
+	if(Ar.UEVer() < VER_UE4_FIX_SLOT_NAME_DUPLICATION)
 	{
 		// In older assets we may have duplicates, remove these while building the map.
 		BuildSlotToGroupMap(true);
@@ -740,7 +740,7 @@ void USkeleton::Serialize( FArchive& Ar )
 	}
 
 #if WITH_EDITORONLY_DATA
-	if (Ar.UE4Ver() < VER_UE4_SKELETON_ASSET_PROPERTY_TYPE_CHANGE)
+	if (Ar.UEVer() < VER_UE4_SKELETON_ASSET_PROPERTY_TYPE_CHANGE)
 	{
 		PreviewAttachedAssetContainer.SaveAttachedObjectsFromDeprecatedProperties();
 	}
diff --git a/Engine/Source/Runtime/Engine/Private/Animation/SkinWeightProfile.cpp b/Engine/Source/Runtime/Engine/Private/Animation/SkinWeightProfile.cpp
--- a/Engine/Source/Runtime/Engine/Private/Animation/SkinWeightProfile.cpp
+++ b/Engine/Source/Runtime/Engine/Private/Animation/SkinWeightProfile.cpp
@@ -138,7 +138,7 @@ FAutoConsoleVariableRef CVarSkinWeightProfilesAllowedFromLOD(
 FArchive& operator<<(FArchive& Ar, FRuntimeSkinWeightProfileData& OverrideData)
 {
 #if WITH_EDITOR
-	if (Ar.UE4Ver() < VER_UE4_SKINWEIGHT_PROFILE_DATA_LAYOUT_CHANGES)
+	if (Ar.UEVer() < VER_UE4_SKINWEIGHT_PROFILE_DATA_LAYOUT_CHANGES)
 	{
 		Ar << OverrideData.OverridesInfo_DEPRECATED;
 		Ar << OverrideData.Weights_DEPRECATED;
diff --git a/Engine/Source/Runtime/Engine/Private/Animation/SmartName.cpp b/Engine/Source/Runtime/Engine/Private/Animation/SmartName.cpp
--- a/Engine/Source/Runtime/Engine/Private/Animation/SmartName.cpp
+++ b/Engine/Source/Runtime/Engine/Private/Animation/SmartName.cpp
@@ -100,7 +100,7 @@ void FSmartNameMapping::Serialize(FArchive& Ar)
 			Ar << TempGuidMap;
 		}
 	}
-	else if(Ar.UE4Ver() >= VER_UE4_SKELETON_ADD_SMARTNAMES)
+	else if(Ar.UEVer() >= VER_UE4_SKELETON_ADD_SMARTNAMES)
 	{
 		SmartName::UID_Type NextUidTemp;
 		Ar << NextUidTemp;
diff --git a/Engine/Source/Runtime/Engine/Private/AnimationAsset.cpp b/Engine/Source/Runtime/Engine/Private/AnimationAsset.cpp
--- a/Engine/Source/Runtime/Engine/Private/AnimationAsset.cpp
+++ b/Engine/Source/Runtime/Engine/Private/AnimationAsset.cpp
@@ -280,7 +280,7 @@ void UAnimationAsset::Serialize(FArchive& Ar)
 
 	Super::Serialize(Ar);
 
-	if (Ar.UE4Ver() >= VER_UE4_SKELETON_GUID_SERIALIZATION)
+	if (Ar.UEVer() >= VER_UE4_SKELETON_GUID_SERIALIZATION)
 	{
 		Ar << SkeletonGuid;
 	}
diff --git a/Engine/Source/Runtime/Engine/Private/Atmosphere/AtmosphericFogComponent.cpp b/Engine/Source/Runtime/Engine/Private/Atmosphere/AtmosphericFogComponent.cpp
--- a/Engine/Source/Runtime/Engine/Private/Atmosphere/AtmosphericFogComponent.cpp
+++ b/Engine/Source/Runtime/Engine/Private/Atmosphere/AtmosphericFogComponent.cpp
@@ -638,12 +638,12 @@ void UAtmosphericFogComponent::Serialize(FArchive& Ar)
 {
 	Super::Serialize(Ar);
 
-	if (Ar.IsLoading() && Ar.UE4Ver() < VER_UE4_ATMOSPHERIC_FOG_DECAY_NAME_CHANGE)
+	if (Ar.IsLoading() && Ar.UEVer() < VER_UE4_ATMOSPHERIC_FOG_DECAY_NAME_CHANGE)
 	{
 		PrecomputeParams.DensityHeight = PrecomputeParams.DecayHeight_DEPRECATED;
 	}
 
-	if (Ar.UE4Ver() >= VER_UE4_ATMOSPHERIC_FOG_CACHE_DATA)
+	if (Ar.UEVer() >= VER_UE4_ATMOSPHERIC_FOG_CACHE_DATA)
 	{
 		TransmittanceData.Serialize(Ar, this, INDEX_NONE, false);
 		IrradianceData.Serialize(Ar, this, INDEX_NONE, false);
@@ -668,7 +668,7 @@ void UAtmosphericFogComponent::Serialize(FArchive& Ar)
 		Ar << CounterVal;
 	}
 
-	if (Ar.IsLoading() && Ar.UE4Ver() < VER_UE4_ATMOSPHERIC_FOG_CACHE_DATA && PrecomputeCounter == EValid) 
+	if (Ar.IsLoading() && Ar.UEVer() < VER_UE4_ATMOSPHERIC_FOG_CACHE_DATA && PrecomputeCounter == EValid) 
 	{
 		// InscatterAltitudeSampleNum default value has been changed (32 -> 2)
 		// Recalculate InscatterAltitudeSampleNum based on Inscatter Size
diff --git a/Engine/Source/Runtime/Engine/Private/AudioVolume.cpp b/Engine/Source/Runtime/Engine/Private/AudioVolume.cpp
--- a/Engine/Source/Runtime/Engine/Private/AudioVolume.cpp
+++ b/Engine/Source/Runtime/Engine/Private/AudioVolume.cpp
@@ -28,7 +28,7 @@ FInteriorSettings::FInteriorSettings()
 #if WITH_EDITORONLY_DATA
 void FInteriorSettings::PostSerialize(const FArchive& Ar)
 {
-	if (Ar.UE4Ver() < VER_UE4_USE_LOW_PASS_FILTER_FREQ)
+	if (Ar.UEVer() < VER_UE4_USE_LOW_PASS_FILTER_FREQ)
 	{
 		if (InteriorLPF > 0.0f && InteriorLPF < 1.0f)
 		{
diff --git a/Engine/Source/Runtime/Engine/Private/Blueprint.cpp b/Engine/Source/Runtime/Engine/Private/Blueprint.cpp
--- a/Engine/Source/Runtime/Engine/Private/Blueprint.cpp
+++ b/Engine/Source/Runtime/Engine/Private/Blueprint.cpp
@@ -285,8 +285,8 @@ void UBlueprintCore::Serialize(FArchive& Ar)
 	}
 #endif
 
-	if ((Ar.UE4Ver() < VER_UE4_BLUEPRINT_SKEL_CLASS_TRANSIENT_AGAIN)
-		&& (Ar.UE4Ver() != VER_UE4_BLUEPRINT_SKEL_TEMPORARY_TRANSIENT))
+	if ((Ar.UEVer() < VER_UE4_BLUEPRINT_SKEL_CLASS_TRANSIENT_AGAIN)
+		&& (Ar.UEVer() != VER_UE4_BLUEPRINT_SKEL_TEMPORARY_TRANSIENT))
 	{
 		Ar << SkeletonGeneratedClass;
 		if( SkeletonGeneratedClass )
@@ -426,7 +426,7 @@ void UBlueprint::Serialize(FArchive& Ar)
 	Super::Serialize(Ar);
 
 #if WITH_EDITORONLY_DATA
-	if(Ar.IsLoading() && Ar.UE4Ver() < VER_UE4_BLUEPRINT_VARS_NOT_READ_ONLY)
+	if(Ar.IsLoading() && Ar.UEVer() < VER_UE4_BLUEPRINT_VARS_NOT_READ_ONLY)
 	{
 		// Allow all blueprint defined vars to be read/write.  undoes previous convention of making exposed variables read-only
 		for (int32 i = 0; i < NewVariables.Num(); ++i)
@@ -436,7 +436,7 @@ void UBlueprint::Serialize(FArchive& Ar)
 		}
 	}
 
-	if (Ar.UE4Ver() < VER_UE4_K2NODE_REFERENCEGUIDS)
+	if (Ar.UEVer() < VER_UE4_K2NODE_REFERENCEGUIDS)
 	{
 		for (int32 Index = 0; Index < NewVariables.Num(); ++Index)
 		{
@@ -463,7 +463,7 @@ void UBlueprint::Serialize(FArchive& Ar)
 
 		// Actor variables can't have default values (because Blueprint templates are library elements that can 
 		// bridge multiple levels and different levels might not have the actor that the default is referencing).
-		if (Ar.UE4Ver() < VER_UE4_FIX_BLUEPRINT_VARIABLE_FLAGS)
+		if (Ar.UEVer() < VER_UE4_FIX_BLUEPRINT_VARIABLE_FLAGS)
 		{
 			const FEdGraphPinType& VarType = Variable.VarType;
 
diff --git a/Engine/Source/Runtime/Engine/Private/Camera/CameraActor.cpp b/Engine/Source/Runtime/Engine/Private/Camera/CameraActor.cpp
--- a/Engine/Source/Runtime/Engine/Private/Camera/CameraActor.cpp
+++ b/Engine/Source/Runtime/Engine/Private/Camera/CameraActor.cpp
@@ -40,7 +40,7 @@ void ACameraActor::Serialize(FArchive& Ar)
 {
 	Super::Serialize(Ar);
 
-	if ((Ar.UE4Ver() < VER_UE4_CAMERA_ACTOR_USING_CAMERA_COMPONENT) && Ar.IsLoading())
+	if ((Ar.UEVer() < VER_UE4_CAMERA_ACTOR_USING_CAMERA_COMPONENT) && Ar.IsLoading())
 	{
 		CameraComponent->bConstrainAspectRatio = bConstrainAspectRatio_DEPRECATED;
 		CameraComponent->ProjectionMode = ECameraProjectionMode::Perspective;
diff --git a/Engine/Source/Runtime/Engine/Private/Camera/CameraComponent.cpp b/Engine/Source/Runtime/Engine/Private/Camera/CameraComponent.cpp
--- a/Engine/Source/Runtime/Engine/Private/Camera/CameraComponent.cpp
+++ b/Engine/Source/Runtime/Engine/Private/Camera/CameraComponent.cpp
@@ -149,9 +149,9 @@ void UCameraComponent::PostLoad()
 {
 	Super::PostLoad();
 
-	const int32 LinkerUE4Ver = GetLinkerUEVersion();
+	const int32 LinkerUEVer = GetLinkerUEVersion();
 
-	if (LinkerUE4Ver < VER_UE4_RENAME_CAMERA_COMPONENT_VIEW_ROTATION)
+	if (LinkerUEVer < VER_UE4_RENAME_CAMERA_COMPONENT_VIEW_ROTATION)
 	{
 		bUsePawnControlRotation = bUseControllerViewRotation_DEPRECATED;
 	}
diff --git a/Engine/Source/Runtime/Engine/Private/Components/CapsuleComponent.cpp b/Engine/Source/Runtime/Engine/Private/Components/CapsuleComponent.cpp
--- a/Engine/Source/Runtime/Engine/Private/Components/CapsuleComponent.cpp
+++ b/Engine/Source/Runtime/Engine/Private/Components/CapsuleComponent.cpp
@@ -118,9 +118,9 @@ void UCapsuleComponent::Serialize(FArchive& Ar)
 {
 	Super::Serialize(Ar);
 
-	if (Ar.IsLoading() && (Ar.UE4Ver() < VER_UE4_AFTER_CAPSULE_HALF_HEIGHT_CHANGE))
+	if (Ar.IsLoading() && (Ar.UEVer() < VER_UE4_AFTER_CAPSULE_HALF_HEIGHT_CHANGE))
 	{
-		if ((CapsuleHeight_DEPRECATED != 0.0f) || (Ar.UE4Ver() < VER_UE4_BLUEPRINT_VARS_NOT_READ_ONLY))
+		if ((CapsuleHeight_DEPRECATED != 0.0f) || (Ar.UEVer() < VER_UE4_BLUEPRINT_VARS_NOT_READ_ONLY))
 		{
 			CapsuleHalfHeight = CapsuleHeight_DEPRECATED;
 			CapsuleHeight_DEPRECATED = 0.0f;
diff --git a/Engine/Source/Runtime/Engine/Private/Components/CharacterMovementComponent.cpp b/Engine/Source/Runtime/Engine/Private/Components/CharacterMovementComponent.cpp
--- a/Engine/Source/Runtime/Engine/Private/Components/CharacterMovementComponent.cpp
+++ b/Engine/Source/Runtime/Engine/Private/Components/CharacterMovementComponent.cpp
@@ -1,4 +1,5 @@
 // Copyright Epic Games, Inc. All Rights Reserved.
+// Copyright Epic Games, Inc. All Rights Reserved.
 
 /*=============================================================================
 	Movement.cpp: Character movement implementation
@@ -575,14 +576,14 @@ void UCharacterMovementComponent::PostLoad()
 	Super::PostLoad();
 
 #if WITH_EDITORONLY_DATA
-	const int32 LinkerUE4Ver = GetLinkerUEVersion();
+	const int32 LinkerUEVer = GetLinkerUEVersion();
 
-	if (LinkerUE4Ver < VER_UE4_CHARACTER_MOVEMENT_DECELERATION)
+	if (LinkerUEVer < VER_UE4_CHARACTER_MOVEMENT_DECELERATION)
 	{
 		BrakingDecelerationWalking = MaxAcceleration;
 	}
 
-	if (LinkerUE4Ver < VER_UE4_CHARACTER_BRAKING_REFACTOR)
+	if (LinkerUEVer < VER_UE4_CHARACTER_BRAKING_REFACTOR)
 	{
 		// This bool used to apply walking braking in flying and swimming modes.
 		if (bForceBraking_DEPRECATED)
@@ -592,13 +593,13 @@ void UCharacterMovementComponent::PostLoad()
 		}
 	}
 
-	if (LinkerUE4Ver < VER_UE4_CHARACTER_MOVEMENT_WALKABLE_FLOOR_REFACTOR)
+	if (LinkerUEVer < VER_UE4_CHARACTER_MOVEMENT_WALKABLE_FLOOR_REFACTOR)
 	{
 		// Compute the walkable floor angle, since we have never done so yet.
 		UCharacterMovementComponent::SetWalkableFloorZ(WalkableFloorZ);
 	}
 
-	if (LinkerUE4Ver < VER_UE4_DEPRECATED_MOVEMENTCOMPONENT_MODIFIED_SPEEDS)
+	if (LinkerUEVer < VER_UE4_DEPRECATED_MOVEMENTCOMPONENT_MODIFIED_SPEEDS)
 	{
 		MaxWalkSpeedCrouched = MaxWalkSpeed * CrouchedSpeedMultiplier_DEPRECATED;
 		MaxCustomMovementSpeed = MaxWalkSpeed;
@@ -1280,7 +1281,7 @@ void UCharacterMovementComponent::Serialize(FArchive& Archive)
 {
 	Super::Serialize(Archive);
 
-	if (Archive.IsLoading() && Archive.UE4Ver() < VER_UE4_ADDED_SWEEP_WHILE_WALKING_FLAG)
+	if (Archive.IsLoading() && Archive.UEVer() < VER_UE4_ADDED_SWEEP_WHILE_WALKING_FLAG)
 	{
 		// We need to update the bSweepWhileNavWalking flag to match the previous behavior.
 		// Since UpdatedComponent is transient, we'll have to wait until we're registered.
diff --git a/Engine/Source/Runtime/Engine/Private/Components/DecalComponent.cpp b/Engine/Source/Runtime/Engine/Private/Components/DecalComponent.cpp
--- a/Engine/Source/Runtime/Engine/Private/Components/DecalComponent.cpp
+++ b/Engine/Source/Runtime/Engine/Private/Components/DecalComponent.cpp
@@ -121,7 +121,7 @@ void UDecalComponent::Serialize(FArchive& Ar)
 {
 	Super::Serialize(Ar);
 
-	if (Ar.UE4Ver() < VER_UE4_DECAL_SIZE)
+	if (Ar.UEVer() < VER_UE4_DECAL_SIZE)
 	{
 		DecalSize = FVector(1.0f, 1.0f, 1.0f);
 	}
diff --git a/Engine/Source/Runtime/Engine/Private/Components/DirectionalLightComponent.cpp b/Engine/Source/Runtime/Engine/Private/Components/DirectionalLightComponent.cpp
--- a/Engine/Source/Runtime/Engine/Private/Components/DirectionalLightComponent.cpp
+++ b/Engine/Source/Runtime/Engine/Private/Components/DirectionalLightComponent.cpp
@@ -1245,7 +1245,7 @@ void UDirectionalLightComponent::Serialize(FArchive& Ar)
 {
 	Super::Serialize(Ar);
 
-	if(Ar.UE4Ver() < VER_UE4_REMOVE_LIGHT_MOBILITY_CLASSES)
+	if(Ar.UEVer() < VER_UE4_REMOVE_LIGHT_MOBILITY_CLASSES)
 	{
 		// If outer is a DirectionalLight, we use the ADirectionalLight::LoadedFromAnotherClass path
 		if( GetOuter() != NULL && !GetOuter()->IsA(ADirectionalLight::StaticClass()) )
diff --git a/Engine/Source/Runtime/Engine/Private/Components/LightComponent.cpp b/Engine/Source/Runtime/Engine/Private/Components/LightComponent.cpp
--- a/Engine/Source/Runtime/Engine/Private/Components/LightComponent.cpp
+++ b/Engine/Source/Runtime/Engine/Private/Components/LightComponent.cpp
@@ -130,7 +130,7 @@ void ULightComponentBase::Serialize(FArchive& Ar)
 {
 	Super::Serialize(Ar);
 
-	if (Ar.UE4Ver() < VER_UE4_INVERSE_SQUARED_LIGHTS_DEFAULT)
+	if (Ar.UEVer() < VER_UE4_INVERSE_SQUARED_LIGHTS_DEFAULT)
 	{
 		Intensity = Brightness_DEPRECATED;
 	}
@@ -550,7 +550,7 @@ void ULightComponent::Serialize(FArchive& Ar)
 
 	Ar.UsingCustomVersion(FRenderingObjectVersion::GUID);
 
-	if (Ar.UE4Ver() >= VER_UE4_STATIC_SHADOW_DEPTH_MAPS)
+	if (Ar.UEVer() >= VER_UE4_STATIC_SHADOW_DEPTH_MAPS)
 	{
 		if (Ar.IsLoading() && Ar.CustomVer(FRenderingObjectVersion::GUID) < FRenderingObjectVersion::MapBuildDataSeparatePackage)
 		{
diff --git a/Engine/Source/Runtime/Engine/Private/Components/LocalLightComponent.cpp b/Engine/Source/Runtime/Engine/Private/Components/LocalLightComponent.cpp
--- a/Engine/Source/Runtime/Engine/Private/Components/LocalLightComponent.cpp
+++ b/Engine/Source/Runtime/Engine/Private/Components/LocalLightComponent.cpp
@@ -86,7 +86,7 @@ void ULocalLightComponent::Serialize(FArchive& Ar)
 {
 	Super::Serialize(Ar);
 
-	if (Ar.UE4Ver() < VER_UE4_INVERSE_SQUARED_LIGHTS_DEFAULT)
+	if (Ar.UEVer() < VER_UE4_INVERSE_SQUARED_LIGHTS_DEFAULT)
 	{
 		AttenuationRadius = Radius_DEPRECATED;
 	}
diff --git a/Engine/Source/Runtime/Engine/Private/Components/ModelComponent.cpp b/Engine/Source/Runtime/Engine/Private/Components/ModelComponent.cpp
--- a/Engine/Source/Runtime/Engine/Private/Components/ModelComponent.cpp
+++ b/Engine/Source/Runtime/Engine/Private/Components/ModelComponent.cpp
@@ -243,7 +243,7 @@ void UModelComponent::Serialize(FArchive& Ar)
 
 	Ar << Model;
 
-	if( Ar.IsLoading() && Ar.UE4Ver() < VER_UE4_REMOVE_ZONES_FROM_MODEL)
+	if( Ar.IsLoading() && Ar.UEVer() < VER_UE4_REMOVE_ZONES_FROM_MODEL)
 	{
 		int32 DummyZoneIndex;
 		Ar << DummyZoneIndex << Elements;
diff --git a/Engine/Source/Runtime/Engine/Private/Components/PointLightComponent.cpp b/Engine/Source/Runtime/Engine/Private/Components/PointLightComponent.cpp
--- a/Engine/Source/Runtime/Engine/Private/Components/PointLightComponent.cpp
+++ b/Engine/Source/Runtime/Engine/Private/Components/PointLightComponent.cpp
@@ -236,12 +236,12 @@ void UPointLightComponent::Serialize(FArchive& Ar)
 {
 	Super::Serialize(Ar);
 
-	if (Ar.UE4Ver() < VER_UE4_INVERSE_SQUARED_LIGHTS_DEFAULT)
+	if (Ar.UEVer() < VER_UE4_INVERSE_SQUARED_LIGHTS_DEFAULT)
 	{
 		bUseInverseSquaredFalloff = InverseSquaredFalloff_DEPRECATED;
 	}
 	// Reorient old light tubes that didn't use an IES profile
-	else if(Ar.UE4Ver() < VER_UE4_POINTLIGHT_SOURCE_ORIENTATION && SourceLength > KINDA_SMALL_NUMBER && IESTexture == nullptr)
+	else if(Ar.UEVer() < VER_UE4_POINTLIGHT_SOURCE_ORIENTATION && SourceLength > KINDA_SMALL_NUMBER && IESTexture == nullptr)
 	{
 		AddLocalRotation( FRotator(-90.f, 0.f, 0.f) );
 	}
diff --git a/Engine/Source/Runtime/Engine/Private/Components/PrimitiveComponent.cpp b/Engine/Source/Runtime/Engine/Private/Components/PrimitiveComponent.cpp
--- a/Engine/Source/Runtime/Engine/Private/Components/PrimitiveComponent.cpp
+++ b/Engine/Source/Runtime/Engine/Private/Components/PrimitiveComponent.cpp
@@ -1210,7 +1210,7 @@ void UPrimitiveComponent::PostLoad()
 {
 	Super::PostLoad();
 	
-	int32 const UE4Version = GetLinkerUEVersion();
+	int32 const UEVersion = GetLinkerUEVersion();
 
 	// as temporary fix for the bug TTP 299926
 	// permanent fix is coming
@@ -1220,7 +1220,7 @@ void UPrimitiveComponent::PostLoad()
 	}
 
 #if WITH_EDITORONLY_DATA
-	if (UE4Version < VER_UE4_RENAME_CANBECHARACTERBASE)
+	if (UEVersion < VER_UE4_RENAME_CANBECHARACTERBASE)
 	{
 		CanCharacterStepUpOn = CanBeCharacterBase_DEPRECATED;
 	}
diff --git a/Engine/Source/Runtime/Engine/Private/Components/ProjectileMovementComponent.cpp b/Engine/Source/Runtime/Engine/Private/Components/ProjectileMovementComponent.cpp
--- a/Engine/Source/Runtime/Engine/Private/Components/ProjectileMovementComponent.cpp
+++ b/Engine/Source/Runtime/Engine/Private/Components/ProjectileMovementComponent.cpp
@@ -57,9 +57,9 @@ void UProjectileMovementComponent::PostLoad()
 {
 	Super::PostLoad();
 
-	const int32 LinkerUE4Ver = GetLinkerUEVersion();
+	const int32 LinkerUEVer = GetLinkerUEVersion();
 
-	if (LinkerUE4Ver < VER_UE4_REFACTOR_PROJECTILE_MOVEMENT)
+	if (LinkerUEVer < VER_UE4_REFACTOR_PROJECTILE_MOVEMENT)
 	{
 		// Old code used to treat Bounciness as Friction as well.
 		Friction = FMath::Clamp(1.f - Bounciness, 0.f, 1.f);
diff --git a/Engine/Source/Runtime/Engine/Private/Components/ReflectionCaptureComponent.cpp b/Engine/Source/Runtime/Engine/Private/Components/ReflectionCaptureComponent.cpp
--- a/Engine/Source/Runtime/Engine/Private/Components/ReflectionCaptureComponent.cpp
+++ b/Engine/Source/Runtime/Engine/Private/Components/ReflectionCaptureComponent.cpp
@@ -834,13 +834,13 @@ void UReflectionCaptureComponent::SerializeLegacyData(FArchive& Ar)
 
 	if (Ar.CustomVer(FReflectionCaptureObjectVersion::GUID) < FReflectionCaptureObjectVersion::MoveReflectionCaptureDataToMapBuildData)
 	{
-		if (Ar.UE4Ver() >= VER_UE4_REFLECTION_CAPTURE_COOKING)
+		if (Ar.UEVer() >= VER_UE4_REFLECTION_CAPTURE_COOKING)
 		{
 			bool bLegacy = false;
 			Ar << bLegacy;
 		}
 
-		if (Ar.UE4Ver() >= VER_UE4_REFLECTION_DATA_IN_PACKAGES)
+		if (Ar.UEVer() >= VER_UE4_REFLECTION_DATA_IN_PACKAGES)
 		{
 			FGuid SavedVersion;
 			Ar << SavedVersion;
diff --git a/Engine/Source/Runtime/Engine/Private/Components/SkeletalMeshComponent.cpp b/Engine/Source/Runtime/Engine/Private/Components/SkeletalMeshComponent.cpp
--- a/Engine/Source/Runtime/Engine/Private/Components/SkeletalMeshComponent.cpp
+++ b/Engine/Source/Runtime/Engine/Private/Components/SkeletalMeshComponent.cpp
@@ -299,7 +299,7 @@ void USkeletalMeshComponent::Serialize(FArchive& Ar)
 		RequiredBones.CountBytes(Ar);
 	}
 
-	if (Ar.UE4Ver() < VER_UE4_REMOVE_SKELETALMESH_COMPONENT_BODYSETUP_SERIALIZATION)
+	if (Ar.UEVer() < VER_UE4_REMOVE_SKELETALMESH_COMPONENT_BODYSETUP_SERIALIZATION)
 	{
 		//we used to serialize bodysetup of skeletal mesh component. We no longer do this, but need to not break existing content
 		if (bEnablePerPolyCollision)
@@ -316,7 +316,7 @@ void USkeletalMeshComponent::Serialize(FArchive& Ar)
 	}
 
 #if WITH_EDITORONLY_DATA
-	if (Ar.IsLoading() && (Ar.UE4Ver() < VER_UE4_EDITORONLY_BLUEPRINTS))
+	if (Ar.IsLoading() && (Ar.UEVer() < VER_UE4_EDITORONLY_BLUEPRINTS))
 	{
 		if ((NULL != AnimationBlueprint_DEPRECATED))
 		{
@@ -327,7 +327,7 @@ void USkeletalMeshComponent::Serialize(FArchive& Ar)
 	}
 #endif
 
-	if (Ar.IsLoading() && (Ar.UE4Ver() < VER_UE4_NO_ANIM_BP_CLASS_IN_GAMEPLAY_CODE))
+	if (Ar.IsLoading() && (Ar.UEVer() < VER_UE4_NO_ANIM_BP_CLASS_IN_GAMEPLAY_CODE))
 	{
 		if (nullptr != AnimBlueprintGeneratedClass)
 		{
@@ -340,7 +340,7 @@ void USkeletalMeshComponent::Serialize(FArchive& Ar)
 		AnimBlueprintGeneratedClass = nullptr;
 	}
 
-	if (Ar.IsLoading() && (Ar.UE4Ver() < VER_UE4_AUTO_WELDING))
+	if (Ar.IsLoading() && (Ar.UEVer() < VER_UE4_AUTO_WELDING))
 	{
 		BodyInstance.bAutoWeld = false;
 	}
diff --git a/Engine/Source/Runtime/Engine/Private/Components/SkyLightComponent.cpp b/Engine/Source/Runtime/Engine/Private/Components/SkyLightComponent.cpp
--- a/Engine/Source/Runtime/Engine/Private/Components/SkyLightComponent.cpp
+++ b/Engine/Source/Runtime/Engine/Private/Components/SkyLightComponent.cpp
@@ -1039,7 +1039,7 @@ void USkyLightComponent::Serialize(FArchive& Ar)
 	Super::Serialize(Ar);
 
 	// if version is between VER_UE4_SKYLIGHT_MOBILE_IRRADIANCE_MAP and FReleaseObjectVersion::SkyLightRemoveMobileIrradianceMap then handle aborted attempt to serialize irradiance data on mobile.
-	if (Ar.UE4Ver() >= VER_UE4_SKYLIGHT_MOBILE_IRRADIANCE_MAP && !(Ar.CustomVer(FReleaseObjectVersion::GUID) >= FReleaseObjectVersion::SkyLightRemoveMobileIrradianceMap))
+	if (Ar.UEVer() >= VER_UE4_SKYLIGHT_MOBILE_IRRADIANCE_MAP && !(Ar.CustomVer(FReleaseObjectVersion::GUID) >= FReleaseObjectVersion::SkyLightRemoveMobileIrradianceMap))
 	{
 		FSHVectorRGB3 DummyIrradianceEnvironmentMap;
 		Ar << DummyIrradianceEnvironmentMap;
diff --git a/Engine/Source/Runtime/Engine/Private/Components/SplineComponent.cpp b/Engine/Source/Runtime/Engine/Private/Components/SplineComponent.cpp
--- a/Engine/Source/Runtime/Engine/Private/Components/SplineComponent.cpp
+++ b/Engine/Source/Runtime/Engine/Private/Components/SplineComponent.cpp
@@ -123,8 +123,8 @@ void USplineComponent::Serialize(FArchive& Ar)
 	}
 
 	// Support old resources which don't have the rotation and scale splines present
-	const int32 ArchiveUE4Version = Ar.UE4Ver();
-	if (ArchiveUE4Version < VER_UE4_INTERPCURVE_SUPPORTS_LOOPING)
+	const int32 ArchiveUEVersion = Ar.UEVer();
+	if (ArchiveUEVersion < VER_UE4_INTERPCURVE_SUPPORTS_LOOPING)
 	{
 		int32 NumPoints = SplineCurves.Position.Points.Num();
 
diff --git a/Engine/Source/Runtime/Engine/Private/Components/SplineMeshComponent.cpp b/Engine/Source/Runtime/Engine/Private/Components/SplineMeshComponent.cpp
--- a/Engine/Source/Runtime/Engine/Private/Components/SplineMeshComponent.cpp
+++ b/Engine/Source/Runtime/Engine/Private/Components/SplineMeshComponent.cpp
@@ -507,7 +507,7 @@ void USplineMeshComponent::Serialize(FArchive& Ar)
 {
 	Super::Serialize(Ar);
 
-	if (Ar.UE4Ver() < VER_UE4_SPLINE_MESH_ORIENTATION)
+	if (Ar.UEVer() < VER_UE4_SPLINE_MESH_ORIENTATION)
 	{
 		ForwardAxis = ESplineMeshAxis::Z;
 		SplineParams.StartRoll -= HALF_PI;
diff --git a/Engine/Source/Runtime/Engine/Private/Components/StaticMeshComponent.cpp b/Engine/Source/Runtime/Engine/Private/Components/StaticMeshComponent.cpp
--- a/Engine/Source/Runtime/Engine/Private/Components/StaticMeshComponent.cpp
+++ b/Engine/Source/Runtime/Engine/Private/Components/StaticMeshComponent.cpp
@@ -304,7 +304,7 @@ void UStaticMeshComponent::Serialize(FArchive& Ar)
 	}
 
 #if WITH_EDITORONLY_DATA
-	if (Ar.UE4Ver() < VER_UE4_COMBINED_LIGHTMAP_TEXTURES)
+	if (Ar.UEVer() < VER_UE4_COMBINED_LIGHTMAP_TEXTURES)
 	{
 		check(AttachmentCounter.GetValue() == 0);
 		// Irrelevant lights were incorrect before VER_UE4_TOSS_IRRELEVANT_LIGHTS
@@ -328,7 +328,7 @@ void UStaticMeshComponent::Serialize(FArchive& Ar)
 		GComponentsWithLegacyLightmaps.AddAnnotation(this, MoveTemp(LegacyComponentData));
 	}
 
-	if (Ar.UE4Ver() < VER_UE4_AUTO_WELDING)
+	if (Ar.UEVer() < VER_UE4_AUTO_WELDING)
 	{
 		GetBodyInstance()->bAutoWeld = false;	//existing content may rely on no auto welding
 	}
diff --git a/Engine/Source/Runtime/Engine/Private/ContentStreaming.cpp b/Engine/Source/Runtime/Engine/Private/ContentStreaming.cpp
--- a/Engine/Source/Runtime/Engine/Private/ContentStreaming.cpp
+++ b/Engine/Source/Runtime/Engine/Private/ContentStreaming.cpp
@@ -1,4 +1,5 @@
 // Copyright Epic Games, Inc. All Rights Reserved.
+// Copyright Epic Games, Inc. All Rights Reserved.
 
 /*=============================================================================
 	ContentStreaming.cpp: Implementation of content streaming classes.
@@ -1344,7 +1345,7 @@ void FStreamingManagerCollection::AddOrRemoveTextureStreamingManagerIfNeeded(boo
  */
 FArchive& operator<<( FArchive& Ar, FStreamableTextureInstance& TextureInstance )
 {
-	if (Ar.UE4Ver() >= VER_UE4_STREAMABLE_TEXTURE_AABB)
+	if (Ar.UEVer() >= VER_UE4_STREAMABLE_TEXTURE_AABB)
 	{
 		Ar << TextureInstance.Bounds;
 	}
@@ -1355,7 +1356,7 @@ FArchive& operator<<( FArchive& Ar, FStreamableTextureInstance& TextureInstance
 		TextureInstance.Bounds = FBoxSphereBounds(BoundingSphere);
 	}
 
-	if (Ar.UE4Ver() >= VER_UE4_STREAMABLE_TEXTURE_MIN_MAX_DISTANCE)
+	if (Ar.UEVer() >= VER_UE4_STREAMABLE_TEXTURE_MIN_MAX_DISTANCE)
 	{
 		Ar << TextureInstance.MinDistance;
 		Ar << TextureInstance.MaxDistance;
diff --git a/Engine/Source/Runtime/Engine/Private/Curves/RichCurve.cpp b/Engine/Source/Runtime/Engine/Private/Curves/RichCurve.cpp
--- a/Engine/Source/Runtime/Engine/Private/Curves/RichCurve.cpp
+++ b/Engine/Source/Runtime/Engine/Private/Curves/RichCurve.cpp
@@ -160,7 +160,7 @@ FRichCurveKey::FRichCurveKey(const FInterpCurvePoint<FTwoVectors>& InPoint, int3
 
 bool FRichCurveKey::Serialize(FArchive& Ar)
 {
-	if (Ar.UE4Ver() < VER_UE4_SERIALIZE_RICH_CURVE_KEY)
+	if (Ar.UEVer() < VER_UE4_SERIALIZE_RICH_CURVE_KEY)
 	{
 		return false;
 	}
diff --git a/Engine/Source/Runtime/Engine/Private/DataAsset.cpp b/Engine/Source/Runtime/Engine/Private/DataAsset.cpp
--- a/Engine/Source/Runtime/Engine/Private/DataAsset.cpp
+++ b/Engine/Source/Runtime/Engine/Private/DataAsset.cpp
@@ -19,7 +19,7 @@ void UDataAsset::Serialize(FStructuredArchive::FRecord Record)
 	FArchive& UnderlyingArchive = Record.GetUnderlyingArchive();
 	Super::Serialize(Record);
 
-	if (UnderlyingArchive.IsLoading() && (UnderlyingArchive.UE4Ver() < VER_UE4_ADD_TRANSACTIONAL_TO_DATA_ASSETS))
+	if (UnderlyingArchive.IsLoading() && (UnderlyingArchive.UEVer() < VER_UE4_ADD_TRANSACTIONAL_TO_DATA_ASSETS))
 	{
 		SetFlags(RF_Transactional);
 	}
diff --git a/Engine/Source/Runtime/Engine/Private/DecalActor.cpp b/Engine/Source/Runtime/Engine/Private/DecalActor.cpp
--- a/Engine/Source/Runtime/Engine/Private/DecalActor.cpp
+++ b/Engine/Source/Runtime/Engine/Private/DecalActor.cpp
@@ -157,7 +157,7 @@ class UMaterialInstanceDynamic* ADecalActor::CreateDynamicMaterialInstance()
 
 void ADecalActor::Serialize(FArchive& Ar)
 {
-	if (Ar.UE4Ver() < VER_UE4_DECAL_SIZE)
+	if (Ar.UEVer() < VER_UE4_DECAL_SIZE)
 	{
 		if(Decal)
 		{
@@ -168,7 +168,7 @@ void ADecalActor::Serialize(FArchive& Ar)
 
 	Super::Serialize(Ar);
 
-	if (Ar.UE4Ver() < VER_UE4_DECAL_SIZE)
+	if (Ar.UEVer() < VER_UE4_DECAL_SIZE)
 	{
 		if(Decal)
 		{
diff --git a/Engine/Source/Runtime/Engine/Private/EdGraph/EdGraph.cpp b/Engine/Source/Runtime/Engine/Private/EdGraph/EdGraph.cpp
--- a/Engine/Source/Runtime/Engine/Private/EdGraph/EdGraph.cpp
+++ b/Engine/Source/Runtime/Engine/Private/EdGraph/EdGraph.cpp
@@ -18,7 +18,7 @@ DEFINE_LOG_CATEGORY_STATIC(LogEdGraph, Log, All);
 void FGraphReference::PostSerialize(const FArchive& Ar)
 {
 #if WITH_EDITORONLY_DATA
-	if (Ar.UE4Ver() >= VER_UE4_K2NODE_REFERENCEGUIDS)
+	if (Ar.UEVer() >= VER_UE4_K2NODE_REFERENCEGUIDS)
 	{
 		// Because the macro instance could have been saved with a GUID that was allocated 
 		// but the macro graph never actually saved with that value we are forced to make 
diff --git a/Engine/Source/Runtime/Engine/Private/EdGraph/EdGraphNode.cpp b/Engine/Source/Runtime/Engine/Private/EdGraph/EdGraphNode.cpp
--- a/Engine/Source/Runtime/Engine/Private/EdGraph/EdGraphNode.cpp
+++ b/Engine/Source/Runtime/Engine/Private/EdGraph/EdGraphNode.cpp
@@ -49,7 +49,7 @@ FArchive& operator<<(FArchive& Ar, FEdGraphTerminalType& T)
 		FString TerminalCategoryStr;
 		Ar << TerminalCategoryStr;
 
-		if (Ar.UE4Ver() < VER_UE4_ADDED_SOFT_OBJECT_PATH)
+		if (Ar.UEVer() < VER_UE4_ADDED_SOFT_OBJECT_PATH)
 		{
 			// Handle asset->soft object rename, this is here instead of BP code because this structure is embedded
 			if (TerminalCategoryStr == TEXT("asset"))
diff --git a/Engine/Source/Runtime/Engine/Private/EdGraph/EdGraphPin.cpp b/Engine/Source/Runtime/Engine/Private/EdGraph/EdGraphPin.cpp
--- a/Engine/Source/Runtime/Engine/Private/EdGraph/EdGraphPin.cpp
+++ b/Engine/Source/Runtime/Engine/Private/EdGraph/EdGraphPin.cpp
@@ -141,7 +141,7 @@ TArray<TPair<UEdGraphPin*, FString>> PinAllocationTracking;
 
 void FEdGraphPinType::PostSerialize(const FArchive& Ar)
 {
-	if (Ar.UE4Ver() < VER_UE4_EDGRAPHPINTYPE_SERIALIZATION)
+	if (Ar.UEVer() < VER_UE4_EDGRAPHPINTYPE_SERIALIZATION)
 	{
 		if (bIsArray_DEPRECATED)
 		{
@@ -152,7 +152,7 @@ void FEdGraphPinType::PostSerialize(const FArchive& Ar)
 
 bool FEdGraphPinType::Serialize(FArchive& Ar)
 {
-	if (Ar.UE4Ver() < VER_UE4_EDGRAPHPINTYPE_SERIALIZATION)
+	if (Ar.UEVer() < VER_UE4_EDGRAPHPINTYPE_SERIALIZATION)
 	{
 		return false;
 	}
@@ -176,7 +176,7 @@ bool FEdGraphPinType::Serialize(FArchive& Ar)
 		PinSubCategory = *PinSubCategoryStr;
 	}
 
-	if (Ar.UE4Ver() < VER_UE4_ADDED_SOFT_OBJECT_PATH)
+	if (Ar.UEVer() < VER_UE4_ADDED_SOFT_OBJECT_PATH)
 	{
 		// Fixup has to be here instead of in BP code because this is embedded in other structures
 		if (PinCategory == TEXT("asset"))
@@ -240,7 +240,7 @@ bool FEdGraphPinType::Serialize(FArchive& Ar)
 	Ar << bIsReferenceBool;
 	Ar << bIsWeakPointerBool;
 
-	if (Ar.UE4Ver() >= VER_UE4_MEMBERREFERENCE_IN_PINTYPE)
+	if (Ar.UEVer() >= VER_UE4_MEMBERREFERENCE_IN_PINTYPE)
 	{
 		Ar << PinSubCategoryMemberReference;
 	}
@@ -259,7 +259,7 @@ bool FEdGraphPinType::Serialize(FArchive& Ar)
 
 	bool bIsConstBool = bIsConst;
 
-	if (Ar.UE4Ver() >= VER_UE4_SERIALIZE_PINTYPE_CONST)
+	if (Ar.UEVer() >= VER_UE4_SERIALIZE_PINTYPE_CONST)
 	{
 		Ar << bIsConstBool;
 	}
diff --git a/Engine/Source/Runtime/Engine/Private/EditorFramework/AssetImportData.cpp b/Engine/Source/Runtime/Engine/Private/EditorFramework/AssetImportData.cpp
--- a/Engine/Source/Runtime/Engine/Private/EditorFramework/AssetImportData.cpp
+++ b/Engine/Source/Runtime/Engine/Private/EditorFramework/AssetImportData.cpp
@@ -377,7 +377,7 @@ void UAssetImportData::Serialize(FStructuredArchive::FRecord Record)
 {
 	FArchive& BaseArchive = Record.GetUnderlyingArchive();
 
-	if (BaseArchive.UE4Ver() >= VER_UE4_ASSET_IMPORT_DATA_AS_JSON)
+	if (BaseArchive.UEVer() >= VER_UE4_ASSET_IMPORT_DATA_AS_JSON)
 	{
 		if (!BaseArchive.IsFilterEditorOnly())
 		{
diff --git a/Engine/Source/Runtime/Engine/Private/EngineUtils.cpp b/Engine/Source/Runtime/Engine/Private/EngineUtils.cpp
--- a/Engine/Source/Runtime/Engine/Private/EngineUtils.cpp
+++ b/Engine/Source/Runtime/Engine/Private/EngineUtils.cpp
@@ -502,7 +502,7 @@ FStripDataFlags::FStripDataFlags( class FArchive& Ar, uint8 InClassFlags /*= 0*/
 	, ClassStripFlags( 0 )
 {
 	check(InVersion >= VER_UE4_OLDEST_LOADABLE_PACKAGE);
-	if (Ar.UE4Ver() >= InVersion)
+	if (Ar.UEVer() >= InVersion)
 	{
 		if (Ar.IsCooking())
 		{
@@ -522,7 +522,7 @@ FStripDataFlags::FStripDataFlags( class FArchive& Ar, uint8 InGlobalFlags, uint8
 	, ClassStripFlags( 0 )
 {
 	check(InVersion >= VER_UE4_OLDEST_LOADABLE_PACKAGE);
-	if (Ar.UE4Ver() >= InVersion)
+	if (Ar.UEVer() >= InVersion)
 	{
 		if (Ar.IsCooking())
 		{
@@ -546,7 +546,7 @@ FStripDataFlags::FStripDataFlags(FStructuredArchive::FSlot Slot, uint8 InClassFl
 	FStructuredArchive::FRecord Record = Slot.EnterRecord();
 
 	check(InVersion >= VER_UE4_OLDEST_LOADABLE_PACKAGE);
-	if (UnderlyingArchive.UE4Ver() >= InVersion)
+	if (UnderlyingArchive.UEVer() >= InVersion)
 	{
 		if (UnderlyingArchive.IsCooking())
 		{
@@ -569,7 +569,7 @@ FStripDataFlags::FStripDataFlags(FStructuredArchive::FSlot Slot, uint8 InGlobalF
 	FStructuredArchive::FRecord Record = Slot.EnterRecord();
 
 	check(InVersion >= VER_UE4_OLDEST_LOADABLE_PACKAGE);
-	if (UnderlyingArchive.UE4Ver() >= InVersion)
+	if (UnderlyingArchive.UEVer() >= InVersion)
 	{
 		if (UnderlyingArchive.IsCooking())
 		{
diff --git a/Engine/Source/Runtime/Engine/Private/GameplayStatics.cpp b/Engine/Source/Runtime/Engine/Private/GameplayStatics.cpp
--- a/Engine/Source/Runtime/Engine/Private/GameplayStatics.cpp
+++ b/Engine/Source/Runtime/Engine/Private/GameplayStatics.cpp
@@ -144,7 +144,7 @@ void FSaveGameHeader::Read(FMemoryReader& MemoryReader)
 		// Note for 4.8 and beyond: if you get a crash loading a pre-4.8 version of your savegame file and 
 		// you don't want to delete it, try uncommenting these lines and changing them to use the version 
 		// information from your previous build. Then load and resave your savegame file.
-		//MemoryReader.SetUE4Ver(MyPreviousUEVersion);				// @see GPackageFileUEVersion
+		//MemoryReader.SetUEVer(MyPreviousUEVersion);				// @see GPackageFileUEVersion
 		//MemoryReader.SetEngineVer(MyPreviousEngineVersion);		// @see FEngineVersion::Current()
 	}
 	else
@@ -157,7 +157,7 @@ void FSaveGameHeader::Read(FMemoryReader& MemoryReader)
 
 		MemoryReader << SavedEngineVersion;
 
-		MemoryReader.SetUE4Ver(PackageFileUEVersion);
+		MemoryReader.SetUEVer(PackageFileUEVersion);
 		MemoryReader.SetEngineVer(SavedEngineVersion);
 
 		if (SaveGameFileVersion >= FSaveGameFileVersion::AddedCustomVersions)
diff --git a/Engine/Source/Runtime/Engine/Private/Interpolation.cpp b/Engine/Source/Runtime/Engine/Private/Interpolation.cpp
--- a/Engine/Source/Runtime/Engine/Private/Interpolation.cpp
+++ b/Engine/Source/Runtime/Engine/Private/Interpolation.cpp
@@ -9815,7 +9815,7 @@ UInterpTrackFloatAnimBPParam::UInterpTrackFloatAnimBPParam(const FObjectInitiali
 void UInterpTrackFloatAnimBPParam::Serialize(FArchive& Ar)
 {
 	PRAGMA_DISABLE_DEPRECATION_WARNINGS	
-	if (Ar.IsSaving() && Ar.UE4Ver() < VER_UE4_NO_ANIM_BP_CLASS_IN_GAMEPLAY_CODE)
+	if (Ar.IsSaving() && Ar.UEVer() < VER_UE4_NO_ANIM_BP_CLASS_IN_GAMEPLAY_CODE)
 	{
 		if ((nullptr != AnimBlueprintClass) && (nullptr == AnimClass))
 		{
@@ -9825,7 +9825,7 @@ void UInterpTrackFloatAnimBPParam::Serialize(FArchive& Ar)
 
 	Super::Serialize(Ar);
 
-	if (Ar.IsLoading() && Ar.UE4Ver() < VER_UE4_NO_ANIM_BP_CLASS_IN_GAMEPLAY_CODE)
+	if (Ar.IsLoading() && Ar.UEVer() < VER_UE4_NO_ANIM_BP_CLASS_IN_GAMEPLAY_CODE)
 	{
 		if ((nullptr != AnimBlueprintClass) && (nullptr == AnimClass))
 		{
diff --git a/Engine/Source/Runtime/Engine/Private/Level.cpp b/Engine/Source/Runtime/Engine/Private/Level.cpp
--- a/Engine/Source/Runtime/Engine/Private/Level.cpp
+++ b/Engine/Source/Runtime/Engine/Private/Level.cpp
@@ -485,7 +485,7 @@ void ULevel::Serialize( FArchive& Ar )
 
 	Ar << ModelComponents;
 
-	if(!Ar.IsFilterEditorOnly() || (Ar.UE4Ver() < VER_UE4_EDITORONLY_BLUEPRINTS) )
+	if(!Ar.IsFilterEditorOnly() || (Ar.UEVer() < VER_UE4_EDITORONLY_BLUEPRINTS) )
 	{
 #if WITH_EDITORONLY_DATA
 		// Skip serializing the LSBP if this is a world duplication for PIE/SIE, as it is not needed, and it causes overhead in startup times
@@ -530,7 +530,7 @@ void ULevel::Serialize( FArchive& Ar )
 			Ar << Len;
 		}
 
-		if(Ar.UE4Ver() < VER_UE4_REMOVE_LEVELBODYSETUP)
+		if(Ar.UEVer() < VER_UE4_REMOVE_LEVELBODYSETUP)
 		{
 			UBodySetup* DummySetup;
 			Ar << DummySetup;
@@ -585,8 +585,8 @@ void ULevel::Serialize( FArchive& Ar )
 	Ar << PrecomputedVisibilityHandler;
 	Ar << PrecomputedVolumeDistanceField;
 
-	if (Ar.UE4Ver() >= VER_UE4_WORLD_LEVEL_INFO &&
-		Ar.UE4Ver() < VER_UE4_WORLD_LEVEL_INFO_UPDATED)
+	if (Ar.UEVer() >= VER_UE4_WORLD_LEVEL_INFO &&
+		Ar.UEVer() < VER_UE4_WORLD_LEVEL_INFO_UPDATED)
 	{
 		FWorldTileInfo Info;
 		Ar << Info;
diff --git a/Engine/Source/Runtime/Engine/Private/LightMap.cpp b/Engine/Source/Runtime/Engine/Private/LightMap.cpp
--- a/Engine/Source/Runtime/Engine/Private/LightMap.cpp
+++ b/Engine/Source/Runtime/Engine/Private/LightMap.cpp
@@ -2982,7 +2982,7 @@ void FLightMap2D::Serialize(FArchive& Ar)
 
 	const bool bUsingVTLightmaps = (CVarVirtualTexturedLightMaps.GetValueOnAnyThread() != 0) && UseVirtualTexturing(GMaxRHIFeatureLevel, Ar.CookingTarget());
 
-	if( Ar.IsLoading() && Ar.UE4Ver() < VER_UE4_LOW_QUALITY_DIRECTIONAL_LIGHTMAPS )
+	if( Ar.IsLoading() && Ar.UEVer() < VER_UE4_LOW_QUALITY_DIRECTIONAL_LIGHTMAPS )
 	{
 		for(uint32 CoefficientIndex = 0;CoefficientIndex < 3;CoefficientIndex++)
 		{
@@ -2993,7 +2993,7 @@ void FLightMap2D::Serialize(FArchive& Ar)
 			Ar << Dummy2;
 		}
 	}
-	else if( Ar.IsLoading() && Ar.UE4Ver() < VER_UE4_COMBINED_LIGHTMAP_TEXTURES )
+	else if( Ar.IsLoading() && Ar.UEVer() < VER_UE4_COMBINED_LIGHTMAP_TEXTURES )
 	{
 		for( uint32 CoefficientIndex = 0; CoefficientIndex < 4; CoefficientIndex++ )
 		{
@@ -3023,7 +3023,7 @@ void FLightMap2D::Serialize(FArchive& Ar)
 			Ar << Textures[1];
 		}
 
-		if (Ar.UE4Ver() >= VER_UE4_SKY_LIGHT_COMPONENT)
+		if (Ar.UEVer() >= VER_UE4_SKY_LIGHT_COMPONENT)
 		{
 			if (Ar.IsCooking())
 			{
@@ -3033,7 +3033,7 @@ void FLightMap2D::Serialize(FArchive& Ar)
 				ULightMapTexture2D*& SkyTexture = bStripHQLightmaps ? Dummy : SkyOcclusionTexture;
 				Ar << SkyTexture;
 
-				if (Ar.UE4Ver() >= VER_UE4_AO_MATERIAL_MASK)
+				if (Ar.UEVer() >= VER_UE4_AO_MATERIAL_MASK)
 				{
 					ULightMapTexture2D*& MaskTexture = bStripHQLightmaps ? Dummy : AOMaterialMaskTexture;
 					Ar << MaskTexture;
@@ -3043,7 +3043,7 @@ void FLightMap2D::Serialize(FArchive& Ar)
 			{
 				Ar << SkyOcclusionTexture;
 
-				if (Ar.UE4Ver() >= VER_UE4_AO_MATERIAL_MASK)
+				if (Ar.UEVer() >= VER_UE4_AO_MATERIAL_MASK)
 				{
 					Ar << AOMaterialMaskTexture;
 				}
@@ -3362,7 +3362,7 @@ FArchive& operator<<(FArchive& Ar, FLightMap*& R)
 			}
 
 			// Dump old lightmaps
-			if (Ar.UE4Ver() < VER_UE4_COMBINED_LIGHTMAP_TEXTURES)
+			if (Ar.UEVer() < VER_UE4_COMBINED_LIGHTMAP_TEXTURES)
 			{
 				delete R; // safe because if we're loading we new'd this above
 				R = nullptr;
diff --git a/Engine/Source/Runtime/Engine/Private/Materials/Material.cpp b/Engine/Source/Runtime/Engine/Private/Materials/Material.cpp
--- a/Engine/Source/Runtime/Engine/Private/Materials/Material.cpp
+++ b/Engine/Source/Runtime/Engine/Private/Materials/Material.cpp
@@ -3355,7 +3355,7 @@ void UMaterial::Serialize(FArchive& Ar)
 
 	Super::Serialize(Ar);
 
-	if (Ar.UE4Ver() >= VER_UE4_PURGED_FMATERIAL_COMPILE_OUTPUTS)
+	if (Ar.UEVer() >= VER_UE4_PURGED_FMATERIAL_COMPILE_OUTPUTS)
 	{
 #if WITH_EDITOR
 		static_assert(!STORE_ONLY_ACTIVE_SHADERMAPS, "Only discard unused SMs in cooked build");
@@ -3400,20 +3400,20 @@ void UMaterial::Serialize(FArchive& Ar)
 	}
 
 #if WITH_EDITOR
-	if (Ar.UE4Ver() < VER_UE4_FLIP_MATERIAL_COORDS)
+	if (Ar.UEVer() < VER_UE4_FLIP_MATERIAL_COORDS)
 	{
 		GMaterialsThatNeedExpressionsFlipped.Set(this);
 	}
-	else if (Ar.UE4Ver() < VER_UE4_FIX_MATERIAL_COORDS)
+	else if (Ar.UEVer() < VER_UE4_FIX_MATERIAL_COORDS)
 	{
 		GMaterialsThatNeedCoordinateCheck.Set(this);
 	}
-	else if (Ar.UE4Ver() < VER_UE4_FIX_MATERIAL_COMMENTS)
+	else if (Ar.UEVer() < VER_UE4_FIX_MATERIAL_COMMENTS)
 	{
 		GMaterialsThatNeedCommentFix.Set(this);
 	}
 
-	if (Ar.UE4Ver() < VER_UE4_ADD_LINEAR_COLOR_SAMPLER)
+	if (Ar.UEVer() < VER_UE4_ADD_LINEAR_COLOR_SAMPLER)
 	{
 		GMaterialsThatNeedSamplerFixup.Set(this);
 	}
@@ -3421,7 +3421,7 @@ void UMaterial::Serialize(FArchive& Ar)
 
 	static_assert(MP_MAX == 33, "New material properties must have DoMaterialAttributeReorder called on them to ensure that any future reordering of property pins is correctly applied.");
 
-	if (Ar.UE4Ver() < VER_UE4_MATERIAL_MASKED_BLENDMODE_TIDY)
+	if (Ar.UEVer() < VER_UE4_MATERIAL_MASKED_BLENDMODE_TIDY)
 	{
 		//Set based on old value. Real check may not be possible here in cooked builds?
 		//Cached using acutal check in PostEditChangProperty().
@@ -4110,7 +4110,7 @@ void UMaterial::PostLoad()
 #endif // WITH_EDITORONLY_DATA
 
 	// Fixup for legacy materials which didn't recreate the lighting guid properly on duplication
-	if (GetLinker() && GetLinker()->UE4Ver() < VER_UE4_BUMPED_MATERIAL_EXPORT_GUIDS)
+	if (GetLinker() && GetLinker()->UEVer() < VER_UE4_BUMPED_MATERIAL_EXPORT_GUIDS)
 	{
 		UMaterialInterface** ExistingMaterial = LightingGuidFixupMap.Find(GetLightingGuid());
 
diff --git a/Engine/Source/Runtime/Engine/Private/Materials/MaterialExpressions.cpp b/Engine/Source/Runtime/Engine/Private/Materials/MaterialExpressions.cpp
--- a/Engine/Source/Runtime/Engine/Private/Materials/MaterialExpressions.cpp
+++ b/Engine/Source/Runtime/Engine/Private/Materials/MaterialExpressions.cpp
@@ -781,7 +781,7 @@ void UMaterialExpression::Serialize(FStructuredArchive::FRecord Record)
 	for (int32 InputIndex = 0; InputIndex < Inputs.Num(); ++InputIndex)
 	{
 		FExpressionInput* Input = Inputs[InputIndex];
-		DoMaterialAttributeReorder(Input, Record.GetUnderlyingArchive().UE4Ver(), Record.GetUnderlyingArchive().CustomVer(FRenderingObjectVersion::GUID));
+		DoMaterialAttributeReorder(Input, Record.GetUnderlyingArchive().UEVer(), Record.GetUnderlyingArchive().CustomVer(FRenderingObjectVersion::GUID));
 	}
 #endif // WITH_EDITORONLY_DATA
 }
@@ -4249,7 +4249,7 @@ void UMaterialExpressionClamp::Serialize(FStructuredArchive::FRecord Record)
 	Super::Serialize(Record);
 	FArchive& UnderlyingArchive = Record.GetUnderlyingArchive();
 
-	if (UnderlyingArchive.IsLoading() && UnderlyingArchive.UE4Ver() < VER_UE4_RETROFIT_CLAMP_EXPRESSIONS_SWAP)
+	if (UnderlyingArchive.IsLoading() && UnderlyingArchive.UEVer() < VER_UE4_RETROFIT_CLAMP_EXPRESSIONS_SWAP)
 	{
 		if (ClampMode == CMODE_ClampMin)
 		{
@@ -8699,7 +8699,7 @@ void UMaterialExpressionFeatureLevelSwitch::Serialize(FStructuredArchive::FRecor
 	FExpressionInput Inputs[ERHIFeatureLevel::Num];
 #endif
 
-	if (UnderlyingArchive.IsLoading() && UnderlyingArchive.UE4Ver() < VER_UE4_RENAME_SM3_TO_ES3_1)
+	if (UnderlyingArchive.IsLoading() && UnderlyingArchive.UEVer() < VER_UE4_RENAME_SM3_TO_ES3_1)
 	{
 		// Copy the ES2 input to SM3 (since SM3 will now become ES3_1 and we don't want broken content)
 		Inputs[ERHIFeatureLevel::ES3_1] = Inputs[ERHIFeatureLevel::ES2_REMOVED];
@@ -11963,7 +11963,7 @@ void UMaterialExpressionCustom::Serialize(FStructuredArchive::FRecord Record)
 
 	bool bDidUpdate = false;
 
-	if (UnderlyingArchive.UE4Ver() < VER_UE4_INSTANCED_STEREO_UNIFORM_UPDATE)
+	if (UnderlyingArchive.UEVer() < VER_UE4_INSTANCED_STEREO_UNIFORM_UPDATE)
 	{
 		// Look for WorldPosition rename
 		if (Code.ReplaceInline(TEXT("Parameters.WorldPosition"), TEXT("Parameters.AbsoluteWorldPosition"), ESearchCase::CaseSensitive) > 0)
@@ -11972,7 +11972,7 @@ void UMaterialExpressionCustom::Serialize(FStructuredArchive::FRecord Record)
 		}
 	}
 	// Fix up uniform references that were moved from View to Frame as part of the instanced stereo implementation
-	else if (UnderlyingArchive.UE4Ver() < VER_UE4_INSTANCED_STEREO_UNIFORM_REFACTOR)
+	else if (UnderlyingArchive.UEVer() < VER_UE4_INSTANCED_STEREO_UNIFORM_REFACTOR)
 	{
 		// Uniform members that were moved from View to Frame
 		static const FString UniformMembers[] = {
@@ -12256,25 +12256,25 @@ void UMaterialFunction::Serialize(FArchive& Ar)
 	Super::Serialize(Ar);
 
 #if WITH_EDITOR
-	if (Ar.UE4Ver() < VER_UE4_FLIP_MATERIAL_COORDS)
+	if (Ar.UEVer() < VER_UE4_FLIP_MATERIAL_COORDS)
 	{
 		GMaterialFunctionsThatNeedExpressionsFlipped.Set(this);
 	}
-	else if (Ar.UE4Ver() < VER_UE4_FIX_MATERIAL_COORDS)
+	else if (Ar.UEVer() < VER_UE4_FIX_MATERIAL_COORDS)
 	{
 		GMaterialFunctionsThatNeedCoordinateCheck.Set(this);
 	}
-	else if (Ar.UE4Ver() < VER_UE4_FIX_MATERIAL_COMMENTS)
+	else if (Ar.UEVer() < VER_UE4_FIX_MATERIAL_COMMENTS)
 	{
 		GMaterialFunctionsThatNeedCommentFix.Set(this);
 	}
 
-	if (Ar.UE4Ver() < VER_UE4_ADD_LINEAR_COLOR_SAMPLER)
+	if (Ar.UEVer() < VER_UE4_ADD_LINEAR_COLOR_SAMPLER)
 	{
 		GMaterialFunctionsThatNeedSamplerFixup.Set(this);
 	}
 
-	if (Ar.UE4Ver() < VER_UE4_LIBRARY_CATEGORIES_AS_FTEXT)
+	if (Ar.UEVer() < VER_UE4_LIBRARY_CATEGORIES_AS_FTEXT)
 	{
 		for (FString& Category : LibraryCategories_DEPRECATED)
 		{
@@ -18049,7 +18049,7 @@ void UMaterialExpressionSpeedTree::Serialize(FStructuredArchive::FRecord Record)
 {
 	Super::Serialize(Record);
 
-	if (Record.GetUnderlyingArchive().UE4Ver() < VER_UE4_SPEEDTREE_WIND_V7)
+	if (Record.GetUnderlyingArchive().UEVer() < VER_UE4_SPEEDTREE_WIND_V7)
 	{
 		// update wind presets for speedtree v7
 		switch (WindType)
diff --git a/Engine/Source/Runtime/Engine/Private/Materials/MaterialInstance.cpp b/Engine/Source/Runtime/Engine/Private/Materials/MaterialInstance.cpp
--- a/Engine/Source/Runtime/Engine/Private/Materials/MaterialInstance.cpp
+++ b/Engine/Source/Runtime/Engine/Private/Materials/MaterialInstance.cpp
@@ -3395,7 +3395,7 @@ void UMaterialInstance::Serialize(FArchive& Ar)
 	// Only serialize the static permutation resource if one exists
 	if (bHasStaticPermutationResource)
 	{
-		if (Ar.UE4Ver() >= VER_UE4_PURGED_FMATERIAL_COMPILE_OUTPUTS)
+		if (Ar.UEVer() >= VER_UE4_PURGED_FMATERIAL_COMPILE_OUTPUTS)
 		{
 			if (Ar.CustomVer(FRenderingObjectVersion::GUID) < FRenderingObjectVersion::MaterialAttributeLayerParameters)
 			{
@@ -3438,10 +3438,10 @@ void UMaterialInstance::Serialize(FArchive& Ar)
 #endif // WITH_EDITOR
 	}
 
-	if (Ar.UE4Ver() >= VER_UE4_MATERIAL_INSTANCE_BASE_PROPERTY_OVERRIDES )
+	if (Ar.UEVer() >= VER_UE4_MATERIAL_INSTANCE_BASE_PROPERTY_OVERRIDES )
 	{
 #if WITH_EDITORONLY_DATA
-		if( Ar.UE4Ver() < VER_UE4_FIX_MATERIAL_PROPERTY_OVERRIDE_SERIALIZE )
+		if( Ar.UEVer() < VER_UE4_FIX_MATERIAL_PROPERTY_OVERRIDE_SERIALIZE )
 		{
 			// awful old native serialize of FMaterialInstanceBasePropertyOverrides UStruct
 			Ar << bOverrideBaseProperties_DEPRECATED;
@@ -3452,7 +3452,7 @@ void UMaterialInstance::Serialize(FArchive& Ar)
 				FArchive_Serialize_BitfieldBool(Ar, BasePropertyOverrides.bOverride_OpacityMaskClipValue);
 				Ar << BasePropertyOverrides.OpacityMaskClipValue;
 
-				if( Ar.UE4Ver() >= VER_UE4_MATERIAL_INSTANCE_BASE_PROPERTY_OVERRIDES_PHASE_2 )
+				if( Ar.UEVer() >= VER_UE4_MATERIAL_INSTANCE_BASE_PROPERTY_OVERRIDES_PHASE_2 )
 				{
 					FArchive_Serialize_BitfieldBool(Ar, BasePropertyOverrides.bOverride_BlendMode);
 					Ar << BasePropertyOverrides.BlendMode;
@@ -3461,13 +3461,13 @@ void UMaterialInstance::Serialize(FArchive& Ar)
 					FArchive_Serialize_BitfieldBool(Ar, BasePropertyOverrides.bOverride_TwoSided);
 					FArchive_Serialize_BitfieldBool(Ar, BasePropertyOverrides.TwoSided);
 
-					if( Ar.UE4Ver() >= VER_UE4_MATERIAL_INSTANCE_BASE_PROPERTY_OVERRIDES_DITHERED_LOD_TRANSITION )
+					if( Ar.UEVer() >= VER_UE4_MATERIAL_INSTANCE_BASE_PROPERTY_OVERRIDES_DITHERED_LOD_TRANSITION )
 					{
 						FArchive_Serialize_BitfieldBool(Ar, BasePropertyOverrides.bOverride_DitheredLODTransition);
 						FArchive_Serialize_BitfieldBool(Ar, BasePropertyOverrides.DitheredLODTransition);
 					}
 					// unrelated but closest change to bug
-					if( Ar.UE4Ver() < VER_UE4_STATIC_SHADOW_DEPTH_MAPS )
+					if( Ar.UEVer() < VER_UE4_STATIC_SHADOW_DEPTH_MAPS )
 					{
 						// switched enum order
 						switch( BasePropertyOverrides.ShadingModel )
@@ -3640,7 +3640,7 @@ void UMaterialInstance::PostLoad()
 	}
 
 	// Fixup for legacy instances which didn't recreate the lighting guid properly on duplication
-	if (GetLinker() && GetLinker()->UE4Ver() < VER_UE4_BUMPED_MATERIAL_EXPORT_GUIDS)
+	if (GetLinker() && GetLinker()->UEVer() < VER_UE4_BUMPED_MATERIAL_EXPORT_GUIDS)
 	{
 		extern TMap<FGuid, UMaterialInterface*> LightingGuidFixupMap;
 		UMaterialInterface** ExistingMaterial = LightingGuidFixupMap.Find(GetLightingGuid());
diff --git a/Engine/Source/Runtime/Engine/Private/Materials/MaterialShader.cpp b/Engine/Source/Runtime/Engine/Private/Materials/MaterialShader.cpp
--- a/Engine/Source/Runtime/Engine/Private/Materials/MaterialShader.cpp
+++ b/Engine/Source/Runtime/Engine/Private/Materials/MaterialShader.cpp
@@ -336,12 +336,12 @@ void FMaterialShaderMapId::Serialize(FArchive& Ar, bool bLoadedByCookedMaterial)
 
 	// Note: FMaterialShaderMapId is saved both in packages (legacy UMaterialInstance) and the DDC (FMaterialShaderMap)
 	// Backwards compatibility only works with FMaterialShaderMapId's stored in packages.  
-	// Only serialized in legacy packages if UE4Ver() < VER_UE4_PURGED_FMATERIAL_COMPILE_OUTPUTS
+	// Only serialized in legacy packages if UEVer() < VER_UE4_PURGED_FMATERIAL_COMPILE_OUTPUTS
 	// You must bump MATERIALSHADERMAP_DERIVEDDATA_VER as well if changing the serialization of FMaterialShaderMapId.
 	Ar.UsingCustomVersion(FEditorObjectVersion::GUID);
 	Ar.UsingCustomVersion(FReleaseObjectVersion::GUID);
 
-	const bool bIsLegacyPackage = Ar.UE4Ver() < VER_UE4_PURGED_FMATERIAL_COMPILE_OUTPUTS;
+	const bool bIsLegacyPackage = Ar.UEVer() < VER_UE4_PURGED_FMATERIAL_COMPILE_OUTPUTS;
 
 	// Ensure saved content is correct
 	check(!Ar.IsSaving() || IsContentValid());
@@ -398,7 +398,7 @@ void FMaterialShaderMapId::Serialize(FArchive& Ar, bool bLoadedByCookedMaterial)
 
 		Ar << ReferencedFunctions;
 
-		if (Ar.UE4Ver() >= VER_UE4_COLLECTIONS_IN_SHADERMAPID)
+		if (Ar.UEVer() >= VER_UE4_COLLECTIONS_IN_SHADERMAPID)
 		{
 			Ar << ReferencedParameterCollections;
 		}
@@ -427,7 +427,7 @@ void FMaterialShaderMapId::Serialize(FArchive& Ar, bool bLoadedByCookedMaterial)
 			Ar << LegacyHash;
 		}
 
-		if (Ar.UE4Ver() >= VER_UE4_MATERIAL_INSTANCE_BASE_PROPERTY_OVERRIDES)
+		if (Ar.UEVer() >= VER_UE4_MATERIAL_INSTANCE_BASE_PROPERTY_OVERRIDES)
 		{
 			Ar << BasePropertyOverridesHash;
 		}
diff --git a/Engine/Source/Runtime/Engine/Private/Materials/MaterialShared.cpp b/Engine/Source/Runtime/Engine/Private/Materials/MaterialShared.cpp
--- a/Engine/Source/Runtime/Engine/Private/Materials/MaterialShared.cpp
+++ b/Engine/Source/Runtime/Engine/Private/Materials/MaterialShared.cpp
@@ -1076,7 +1076,7 @@ FArchive& operator<<(FArchive& Ar, FLegacyTextureLookup& Ref)
 
 void FMaterial::LegacySerialize(FArchive& Ar)
 {
-	if (Ar.UE4Ver() < VER_UE4_PURGED_FMATERIAL_COMPILE_OUTPUTS)
+	if (Ar.UEVer() < VER_UE4_PURGED_FMATERIAL_COMPILE_OUTPUTS)
 	{
 		TArray<FString> LegacyStrings;
 		Ar << LegacyStrings;
@@ -1204,7 +1204,7 @@ void FMaterialResource::LegacySerialize(FArchive& Ar)
 {
 	FMaterial::LegacySerialize(Ar);
 
-	if (Ar.UE4Ver() < VER_UE4_PURGED_FMATERIAL_COMPILE_OUTPUTS)
+	if (Ar.UEVer() < VER_UE4_PURGED_FMATERIAL_COMPILE_OUTPUTS)
 	{
 		int32 BlendModeOverrideValueTemp = 0;
 		Ar << BlendModeOverrideValueTemp;
diff --git a/Engine/Source/Runtime/Engine/Private/Model.cpp b/Engine/Source/Runtime/Engine/Private/Model.cpp
--- a/Engine/Source/Runtime/Engine/Private/Model.cpp
+++ b/Engine/Source/Runtime/Engine/Private/Model.cpp
@@ -191,7 +191,7 @@ void UModel::Serialize( FArchive& Ar )
 
 	Ar << Bounds;
 
-	if (Ar.IsLoading() && Ar.UE4Ver() < VER_UE4_BSP_UNDO_FIX )
+	if (Ar.IsLoading() && Ar.UEVer() < VER_UE4_BSP_UNDO_FIX )
 	{
 		TTransArray<FVector> OldVectors(this);
 		TTransArray<FVector> OldPoints(this);
@@ -218,7 +218,7 @@ void UModel::Serialize( FArchive& Ar )
 		}
 	}
 	
-	if (Ar.IsLoading() && Ar.UE4Ver() < VER_UE4_BSP_UNDO_FIX )
+	if (Ar.IsLoading() && Ar.UEVer() < VER_UE4_BSP_UNDO_FIX )
 	{
 		TTransArray<FBspSurf> OldSurfs(this);
 		TTransArray<FVert> OldVerts(this);
@@ -235,7 +235,7 @@ void UModel::Serialize( FArchive& Ar )
 		Verts.BulkSerialize(Ar);
 	}
 
-	if( Ar.IsLoading() && Ar.UE4Ver() < VER_UE4_REMOVE_ZONES_FROM_MODEL)
+	if( Ar.IsLoading() && Ar.UEVer() < VER_UE4_REMOVE_ZONES_FROM_MODEL)
 	{
 		int32 NumZones;
 		Ar << NumSharedSides << NumZones;
@@ -254,7 +254,7 @@ void UModel::Serialize( FArchive& Ar )
 #if WITH_EDITOR
 	bool bHasEditorOnlyData = !Ar.IsFilterEditorOnly();
 	
-	if ( Ar.UE4Ver() < VER_UE4_REMOVE_UNUSED_UPOLYS_FROM_UMODEL )
+	if ( Ar.UEVer() < VER_UE4_REMOVE_UNUSED_UPOLYS_FROM_UMODEL )
 	{
 		bHasEditorOnlyData = true;
 	}
@@ -269,7 +269,7 @@ void UModel::Serialize( FArchive& Ar )
 #else
 	bool bHasEditorOnlyData = !Ar.IsFilterEditorOnly();
 	
-	if ( Ar.UE4Ver() < VER_UE4_REMOVE_UNUSED_UPOLYS_FROM_UMODEL )
+	if ( Ar.UEVer() < VER_UE4_REMOVE_UNUSED_UPOLYS_FROM_UMODEL )
 	{
 		bHasEditorOnlyData = true;
 	}
@@ -289,7 +289,7 @@ void UModel::Serialize( FArchive& Ar )
 
 	Ar << RootOutside << Linked;
 
-	if(Ar.IsLoading() && Ar.UE4Ver() < VER_UE4_REMOVE_ZONES_FROM_MODEL)
+	if(Ar.IsLoading() && Ar.UEVer() < VER_UE4_REMOVE_ZONES_FROM_MODEL)
 	{
 		TArray<int32> DummyPortalNodes;
 		DummyPortalNodes.BulkSerialize( Ar );
diff --git a/Engine/Source/Runtime/Engine/Private/Particles/ParticleModules.cpp b/Engine/Source/Runtime/Engine/Private/Particles/ParticleModules.cpp
--- a/Engine/Source/Runtime/Engine/Private/Particles/ParticleModules.cpp
+++ b/Engine/Source/Runtime/Engine/Private/Particles/ParticleModules.cpp
@@ -3410,7 +3410,7 @@ void UParticleModuleTypeDataMesh::SetToSensibleDefaults(UParticleEmitter* Owner)
 void UParticleModuleTypeDataMesh::Serialize(FArchive& Ar)
 {
 	Super::Serialize(Ar);
-	if (Ar.IsLoading() && Ar.UE4Ver() < VER_UE4_MESH_EMITTER_INITIAL_ORIENTATION_DISTRIBUTION)
+	if (Ar.IsLoading() && Ar.UEVer() < VER_UE4_MESH_EMITTER_INITIAL_ORIENTATION_DISTRIBUTION)
 	{
 		FVector oldOrient(0.0f, 0.0f, 0.0f);
 		CreateDistribution();
@@ -3419,7 +3419,7 @@ void UParticleModuleTypeDataMesh::Serialize(FArchive& Ar)
 		RPYDistribution->Max = oldOrient;
 		RPYDistribution->bIsDirty = true;
 	}
-	if (Ar.IsLoading() && Ar.UE4Ver() < VER_UE4_MESH_PARTICLE_COLLISIONS_CONSIDER_PARTICLE_SIZE)
+	if (Ar.IsLoading() && Ar.UEVer() < VER_UE4_MESH_PARTICLE_COLLISIONS_CONSIDER_PARTICLE_SIZE)
 	{
 		bCollisionsConsiderPartilceSize = false;//Old data should default to the old behavior of not considering particle size.
 	}
diff --git a/Engine/Source/Runtime/Engine/Private/PhysicsEngine/BodyInstance.cpp b/Engine/Source/Runtime/Engine/Private/PhysicsEngine/BodyInstance.cpp
--- a/Engine/Source/Runtime/Engine/Private/PhysicsEngine/BodyInstance.cpp
+++ b/Engine/Source/Runtime/Engine/Private/PhysicsEngine/BodyInstance.cpp
@@ -387,7 +387,7 @@ FArchive& operator<<(FArchive& Ar,FBodyInstance& BodyInst)
 		Ar << BodyInst.PhysMaterialOverride;
 	}
 
-	if (Ar.IsLoading() && Ar.UE4Ver() < VER_UE4_MAX_ANGULAR_VELOCITY_DEFAULT)
+	if (Ar.IsLoading() && Ar.UEVer() < VER_UE4_MAX_ANGULAR_VELOCITY_DEFAULT)
 	{
 		if(BodyInst.MaxAngularVelocity != 400.f)
 		{
@@ -3979,10 +3979,10 @@ void FBodyInstance::FixupData(class UObject* Loader)
 {
 	check (Loader);
 
-	int32 const UE4Version = Loader->GetLinkerUEVersion();
+	int32 const UEVersion = Loader->GetLinkerUEVersion();
 
 #if WITH_EDITOR
-	if (UE4Version < VER_UE4_ADD_CUSTOMPROFILENAME_CHANGE)
+	if (UEVersion < VER_UE4_ADD_CUSTOMPROFILENAME_CHANGE)
 	{
 		if (CollisionProfileName == NAME_None)
 		{
@@ -3990,22 +3990,22 @@ void FBodyInstance::FixupData(class UObject* Loader)
 		}
 	}
 
-	if (UE4Version < VER_UE4_SAVE_COLLISIONRESPONSE_PER_CHANNEL)
+	if (UEVersion < VER_UE4_SAVE_COLLISIONRESPONSE_PER_CHANNEL)
 	{
 		CollisionResponses.SetCollisionResponseContainer(ResponseToChannels_DEPRECATED);
 	}
 #endif // WITH_EDITORONLY_DATA
 
 	// Load profile. If older version, please verify profile name first
-	bool bNeedToVerifyProfile = (UE4Version < VER_UE4_COLLISION_PROFILE_SETTING) || 
+	bool bNeedToVerifyProfile = (UEVersion < VER_UE4_COLLISION_PROFILE_SETTING) || 
 		// or shape component needs to convert since we added profile
-		(UE4Version < VER_UE4_SAVE_COLLISIONRESPONSE_PER_CHANNEL && Loader->IsA(UShapeComponent::StaticClass()));
+		(UEVersion < VER_UE4_SAVE_COLLISIONRESPONSE_PER_CHANNEL && Loader->IsA(UShapeComponent::StaticClass()));
 	LoadProfileData(bNeedToVerifyProfile);
 
 	// if profile isn't set, then fix up channel responses
 	if( CollisionProfileName == UCollisionProfile::CustomCollisionProfileName ) 
 	{
-		if (UE4Version >= VER_UE4_SAVE_COLLISIONRESPONSE_PER_CHANNEL)
+		if (UEVersion >= VER_UE4_SAVE_COLLISIONRESPONSE_PER_CHANNEL)
 		{
 			CollisionResponses.UpdateResponseContainerFromArray();
 		}
diff --git a/Engine/Source/Runtime/Engine/Private/PhysicsEngine/BodySetup.cpp b/Engine/Source/Runtime/Engine/Private/PhysicsEngine/BodySetup.cpp
--- a/Engine/Source/Runtime/Engine/Private/PhysicsEngine/BodySetup.cpp
+++ b/Engine/Source/Runtime/Engine/Private/PhysicsEngine/BodySetup.cpp
@@ -1083,7 +1083,7 @@ void UBodySetup::Serialize(FArchive& Ar)
 		else
 #endif
 		{
-			if (Ar.UE4Ver() >= VER_UE4_STORE_HASCOOKEDDATA_FOR_BODYSETUP)
+			if (Ar.UEVer() >= VER_UE4_STORE_HASCOOKEDDATA_FOR_BODYSETUP)
 			{
 				// CL#14327190 Removed cooked implicit collision structures from the UBodySetup.
 				// UBodySetups saved with support for cooked implicit geometry store a counter for the number 
@@ -1989,7 +1989,7 @@ TArray<int32> FKConvexElem::GetChaosConvexIndices() const
 #if WITH_EDITORONLY_DATA
 void FKSphereElem::FixupDeprecated( FArchive& Ar )
 {
-	if ( Ar.IsLoading() && Ar.UE4Ver() < VER_UE4_REFACTOR_PHYSICS_TRANSFORMS )
+	if ( Ar.IsLoading() && Ar.UEVer() < VER_UE4_REFACTOR_PHYSICS_TRANSFORMS )
 	{
 		Center = TM_DEPRECATED.GetOrigin();
 	}
@@ -2060,7 +2060,7 @@ FKSphereElem FKSphereElem::GetFinalScaled(const FVector& Scale3D, const FTransfo
 #if WITH_EDITORONLY_DATA
 void FKBoxElem::FixupDeprecated( FArchive& Ar )
 {
-	if ( Ar.IsLoading() && Ar.UE4Ver() < VER_UE4_REFACTOR_PHYSICS_TRANSFORMS )
+	if ( Ar.IsLoading() && Ar.UEVer() < VER_UE4_REFACTOR_PHYSICS_TRANSFORMS )
 	{
 		Center = TM_DEPRECATED.GetOrigin();
 		Orientation_DEPRECATED = TM_DEPRECATED.ToQuat();
@@ -2146,7 +2146,7 @@ float FKBoxElem::GetClosestPointAndNormal(const FVector& WorldPosition, const FT
 #if WITH_EDITORONLY_DATA
 void FKSphylElem::FixupDeprecated( FArchive& Ar )
 {
-	if ( Ar.IsLoading() && Ar.UE4Ver() < VER_UE4_REFACTOR_PHYSICS_TRANSFORMS )
+	if ( Ar.IsLoading() && Ar.UEVer() < VER_UE4_REFACTOR_PHYSICS_TRANSFORMS )
 	{
 		Center = TM_DEPRECATED.GetOrigin();
 		Orientation_DEPRECATED = TM_DEPRECATED.ToQuat();
diff --git a/Engine/Source/Runtime/Engine/Private/PhysicsEngine/ConstraintInstance.cpp b/Engine/Source/Runtime/Engine/Private/PhysicsEngine/ConstraintInstance.cpp
--- a/Engine/Source/Runtime/Engine/Private/PhysicsEngine/ConstraintInstance.cpp
+++ b/Engine/Source/Runtime/Engine/Private/PhysicsEngine/ConstraintInstance.cpp
@@ -970,7 +970,7 @@ bool FConstraintInstance::Serialize(FArchive& Ar)
 #if WITH_EDITORONLY_DATA
 void FConstraintInstance::PostSerialize(const FArchive& Ar)
 {
-	if (Ar.IsLoading() && Ar.UE4Ver() < VER_UE4_FIXUP_STIFFNESS_AND_DAMPING_SCALE)
+	if (Ar.IsLoading() && Ar.UEVer() < VER_UE4_FIXUP_STIFFNESS_AND_DAMPING_SCALE)
 	{
 		LinearLimitStiffness_DEPRECATED		/= CVarConstraintAngularStiffnessScale.GetValueOnGameThread();
 		SwingLimitStiffness_DEPRECATED		/= CVarConstraintAngularStiffnessScale.GetValueOnGameThread();
@@ -980,12 +980,12 @@ void FConstraintInstance::PostSerialize(const FArchive& Ar)
 		TwistLimitDamping_DEPRECATED		/=  CVarConstraintAngularDampingScale.GetValueOnGameThread();
 	}
 
-	if (Ar.IsLoading() && Ar.UE4Ver() < VER_UE4_FIXUP_MOTOR_UNITS)
+	if (Ar.IsLoading() && Ar.UEVer() < VER_UE4_FIXUP_MOTOR_UNITS)
 	{
 		AngularVelocityTarget_DEPRECATED *= 1.f / (2.f * PI);	//we want to use revolutions per second - old system was using radians directly
 	}
 
-	if (Ar.IsLoading() && Ar.UE4Ver() < VER_UE4_CONSTRAINT_INSTANCE_MOTOR_FLAGS)
+	if (Ar.IsLoading() && Ar.UEVer() < VER_UE4_CONSTRAINT_INSTANCE_MOTOR_FLAGS)
 	{
 		bLinearXVelocityDrive_DEPRECATED = LinearVelocityTarget_DEPRECATED.X != 0.f;
 		bLinearYVelocityDrive_DEPRECATED = LinearVelocityTarget_DEPRECATED.Y != 0.f;
diff --git a/Engine/Source/Runtime/Engine/Private/PhysicsEngine/PhysicsConstraintTemplate.cpp b/Engine/Source/Runtime/Engine/Private/PhysicsEngine/PhysicsConstraintTemplate.cpp
--- a/Engine/Source/Runtime/Engine/Private/PhysicsEngine/PhysicsConstraintTemplate.cpp
+++ b/Engine/Source/Runtime/Engine/Private/PhysicsEngine/PhysicsConstraintTemplate.cpp
@@ -40,7 +40,7 @@ void UPhysicsConstraintTemplate::Serialize(FArchive& Ar)
 	Super::Serialize(Ar);
 
 	// If old content, copy properties out of setup into instance
-	if(Ar.UE4Ver() < VER_UE4_ALL_PROPS_TO_CONSTRAINTINSTANCE)
+	if(Ar.UEVer() < VER_UE4_ALL_PROPS_TO_CONSTRAINTINSTANCE)
 	{
 		CopySetupPropsToInstance(&DefaultInstance);
 	}
diff --git a/Engine/Source/Runtime/Engine/Private/Polygon.cpp b/Engine/Source/Runtime/Engine/Private/Polygon.cpp
--- a/Engine/Source/Runtime/Engine/Private/Polygon.cpp
+++ b/Engine/Source/Runtime/Engine/Private/Polygon.cpp
@@ -39,7 +39,7 @@ FArchive& operator<<(FArchive& Ar, FLightmassPrimitiveSettings& Settings)
 	Ar << Temp;
 	Settings.bUseEmissiveForStaticLighting = Temp;
 
-	if (Ar.UE4Ver() >= VER_UE4_NEW_LIGHTMASS_PRIMITIVE_SETTING)
+	if (Ar.UEVer() >= VER_UE4_NEW_LIGHTMASS_PRIMITIVE_SETTING)
 	{
 		Temp = Settings.bUseVertexNormalForHemisphereGather;
 		Ar << Temp;
@@ -1009,7 +1009,7 @@ void UPolys::Serialize( FArchive& Ar )
 	}
 	else
 	{
-		if( Ar.IsLoading() && Ar.UE4Ver() < VER_UE4_BSP_UNDO_FIX )
+		if( Ar.IsLoading() && Ar.UEVer() < VER_UE4_BSP_UNDO_FIX )
 		{
 			Element.CountBytes(Ar);
 			int32 DbNum = Element.Num(), DbMax = DbNum;
diff --git a/Engine/Source/Runtime/Engine/Private/PrecomputedLightVolume.cpp b/Engine/Source/Runtime/Engine/Private/PrecomputedLightVolume.cpp
--- a/Engine/Source/Runtime/Engine/Private/PrecomputedLightVolume.cpp
+++ b/Engine/Source/Runtime/Engine/Private/PrecomputedLightVolume.cpp
@@ -75,12 +75,12 @@ FArchive& operator<<(FArchive& Ar, TVolumeLightingSample<2>& Sample)
 	Ar << Sample.Radius;
 	Ar << Sample.Lighting;
 
-	if (Ar.UE4Ver() >= VER_UE4_SKY_BENT_NORMAL)
+	if (Ar.UEVer() >= VER_UE4_SKY_BENT_NORMAL)
 	{
 		Ar << Sample.PackedSkyBentNormal;
 	}
 
-	if (Ar.UE4Ver() >= VER_UE4_VOLUME_SAMPLE_LOW_QUALITY_SUPPORT)
+	if (Ar.UEVer() >= VER_UE4_VOLUME_SAMPLE_LOW_QUALITY_SUPPORT)
 	{
 		Ar << Sample.DirectionalLightShadowing;
 	}
@@ -187,7 +187,7 @@ FArchive& operator<<(FArchive& Ar,FPrecomputedLightVolumeData& Volume)
 
 			TArray<FVolumeLightingSample> LowQualitySamples;
 
-			if (Ar.UE4Ver() >= VER_UE4_VOLUME_SAMPLE_LOW_QUALITY_SUPPORT)
+			if (Ar.UEVer() >= VER_UE4_VOLUME_SAMPLE_LOW_QUALITY_SUPPORT)
 			{
 				LoadVolumeLightSamples(Ar, NumSHSamples, LowQualitySamples);
 			}
diff --git a/Engine/Source/Runtime/Engine/Private/RawIndexBuffer.cpp b/Engine/Source/Runtime/Engine/Private/RawIndexBuffer.cpp
--- a/Engine/Source/Runtime/Engine/Private/RawIndexBuffer.cpp
+++ b/Engine/Source/Runtime/Engine/Private/RawIndexBuffer.cpp
@@ -380,7 +380,7 @@ void FRawStaticIndexBuffer::Serialize(FArchive& Ar, bool bNeedsCPUAccess)
 {
 	IndexStorage.SetAllowCPUAccess(bNeedsCPUAccess);
 
-	if (Ar.UE4Ver() < VER_UE4_SUPPORT_32BIT_STATIC_MESH_INDICES)
+	if (Ar.UEVer() < VER_UE4_SUPPORT_32BIT_STATIC_MESH_INDICES)
 	{
 		TResourceArray<uint16,INDEXBUFFER_ALIGNMENT> LegacyIndices;
 
diff --git a/Engine/Source/Runtime/Engine/Private/Rendering/MultiSizeIndexContainer.cpp b/Engine/Source/Runtime/Engine/Private/Rendering/MultiSizeIndexContainer.cpp
--- a/Engine/Source/Runtime/Engine/Private/Rendering/MultiSizeIndexContainer.cpp
+++ b/Engine/Source/Runtime/Engine/Private/Rendering/MultiSizeIndexContainer.cpp
@@ -141,7 +141,7 @@ void FMultiSizeIndexContainer::CopyIndexBuffer(const TArray<uint32>& NewArray)
 void FMultiSizeIndexContainer::Serialize(FArchive& Ar, bool bNeedsCPUAccess)
 {
 	DECLARE_SCOPE_CYCLE_COUNTER(TEXT("FMultiSizeIndexContainer::Serialize"), STAT_MultiSizeIndexContainer_Serialize, STATGROUP_LoadTime);
-	if (Ar.UE4Ver() < VER_UE4_KEEP_SKEL_MESH_INDEX_DATA)
+	if (Ar.UEVer() < VER_UE4_KEEP_SKEL_MESH_INDEX_DATA)
 	{
 		bool bOldNeedsCPUAccess = true;
 		Ar << bOldNeedsCPUAccess;
diff --git a/Engine/Source/Runtime/Engine/Private/Rendering/SkeletalMeshVertexBuffer.cpp b/Engine/Source/Runtime/Engine/Private/Rendering/SkeletalMeshVertexBuffer.cpp
--- a/Engine/Source/Runtime/Engine/Private/Rendering/SkeletalMeshVertexBuffer.cpp
+++ b/Engine/Source/Runtime/Engine/Private/Rendering/SkeletalMeshVertexBuffer.cpp
@@ -39,7 +39,7 @@ FArchive& operator<<(FArchive& Ar, FDummySkeletalMeshVertexBuffer& VertexBuffer)
 
 	Ar.UsingCustomVersion(FSkeletalMeshCustomVersion::GUID);
 
-	if (Ar.UE4Ver() >= VER_UE4_SUPPORT_GPUSKINNING_8_BONE_INFLUENCES && Ar.CustomVer(FSkeletalMeshCustomVersion::GUID) < FSkeletalMeshCustomVersion::UseSeparateSkinWeightBuffer)
+	if (Ar.UEVer() >= VER_UE4_SUPPORT_GPUSKINNING_8_BONE_INFLUENCES && Ar.CustomVer(FSkeletalMeshCustomVersion::GUID) < FSkeletalMeshCustomVersion::UseSeparateSkinWeightBuffer)
 	{
 		Ar << bBackCompatExtraBoneInfluences;
 	}
diff --git a/Engine/Source/Runtime/Engine/Private/ReverbSettings.cpp b/Engine/Source/Runtime/Engine/Private/ReverbSettings.cpp
--- a/Engine/Source/Runtime/Engine/Private/ReverbSettings.cpp
+++ b/Engine/Source/Runtime/Engine/Private/ReverbSettings.cpp
@@ -18,7 +18,7 @@ bool FReverbSettings::operator==(const FReverbSettings& Other) const
 #if WITH_EDITORONLY_DATA
 void FReverbSettings::PostSerialize(const FArchive& Ar)
 {
-	if (Ar.UE4Ver() < VER_UE4_REVERB_EFFECT_ASSET_TYPE)
+	if (Ar.UEVer() < VER_UE4_REVERB_EFFECT_ASSET_TYPE)
 	{
 		FString ReverbAssetName;
 		switch (ReverbType_DEPRECATED)
diff --git a/Engine/Source/Runtime/Engine/Private/ShadowMap.cpp b/Engine/Source/Runtime/Engine/Private/ShadowMap.cpp
--- a/Engine/Source/Runtime/Engine/Private/ShadowMap.cpp
+++ b/Engine/Source/Runtime/Engine/Private/ShadowMap.cpp
@@ -569,7 +569,7 @@ void FShadowMap2D::Serialize(FArchive& Ar)
 		Ar << bChannelValid[Channel];
 	}
 
-	if (Ar.UE4Ver() >= VER_UE4_STATIC_SHADOWMAP_PENUMBRA_SIZE)
+	if (Ar.UEVer() >= VER_UE4_STATIC_SHADOWMAP_PENUMBRA_SIZE)
 	{
 		Ar << InvUniformPenumbraSize;
 	}
@@ -1212,7 +1212,7 @@ FArchive& operator<<(FArchive& Ar,FShadowMap*& R)
 		if (Ar.IsLoading())
 		{
 			// Dump old Shadowmaps
-			if (Ar.UE4Ver() < VER_UE4_COMBINED_LIGHTMAP_TEXTURES)
+			if (Ar.UEVer() < VER_UE4_COMBINED_LIGHTMAP_TEXTURES)
 			{
 				delete R; // safe because if we're loading we new'd this above
 				R = nullptr;
diff --git a/Engine/Source/Runtime/Engine/Private/SimpleConstructionScript.cpp b/Engine/Source/Runtime/Engine/Private/SimpleConstructionScript.cpp
--- a/Engine/Source/Runtime/Engine/Private/SimpleConstructionScript.cpp
+++ b/Engine/Source/Runtime/Engine/Private/SimpleConstructionScript.cpp
@@ -45,7 +45,7 @@ void USimpleConstructionScript::Serialize(FArchive& Ar)
 #if WITH_EDITORONLY_DATA
 	if(Ar.IsLoading())
 	{
-		if(Ar.UE4Ver() < VER_UE4_REMOVE_NATIVE_COMPONENTS_FROM_BLUEPRINT_SCS)
+		if(Ar.UEVer() < VER_UE4_REMOVE_NATIVE_COMPONENTS_FROM_BLUEPRINT_SCS)
 		{
 			// If we previously had a root node, we need to move it into the new RootNodes array. This is done in Serialize() in order to support SCS preloading (which relies on a valid RootNodes array).
 			if(RootNode_DEPRECATED != NULL)
diff --git a/Engine/Source/Runtime/Engine/Private/SkeletalMesh.cpp b/Engine/Source/Runtime/Engine/Private/SkeletalMesh.cpp
--- a/Engine/Source/Runtime/Engine/Private/SkeletalMesh.cpp
+++ b/Engine/Source/Runtime/Engine/Private/SkeletalMesh.cpp
@@ -1248,7 +1248,7 @@ FArchive &operator<<( FArchive& Ar, FSkeletalMeshLODInfo& I )
 	Ar << I.LODMaterialMap;
 
 #if WITH_EDITORONLY_DATA
-	if ( Ar.IsLoading() && Ar.UE4Ver() < VER_UE4_MOVE_SKELETALMESH_SHADOWCASTING )
+	if ( Ar.IsLoading() && Ar.UEVer() < VER_UE4_MOVE_SKELETALMESH_SHADOWCASTING )
 	{
 		Ar << I.bEnableShadowCasting_DEPRECATED;
 	}
@@ -1724,7 +1724,7 @@ void USkeletalMesh::Serialize( FArchive& Ar )
 		Ar << GetRefBasesInvMatrix();
 	}
 
-	if( Ar.UE4Ver() < VER_UE4_REFERENCE_SKELETON_REFACTOR )
+	if( Ar.UEVer() < VER_UE4_REFERENCE_SKELETON_REFACTOR )
 	{
 		TMap<FName, int32> DummyNameIndexMap;
 		Ar << DummyNameIndexMap;
@@ -1758,14 +1758,14 @@ void USkeletalMesh::Serialize( FArchive& Ar )
 		}
 	}
 
-	if ( Ar.IsLoading() && Ar.UE4Ver() < VER_UE4_ASSET_IMPORT_DATA_AS_JSON && !GetAssetImportData())
+	if ( Ar.IsLoading() && Ar.UEVer() < VER_UE4_ASSET_IMPORT_DATA_AS_JSON && !GetAssetImportData())
 	{
 		// AssetImportData should always be valid
 		SetAssetImportData(NewObject<UAssetImportData>(this, TEXT("AssetImportData")));
 	}
 	
 	// SourceFilePath and SourceFileTimestamp were moved into a subobject
-	if ( Ar.IsLoading() && Ar.UE4Ver() < VER_UE4_ADDED_FBX_ASSET_IMPORT_DATA && GetAssetImportData())
+	if ( Ar.IsLoading() && Ar.UEVer() < VER_UE4_ADDED_FBX_ASSET_IMPORT_DATA && GetAssetImportData())
 	{
 		// AssetImportData should always have been set up in the constructor where this is relevant
 		FAssetImportInfo Info;
@@ -1776,7 +1776,7 @@ void USkeletalMesh::Serialize( FArchive& Ar )
 		SourceFileTimestamp_DEPRECATED = TEXT("");
 	}
 
-	if (Ar.UE4Ver() >= VER_UE4_APEX_CLOTH)
+	if (Ar.UEVer() >= VER_UE4_APEX_CLOTH)
 	{
 		if(Ar.CustomVer(FSkeletalMeshCustomVersion::GUID) < FSkeletalMeshCustomVersion::NewClothingSystemAdded)
 		{
@@ -1787,20 +1787,20 @@ void USkeletalMesh::Serialize( FArchive& Ar )
 			}
 		}
 
-		if (Ar.UE4Ver() < VER_UE4_REFERENCE_SKELETON_REFACTOR)
+		if (Ar.UEVer() < VER_UE4_REFERENCE_SKELETON_REFACTOR)
 		{
 			RebuildRefSkeletonNameToIndexMap();
 		}
 	}
 
-	if ( Ar.IsLoading() && Ar.UE4Ver() < VER_UE4_MOVE_SKELETALMESH_SHADOWCASTING )
+	if ( Ar.IsLoading() && Ar.UEVer() < VER_UE4_MOVE_SKELETALMESH_SHADOWCASTING )
 	{
 		// Previous to this version, shadowcasting flags were stored in the LODInfo array
 		// now they're in the Materials array so we need to move them over
 		MoveDeprecatedShadowFlagToMaterials();
 	}
 
-	if (Ar.UE4Ver() < VER_UE4_SKELETON_ASSET_PROPERTY_TYPE_CHANGE)
+	if (Ar.UEVer() < VER_UE4_SKELETON_ASSET_PROPERTY_TYPE_CHANGE)
 	{
 		GetPreviewAttachedAssetContainer().SaveAttachedObjectsFromDeprecatedProperties();
 	}
@@ -4123,7 +4123,7 @@ FArchive& operator<<(FArchive& Ar, FSkeletalMaterial& Elem)
 #if WITH_EDITORONLY_DATA
 	else
 	{
-		if (Ar.UE4Ver() >= VER_UE4_MOVE_SKELETALMESH_SHADOWCASTING)
+		if (Ar.UEVer() >= VER_UE4_MOVE_SKELETALMESH_SHADOWCASTING)
 		{
 			Ar << Elem.bEnableShadowCasting_DEPRECATED;
 		}
diff --git a/Engine/Source/Runtime/Engine/Private/SkeletalMeshLODModel.cpp b/Engine/Source/Runtime/Engine/Private/SkeletalMeshLODModel.cpp
--- a/Engine/Source/Runtime/Engine/Private/SkeletalMeshLODModel.cpp
+++ b/Engine/Source/Runtime/Engine/Private/SkeletalMeshLODModel.cpp
@@ -81,7 +81,7 @@ FArchive& operator<<(FArchive& Ar, FSoftSkinVertex& V)
 		}
 	}
 
-	if (Ar.UE4Ver() >= VER_UE4_SUPPORT_8_BONE_INFLUENCES_SKELETAL_MESHES)
+	if (Ar.UEVer() >= VER_UE4_SUPPORT_8_BONE_INFLUENCES_SKELETAL_MESHES)
 	{
 		for (uint32 InfluenceIndex = MAX_INFLUENCES_PER_STREAM; InfluenceIndex < EXTRA_BONE_INFLUENCES; InfluenceIndex++)
 		{
@@ -111,7 +111,7 @@ FArchive& operator<<(FArchive& Ar, FSoftSkinVertex& V)
 		Ar << V.InfluenceWeights[InfluenceIndex];
 	}
 
-	if (Ar.UE4Ver() >= VER_UE4_SUPPORT_8_BONE_INFLUENCES_SKELETAL_MESHES)
+	if (Ar.UEVer() >= VER_UE4_SUPPORT_8_BONE_INFLUENCES_SKELETAL_MESHES)
 	{
 		for (uint32 InfluenceIndex = MAX_INFLUENCES_PER_STREAM; InfluenceIndex < EXTRA_BONE_INFLUENCES; InfluenceIndex++)
 		{
@@ -336,7 +336,7 @@ FArchive& operator<<(FArchive& Ar, FSkelMeshSection& S)
 	}
 
 	// for clothing info
-	if (Ar.UE4Ver() >= VER_UE4_APEX_CLOTH)
+	if (Ar.UEVer() >= VER_UE4_APEX_CLOTH)
 	{
 		// Load old 'disabled' flag on sections, as this was used to identify legacy clothing sections for conversion
 		if (Ar.CustomVer(FSkeletalMeshCustomVersion::GUID) < FSkeletalMeshCustomVersion::DeprecateSectionDisabledFlag)
@@ -351,7 +351,7 @@ FArchive& operator<<(FArchive& Ar, FSkelMeshSection& S)
 		}
 	}
 
-	if (Ar.UE4Ver() >= VER_UE4_APEX_CLOTH_LOD)
+	if (Ar.UEVer() >= VER_UE4_APEX_CLOTH_LOD)
 	{
 		if (Ar.CustomVer(FSkeletalMeshCustomVersion::GUID) < FSkeletalMeshCustomVersion::RemoveEnableClothLOD)
 		{
@@ -646,7 +646,7 @@ struct FLegacySkelMeshChunk
 		Ar << C.MaxBoneInfluences;
 
 
-		if (Ar.UE4Ver() >= VER_UE4_APEX_CLOTH)
+		if (Ar.UEVer() >= VER_UE4_APEX_CLOTH)
 		{
 			Ar << C.ApexClothMappingData;
 			Ar << C.PhysicalMeshVertices;
@@ -860,7 +860,7 @@ void FSkeletalMeshLODModel::Serialize(FArchive& Ar, UObject* Owner, int32 Idx)
 				}
 			}
 
-			if (Ar.UE4Ver() >= VER_UE4_APEX_CLOTH && HasClothData())
+			if (Ar.UEVer() >= VER_UE4_APEX_CLOTH && HasClothData())
 			{
 				FStripDataFlags StripFlags2(Ar, 0, VER_UE4_STATIC_SKELETAL_MESH_SERIALIZATION_FIX);
 				TSkeletalMeshVertexData<FMeshToMeshVertData> DummyClothData(true);
diff --git a/Engine/Source/Runtime/Engine/Private/SoundAttenuation.cpp b/Engine/Source/Runtime/Engine/Private/SoundAttenuation.cpp
--- a/Engine/Source/Runtime/Engine/Private/SoundAttenuation.cpp
+++ b/Engine/Source/Runtime/Engine/Private/SoundAttenuation.cpp
@@ -13,7 +13,7 @@
 #if WITH_EDITORONLY_DATA
 void FSoundAttenuationSettings::PostSerialize(const FArchive& Ar)
 {
-	if (Ar.UE4Ver() < VER_UE4_ATTENUATION_SHAPES)
+	if (Ar.UEVer() < VER_UE4_ATTENUATION_SHAPES)
 	{
 		FalloffDistance = RadiusMax_DEPRECATED - RadiusMin_DEPRECATED;
 
diff --git a/Engine/Source/Runtime/Engine/Private/SoundClass.cpp b/Engine/Source/Runtime/Engine/Private/SoundClass.cpp
--- a/Engine/Source/Runtime/Engine/Private/SoundClass.cpp
+++ b/Engine/Source/Runtime/Engine/Private/SoundClass.cpp
@@ -277,7 +277,7 @@ void USoundClass::Serialize( FArchive& Ar )
 {
 	Super::Serialize( Ar );
 
-	if (Ar.UE4Ver() < VER_UE4_SOUND_CLASS_GRAPH_EDITOR)
+	if (Ar.UEVer() < VER_UE4_SOUND_CLASS_GRAPH_EDITOR)
 	{
 		// load this to match size and then throw away
 		TMap<USoundClass*, FSoundClassEditorData>	EditorData_DEPRECATED;
diff --git a/Engine/Source/Runtime/Engine/Private/SoundCue.cpp b/Engine/Source/Runtime/Engine/Private/SoundCue.cpp
--- a/Engine/Source/Runtime/Engine/Private/SoundCue.cpp
+++ b/Engine/Source/Runtime/Engine/Private/SoundCue.cpp
@@ -167,7 +167,7 @@ void USoundCue::Serialize(FStructuredArchive::FRecord Record)
 		Super::Serialize(Record);
 	}
 
-	if (UnderlyingArchive.UE4Ver() >= VER_UE4_COOKED_ASSETS_IN_EDITOR_SUPPORT)
+	if (UnderlyingArchive.UEVer() >= VER_UE4_COOKED_ASSETS_IN_EDITOR_SUPPORT)
 	{
 		FStripDataFlags StripFlags(Record.EnterField(SA_FIELD_NAME(TEXT("SoundCueStripFlags"))));
 #if WITH_EDITORONLY_DATA
diff --git a/Engine/Source/Runtime/Engine/Private/SoundNode.cpp b/Engine/Source/Runtime/Engine/Private/SoundNode.cpp
--- a/Engine/Source/Runtime/Engine/Private/SoundNode.cpp
+++ b/Engine/Source/Runtime/Engine/Private/SoundNode.cpp
@@ -31,7 +31,7 @@ void USoundNode::Serialize(FArchive& Ar)
 {
 	Super::Serialize(Ar);
 
-	if (Ar.UE4Ver() >= VER_UE4_COOKED_ASSETS_IN_EDITOR_SUPPORT)
+	if (Ar.UEVer() >= VER_UE4_COOKED_ASSETS_IN_EDITOR_SUPPORT)
 	{
 		FStripDataFlags StripFlags(Ar);
 #if WITH_EDITORONLY_DATA
diff --git a/Engine/Source/Runtime/Engine/Private/SoundNodeEnveloper.cpp b/Engine/Source/Runtime/Engine/Private/SoundNodeEnveloper.cpp
--- a/Engine/Source/Runtime/Engine/Private/SoundNodeEnveloper.cpp
+++ b/Engine/Source/Runtime/Engine/Private/SoundNodeEnveloper.cpp
@@ -33,7 +33,7 @@ void USoundNodeEnveloper::Serialize(FArchive& Ar)
 	Super::Serialize(Ar);
 	if (Ar.IsLoading())
 	{
-		if (Ar.UE4Ver() < VER_UE4_SOUND_NODE_ENVELOPER_CURVE_CHANGE)
+		if (Ar.UEVer() < VER_UE4_SOUND_NODE_ENVELOPER_CURVE_CHANGE)
 		{
 			if (VolumeInterpCurve_DEPRECATED)
 			{
diff --git a/Engine/Source/Runtime/Engine/Private/SoundWave.cpp b/Engine/Source/Runtime/Engine/Private/SoundWave.cpp
--- a/Engine/Source/Runtime/Engine/Private/SoundWave.cpp
+++ b/Engine/Source/Runtime/Engine/Private/SoundWave.cpp
@@ -399,7 +399,7 @@ void USoundWave::Serialize( FArchive& Ar )
 
 	Ar.UsingCustomVersion(FFrameworkObjectVersion::GUID);
 
-	if (Ar.IsLoading() && (Ar.UE4Ver() >= VER_UE4_SOUND_COMPRESSION_TYPE_ADDED) && (Ar.CustomVer(FFrameworkObjectVersion::GUID) < FFrameworkObjectVersion::RemoveSoundWaveCompressionName))
+	if (Ar.IsLoading() && (Ar.UEVer() >= VER_UE4_SOUND_COMPRESSION_TYPE_ADDED) && (Ar.CustomVer(FFrameworkObjectVersion::GUID) < FFrameworkObjectVersion::RemoveSoundWaveCompressionName))
 	{
 		FName DummyCompressionName;
 		Ar << DummyCompressionName;
diff --git a/Engine/Source/Runtime/Engine/Private/SpeedTreeWind.cpp b/Engine/Source/Runtime/Engine/Private/SpeedTreeWind.cpp
--- a/Engine/Source/Runtime/Engine/Private/SpeedTreeWind.cpp
+++ b/Engine/Source/Runtime/Engine/Private/SpeedTreeWind.cpp
@@ -777,7 +777,7 @@ FArchive& operator<<(FArchive& Ar, FSpeedTreeWind& Wind)
 		Ar << Params.m_asBranch[BranchIndex].m_fTwitchFreqScale;
 	}
 
-	if (Ar.UE4Ver() < VER_UE4_SPEEDTREE_WIND_V7)
+	if (Ar.UEVer() < VER_UE4_SPEEDTREE_WIND_V7)
 	{
 		float fDiscardOldRolling = 0.0f;
 		Ar << fDiscardOldRolling;
@@ -805,7 +805,7 @@ FArchive& operator<<(FArchive& Ar, FSpeedTreeWind& Wind)
 	Ar << Params.m_fFrondRippleTile;
 	Ar << Params.m_fFrondRippleLightingScalar;
 	
-	if (Ar.UE4Ver() >= VER_UE4_SPEEDTREE_WIND_V7)
+	if (Ar.UEVer() >= VER_UE4_SPEEDTREE_WIND_V7)
 	{
 		Ar << Params.m_fRollingNoiseSize;
 		Ar << Params.m_fRollingNoiseTwist;
@@ -830,7 +830,7 @@ FArchive& operator<<(FArchive& Ar, FSpeedTreeWind& Wind)
 	bool Options[FSpeedTreeWind::NUM_WIND_OPTIONS];
 
 	#define SERIALIZE_OPTION(name) { Options[FSpeedTreeWind::name] = Wind.IsOptionEnabled(FSpeedTreeWind::name); Ar << Options[FSpeedTreeWind::name]; }
-	#define SKIP_OLD_OPTION() if (Ar.UE4Ver() < VER_UE4_SPEEDTREE_WIND_V7) { bool bDiscard = false; Ar << bDiscard; }
+	#define SKIP_OLD_OPTION() if (Ar.UEVer() < VER_UE4_SPEEDTREE_WIND_V7) { bool bDiscard = false; Ar << bDiscard; }
 	
 	SERIALIZE_OPTION(GLOBAL_WIND);
 	SERIALIZE_OPTION(GLOBAL_PRESERVE_SHAPE);
@@ -869,7 +869,7 @@ FArchive& operator<<(FArchive& Ar, FSpeedTreeWind& Wind)
 	SERIALIZE_OPTION(FROND_RIPPLE_TWO_SIDED);
 	SERIALIZE_OPTION(FROND_RIPPLE_ADJUST_LIGHTING);
 
-	if (Ar.UE4Ver() >= VER_UE4_SPEEDTREE_WIND_V7)
+	if (Ar.UEVer() >= VER_UE4_SPEEDTREE_WIND_V7)
 	{
 		SERIALIZE_OPTION(ROLLING);
 	}
diff --git a/Engine/Source/Runtime/Engine/Private/StaticMesh.cpp b/Engine/Source/Runtime/Engine/Private/StaticMesh.cpp
--- a/Engine/Source/Runtime/Engine/Private/StaticMesh.cpp
+++ b/Engine/Source/Runtime/Engine/Private/StaticMesh.cpp
@@ -2340,7 +2340,7 @@ static void SerializeBuildSettingsForDDC(FArchive& Ar, FMeshBuildSettings& Build
 	Ar << BuildSettings.SrcLightmapIndex;
 	Ar << BuildSettings.DstLightmapIndex;
 
-	if (Ar.IsLoading() && Ar.UE4Ver() < VER_UE4_BUILD_SCALE_VECTOR)
+	if (Ar.IsLoading() && Ar.UEVer() < VER_UE4_BUILD_SCALE_VECTOR)
 	{
 		float BuildScale(1.0f);
 		Ar << BuildScale;
@@ -4759,7 +4759,7 @@ void UStaticMesh::Serialize(FArchive& Ar)
 	Ar << bCooked;
 
 #if WITH_EDITORONLY_DATA
-	if (Ar.IsLoading() && Ar.UE4Ver() < VER_UE4_REMOVE_ZERO_TRIANGLE_SECTIONS)
+	if (Ar.IsLoading() && Ar.UEVer() < VER_UE4_REMOVE_ZERO_TRIANGLE_SECTIONS)
 	{
 		GStaticMeshesThatNeedMaterialFixup.Set(this);
 	}
@@ -4769,7 +4769,7 @@ void UStaticMesh::Serialize(FArchive& Ar)
 	Ar << LocalBodySetup;
 	SetBodySetup(LocalBodySetup);
 
-	if (Ar.UE4Ver() >= VER_UE4_STATIC_MESH_STORE_NAV_COLLISION)
+	if (Ar.UEVer() >= VER_UE4_STATIC_MESH_STORE_NAV_COLLISION)
 	{
 		UNavCollisionBase* LocalNavCollision = GetNavCollision();
 		Ar << LocalNavCollision;
@@ -4804,7 +4804,7 @@ void UStaticMesh::Serialize(FArchive& Ar)
 #if WITH_EDITORONLY_DATA
 	if( !StripFlags.IsEditorDataStripped() )
 	{
-		if ( Ar.IsLoading() && Ar.UE4Ver() < VER_UE4_DEPRECATED_STATIC_MESH_THUMBNAIL_PROPERTIES_REMOVED )
+		if ( Ar.IsLoading() && Ar.UEVer() < VER_UE4_DEPRECATED_STATIC_MESH_THUMBNAIL_PROPERTIES_REMOVED )
 		{
 			FRotator DummyThumbnailAngle;
 			float DummyThumbnailDistance;
@@ -4862,7 +4862,7 @@ void UStaticMesh::Serialize(FArchive& Ar)
 
 		// Need to set a flag rather than do conversion in place as RenderData is not
 		// created until postload and it is needed for bounding information
-		bRequiresLODDistanceConversion = Ar.UE4Ver() < VER_UE4_STATIC_MESH_SCREEN_SIZE_LODS;
+		bRequiresLODDistanceConversion = Ar.UEVer() < VER_UE4_STATIC_MESH_SCREEN_SIZE_LODS;
 		bRequiresLODScreenSizeConversion = Ar.CustomVer(FFrameworkObjectVersion::GUID) < FFrameworkObjectVersion::LODsUseResolutionIndependentScreenSize;
 	}
 #endif // #if WITH_EDITOR
@@ -4897,7 +4897,7 @@ void UStaticMesh::Serialize(FArchive& Ar)
 #endif
 	}
 
-	if (Ar.UE4Ver() >= VER_UE4_SPEEDTREE_STATICMESH)
+	if (Ar.UEVer() >= VER_UE4_SPEEDTREE_STATICMESH)
 	{
 		bool bHasSpeedTreeWind = SpeedTreeWind.IsValid();
 		Ar << bHasSpeedTreeWind;
@@ -4914,14 +4914,14 @@ void UStaticMesh::Serialize(FArchive& Ar)
 	}
 
 #if WITH_EDITORONLY_DATA
-	if ( Ar.IsLoading() && Ar.UE4Ver() < VER_UE4_ASSET_IMPORT_DATA_AS_JSON && !AssetImportData)
+	if ( Ar.IsLoading() && Ar.UEVer() < VER_UE4_ASSET_IMPORT_DATA_AS_JSON && !AssetImportData)
 	{
 		// AssetImportData should always be valid
 		AssetImportData = NewObject<UAssetImportData>(this, TEXT("AssetImportData"));
 	}
 	
 	// SourceFilePath and SourceFileTimestamp were moved into a subobject
-	if ( Ar.IsLoading() && Ar.UE4Ver() < VER_UE4_ADDED_FBX_ASSET_IMPORT_DATA && AssetImportData )
+	if ( Ar.IsLoading() && Ar.UEVer() < VER_UE4_ADDED_FBX_ASSET_IMPORT_DATA && AssetImportData )
 	{
 		// AssetImportData should always have been set up in the constructor where this is relevant
 		FAssetImportInfo Info;
diff --git a/Engine/Source/Runtime/Engine/Private/UnrealEngine.cpp b/Engine/Source/Runtime/Engine/Private/UnrealEngine.cpp
--- a/Engine/Source/Runtime/Engine/Private/UnrealEngine.cpp
+++ b/Engine/Source/Runtime/Engine/Private/UnrealEngine.cpp
@@ -10461,21 +10461,6 @@ float DrawMapWarnings(UWorld* World, FViewport* Viewport, FCanvas* Canvas, UCanv
 	}
 #endif
 
-	/* @todo ue4 temporarily disabled
-	AWorldSettings* WorldSettings = World->GetWorldSettings();
-	if( !WorldSettings->IsNavigationRebuilt() )
-	{
-	DrawShadowedString(Canvas,
-	MessageX,
-	MessageY,
-	TEXT("PATHS NEED TO BE REBUILT"),
-	GEngine->GetSmallFont(),
-	FColor(128,128,128)
-	);
-	MessageY += FontSizeY;
-	}
-	*/
-
 	if (World->bIsLevelStreamingFrozen)
 	{
 		SmallTextItem.Text = LOCTEXT("Levelstreamingfrozen", "Level streaming frozen...");
diff --git a/Engine/Source/Runtime/Engine/Private/World.cpp b/Engine/Source/Runtime/Engine/Private/World.cpp
--- a/Engine/Source/Runtime/Engine/Private/World.cpp
+++ b/Engine/Source/Runtime/Engine/Private/World.cpp
@@ -415,7 +415,7 @@ void UWorld::Serialize( FArchive& Ar )
 
 	Ar << PersistentLevel;
 
-	if (Ar.UE4Ver() < VER_UE4_ADD_EDITOR_VIEWS)
+	if (Ar.UEVer() < VER_UE4_ADD_EDITOR_VIEWS)
 	{
 #if WITH_EDITOR
 		EditorViews.SetNum(4);
@@ -448,7 +448,7 @@ void UWorld::Serialize( FArchive& Ar )
 	}
 #endif
 
-	if (Ar.UE4Ver() < VER_UE4_REMOVE_SAVEGAMESUMMARY)
+	if (Ar.UEVer() < VER_UE4_REMOVE_SAVEGAMESUMMARY)
 	{
 		UObject* DummyObject;
 		Ar << DummyObject;
diff --git a/Engine/Source/Runtime/Engine/Private/WorldSettings.cpp b/Engine/Source/Runtime/Engine/Private/WorldSettings.cpp
--- a/Engine/Source/Runtime/Engine/Private/WorldSettings.cpp
+++ b/Engine/Source/Runtime/Engine/Private/WorldSettings.cpp
@@ -304,7 +304,7 @@ void AWorldSettings::Serialize( FArchive& Ar )
 	Ar.UsingCustomVersion(FEnterpriseObjectVersion::GUID);
 	Ar.UsingCustomVersion(FWorldSettingCustomVersion::GUID);
 
-	if (Ar.UE4Ver() < VER_UE4_ADD_OVERRIDE_GRAVITY_FLAG)
+	if (Ar.UEVer() < VER_UE4_ADD_OVERRIDE_GRAVITY_FLAG)
 	{
 		//before we had override flag we would use GlobalGravityZ != 0
 		if(GlobalGravityZ != 0.0f)
diff --git a/Engine/Source/Runtime/Engine/Public/Animation/AnimTypes.h b/Engine/Source/Runtime/Engine/Public/Animation/AnimTypes.h
--- a/Engine/Source/Runtime/Engine/Public/Animation/AnimTypes.h
+++ b/Engine/Source/Runtime/Engine/Public/Animation/AnimTypes.h
@@ -865,7 +865,7 @@ struct ENGINE_API FRawAnimSequenceTrack
 		T.PosKeys.BulkSerialize(Ar);
 		T.RotKeys.BulkSerialize(Ar);
 
-		if (Ar.UE4Ver() >= VER_UE4_ANIM_SUPPORT_NONUNIFORM_SCALE_ANIMATION)
+		if (Ar.UEVer() >= VER_UE4_ANIM_SUPPORT_NONUNIFORM_SCALE_ANIMATION)
 		{
 			T.ScaleKeys.BulkSerialize(Ar);
 		}
diff --git a/Engine/Source/Runtime/Engine/Public/EngineMinimal.h b/Engine/Source/Runtime/Engine/Public/EngineMinimal.h
--- a/Engine/Source/Runtime/Engine/Public/EngineMinimal.h
+++ b/Engine/Source/Runtime/Engine/Public/EngineMinimal.h
@@ -1,7 +1,7 @@
 // Copyright Epic Games, Inc. All Rights Reserved.
 
 /*=============================================================================
-	EngineMinimal.h: Commonly used include for developing projects with UE4
+	EngineMinimal.h: Commonly used include for developing projects with UE
 =============================================================================*/
 
 #pragma once
diff --git a/Engine/Source/Runtime/Foliage/Private/FoliageTypeObject.cpp b/Engine/Source/Runtime/Foliage/Private/FoliageTypeObject.cpp
--- a/Engine/Source/Runtime/Foliage/Private/FoliageTypeObject.cpp
+++ b/Engine/Source/Runtime/Foliage/Private/FoliageTypeObject.cpp
@@ -67,7 +67,7 @@ void FFoliageTypeObject::SetClean()
 
 void FFoliageTypeObject::PostSerialize(const FArchive& Ar)
 {
-	if (Ar.UE4Ver() < VER_UE4_FOLIAGE_WITH_ASSET_OR_CLASS)
+	if (Ar.UEVer() < VER_UE4_FOLIAGE_WITH_ASSET_OR_CLASS)
 	{
 		// Set the type object to be the previously set FoliageType BP class
 		if (Type_DEPRECATED)
diff --git a/Engine/Source/Runtime/Foliage/Private/InstancedFoliage.cpp b/Engine/Source/Runtime/Foliage/Private/InstancedFoliage.cpp
--- a/Engine/Source/Runtime/Foliage/Private/InstancedFoliage.cpp
+++ b/Engine/Source/Runtime/Foliage/Private/InstancedFoliage.cpp
@@ -306,7 +306,7 @@ struct FFoliageInstanceCluster_Deprecated
 
 #if WITH_EDITORONLY_DATA
 		if (!Ar.ArIsFilterEditorOnly ||
-			Ar.UE4Ver() < VER_UE4_FOLIAGE_SETTINGS_TYPE)
+			Ar.UEVer() < VER_UE4_FOLIAGE_SETTINGS_TYPE)
 		{
 			Ar << OldCluster.InstanceIndices;
 		}
@@ -329,7 +329,7 @@ FArchive& operator<<(FArchive& Ar, FFoliageMeshInfo_Deprecated& MeshInfo)
 	}
 
 #if WITH_EDITORONLY_DATA
-	if ((!Ar.ArIsFilterEditorOnly || Ar.UE4Ver() < VER_UE4_FOLIAGE_SETTINGS_TYPE) &&
+	if ((!Ar.ArIsFilterEditorOnly || Ar.UEVer() < VER_UE4_FOLIAGE_SETTINGS_TYPE) &&
 		(!(Ar.GetPortFlags() & PPF_DuplicateForPIE)))
 	{
 		Ar << MeshInfo.Instances;
@@ -3845,7 +3845,7 @@ void AInstancedFoliageActor::Serialize(FArchive& Ar)
 	}
 #endif
 
-	if (Ar.UE4Ver() < VER_UE4_FOLIAGE_SETTINGS_TYPE)
+	if (Ar.UEVer() < VER_UE4_FOLIAGE_SETTINGS_TYPE)
 	{
 #if WITH_EDITORONLY_DATA
 		TMap<UFoliageType*, TUniqueObj<FFoliageMeshInfo_Deprecated>> FoliageMeshesDeprecated;
diff --git a/Engine/Source/Runtime/GameplayTags/Private/GameplayTagContainer.cpp b/Engine/Source/Runtime/GameplayTags/Private/GameplayTagContainer.cpp
--- a/Engine/Source/Runtime/GameplayTags/Private/GameplayTagContainer.cpp
+++ b/Engine/Source/Runtime/GameplayTags/Private/GameplayTagContainer.cpp
@@ -831,7 +831,7 @@ bool FGameplayTagContainer::Serialize(FStructuredArchive::FSlot Slot)
 {
 	FArchive& UnderlyingArchive = Slot.GetUnderlyingArchive();
 
-	const bool bOldTagVer = UnderlyingArchive.UE4Ver() < VER_UE4_GAMEPLAY_TAG_CONTAINER_TAG_TYPE_CHANGE;
+	const bool bOldTagVer = UnderlyingArchive.UEVer() < VER_UE4_GAMEPLAY_TAG_CONTAINER_TAG_TYPE_CHANGE;
 	
 	if (bOldTagVer)
 	{
diff --git a/Engine/Source/Runtime/Landscape/Private/Landscape.cpp b/Engine/Source/Runtime/Landscape/Private/Landscape.cpp
--- a/Engine/Source/Runtime/Landscape/Private/Landscape.cpp
+++ b/Engine/Source/Runtime/Landscape/Private/Landscape.cpp
@@ -493,7 +493,7 @@ void ULandscapeComponent::Serialize(FArchive& Ar)
 #endif
 	}
 
-	if (Ar.UE4Ver() >= VER_UE4_SERIALIZE_LANDSCAPE_GRASS_DATA)
+	if (Ar.UEVer() >= VER_UE4_SERIALIZE_LANDSCAPE_GRASS_DATA)
 	{
 		// Share the shared ref so PIE can share this data
 		if (Ar.GetPortFlags() & PPF_DuplicateForPIE)
@@ -526,7 +526,7 @@ void ULandscapeComponent::Serialize(FArchive& Ar)
 
 	bool bCooked = false;
 
-	if (Ar.UE4Ver() >= VER_UE4_LANDSCAPE_PLATFORMDATA_COOKING && !HasAnyFlags(RF_ClassDefaultObject))
+	if (Ar.UEVer() >= VER_UE4_LANDSCAPE_PLATFORMDATA_COOKING && !HasAnyFlags(RF_ClassDefaultObject))
 	{
 		bCooked = Ar.IsCooking() || (FPlatformProperties::RequiresCookedData() && Ar.IsSaving());
 		// This is needed when loading cooked data, to know to serialize differently
@@ -2421,7 +2421,7 @@ void ALandscapeProxy::PostLoad()
 		BodyInstance.FixupData(this);
 	}
 
-	if ((GetLinker() && (GetLinker()->UE4Ver() < VER_UE4_LANDSCAPE_COMPONENT_LAZY_REFERENCES)) ||
+	if ((GetLinker() && (GetLinker()->UEVer() < VER_UE4_LANDSCAPE_COMPONENT_LAZY_REFERENCES)) ||
 		LandscapeComponents.Num() != CollisionComponents.Num() ||
 		LandscapeComponents.ContainsByPredicate([](ULandscapeComponent* Comp) { return ((Comp != nullptr) && !Comp->CollisionComponent.IsValid()); }))
 	{
@@ -2430,7 +2430,7 @@ void ALandscapeProxy::PostLoad()
 #if WITH_EDITOR
 	if (GIsEditor && !GetWorld()->IsGameWorld())
 	{
-		if ((GetLinker() && (GetLinker()->UE4Ver() < VER_UE4_LANDSCAPE_COMPONENT_LAZY_REFERENCES)) ||
+		if ((GetLinker() && (GetLinker()->UEVer() < VER_UE4_LANDSCAPE_COMPONENT_LAZY_REFERENCES)) ||
 			LandscapeComponents.Num() != CollisionComponents.Num() ||
 			LandscapeComponents.ContainsByPredicate([](ULandscapeComponent* Comp) { return ((Comp != nullptr) && !Comp->CollisionComponent.IsValid()); }))
 		{
diff --git a/Engine/Source/Runtime/Landscape/Private/LandscapeCollision.cpp b/Engine/Source/Runtime/Landscape/Private/LandscapeCollision.cpp
--- a/Engine/Source/Runtime/Landscape/Private/LandscapeCollision.cpp
+++ b/Engine/Source/Runtime/Landscape/Private/LandscapeCollision.cpp
@@ -2143,7 +2143,7 @@ bool ULandscapeMeshCollisionComponent::RecreateCollision()
 void ULandscapeHeightfieldCollisionComponent::Serialize(FArchive& Ar)
 {
 #if WITH_EDITOR
-	if (Ar.UE4Ver() >= VER_UE4_LANDSCAPE_COLLISION_DATA_COOKING)
+	if (Ar.UEVer() >= VER_UE4_LANDSCAPE_COLLISION_DATA_COOKING)
 	{
 		// Cook data here so CookedPhysicalMaterials is always up to date
 		if (Ar.IsCooking() && !HasAnyFlags(RF_ClassDefaultObject))
@@ -2157,7 +2157,7 @@ void ULandscapeHeightfieldCollisionComponent::Serialize(FArchive& Ar)
 	// this will also serialize CookedPhysicalMaterials
 	Super::Serialize(Ar);
 
-	if (Ar.UE4Ver() < VER_UE4_LANDSCAPE_COLLISION_DATA_COOKING)
+	if (Ar.UEVer() < VER_UE4_LANDSCAPE_COLLISION_DATA_COOKING)
 	{
 #if WITH_EDITORONLY_DATA
 		CollisionHeightData.Serialize(Ar, this);
@@ -2201,7 +2201,7 @@ void ULandscapeMeshCollisionComponent::Serialize(FArchive& Ar)
 {
 	Super::Serialize(Ar);
 
-	if (Ar.UE4Ver() < VER_UE4_LANDSCAPE_COLLISION_DATA_COOKING)
+	if (Ar.UEVer() < VER_UE4_LANDSCAPE_COLLISION_DATA_COOKING)
 	{
 #if WITH_EDITORONLY_DATA
 		// conditional serialization in later versions
@@ -2211,7 +2211,7 @@ void ULandscapeMeshCollisionComponent::Serialize(FArchive& Ar)
 
 	// PhysX cooking mesh data
 	bool bCooked = false;
-	if (Ar.UE4Ver() >= VER_UE4_ADD_COOKED_TO_LANDSCAPE)
+	if (Ar.UEVer() >= VER_UE4_ADD_COOKED_TO_LANDSCAPE)
 	{
 		bCooked = Ar.IsCooking();
 		Ar << bCooked;
@@ -2226,7 +2226,7 @@ void ULandscapeMeshCollisionComponent::Serialize(FArchive& Ar)
 	{
 		// triangle mesh cooked data should be serialized in ULandscapeHeightfieldCollisionComponent
 	}
-	else if (Ar.UE4Ver() >= VER_UE4_LANDSCAPE_COLLISION_DATA_COOKING)
+	else if (Ar.UEVer() >= VER_UE4_LANDSCAPE_COLLISION_DATA_COOKING)
 	{
 #if WITH_EDITORONLY_DATA		
 		// we serialize raw collision data only with non-cooked content
diff --git a/Engine/Source/Runtime/Landscape/Private/LandscapeGrass.cpp b/Engine/Source/Runtime/Landscape/Private/LandscapeGrass.cpp
--- a/Engine/Source/Runtime/Landscape/Private/LandscapeGrass.cpp
+++ b/Engine/Source/Runtime/Landscape/Private/LandscapeGrass.cpp
@@ -1323,7 +1323,7 @@ FArchive& operator<<(FArchive& Ar, FLandscapeComponentGrassData& Data)
 		else
 		{
 			Data.MaterialStateIds.Empty(1);
-			if (Ar.UE4Ver() >= VER_UE4_SERIALIZE_LANDSCAPE_GRASS_DATA_MATERIAL_GUID)
+			if (Ar.UEVer() >= VER_UE4_SERIALIZE_LANDSCAPE_GRASS_DATA_MATERIAL_GUID)
 			{
 				FGuid MaterialStateId;
 				Ar << MaterialStateId;
diff --git a/Engine/Source/Runtime/Landscape/Private/LandscapeSplines.cpp b/Engine/Source/Runtime/Landscape/Private/LandscapeSplines.cpp
--- a/Engine/Source/Runtime/Landscape/Private/LandscapeSplines.cpp
+++ b/Engine/Source/Runtime/Landscape/Private/LandscapeSplines.cpp
@@ -625,7 +625,7 @@ void ULandscapeSplinesComponent::Serialize(FArchive& Ar)
 	Super::Serialize(Ar);
 
 #if WITH_EDITORONLY_DATA
-	if (Ar.UE4Ver() >= VER_UE4_LANDSCAPE_SPLINE_CROSS_LEVEL_MESHES &&
+	if (Ar.UEVer() >= VER_UE4_LANDSCAPE_SPLINE_CROSS_LEVEL_MESHES &&
 		!Ar.IsFilterEditorOnly())
 	{
 		Ar.UsingCustomVersion(FLandscapeCustomVersion::GUID);
@@ -1572,7 +1572,7 @@ void ULandscapeSplineControlPoint::Serialize(FArchive& Ar)
 	Ar.UsingCustomVersion(FLandscapeCustomVersion::GUID);
 
 #if WITH_EDITOR
-	if (Ar.UE4Ver() < VER_UE4_LANDSCAPE_SPLINE_CROSS_LEVEL_MESHES)
+	if (Ar.UEVer() < VER_UE4_LANDSCAPE_SPLINE_CROSS_LEVEL_MESHES)
 	{
 		bPlaceSplineMeshesInStreamingLevels = false;
 	}
@@ -2377,7 +2377,7 @@ void ULandscapeSplineSegment::Serialize(FArchive& Ar)
 	Ar.UsingCustomVersion(FLandscapeCustomVersion::GUID);
 
 #if WITH_EDITOR
-	if (Ar.UE4Ver() < VER_UE4_SPLINE_MESH_ORIENTATION)
+	if (Ar.UEVer() < VER_UE4_SPLINE_MESH_ORIENTATION)
 	{
 		for (FLandscapeSplineMeshEntry& MeshEntry : SplineMeshes)
 		{
@@ -2395,7 +2395,7 @@ void ULandscapeSplineSegment::Serialize(FArchive& Ar)
 		}
 	}
 
-	if (Ar.UE4Ver() < VER_UE4_LANDSCAPE_SPLINE_CROSS_LEVEL_MESHES)
+	if (Ar.UEVer() < VER_UE4_LANDSCAPE_SPLINE_CROSS_LEVEL_MESHES)
 	{
 		bPlaceSplineMeshesInStreamingLevels = false;
 	}
diff --git a/Engine/Source/Runtime/Landscape/Private/Materials/MaterialExpressionLandscapeLayerBlend.cpp b/Engine/Source/Runtime/Landscape/Private/Materials/MaterialExpressionLandscapeLayerBlend.cpp
--- a/Engine/Source/Runtime/Landscape/Private/Materials/MaterialExpressionLandscapeLayerBlend.cpp
+++ b/Engine/Source/Runtime/Landscape/Private/Materials/MaterialExpressionLandscapeLayerBlend.cpp
@@ -48,7 +48,7 @@ void UMaterialExpressionLandscapeLayerBlend::Serialize(FStructuredArchive::FReco
 	Super::Serialize(Record);
 	FArchive& UnderlyingArchive = Record.GetUnderlyingArchive();
 
-	if (UnderlyingArchive.IsLoading() && UnderlyingArchive.UE4Ver() < VER_UE4_ADD_LB_WEIGHTBLEND)
+	if (UnderlyingArchive.IsLoading() && UnderlyingArchive.UEVer() < VER_UE4_ADD_LB_WEIGHTBLEND)
 	{
 		// convert any LB_AlphaBlend entries to LB_WeightBlend
 		for (FLayerBlendInput& LayerInput : Layers)
diff --git a/Engine/Source/Runtime/Landscape/Private/Materials/MaterialExpressionLandscapeLayerSwitch.cpp b/Engine/Source/Runtime/Landscape/Private/Materials/MaterialExpressionLandscapeLayerSwitch.cpp
--- a/Engine/Source/Runtime/Landscape/Private/Materials/MaterialExpressionLandscapeLayerSwitch.cpp
+++ b/Engine/Source/Runtime/Landscape/Private/Materials/MaterialExpressionLandscapeLayerSwitch.cpp
@@ -113,7 +113,7 @@ void UMaterialExpressionLandscapeLayerSwitch::Serialize(FStructuredArchive::FRec
 {
 	Super::Serialize(Record);
 
-	if (Record.GetUnderlyingArchive().UE4Ver() < VER_UE4_FIX_TERRAIN_LAYER_SWITCH_ORDER)
+	if (Record.GetUnderlyingArchive().UEVer() < VER_UE4_FIX_TERRAIN_LAYER_SWITCH_ORDER)
 	{
 		Swap(LayerUsed, LayerNotUsed);
 	}
diff --git a/Engine/Source/Runtime/NavigationSystem/Private/NavMesh/PImplRecastNavMesh.cpp b/Engine/Source/Runtime/NavigationSystem/Private/NavMesh/PImplRecastNavMesh.cpp
--- a/Engine/Source/Runtime/NavigationSystem/Private/NavMesh/PImplRecastNavMesh.cpp
+++ b/Engine/Source/Runtime/NavigationSystem/Private/NavMesh/PImplRecastNavMesh.cpp
@@ -392,7 +392,7 @@ void FPImplRecastNavMesh::Serialize( FArchive& Ar, int32 NavMeshVersion )
 					dtFree(TileHeader);
 
 					//
-					if (Ar.UE4Ver() >= VER_UE4_ADD_MODIFIERS_RUNTIME_GENERATION && 
+					if (Ar.UEVer() >= VER_UE4_ADD_MODIFIERS_RUNTIME_GENERATION && 
 						(Ar.EngineVer().GetMajor() != 4 || Ar.EngineVer().GetMinor() != 7)) // Merged package from 4.7 branch
 					{
 						unsigned char* ComressedTileData = NULL;
@@ -433,7 +433,7 @@ void FPImplRecastNavMesh::Serialize( FArchive& Ar, int32 NavMeshVersion )
 					DetourNavMesh->addTile(TileData, TileDataSize, DT_TILE_FREE_DATA, TileRef, NULL);
 
 					// Serialize compressed tile cache layer
-					if (Ar.UE4Ver() >= VER_UE4_ADD_MODIFIERS_RUNTIME_GENERATION &&
+					if (Ar.UEVer() >= VER_UE4_ADD_MODIFIERS_RUNTIME_GENERATION &&
 						(Ar.EngineVer().GetMajor() != 4 || Ar.EngineVer().GetMinor() != 7)) // Merged package from 4.7 branch
 					{
 						uint8* ComressedTileData = nullptr;
diff --git a/Engine/Source/Runtime/NavigationSystem/Private/NavMesh/RecastNavMeshDataChunk.cpp b/Engine/Source/Runtime/NavigationSystem/Private/NavMesh/RecastNavMeshDataChunk.cpp
--- a/Engine/Source/Runtime/NavigationSystem/Private/NavMesh/RecastNavMeshDataChunk.cpp
+++ b/Engine/Source/Runtime/NavigationSystem/Private/NavMesh/RecastNavMeshDataChunk.cpp
@@ -140,7 +140,7 @@ void URecastNavMeshDataChunk::SerializeRecastData(FArchive& Ar, int32 NavMeshVer
 				// Load compressed tile cache layer
 				int32 TileCacheDataSize = 0;
 				uint8* TileCacheRawData = nullptr;
-				if (Ar.UE4Ver() >= VER_UE4_ADD_MODIFIERS_RUNTIME_GENERATION && 
+				if (Ar.UEVer() >= VER_UE4_ADD_MODIFIERS_RUNTIME_GENERATION && 
 					(Ar.EngineVer().GetMajor() != 4 || Ar.EngineVer().GetMinor() != 7)) // Merged package from 4.7 branch
 				{
 					FPImplRecastNavMesh::SerializeCompressedTileCacheData(Ar, NavMeshVersion, TileCacheRawData, TileCacheDataSize); //allocates TileCacheRawData on load
diff --git a/Engine/Source/Runtime/Online/BuildPatchServices/Private/Data/ManifestUObject.cpp b/Engine/Source/Runtime/Online/BuildPatchServices/Private/Data/ManifestUObject.cpp
--- a/Engine/Source/Runtime/Online/BuildPatchServices/Private/Data/ManifestUObject.cpp
+++ b/Engine/Source/Runtime/Online/BuildPatchServices/Private/Data/ManifestUObject.cpp
@@ -519,7 +519,7 @@ bool FManifestUObject::LoadInternal(FArchive& Ar, FBuildPatchAppManifest& AppMan
 
 	// Make sure we use the correct serialization version, this is now fixed and must never use a newer version,
 	// because the property tag has changed in structure meaning older clients would not read correctly.
-	Ar.SetUE4Ver(VER_UE4_STRUCT_GUID_IN_PROPERTY_TAG - 1);
+	Ar.SetUEVer(VER_UE4_STRUCT_GUID_IN_PROPERTY_TAG - 1);
 
 	if (Ar.IsLoading())
 	{
@@ -601,7 +601,7 @@ bool FManifestUObject::SaveInternal(FArchive& Ar, const FBuildPatchAppManifest&
 
 	// Make sure we use the correct serialization version, this is now fixed and must never use a newer version,
 	// because the property tag has changed in structure meaning older clients would not read correctly.
-	Ar.SetUE4Ver(VER_UE4_STRUCT_GUID_IN_PROPERTY_TAG - 1);
+	Ar.SetUEVer(VER_UE4_STRUCT_GUID_IN_PROPERTY_TAG - 1);
 
 	if (Ar.IsLoading())
 	{
diff --git a/Engine/Source/Runtime/SlateCore/Private/Fonts/SlateFontInfo.cpp b/Engine/Source/Runtime/SlateCore/Private/Fonts/SlateFontInfo.cpp
--- a/Engine/Source/Runtime/SlateCore/Private/Fonts/SlateFontInfo.cpp
+++ b/Engine/Source/Runtime/SlateCore/Private/Fonts/SlateFontInfo.cpp
@@ -200,7 +200,7 @@ const FCompositeFont* FSlateFontInfo::GetCompositeFont() const
 #if WITH_EDITORONLY_DATA
 void FSlateFontInfo::PostSerialize(const FArchive& Ar)
 {
-	if (Ar.UE4Ver() < VER_UE4_SLATE_COMPOSITE_FONTS && !FontObject)
+	if (Ar.UEVer() < VER_UE4_SLATE_COMPOSITE_FONTS && !FontObject)
 	{
 		UpgradeLegacyFontInfo(FontName_DEPRECATED, Hinting_DEPRECATED);
 	}
diff --git a/Engine/Source/Runtime/SlateCore/Private/Styling/SlateTypes.cpp b/Engine/Source/Runtime/SlateCore/Private/Styling/SlateTypes.cpp
--- a/Engine/Source/Runtime/SlateCore/Private/Styling/SlateTypes.cpp
+++ b/Engine/Source/Runtime/SlateCore/Private/Styling/SlateTypes.cpp
@@ -59,7 +59,7 @@ void FCheckBoxStyle::GetResources( TArray< const FSlateBrush* > & OutBrushes ) c
 #if WITH_EDITORONLY_DATA
 void FCheckBoxStyle::PostSerialize(const FArchive& Ar)
 {
-	if(Ar.IsLoading() && Ar.UE4Ver() < VER_UE4_FSLATESOUND_CONVERSION)
+	if(Ar.IsLoading() && Ar.UEVer() < VER_UE4_FSLATESOUND_CONVERSION)
 	{
 		// eg, SoundCue'/Game/QA_Assets/Sound/TEST_Music_Ambient.TEST_Music_Ambient'
 		CheckedSlateSound = FSlateSound::FromName_DEPRECATED(CheckedSound_DEPRECATED);
@@ -133,7 +133,7 @@ const FButtonStyle& FButtonStyle::GetDefault()
 #if WITH_EDITORONLY_DATA
 void FButtonStyle::PostSerialize(const FArchive& Ar)
 {
-	if(Ar.IsLoading() && Ar.UE4Ver() < VER_UE4_FSLATESOUND_CONVERSION)
+	if(Ar.IsLoading() && Ar.UEVer() < VER_UE4_FSLATESOUND_CONVERSION)
 	{
 		// eg, SoundCue'/Game/QA_Assets/Sound/TEST_Music_Ambient.TEST_Music_Ambient'
 		PressedSlateSound = FSlateSound::FromName_DEPRECATED(PressedSound_DEPRECATED);
@@ -190,7 +190,7 @@ const FComboBoxStyle& FComboBoxStyle::GetDefault()
 #if WITH_EDITORONLY_DATA
 void FComboBoxStyle::PostSerialize(const FArchive& Ar)
 {
-	if(Ar.IsLoading() && Ar.UE4Ver() < VER_UE4_FSLATESOUND_CONVERSION)
+	if(Ar.IsLoading() && Ar.UEVer() < VER_UE4_FSLATESOUND_CONVERSION)
 	{
 		// eg, SoundCue'/Game/QA_Assets/Sound/TEST_Music_Ambient.TEST_Music_Ambient'
 		PressedSlateSound = FSlateSound::FromName_DEPRECATED(PressedSound_DEPRECATED);
diff --git a/Engine/Source/Runtime/UMG/Private/UserWidget.cpp b/Engine/Source/Runtime/UMG/Private/UserWidget.cpp
--- a/Engine/Source/Runtime/UMG/Private/UserWidget.cpp
+++ b/Engine/Source/Runtime/UMG/Private/UserWidget.cpp
@@ -1867,7 +1867,7 @@ void UUserWidget::Serialize(FArchive& Ar)
 
 	if ( Ar.IsLoading() )
 	{
-		if ( Ar.UE4Ver() < VER_UE4_USERWIDGET_DEFAULT_FOCUSABLE_FALSE )
+		if ( Ar.UEVer() < VER_UE4_USERWIDGET_DEFAULT_FOCUSABLE_FALSE )
 		{
 			bIsFocusable = bSupportsKeyboardFocus_DEPRECATED;
 		}
