Fix an issue with deferred delete of modes not calling exit on the modes when the manager is cleaned up.

diff --git a/Engine/Source/Editor/UnrealEd/Private/EditorModeManager.cpp b/Engine/Source/Editor/UnrealEd/Private/EditorModeManager.cpp
--- a/Engine/Source/Editor/UnrealEd/Private/EditorModeManager.cpp
+++ b/Engine/Source/Editor/UnrealEd/Private/EditorModeManager.cpp
@@ -98,6 +98,7 @@ FEditorModeTools::~FEditorModeTools()
 
 	SetDefaultMode(FBuiltinEditorModes::EM_Default);
 	DeactivateAllModes();
+	DeactivateAllModesPendingDeletion();
 	RecycledScriptableModes.Empty();
 
 	// We may be destroyed after the UObject system has already shutdown, 
@@ -590,6 +591,18 @@ void FEditorModeTools::ForEachEdMode(TFunctionRef<bool(UEdMode*)> InCalllback) c
 	}
 }
 
+void FEditorModeTools::DeactivateAllModesPendingDeletion()
+{
+	// Reverse iterate since we are modifying the active modes list.
+	for (int32 Index = ActiveScriptableModes.Num() - 1; Index >= 0; --Index)
+	{
+		if (ActiveScriptableModes[Index]->IsPendingDeletion())
+		{
+			DeactivateScriptableModeAtIndex(Index);
+		}
+	}
+}
+
 void FEditorModeTools::SetPivotLocation( const FVector& Location, const bool bIncGridBase )
 {
 	CachedLocation = PivotLocation = SnappedLocation = Location;
@@ -875,23 +888,25 @@ void FEditorModeTools::RemoveAllDelegateHandlers()
 void FEditorModeTools::DeactivateMode( FEditorModeID InID )
 {
 	// Find the mode from the ID and exit it.
-	for (int32 Index = ActiveScriptableModes.Num() - 1; Index >= 0; --Index)
+	ForEachEdMode([InID](UEdMode* Mode)
 	{
-		auto& Mode = ActiveScriptableModes[Index];
 		if (Mode->GetID() == InID)
 		{
 			Mode->RequestDeletion();
-			break;
+			return false;
 		}
-	}
+
+		return true;
+	});
 }
 
 void FEditorModeTools::DeactivateAllModes()
 {
-	for (int32 Index = ActiveScriptableModes.Num() - 1; Index >= 0; --Index)
+	ForEachEdMode([](UEdMode* Mode)
 	{
-		ActiveScriptableModes[Index]->RequestDeletion();
-	}
+		Mode->RequestDeletion();
+		return true;
+	});
 }
 
 void FEditorModeTools::DestroyMode( FEditorModeID InID )
@@ -1356,13 +1371,7 @@ bool FEditorModeTools::ShouldDrawBrushVertices() const
 void FEditorModeTools::Tick( FEditorViewportClient* ViewportClient, float DeltaTime )
 {	
 	// Remove anything pending destruction
-	for (int32 Index = ActiveScriptableModes.Num() - 1; Index >= 0; --Index)
-	{
-		if (ActiveScriptableModes[Index]->IsPendingDeletion())
-		{
-			DeactivateScriptableModeAtIndex(Index);
-		}
-	}
+	DeactivateAllModesPendingDeletion();
 
 	if (ActiveScriptableModes.Num() == 0)
 	{
diff --git a/Engine/Source/Editor/UnrealEd/Public/EditorModeManager.h b/Engine/Source/Editor/UnrealEd/Public/EditorModeManager.h
--- a/Engine/Source/Editor/UnrealEd/Public/EditorModeManager.h
+++ b/Engine/Source/Editor/UnrealEd/Public/EditorModeManager.h
@@ -556,6 +556,8 @@ protected:
 		});
 	}
 
+	void DeactivateAllModesPendingDeletion();
+
 	/** List of default modes for this tool.  These must all be compatible with each other. */
 	TArray<FEditorModeID> DefaultModeIDs;
 
