Move ddc indicator to the status bar

diff --git a/Engine/Source/Editor/LevelEditor/Private/LevelEditor.cpp b/Engine/Source/Editor/LevelEditor/Private/LevelEditor.cpp
--- a/Engine/Source/Editor/LevelEditor/Private/LevelEditor.cpp
+++ b/Engine/Source/Editor/LevelEditor/Private/LevelEditor.cpp
@@ -41,7 +41,6 @@
 #include "Misc/EngineBuildSettings.h"
 #include "Subsystems/PanelExtensionSubsystem.h"
 #include "Classes/EditorStyleSettings.h"
-#include "DDC/SDDCStatusIndicator.h"
 
 #define LOCTEXT_NAMESPACE "LevelEditor"
 
@@ -224,13 +223,6 @@ TSharedRef<SDockTab> FLevelEditorModule::SpawnLevelEditor( const FSpawnTabArgs&
 			[
 				LevelEditorTmp->GetTitleBarMessageWidget()
 			]
-			
-			+ SHorizontalBox::Slot()
-			.AutoWidth()
-			[
-				SNew(SDDCStatusIndicator)
-			]
-
 			+ SHorizontalBox::Slot()
 			.AutoWidth()
 			.Padding(0.0f, 0.0f, 8.0f, 0.0f)
diff --git a/Engine/Source/Editor/LevelEditor/Private/SLevelEditor.cpp b/Engine/Source/Editor/LevelEditor/Private/SLevelEditor.cpp
--- a/Engine/Source/Editor/LevelEditor/Private/SLevelEditor.cpp
+++ b/Engine/Source/Editor/LevelEditor/Private/SLevelEditor.cpp
@@ -69,7 +69,7 @@
 #include "Elements/Component/ComponentElementLevelEditorCommonActionsCustomization.h"
 #include "Elements/SMInstance/SMInstanceElementId.h"
 #include "Elements/SMInstance/SMInstanceElementLevelEditorSelectionCustomization.h"
-#include <atomic>
+#include "DDC/SDDCStatusIndicator.h"
 
 #define LOCTEXT_NAMESPACE "SLevelEditor"
 
@@ -1912,9 +1912,9 @@ void SLevelEditor::RegisterStatusBarTools()
 
 	UToolMenu* Menu = UToolMenus::Get()->ExtendMenu("LevelEditor.StatusBar.ToolBar");
 	
-	FToolMenuSection& Section = Menu->AddSection("Compile", FText::GetEmpty(), FToolMenuInsert("SourceControl", EToolMenuInsertType::Before));
+	FToolMenuSection& CompileSection = Menu->AddSection("Compile", FText::GetEmpty(), FToolMenuInsert("SourceControl", EToolMenuInsertType::Before));
 
-	Section.AddDynamicEntry("CompilerAvailable", 
+	CompileSection.AddDynamicEntry("CompilerAvailable",
 		FNewToolMenuSectionDelegate::CreateLambda([](FToolMenuSection& InSection)
 		{
 			// Only show the compile options on machines with the solution (assuming they can build it)
@@ -1950,6 +1950,15 @@ void SLevelEditor::RegisterStatusBarTools()
 #endif
 			}
 		}));
+
+
+	FToolMenuSection& DDCSection = Menu->AddSection("DDC", FText::GetEmpty(), FToolMenuInsert("Compile", EToolMenuInsertType::Before));
+
+	DDCSection.AddEntry(
+		FToolMenuEntry::InitWidget("DDCStatusIndicator", SNew(SDDCStatusIndicator), FText::GetEmpty(), true, false)
+	);
+
+
 }
 
 void SLevelEditor::AddStandaloneLevelViewport( const TSharedRef<SLevelViewport>& LevelViewport )
diff --git a/Engine/Source/Editor/UnrealEd/Private/DDC/SDDCInformation.cpp b/Engine/Source/Editor/UnrealEd/Private/DDC/SDDCInformation.cpp
--- a/Engine/Source/Editor/UnrealEd/Private/DDC/SDDCInformation.cpp
+++ b/Engine/Source/Editor/UnrealEd/Private/DDC/SDDCInformation.cpp
@@ -18,6 +18,8 @@
 #include "Widgets/Layout/SGridPanel.h"
 #include "Algo/Compare.h"
 #include "Widgets/Input/SCheckBox.h"
+#include "Styling/StyleColors.h"
+#include "Widgets/Images/SImage.h"
 
 #define LOCTEXT_NAMESPACE "SDDCInformation"
 
@@ -174,7 +176,7 @@ TSharedRef<SWidget> SDDCInformation::GetComplexDataGrid()
 	Panel->AddSlot(0, Row)
 	[
 		SNew(STextBlock)
-		.ColorAndOpacity(FSlateColor::UseSubduedForeground())
+		.ColorAndOpacity(FStyleColors::ForegroundHover)
 		.Text(LOCTEXT("DDC", "DDC"))
 	];
 
@@ -182,7 +184,7 @@ TSharedRef<SWidget> SDDCInformation::GetComplexDataGrid()
 	[
 		SNew(STextBlock)
 		.Margin(FMargin(5, 0))
-		.ColorAndOpacity(FSlateColor::UseSubduedForeground())
+		.ColorAndOpacity(FStyleColors::ForegroundHover)
 		.Text(LOCTEXT("Speed", "Speed"))
 	];
 
@@ -190,7 +192,7 @@ TSharedRef<SWidget> SDDCInformation::GetComplexDataGrid()
 	[
 		SNew(STextBlock)
 		.Margin(FMargin(5, 0))
-		.ColorAndOpacity(FSlateColor::UseSubduedForeground())
+		.ColorAndOpacity(FStyleColors::ForegroundHover)
 		.Text(LOCTEXT("Loaded", "Loaded"))
 	];
 
@@ -198,17 +200,16 @@ TSharedRef<SWidget> SDDCInformation::GetComplexDataGrid()
 	[
 		SNew(STextBlock)
 		.Margin(FMargin(5, 0))
-		.ColorAndOpacity(FSlateColor::UseSubduedForeground())
+		.ColorAndOpacity(FStyleColors::ForegroundHover)
 		.Text(LOCTEXT("HitPercentage", "Hit%"))
 	];
 
 	Panel->AddSlot(4, Row)
 	[
-		SNew(STextBlock)
-		.Margin(FMargin(5, 0))
-		.ColorAndOpacity(FSlateColor::UseSubduedForeground())
-		.Font(FEditorStyle::Get().GetFontStyle("FontAwesome.12"))
-		.Text(FEditorFontGlyphs::Pencil)
+		SNew(SImage)
+		.Image(FAppStyle::Get().GetBrush("Icons.Edit"))
+		.ColorAndOpacity(FStyleColors::ForegroundHover)
+		
 	];
 
 	Row++;
@@ -270,8 +271,9 @@ TSharedRef<SWidget> SDDCInformation::GetComplexDataGrid()
 			Panel->AddSlot(4, Row)
 			.HAlign(HAlign_Center)
 			[
-				SNew(SCheckBox)
-				.IsChecked(Backend->IsWritable() ? ECheckBoxState::Checked : ECheckBoxState::Unchecked)
+				SNew(SImage)
+				.Image(FAppStyle::Get().GetBrush("Icons.Check"))
+				.Visibility(Backend->IsWritable() ? EVisibility::Visible : EVisibility::Hidden)
 			];
 
 			Row++;
diff --git a/Engine/Source/Editor/UnrealEd/Private/DDC/SDDCStatusIndicator.cpp b/Engine/Source/Editor/UnrealEd/Private/DDC/SDDCStatusIndicator.cpp
--- a/Engine/Source/Editor/UnrealEd/Private/DDC/SDDCStatusIndicator.cpp
+++ b/Engine/Source/Editor/UnrealEd/Private/DDC/SDDCStatusIndicator.cpp
@@ -19,6 +19,8 @@
 #include "Settings/EditorProjectSettings.h"
 #include "Modules/ModuleManager.h"
 #include "ISettingsModule.h"
+#include "Widgets/Images/SImage.h"
+#include "Styling/StyleColors.h"
 
 #define LOCTEXT_NAMESPACE "SDDCStatusIndicator"
 
@@ -29,14 +31,14 @@ void SDDCStatusIndicator::Construct(const FArguments& InArgs)
 	const FText DDCLabel = FText::FromString(DDC->IsDefaultGraph() ? TEXT("DDC") : DDCGraphName);
 	/*FDerivedDataCacheInterface::FOnDDCNotification& DDCNotificationEvent = GetDerivedDataCacheRef().GetDDCNotificationEvent();
 
-if (bSubscribe)
-{
-	DDCNotificationEvent.AddRaw(this, &FDDCNotifications::OnDDCNotificationEvent);
-}
-else
-{
-	DDCNotificationEvent.RemoveAll(this);
-}*/
+	if (bSubscribe)
+	{
+		DDCNotificationEvent.AddRaw(this, &FDDCNotifications::OnDDCNotificationEvent);
+	}
+	else
+	{
+		DDCNotificationEvent.RemoveAll(this);
+	}*/
 
 	BusyPulseSequence = FCurveSequence(0.f, 1.0f, ECurveEaseFunction::QuadInOut);
 	FadeGetSequence = FCurveSequence(0.f, 0.5f, ECurveEaseFunction::Linear);
@@ -44,66 +46,56 @@ else
 
 	this->ChildSlot
 	[
-		SNew(SBorder)
-		.BorderImage(FEditorStyle::GetBrush("SProjectBadge.BadgeShape"))
-		.Padding(FMargin(10.0f, 4.0f))
-		.BorderBackgroundColor(FColor::Black)
-		.ToolTip(SNew(SToolTip)[ SNew(SDDCInformation) ])
+		SNew(SHorizontalBox)
+		.ToolTip(SNew(SToolTip)[SNew(SDDCInformation)])
+		+ SHorizontalBox::Slot()
+		.AutoWidth()
+		.VAlign(VAlign_Center)
+		.Padding(0, 0, 3, 0)
 		[
-			SNew(SHorizontalBox)
-		
-			+ SHorizontalBox::Slot()
-			.AutoWidth()
-			.VAlign(VAlign_Center)
-			.Padding(0, 0, 3, 0)
-			[
-				SNew(SOverlay)
+			SNew(SOverlay)
 				
-				+ SOverlay::Slot()
-				.HAlign(HAlign_Center)
-				.VAlign(VAlign_Top)
-				.Padding(0, 0, 4, 4)
-				[
-					SNew(STextBlock)
-					.ColorAndOpacity_Lambda([this] { return FLinearColor::Red.CopyWithNewOpacity(FadePutSequence.GetLerp());} )
-					.Font(FEditorStyle::Get().GetFontStyle("FontAwesome.8"))
-					.Text(FEditorFontGlyphs::Arrow_Up)
-				]
-
-				+ SOverlay::Slot()
-				.HAlign(HAlign_Center)
-				.VAlign(VAlign_Bottom)
-				.Padding(4, 4, 0, 0)
-				[
-					SNew(STextBlock)
-					.ColorAndOpacity_Lambda([this] { return FLinearColor::Green.CopyWithNewOpacity(FadeGetSequence.GetLerp()); })
-					.Font(FEditorStyle::Get().GetFontStyle("FontAwesome.8"))
-					.Text(FEditorFontGlyphs::Arrow_Down)
-				]
-			]
-
-			+ SHorizontalBox::Slot()
-			.AutoWidth()
-			.VAlign(VAlign_Center)
-			.Padding(0,0,3,0)
+			+ SOverlay::Slot()
+			.HAlign(HAlign_Center)
+			.VAlign(VAlign_Top)
+			.Padding(0, 0, 4, 4)
 			[
-				SNew(STextBlock)
-				.ColorAndOpacity_Lambda([this] {
-					return FLinearColor::White.CopyWithNewOpacity(0.5f + (0.5f * FMath::MakePulsatingValue(BusyPulseSequence.GetLerp(), 1)));
-				})
-				.Font(FEditorStyle::Get().GetFontStyle("FontAwesome.12"))
-				.Text(FEditorFontGlyphs::Server)
+				SNew(SImage)
+				.Image(FAppStyle::Get().GetBrush("Icons.ArrowUp"))
+				.ColorAndOpacity_Lambda([this] { return FLinearColor::Red.CopyWithNewOpacity(FadePutSequence.GetLerp()); })
 			]
 
-			+ SHorizontalBox::Slot()
-			.AutoWidth()
-			.VAlign(VAlign_Center)
-			.Padding(0, 0, 10, 0)
+			+ SOverlay::Slot()
+			.HAlign(HAlign_Center)
+			.VAlign(VAlign_Bottom)
+			.Padding(4, 4, 0, 0)
 			[
-				SNew(STextBlock)
-				.Text(DDCLabel)
+				SNew(SImage)
+				.Image(FAppStyle::Get().GetBrush("Icons.ArrowDown"))
+				.ColorAndOpacity_Lambda([this] { return FLinearColor::Green.CopyWithNewOpacity(FadeGetSequence.GetLerp()); })
 			]
 		]
+
+		+ SHorizontalBox::Slot()
+		.AutoWidth()
+		.VAlign(VAlign_Center)
+		.Padding(0,0,3,0)
+		[
+			SNew(SImage)
+			.Image(FAppStyle::Get().GetBrush("Icons.Server"))
+			.ColorAndOpacity_Lambda([this] {
+				return FStyleColors::Foreground.GetSpecifiedColor().CopyWithNewOpacity(0.5f + (0.5f * FMath::MakePulsatingValue(BusyPulseSequence.GetLerp(), 1)));
+			})
+		]
+
+		+ SHorizontalBox::Slot()
+		.AutoWidth()
+		.VAlign(VAlign_Center)
+		.Padding(0, 0, 10, 0)
+		[
+			SNew(STextBlock)
+			.Text(DDCLabel)
+		]
 	];
 
 	LastDDCGetTime = GetDDCTime(true);
@@ -180,7 +172,6 @@ EActiveTimerReturnType SDDCStatusIndicator::UpdateWarnings(double InCurrentTime,
 
 		FNotificationInfo Info(LOCTEXT("SharedProjectLocalDDC", "This project recommends you setup the 'Global Local DDC Path', \nso that all copies of this project use the same local DDC cache."));
 		Info.bUseSuccessFailIcons = true;
-		Info.Image = FEditorStyle::GetBrush(TEXT("MessageLog.Warning"));
 		Info.bFireAndForget = false;
 		Info.bUseThrobber = false;
 		Info.FadeOutDuration = 0.0f;
@@ -195,13 +186,14 @@ EActiveTimerReturnType SDDCStatusIndicator::UpdateWarnings(double InCurrentTime,
 				Notification->SetCompletionState(SNotificationItem::CS_None);
 				Notification->ExpireAndFadeout();
 			}
-		})));
+		}),
+		SNotificationItem::ECompletionState::CS_Fail));
 
 		TSharedPtr<SNotificationItem> NotificationItem = FSlateNotificationManager::Get().AddNotification(Info);
 		if (NotificationItem.IsValid())
 		{
 			NotificationPromise.SetValue(NotificationItem);
-			NotificationItem->SetCompletionState(SNotificationItem::CS_Pending);
+			NotificationItem->SetCompletionState(SNotificationItem::CS_Fail);
 		}
 	}
 
@@ -211,7 +203,6 @@ EActiveTimerReturnType SDDCStatusIndicator::UpdateWarnings(double InCurrentTime,
 
 		FNotificationInfo Info(LOCTEXT("SharedProjectS3DDC", "This project recommends you setup the 'Global Local S3 DDC Path', \nso that all copies of this project use the same local S3 DDC cache."));
 		Info.bUseSuccessFailIcons = true;
-		Info.Image = FEditorStyle::GetBrush(TEXT("MessageLog.Warning"));
 		Info.bFireAndForget = false;
 		Info.bUseThrobber = false;
 		Info.FadeOutDuration = 0.0f;
@@ -226,13 +217,14 @@ EActiveTimerReturnType SDDCStatusIndicator::UpdateWarnings(double InCurrentTime,
 				Notification->SetCompletionState(SNotificationItem::CS_None);
 				Notification->ExpireAndFadeout();
 			}
-		})));
+		}),
+		SNotificationItem::ECompletionState::CS_Fail));
 
 		TSharedPtr<SNotificationItem> NotificationItem = FSlateNotificationManager::Get().AddNotification(Info);
 		if (NotificationItem.IsValid())
 		{
 			NotificationPromise.SetValue(NotificationItem);
-			NotificationItem->SetCompletionState(SNotificationItem::CS_Pending);
+			NotificationItem->SetCompletionState(SNotificationItem::CS_Fail);
 		}
 	}
 
diff --git a/Engine/Source/Runtime/SlateCore/Private/Styling/StarshipCoreStyle.cpp b/Engine/Source/Runtime/SlateCore/Private/Styling/StarshipCoreStyle.cpp
--- a/Engine/Source/Runtime/SlateCore/Private/Styling/StarshipCoreStyle.cpp
+++ b/Engine/Source/Runtime/SlateCore/Private/Styling/StarshipCoreStyle.cpp
@@ -294,6 +294,7 @@ TSharedRef<ISlateStyle> FStarshipCoreStyle::Create()
 		Style->Set("Icons.Help", new IMAGE_BRUSH("Icons/icon_help_16x", Icon16x16));
 	
 		Style->Set("Icons.Download", new IMAGE_BRUSH("Icons/icon_Downloads_16x", Icon16x16));
+		Style->Set("Icons.Server", new IMAGE_BRUSH_SVG("Starship/Common/server", Icon16x16));
 
 		Style->Set("Icons.Error", new IMAGE_BRUSH_SVG("Starship/Common/alert-triangle", Icon16x16));
 		Style->Set("Icons.ErrorWithColor", new IMAGE_BRUSH_SVG("Starship/Common/alert-triangle", Icon16x16, FStyleColors::Error));
@@ -441,7 +442,7 @@ TSharedRef<ISlateStyle> FStarshipCoreStyle::Create()
 	// SToolTip defaults...
 	{
 		Style->Set("ToolTip.Font", StyleFonts.Small);
-		Style->Set("ToolTip.Background", new BOX_BRUSH("Old/ToolTip_Background", FMargin(8.0f / 64.0f)));
+		Style->Set("ToolTip.Background", new FSlateRoundedBoxBrush(FStyleColors::Dropdown, 0.0f, FStyleColors::DropdownOutline, 1.0f));
 
 		Style->Set("ToolTip.LargerFont", StyleFonts.Normal);
 		Style->Set("ToolTip.BrightBackground", new BOX_BRUSH("Old/ToolTip_BrightBackground", FMargin(8.0f / 64.0f)));
