Fixed Build Nanite option in FBX import to not be an Advanced option (users will have trouble finding it).

diff --git a/Engine/Source/Editor/UnrealEd/Classes/Factories/FbxStaticMeshImportData.h b/Engine/Source/Editor/UnrealEd/Classes/Factories/FbxStaticMeshImportData.h
--- a/Engine/Source/Editor/UnrealEd/Classes/Factories/FbxStaticMeshImportData.h
+++ b/Engine/Source/Editor/UnrealEd/Classes/Factories/FbxStaticMeshImportData.h
@@ -38,7 +38,7 @@ class UFbxStaticMeshImportData : public UFbxMeshImportData
 	uint32 bBuildReversedIndexBuffer:1;
 
 	/** For static meshes, enabling this option will allow you to use Nanite rendering at runtime. Can only be used with simple opaque materials. */
-	UPROPERTY(EditAnywhere, BlueprintReadWrite, AdvancedDisplay, config, Category = Mesh, meta = (ToolTip = "If enabled, allows to render objects with Nanite", ImportType = "StaticMesh"))
+	UPROPERTY(EditAnywhere, BlueprintReadWrite, config, Category = Mesh, meta = (ToolTip = "If enabled, allows to render objects with Nanite", ImportType = "StaticMesh"))
 	uint32 bBuildNanite : 1;
 
 	UPROPERTY(EditAnywhere, BlueprintReadWrite, config, AdvancedDisplay, Category= Mesh, meta=(ImportType="StaticMesh", ReimportRestrict = "true"))
