Add the ability for reload to update the blueprints with the new UScriptStruct and UEnum classes.

Share as much code as possible without excessive complications between user defined structs and reload.

#rnx

diff --git a/Engine/Source/Editor/KismetCompiler/Private/UserDefinedStructureCompilerUtils.cpp b/Engine/Source/Editor/KismetCompiler/Private/UserDefinedStructureCompilerUtils.cpp
--- a/Engine/Source/Editor/KismetCompiler/Private/UserDefinedStructureCompilerUtils.cpp
+++ b/Engine/Source/Editor/KismetCompiler/Private/UserDefinedStructureCompilerUtils.cpp
@@ -23,6 +23,7 @@
 #include "Serialization/ObjectWriter.h"
 #include "Serialization/ObjectReader.h"
 #include "UObject/FieldIterator.h"
+#include "Algo/Copy.h"
 
 #define LOCTEXT_NAMESPACE "StructureCompiler"
 
@@ -441,100 +442,33 @@ void FUserDefinedStructureCompilerUtils::CompileStruct(class UUserDefinedStruct*
 		FUserDefinedStructureCompilerInner::BuildDependencyMapAndCompile(ChangedStructs, MessageLog);
 
 		// UPDATE ALL THINGS DEPENDENT ON COMPILED STRUCTURES
+		TSet<UScriptStruct*> ChangedStructsSet;
+		ChangedStructsSet.Reserve(ChangedStructs.Num());
+		Algo::Copy(ChangedStructs, ChangedStructsSet);
 		TSet<UBlueprint*> BlueprintsThatHaveBeenRecompiled;
-		for (TObjectIterator<UK2Node> It(RF_Transient | RF_ClassDefaultObject, /** bIncludeDerivedClasses */ true, /** InternalExcludeFlags */ EInternalObjectFlags::PendingKill); It && ChangedStructs.Num(); ++It)
-		{
-			bool bReconstruct = false;
-
-			UK2Node* Node = *It;
-
-			if (Node && !Node->HasAnyFlags(RF_Transient) && !Node->IsPendingKill())
-			{
-				// If this is a struct operation node operation on the changed struct we must reconstruct
-				if (UK2Node_StructOperation* StructOpNode = Cast<UK2Node_StructOperation>(Node))
-				{
-					UUserDefinedStruct* StructInNode = Cast<UUserDefinedStruct>(StructOpNode->StructType);
-					if (StructInNode && ChangedStructs.Contains(StructInNode))
-					{
-						bReconstruct = true;
-					}
-				}
-				if (!bReconstruct)
-				{
-					// Look through the nodes pins and if any of them are split and the type of the split pin is a user defined struct we need to reconstruct
-					for (UEdGraphPin* Pin : Node->Pins)
-					{
-						if (Pin->SubPins.Num() > 0)
-						{
-							UUserDefinedStruct* StructType = Cast<UUserDefinedStruct>(Pin->PinType.PinSubCategoryObject.Get());
-							if (StructType && ChangedStructs.Contains(StructType))
-							{
-								bReconstruct = true;
-								break;
-							}
-						}
-
-					}
-				}
-			}
-
-			if (bReconstruct)
+		FBlueprintEditorUtils::FindScriptStructsInNodes(ChangedStructsSet, [&BlueprintsThatHaveBeenRecompiled, &BlueprintsToRecompile](UBlueprint* Blueprint, UK2Node* Node)
 			{
 				// We need to recombine any nested subpins on this node, otherwise there will be an
 				// unexpected amount of pins during reconstruction. 
-				{
-					TArray<UEdGraphPin*> NestedSplitPins;
-					for (int32 i = Node->Pins.Num() - 1; i >= 0; --i)
-					{
-						UEdGraphPin* Pin = Node->Pins[i];
-						if (Pin->ParentPin != nullptr && Pin->ParentPin->ParentPin != nullptr && !Pin->bOrphanedPin)
-						{
-							NestedSplitPins.Add(Pin);
-							
-							// If there was nothing connected to or changed about this pin, then skip it
-							if (Pin->LinkedTo.Num() > 0 || !Pin->DoesDefaultValueMatchAutogenerated())
-							{
-								// Otherwise add an orphan pin so warning/connections are not silently lost
-								UEdGraphPin* OrphanPin = Node->CreatePin(Pin->Direction, Pin->PinType, Pin->PinName);
-								OrphanPin->bOrphanedPin = true;
-								OrphanPin->bNotConnectable = true;
-								OrphanPin->DefaultValue = Pin->DefaultValue;
-								OrphanPin->DefaultObject = Pin->DefaultObject;
-
-								for (UEdGraphPin* OldLink : Pin->LinkedTo)
-								{
-									OrphanPin->MakeLinkTo(OldLink);
-								}
-							}
-						}
-					}
+				FBlueprintEditorUtils::RecombineNestedSubPins(Node);
 
-					// Wait to recombine because otherwise we could end up combining pins that that haven't had their orphan created yet
-					const UEdGraphSchema* Schema = Node->GetSchema();
-					for (int32 i = NestedSplitPins.Num() - 1; i >= 0; --i)
-					{
-						Schema->RecombinePin(NestedSplitPins[i]);
-					}
-				}
-
-				if (Node->HasValidBlueprint())
+				if (Blueprint)
 				{
-					UBlueprint* FoundBlueprint = Node->GetBlueprint();
 					// The blueprint skeleton needs to be updated before we reconstruct the node
 					// or else we may have member references that point to the old skeleton
-					if (!BlueprintsThatHaveBeenRecompiled.Contains(FoundBlueprint))
+					if (!BlueprintsThatHaveBeenRecompiled.Contains(Blueprint))
 					{
-						BlueprintsThatHaveBeenRecompiled.Add(FoundBlueprint);
-						BlueprintsToRecompile.Remove(FoundBlueprint);
+						BlueprintsThatHaveBeenRecompiled.Add(Blueprint);
+						BlueprintsToRecompile.Remove(Blueprint);
 
 						// Reapply CDO data
 
-						FBlueprintEditorUtils::MarkBlueprintAsStructurallyModified(FoundBlueprint);
+						FBlueprintEditorUtils::MarkBlueprintAsStructurallyModified(Blueprint);
 					}
 					Node->ReconstructNode();
 				}
 			}
-		}
+		);
 
 		for (TPair<UBlueprint*, FUserDefinedStructureCompilerInner::FBlueprintUserStructData>& Pair : BlueprintsToRecompile)
 		{
diff --git a/Engine/Source/Editor/UnrealEd/Private/Kismet2/BlueprintEditorUtils.cpp b/Engine/Source/Editor/UnrealEd/Private/Kismet2/BlueprintEditorUtils.cpp
--- a/Engine/Source/Editor/UnrealEd/Private/Kismet2/BlueprintEditorUtils.cpp
+++ b/Engine/Source/Editor/UnrealEd/Private/Kismet2/BlueprintEditorUtils.cpp
@@ -75,6 +75,7 @@
 #include "K2Node_MathExpression.h"
 #include "K2Node_MatineeController.h"
 #include "K2Node_SpawnActorFromClass.h"
+#include "K2Node_StructOperation.h"
 #include "K2Node_TemporaryVariable.h"
 #include "K2Node_Timeline.h"
 #include "K2Node_Knot.h"
@@ -96,6 +97,7 @@
 #include "ScopedTransaction.h"
 #include "ClassViewerFilter.h"
 #include "InstancedReferenceSubobjectHelper.h"
+#include "NodeDependingOnEnumInterface.h"
 
 #include "BlueprintEditorModule.h"
 #include "BlueprintEditor.h"
@@ -9872,4 +9874,204 @@ void FBlueprintEditorUtils::BuildComponentInstancingData(UActorComponent* Compon
 	}
 }
 
+namespace 
+{
+	// This structure provides the ability to find/update the nodes primary object.  This must be specialized based
+	// on the type of the object being found/updated
+	template <typename TObjectType, bool bIsFind>
+	struct FFindOrUpdateNodeHelper
+	{
+		template <typename FindExisting>
+		static bool FindOrUpdateNode(UK2Node* Node, FindExisting& InFindExisting);
+	};
+
+	template <bool bIsFind>
+	struct FFindOrUpdateNodeHelper<UScriptStruct, bIsFind>
+	{
+		template <typename FindExisting>
+		static bool FindOrUpdateNode(UK2Node* Node, FindExisting InFindExisting)
+		{
+			// If this is a struct operation node operation on the changed struct we must reconstruct
+			if (UK2Node_StructOperation* StructOpNode = Cast<UK2Node_StructOperation>(Node))
+			{
+				if (UScriptStruct* StructInNode = Cast<UScriptStruct>(StructOpNode->StructType))
+				{
+					if (TOptional<UScriptStruct*> NewStructInNode = InFindExisting(StructInNode))
+					{
+						if (!bIsFind)
+						{
+							StructOpNode->StructType = *NewStructInNode;
+						}
+						return true;
+					}
+				}
+			}
+			return false;
+		}
+	};
+
+	template <bool bIsFind>
+	struct FFindOrUpdateNodeHelper<UEnum, bIsFind>
+	{
+		template <typename FindExisting>
+		static bool FindOrUpdateNode(UK2Node* Node, FindExisting InFindExisting)
+		{
+			if (INodeDependingOnEnumInterface* EnumInterface = Cast<INodeDependingOnEnumInterface>(Node))
+			{
+				if (UEnum* EnumInNode = EnumInterface->GetEnum())
+				{
+					if (TOptional<UEnum*> NewEnumInNode = InFindExisting(EnumInNode))
+					{
+						if (!bIsFind)
+						{
+							EnumInterface->ReloadEnum(*NewEnumInNode);
+						}
+						return true;
+					}
+				}
+			}
+			return false;
+		}
+	};
+
+	// Scan all the nodes looking for references to objects. 
+	template <typename TObject, bool bIsFind, typename FindExisting>
+	void FindOrUpdateNodes(FBlueprintEditorUtils::FOnNodeFoundOrUpdated InOnNodeFoundOrUpdated, FindExisting InFindExisting)
+	{
+		for (TObjectIterator<UK2Node> It(RF_Transient | RF_ClassDefaultObject, /** bIncludeDerivedClasses */ true, /** InternalExcludeFlags */ EInternalObjectFlags::PendingKill); It; ++It)
+		{
+			UK2Node* Node = *It;
+
+			if (Node && !Node->HasAnyFlags(RF_Transient) && !Node->IsPendingKill())
+			{
+				bool bReconstruct = FFindOrUpdateNodeHelper<TObject, bIsFind>::FindOrUpdateNode(Node, InFindExisting);
+
+				// Look through the nodes pins and if any of them are split and the type of the split pin is a something we need to reconstruct
+				if (!bIsFind || !bReconstruct)
+				{
+					for (UEdGraphPin* Pin : Node->Pins)
+					{
+						if (TObject* Object = Cast<TObject>(Pin->PinType.PinSubCategoryObject.Get()))
+						{
+							if (TOptional<TObject*> NewObject = InFindExisting(Object))
+							{
+								bReconstruct = true;
+								if (bIsFind)
+								{
+									break;
+								}
+								Pin->PinType.PinSubCategoryObject = *NewObject;
+							}
+						}
+					}
+				}
+
+				if (bReconstruct)
+				{
+					UBlueprint* FoundBlueprint = Node->HasValidBlueprint() ? FoundBlueprint = Node->GetBlueprint() : nullptr;
+					InOnNodeFoundOrUpdated(FoundBlueprint, Node);
+				}
+			}
+		}
+	}
+}
+
+void FBlueprintEditorUtils::FindScriptStructsInNodes(const TSet<UScriptStruct*>& Structs, FOnNodeFoundOrUpdated InOnNodeFoundOrUpdated)
+{
+	if (Structs.Num() == 0)
+	{
+		return;
+	}
+
+	FindOrUpdateNodes<UScriptStruct, true>(InOnNodeFoundOrUpdated, [&Structs](UScriptStruct* ScriptStruct)
+		{
+			return Structs.Contains(ScriptStruct) ? TOptional(ScriptStruct) : TOptional<UScriptStruct*>();
+		}
+	);
+}
+
+void FBlueprintEditorUtils::FindEnumsInNodes(const TSet<UEnum*>& Enums, FOnNodeFoundOrUpdated InOnNodeFoundOrUpdated)
+{
+	if (Enums.Num() == 0)
+	{
+		return;
+	}
+
+	FindOrUpdateNodes<UEnum, true>(InOnNodeFoundOrUpdated, [&Enums](UEnum* Enum)
+		{
+			return Enums.Contains(Enum) ? TOptional(Enum) : TOptional<UEnum*>();
+		}
+	);
+}
+
+void FBlueprintEditorUtils::UpdateScriptStructsInNodes(const TMap<UScriptStruct*, UScriptStruct*>& Structs, FOnNodeFoundOrUpdated InOnNodeFoundOrUpdated)
+{
+	if (Structs.Num() == 0)
+	{
+		return;
+	}
+
+	Structs.Find(nullptr);
+
+	FindOrUpdateNodes<UScriptStruct, false>(InOnNodeFoundOrUpdated, [&Structs] (UScriptStruct* ScriptStruct)
+		{
+			UScriptStruct* const* Found = Structs.Find(ScriptStruct);
+			return Found ? TOptional(*Found) : TOptional<UScriptStruct*>();
+		}
+	);
+}
+
+void FBlueprintEditorUtils::UpdateEnumsInNodes(const TMap<UEnum*, UEnum*>& Enums, FOnNodeFoundOrUpdated InOnNodeFoundOrUpdated)
+{
+	if (Enums.Num() == 0)
+	{
+		return;
+	}
+
+	FindOrUpdateNodes<UEnum, false>(InOnNodeFoundOrUpdated, [&Enums](UEnum* Enum)
+		{
+			UEnum* const* Found = Enums.Find(Enum);
+			return Found ? TOptional(*Found) : TOptional<UEnum*>();
+		}
+	);
+}
+
+void FBlueprintEditorUtils::RecombineNestedSubPins(UK2Node* Node)
+{
+	checkSlow(Node);
+
+	TArray<UEdGraphPin*> NestedSplitPins;
+	for (int32 i = Node->Pins.Num() - 1; i >= 0; --i)
+	{
+		UEdGraphPin* Pin = Node->Pins[i];
+		if (Pin->ParentPin != nullptr && Pin->ParentPin->ParentPin != nullptr && !Pin->bOrphanedPin)
+		{
+			NestedSplitPins.Add(Pin);
+
+			// If there was nothing connected to or changed about this pin, then skip it
+			if (Pin->LinkedTo.Num() > 0 || !Pin->DoesDefaultValueMatchAutogenerated())
+			{
+				// Otherwise add an orphan pin so warning/connections are not silently lost
+				UEdGraphPin* OrphanPin = Node->CreatePin(Pin->Direction, Pin->PinType, Pin->PinName);
+				OrphanPin->bOrphanedPin = true;
+				OrphanPin->bNotConnectable = true;
+				OrphanPin->DefaultValue = Pin->DefaultValue;
+				OrphanPin->DefaultObject = Pin->DefaultObject;
+
+				for (UEdGraphPin* OldLink : Pin->LinkedTo)
+				{
+					OrphanPin->MakeLinkTo(OldLink);
+				}
+			}
+		}
+	}
+
+	// Wait to recombine because otherwise we could end up combining pins that that haven't had their orphan created yet
+	const UEdGraphSchema* Schema = Node->GetSchema();
+	for (int32 i = NestedSplitPins.Num() - 1; i >= 0; --i)
+	{
+		Schema->RecombinePin(NestedSplitPins[i]);
+	}
+}
+
 #undef LOCTEXT_NAMESPACE
diff --git a/Engine/Source/Editor/UnrealEd/Private/Kismet2/ReloadUtilities.cpp b/Engine/Source/Editor/UnrealEd/Private/Kismet2/ReloadUtilities.cpp
--- a/Engine/Source/Editor/UnrealEd/Private/Kismet2/ReloadUtilities.cpp
+++ b/Engine/Source/Editor/UnrealEd/Private/Kismet2/ReloadUtilities.cpp
@@ -8,6 +8,7 @@
 #include "Async/AsyncWork.h"
 #include "Engine/Engine.h"
 #include "Engine/EngineTypes.h"
+#include "Kismet2/BlueprintEditorUtils.h"
 #include "Kismet2/KismetReinstanceUtilities.h"
 #include "Misc/QueuedThreadPool.h"
 #include "Serialization/ArchiveReplaceObjectRef.h"
@@ -104,7 +105,7 @@ class FReloadClassReinstancer : public FBlueprintCompileReinstancer
 public:
 
 	/** Sets the re-instancer up to re-instance native classes */
-	FReloadClassReinstancer(UClass* InNewClass, UClass* InOldClass, const TMap<UClass*, UClass*>& OldToNewClassesMap, TMap<UObject*, UObject*>& OutReconstructedCDOsMap, TSet<UBlueprint*>& InBPSetToRecompile, TSet<UBlueprint*>& InBPSetToRecompileBytecodeOnly);
+	FReloadClassReinstancer(UClass* InNewClass, UClass* InOldClass, const TSet<UObject*>& InReinstancingObjects, TMap<UObject*, UObject*>& OutReconstructedCDOsMap, TSet<UBlueprint*>& InCompiledBlueprints);
 
 	/** Destructor */
 	virtual ~FReloadClassReinstancer();
@@ -119,9 +120,9 @@ public:
 	void ReinstanceObjectsAndUpdateDefaults();
 
 	/** Creates the reinstancer as a sharable object */
-	static TSharedPtr<FReloadClassReinstancer> Create(UClass* InNewClass, UClass* InOldClass, const TMap<UClass*, UClass*>& OldToNewClassesMap, TMap<UObject*, UObject*>& OutReconstructedCDOsMap, TSet<UBlueprint*>& InBPSetToRecompile, TSet<UBlueprint*>& InBPSetToRecompileBytecodeOnly)
+	static TSharedPtr<FReloadClassReinstancer> Create(UClass* InNewClass, UClass* InOldClass, const TSet<UObject*>& InReinstancingObjects, TMap<UObject*, UObject*>& OutReconstructedCDOsMap, TSet<UBlueprint*>& InCompiledBlueprints)
 	{
-		return MakeShareable(new FReloadClassReinstancer(InNewClass, InOldClass, OldToNewClassesMap, OutReconstructedCDOsMap, InBPSetToRecompile, InBPSetToRecompileBytecodeOnly));
+		return MakeShareable(new FReloadClassReinstancer(InNewClass, InOldClass, InReinstancingObjects, OutReconstructedCDOsMap, InCompiledBlueprints));
 	}
 
 	// FSerializableObject interface
@@ -130,7 +131,6 @@ public:
 
 	virtual bool IsClassObjectReplaced() const override { return true; }
 
-	virtual void EnlistDependentBlueprintToRecompile(UBlueprint* BP, bool bBytecodeOnly) override;
 	virtual void BlueprintWasRecompiled(UBlueprint* BP, bool bBytecodeOnly) override;
 
 protected:
@@ -142,10 +142,10 @@ protected:
 private:
 	/** Reference to reconstructed CDOs map in this hot-reload session. */
 	TMap<UObject*, UObject*>& ReconstructedCDOsMap;
-	TSet<UBlueprint*>& BPSetToRecompile;
-	TSet<UBlueprint*>& BPSetToRecompileBytecodeOnly;
-	const TMap<UClass*, UClass*>& OldToNewClassesMap;
-};
+
+	/** Collection of blueprints already recompiled */
+	TSet<UBlueprint*>& CompiledBlueprints;
+ };
 
 void FReloadClassReinstancer::SetupNewClassReinstancing(UClass* InNewClass, UClass* InOldClass)
 {
@@ -419,23 +419,21 @@ void FReloadClassReinstancer::RecreateCDOAndSetupOldClassReinstancing(UClass* In
 	}
 }
 
-FReloadClassReinstancer::FReloadClassReinstancer(UClass* InNewClass, UClass* InOldClass, const TMap<UClass*, UClass*>& InOldToNewClassesMap, TMap<UObject*, UObject*>& OutReconstructedCDOsMap, TSet<UBlueprint*>& InBPSetToRecompile, TSet<UBlueprint*>& InBPSetToRecompileBytecodeOnly)
+FReloadClassReinstancer::FReloadClassReinstancer(UClass* InNewClass, UClass* InOldClass, const TSet<UObject*>& InReinstancingObjects, TMap<UObject*, UObject*>& OutReconstructedCDOsMap, TSet<UBlueprint*>& InCompiledBlueprints)
 	: NewClass(nullptr)
 	, bNeedsReinstancing(false)
 	, CopyOfPreviousCDO(nullptr)
 	, ReconstructedCDOsMap(OutReconstructedCDOsMap)
-	, BPSetToRecompile(InBPSetToRecompile)
-	, BPSetToRecompileBytecodeOnly(InBPSetToRecompileBytecodeOnly)
-	, OldToNewClassesMap(InOldToNewClassesMap)
+	, CompiledBlueprints(InCompiledBlueprints)
 {
 	ensure(InOldClass);
 	ensure(!HotReloadedOldClass && !HotReloadedNewClass);
 	HotReloadedOldClass = InOldClass;
 	HotReloadedNewClass = InNewClass ? InNewClass : InOldClass;
 
-	for (const TPair<UClass*, UClass*>& OldToNewClass : OldToNewClassesMap)
+	for (UObject* Object : InReinstancingObjects)
 	{
-		ObjectsThatShouldUseOldStuff.Add(OldToNewClass.Key);
+		ObjectsThatShouldUseOldStuff.Add(Object);
 	}
 
 	// If InNewClass is NULL, then the old class has not changed after hot-reload.
@@ -450,10 +448,6 @@ FReloadClassReinstancer::FReloadClassReinstancer(UClass* InNewClass, UClass* InO
 		for (TObjectIterator<UBlueprint> BlueprintIt; BlueprintIt; ++BlueprintIt)
 		{
 			FArchiveReplaceObjectRef<UObject> ReplaceObjectArch(*BlueprintIt, ClassRedirects, false, true, true);
-			if (ReplaceObjectArch.GetCount())
-			{
-				EnlistDependentBlueprintToRecompile(*BlueprintIt, false);
-			}
 		}
 	}
 	else
@@ -687,36 +681,9 @@ void FReloadClassReinstancer::AddReferencedObjects(FReferenceCollector& Collecto
 	Collector.AllowEliminatingReferences(true);
 }
 
-void FReloadClassReinstancer::EnlistDependentBlueprintToRecompile(UBlueprint* BP, bool bBytecodeOnly)
-{
-	if (IsValid(BP))
-	{
-		if (bBytecodeOnly)
-		{
-			if (!BPSetToRecompile.Contains(BP) && !BPSetToRecompileBytecodeOnly.Contains(BP))
-			{
-				BPSetToRecompileBytecodeOnly.Add(BP);
-			}
-		}
-		else
-		{
-			if (!BPSetToRecompile.Contains(BP))
-			{
-				if (BPSetToRecompileBytecodeOnly.Contains(BP))
-				{
-					BPSetToRecompileBytecodeOnly.Remove(BP);
-				}
-
-				BPSetToRecompile.Add(BP);
-			}
-		}
-	}
-}
-
 void FReloadClassReinstancer::BlueprintWasRecompiled(UBlueprint* BP, bool bBytecodeOnly)
 {
-	BPSetToRecompile.Remove(BP);
-	BPSetToRecompileBytecodeOnly.Remove(BP);
+	CompiledBlueprints.Add(BP);
 
 	FBlueprintCompileReinstancer::BlueprintWasRecompiled(BP, bBytecodeOnly);
 }
@@ -779,13 +746,11 @@ bool FReload::GetEnableReinstancing(bool bHasChanged) const
 void FReload::Reset()
 {
 	FunctionRemap.Empty();
-	BPSetToRecompile.Empty();
-	BPSetToRecompileBytecodeOnly.Empty();
 	ReconstructedCDOsMap.Empty();
 	ReinstancedClasses.Empty();
 }
 
-void FReload::UpdateStats(ReinstanceStats& Stats, void* New, void* Old)
+void FReload::UpdateStats(FReinstanceStats& Stats, void* New, void* Old)
 {
 	if (Old == nullptr)
 	{
@@ -801,7 +766,7 @@ void FReload::UpdateStats(ReinstanceStats& Stats, void* New, void* Old)
 	}
 }
 
-void FReload::FormatStats(FStringBuilderBase& Out, const TCHAR* Singular, const TCHAR* Plural, const ReinstanceStats& Stats)
+void FReload::FormatStats(FStringBuilderBase& Out, const TCHAR* Singular, const TCHAR* Plural, const FReinstanceStats& Stats)
 {
 	FormatStat(Out, Singular, Plural, TEXT("new"), Stats.New);
 	FormatStat(Out, Singular, Plural, TEXT("changed"), Stats.Changed);
@@ -856,11 +821,25 @@ void FReload::NotifyChange(UClass* New, UClass* Old)
 void FReload::NotifyChange(UEnum* New, UEnum* Old)
 {
 	UpdateStats(EnumStats, New, Old);
+
+	if (Old != nullptr)
+	{
+		UEnum* NewIfChanged = Old != New ? New : nullptr; // supporting code detects unchanged based on null new pointer
+		checkf(!ReinstancedEnums.Contains(Old) || ReinstancedEnums[Old] == NewIfChanged, TEXT("Attempting to reload an enumeration which is already being reloaded as a different enumeration"));
+		ReinstancedEnums.Add(Old, NewIfChanged);
+	}
 }
 
 void FReload::NotifyChange(UScriptStruct* New, UScriptStruct* Old)
 {
 	UpdateStats(StructStats, New, Old);
+
+	if (Old != nullptr)
+	{
+		UScriptStruct* NewIfChanged = Old != New ? New : nullptr; // supporting code detects unchanged based on null new pointer
+		checkf(!ReinstancedStructs.Contains(Old) || ReinstancedStructs[Old] == NewIfChanged, TEXT("Attempting to reload a structure which is already being reloaded as a different structure"));
+		ReinstancedStructs.Add(Old, NewIfChanged);
+	}
 }
 
 void FReload::Reinstance()
@@ -872,18 +851,78 @@ void FReload::Reinstance()
 
 	TMap<UClass*, UClass*>& ClassesToReinstance = GetClassesToReinstanceForHotReload();
 
-	TMap<UClass*, UClass*> OldToNewClassesMap;
+	TSet<UObject*> ReinstancingObjects;
+	ReinstancingObjects.Reserve(ClassesToReinstance.Num() + ReinstancedStructs.Num() + ReinstancedEnums.Num());
 	for (const TPair<UClass*, UClass*>& Pair : ClassesToReinstance)
 	{
-		if (Pair.Value != nullptr)
+		ReinstancingObjects.Add(Pair.Key);
+	}
+
+	// Collect all of the blueprint nodes that are getting updated due to enum/struct changes
+	TMap<UBlueprint*, FBlueprintUpdateInfo> ModifiedBlueprints;
+	FBlueprintEditorUtils::FOnNodeFoundOrUpdated OnNodeFoundOrUpdated = [&ModifiedBlueprints](UBlueprint* Blueprint, UK2Node* Node)
+	{
+		// Blueprint can be nullptr
+		FBlueprintUpdateInfo& BlueprintUpdateInfo = ModifiedBlueprints.FindOrAdd(Blueprint);
+		BlueprintUpdateInfo.Nodes.Add(Node);
+	};
+
+	TMap<UScriptStruct*, UScriptStruct*> ChangedStructs;
+	for (const TPair<UScriptStruct*, UScriptStruct*>& Pair : ReinstancedStructs)
+	{
+		ReinstancingObjects.Add(Pair.Key);
+		if (Pair.Value)
+		{
+			ChangedStructs.Emplace(Pair.Key, Pair.Value);
+		}
+	}
+	FBlueprintEditorUtils::UpdateScriptStructsInNodes(ChangedStructs, OnNodeFoundOrUpdated);
+
+	TMap<UEnum*, UEnum*> ChangedEnums;
+	for (const TPair<UEnum*, UEnum*>& Pair : ReinstancedEnums)
+	{
+		ReinstancingObjects.Add(Pair.Key);
+		if (Pair.Value)
+		{
+			ChangedEnums.Emplace(Pair.Key, Pair.Value);
+		}
+	}
+	FBlueprintEditorUtils::UpdateEnumsInNodes(ChangedEnums, OnNodeFoundOrUpdated);
+
+	// Update all the nodes before we could possibly recompile
+	for (TPair<UBlueprint*, FBlueprintUpdateInfo>& KVP : ModifiedBlueprints)
+	{
+		UBlueprint* Blueprint = KVP.Key;
+		FBlueprintUpdateInfo& Info = KVP.Value;
+
+		for (UK2Node* Node : Info.Nodes)
 		{
-			OldToNewClassesMap.Add(Pair.Key, Pair.Value);
+			FBlueprintEditorUtils::RecombineNestedSubPins(Node);
 		}
 	}
 
+	TSet<UBlueprint*> CompiledBlueprints;
 	for (const TPair<UClass*, UClass*>& Pair : ClassesToReinstance)
 	{
-		ReinstanceClass(Pair.Value, Pair.Key, OldToNewClassesMap);
+		ReinstanceClass(Pair.Value, Pair.Key, ReinstancingObjects, CompiledBlueprints);
+	}
+
+	// Reconstruct the nodes (and recompile blueprints if they haven't already been recompiled)
+	for (TPair<UBlueprint*, FBlueprintUpdateInfo>& KVP : ModifiedBlueprints)
+	{
+		UBlueprint* Blueprint = KVP.Key;
+		FBlueprintUpdateInfo& Info = KVP.Value;
+
+		if (Blueprint && !CompiledBlueprints.Contains(Blueprint))
+		{
+			EBlueprintCompileOptions Options = EBlueprintCompileOptions::SkipGarbageCollection;
+			FKismetEditorUtilities::CompileBlueprint(Blueprint, Options);
+		}
+
+		for (UK2Node* Node : Info.Nodes)
+		{
+			Node->ReconstructNode();
+		}
 	}
 
 	ReinstancedClasses = MoveTemp(ClassesToReinstance);
@@ -891,9 +930,9 @@ void FReload::Reinstance()
 	FCoreUObjectDelegates::ReloadReinstancingCompleteDelegate.Broadcast();
 }
 
-void FReload::ReinstanceClass(UClass* NewClass, UClass* OldClass, const TMap<UClass*, UClass*>& OldToNewClassesMap)
+void FReload::ReinstanceClass(UClass* NewClass, UClass* OldClass, const TSet<UObject*>& ReinstancingObjects, TSet<UBlueprint*>& CompiledBlueprints)
 {
-	TSharedPtr<FReloadClassReinstancer> ReinstanceHelper = FReloadClassReinstancer::Create(NewClass, OldClass, OldToNewClassesMap, ReconstructedCDOsMap, BPSetToRecompile, BPSetToRecompileBytecodeOnly);
+	TSharedPtr<FReloadClassReinstancer> ReinstanceHelper = FReloadClassReinstancer::Create(NewClass, OldClass, ReinstancingObjects, ReconstructedCDOsMap, CompiledBlueprints);
 	if (ReinstanceHelper->ClassNeedsReinstancing())
 	{
 		Ar.Logf(ELogVerbosity::Log, TEXT("Re-instancing %s after reload."), NewClass ? *NewClass->GetName() : *OldClass->GetName());
diff --git a/Engine/Source/Editor/UnrealEd/Public/Kismet2/BlueprintEditorUtils.h b/Engine/Source/Editor/UnrealEd/Public/Kismet2/BlueprintEditorUtils.h
--- a/Engine/Source/Editor/UnrealEd/Public/Kismet2/BlueprintEditorUtils.h
+++ b/Engine/Source/Editor/UnrealEd/Public/Kismet2/BlueprintEditorUtils.h
@@ -1713,6 +1713,36 @@ public:
 	 */
 	static void BuildComponentInstancingData(UActorComponent* ComponentTemplate, FBlueprintCookedComponentInstancingData& OutData, bool bUseTemplateArchetype = false);
 
+	/**
+	 * Callback for when a node has been found
+	 */
+	using FOnNodeFoundOrUpdated = TFunction<void(UBlueprint*, UK2Node*)>;
+
+	/**
+	 * Search the blueprints looking for nodes that contain the given script structs
+	 */
+	static void FindScriptStructsInNodes(const TSet<UScriptStruct*>& Structs, FOnNodeFoundOrUpdated InOnNodeFoundOrUpdated);
+
+	/**
+	 * Search the blueprints looking for nodes that contain the given enumerations
+	 */
+	static void FindEnumsInNodes(const TSet<UEnum*>& UEnums, FOnNodeFoundOrUpdated InOnNodeFoundOrUpdated);
+
+	/**
+	 * Search the blueprints looking for nodes that contain the given script structs and replace the references
+	 */
+	static void UpdateScriptStructsInNodes(const TMap<UScriptStruct*, UScriptStruct*>& Structs, FOnNodeFoundOrUpdated InOnNodeFoundOrUpdated);
+
+	/**
+	 * Search the blueprints looking for nodes that contain the given enumerations and replace the references
+	 */
+	static void UpdateEnumsInNodes(const TMap<UEnum*, UEnum*>& Structs, FOnNodeFoundOrUpdated InOnNodeFoundOrUpdated);
+
+	/**
+	 * Recombine any nested subpins on this node
+	 */
+	static void RecombineNestedSubPins(UK2Node* Node);
+
 protected:
 	// Removes all NULL graph references from the SubGraphs array and recurses thru the non-NULL ones
 	static void CleanNullGraphReferencesRecursive(UEdGraph* Graph);
diff --git a/Engine/Source/Editor/UnrealEd/Public/Kismet2/KismetReinstanceUtilities.h b/Engine/Source/Editor/UnrealEd/Public/Kismet2/KismetReinstanceUtilities.h
--- a/Engine/Source/Editor/UnrealEd/Public/Kismet2/KismetReinstanceUtilities.h
+++ b/Engine/Source/Editor/UnrealEd/Public/Kismet2/KismetReinstanceUtilities.h
@@ -142,6 +142,7 @@ public:
 	static void OptionallyRefreshNodes(UBlueprint* BP);
 
 	void ListDependentBlueprintsToRefresh(const TArray<UBlueprint*>& DependentBPs);
+	UE_DEPRECATED(5.0, "This method performs no function and isn't invoked by the base class.  Remove all calls to the base class method.")
 	virtual void EnlistDependentBlueprintToRecompile(UBlueprint* BP, bool bBytecodeOnly);
 	virtual void BlueprintWasRecompiled(UBlueprint* BP, bool bBytecodeOnly);
 
diff --git a/Engine/Source/Editor/UnrealEd/Public/Kismet2/ReloadUtilities.h b/Engine/Source/Editor/UnrealEd/Public/Kismet2/ReloadUtilities.h
--- a/Engine/Source/Editor/UnrealEd/Public/Kismet2/ReloadUtilities.h
+++ b/Engine/Source/Editor/UnrealEd/Public/Kismet2/ReloadUtilities.h
@@ -10,6 +10,7 @@
 #include "Containers/StringFwd.h"
 
 class UBlueprint;
+class UK2Node;
 
 class UNREALED_API FReload : public IReload
 {
@@ -69,7 +70,7 @@ public:
 
 private:
 
-	struct ReinstanceStats
+	struct FReinstanceStats
 	{
 		int32 New = 0;
 		int32 Changed = 0;
@@ -81,6 +82,11 @@ private:
 		}
 	};
 
+	struct FBlueprintUpdateInfo
+	{
+		TSet<UK2Node*> Nodes;
+	};
+
 	/**
 	 * Finds all references to old CDOs and replaces them with the new ones.
 	 * Skipping UBlueprintGeneratedClass::OverridenArchetypeForCDO as it's the
@@ -91,17 +97,17 @@ private:
 	/**
 	 * Given a new/old class pair, re-instance the class
 	 */
-	void ReinstanceClass(UClass* NewClass, UClass* OldClass, const TMap<UClass*, UClass*>& OldToNewClassesMap);
+	void ReinstanceClass(UClass* NewClass, UClass* OldClass, const TSet<UObject*>& ReinstancingObjects, TSet<UBlueprint*>& CompiledBlueprints);
 
 	/**
 	 * Based on the pointers, update the given stat
 	 */
-	void UpdateStats(ReinstanceStats& Stats, void* New, void* Old);
+	void UpdateStats(FReinstanceStats& Stats, void* New, void* Old);
 
 	/**
 	 * Helper method to format all the stats
 	 */
-	static void FormatStats(FStringBuilderBase& Out, const TCHAR* Singular, const TCHAR* Plural, const ReinstanceStats& Stats);
+	static void FormatStats(FStringBuilderBase& Out, const TCHAR* Singular, const TCHAR* Plural, const FReinstanceStats& Stats);
 
 	/**
 	 * Helper method to format a specific stat
@@ -123,16 +129,18 @@ private:
 	/** Map from old function pointer to new function pointer for hot reload. */
 	TFunctionRemap FunctionRemap;
 
-	/* Unused */
-	TSet<UBlueprint*> BPSetToRecompile;
-	TSet<UBlueprint*> BPSetToRecompileBytecodeOnly;
-
 	/** Map of the reconstructed CDOs during the reinstancing process */
 	TMap<UObject*, UObject*> ReconstructedCDOsMap;
 
 	/** Map from old class to new class.  New class may be null */
 	TMap<UClass*, UClass*> ReinstancedClasses;
 
+	/** Map from old struct to new struct.  New struct may be null */
+	TMap<UScriptStruct*, UScriptStruct*> ReinstancedStructs;
+
+	/** Map from old enum to new enum.  New enum may be null */
+	TMap<UEnum*, UEnum*> ReinstancedEnums;
+
 	/** If true, we have to collect the package list from the context */
 	bool bCollectPackages;
 
@@ -142,9 +150,9 @@ private:
 	/** If true, reinstancing is enabled */
 	bool bEnableReinstancing = true;
 
-	ReinstanceStats ClassStats;
-	ReinstanceStats EnumStats;
-	ReinstanceStats StructStats;
+	FReinstanceStats ClassStats;
+	FReinstanceStats EnumStats;
+	FReinstanceStats StructStats;
 	int32 NumFunctionsRemapped = 0;
 	int32 NumScriptStructsRemapped = 0;
 	mutable bool bEnabledMessage = false;
