Fix a leak when a world is destroyed and a mode in the global mode tools is active which holds references to that world.

diff --git a/Engine/Source/Editor/UnrealEd/Private/EditorModeManager.cpp b/Engine/Source/Editor/UnrealEd/Private/EditorModeManager.cpp
--- a/Engine/Source/Editor/UnrealEd/Private/EditorModeManager.cpp
+++ b/Engine/Source/Editor/UnrealEd/Private/EditorModeManager.cpp
@@ -90,6 +90,8 @@ FEditorModeTools::FEditorModeTools()
 		// This binding ensures the mode is destroyed if the type is unregistered outside of normal shutdown process
 		GEditor->GetEditorSubsystem<UAssetEditorSubsystem>()->OnEditorModeUnregistered().AddRaw(this, &FEditorModeTools::OnModeUnregistered);
 	}
+
+	FWorldDelegates::OnWorldCleanup.AddRaw(this, &FEditorModeTools::OnWorldCleanup);
 }
 
 FEditorModeTools::~FEditorModeTools()
@@ -864,6 +866,15 @@ void FEditorModeTools::InvokeToolPaletteTab(FEditorModeID InModeID, FName InPale
 	}	
 }
 
+void FEditorModeTools::OnWorldCleanup(UWorld* InWorld, bool bSessionEnded, bool bCleanupResources)
+{
+	UWorld* World = GetWorld();
+	if (InWorld == World)
+	{
+		DeactivateAllModesPendingDeletion();
+	}
+}
+
 void FEditorModeTools::RemoveAllDelegateHandlers()
 {
 	if (GEditor)
@@ -872,6 +883,8 @@ void FEditorModeTools::RemoveAllDelegateHandlers()
 		GEditor->GetEditorSubsystem<UAssetEditorSubsystem>()->OnEditorModeUnregistered().RemoveAll(this);
 	}
 
+	FWorldDelegates::OnWorldCleanup.RemoveAll(this);
+	
 	// For now, check that UObjects are even valid, because the level editor has a global static mode tools
 	if (UObjectInitialized())
 	{
diff --git a/Engine/Source/Editor/UnrealEd/Public/EditorModeManager.h b/Engine/Source/Editor/UnrealEd/Public/EditorModeManager.h
--- a/Engine/Source/Editor/UnrealEd/Public/EditorModeManager.h
+++ b/Engine/Source/Editor/UnrealEd/Public/EditorModeManager.h
@@ -538,6 +538,10 @@ protected:
 	void OnEditorSelectionChanged(UObject* NewSelection);
 	void OnEditorSelectNone();
 
+
+	/** Handles the notification when a world is going through GC to clean up any modes pending deactivation. */
+	void OnWorldCleanup(UWorld* InWorld, bool bSessionEnded, bool bCleanupResources);
+
 	virtual void DrawBrackets(FEditorViewportClient* ViewportClient, FViewport* Viewport, const FSceneView* View, FCanvas* Canvas);
 
 	void ForEachEdMode(TFunctionRef<bool(UEdMode*)> InCalllback) const;
