Fix copy/paste of material nodes

diff --git a/Engine/Source/Editor/UnrealEd/Classes/MaterialGraph/MaterialGraphNode_Base.h b/Engine/Source/Editor/UnrealEd/Classes/MaterialGraph/MaterialGraphNode_Base.h
--- a/Engine/Source/Editor/UnrealEd/Classes/MaterialGraph/MaterialGraphNode_Base.h
+++ b/Engine/Source/Editor/UnrealEd/Classes/MaterialGraph/MaterialGraphNode_Base.h
@@ -79,6 +79,7 @@ class UMaterialGraphNode_Base : public UEdGraphNode
 	virtual void AutowireNewNode(UEdGraphPin* FromPin) override;
 	virtual bool CanCreateUnderSpecifiedSchema(const UEdGraphSchema* Schema) const override;
 	virtual FString GetDocumentationLink() const override;
+	virtual void PostPasteNode() override;
 	//~ End UEdGraphNode Interface.
 
 protected:
diff --git a/Engine/Source/Editor/UnrealEd/Private/MaterialGraphNode_Base.cpp b/Engine/Source/Editor/UnrealEd/Private/MaterialGraphNode_Base.cpp
--- a/Engine/Source/Editor/UnrealEd/Private/MaterialGraphNode_Base.cpp
+++ b/Engine/Source/Editor/UnrealEd/Private/MaterialGraphNode_Base.cpp
@@ -153,6 +153,35 @@ void UMaterialGraphNode_Base::AllocateDefaultPins()
 	CreateOutputPins();
 }
 
+void UMaterialGraphNode_Base::PostPasteNode()
+{
+	check(PinInfoMap.Num() == 0)
+
+	int32 NumInputDataPins = 0;
+	int32 NumOutputDataPins = 0;
+	int32 NumInputExecPins = 0;
+	int32 NumOutputExecPins = 0;
+	for (UEdGraphPin* Pin : Pins)
+	{
+		int32 Index = INDEX_NONE;
+		EMaterialGraphPinType Type;
+
+		if (Pin->PinType.PinCategory == UMaterialGraphSchema::PC_Exec)
+		{
+			Type = EMaterialGraphPinType::Exec;
+			if (Pin->Direction == EGPD_Input) Index = NumInputExecPins++;
+			else Index = NumOutputExecPins++;
+		}
+		else
+		{
+			Type = EMaterialGraphPinType::Data;
+			if (Pin->Direction == EGPD_Input) Index = NumInputDataPins++;
+			else Index = NumOutputDataPins++;
+		}
+		RegisterPin(Pin, Type, Index);
+	}
+}
+
 void UMaterialGraphNode_Base::RegisterPin(UEdGraphPin* Pin, EMaterialGraphPinType Type, int32 Index)
 {
 	FMaterialGraphPinInfo& PinInfo = PinInfoMap.FindOrAdd(Pin);
