Downgrade cooker has been blocked warning to a display when the shader compiler is still working.

diff --git a/Engine/Source/Editor/UnrealEd/Private/CookOnTheFlyServer.cpp b/Engine/Source/Editor/UnrealEd/Private/CookOnTheFlyServer.cpp
--- a/Engine/Source/Editor/UnrealEd/Private/CookOnTheFlyServer.cpp
+++ b/Engine/Source/Editor/UnrealEd/Private/CookOnTheFlyServer.cpp
@@ -1288,8 +1288,40 @@ void UCookOnTheFlyServer::SetSaveBusy(bool bInBusy)
 		const float CurrentTime = FPlatformTime::Seconds();
 		if (CurrentTime - SaveBusyTimeStarted > GCookProgressWarnBusyTime)
 		{
-			UE_LOG(LogCook, Warning, TEXT("Cooker has been blocked from saving the current packages for %f seconds."), GCookProgressWarnBusyTime);
+			// Issue a status update. For each UObject we're still waiting on, check whether the long duration is expected using type-specific checks
+			// Make the status update a warning if the long duration is not reported as expected.
+			TArray<UObject*> PendingObjects;
+			TSet<UPackage*> PackagesOfPendingObjects;
 			FPackageDataQueue& SaveQueue = PackageDatas->GetSaveQueue();
+			const TArray<FPendingCookedPlatformData>& PendingCookedPlatformDatas = PackageDatas->GetPendingCookedPlatformDatas();
+			TArray<FName> CompilationUsers({ FName(TEXT("Material")), FName(TEXT("MaterialInstance")), FName(TEXT("NiagaraScript")) });
+
+			bool bAllExpected = !PendingCookedPlatformDatas.IsEmpty();
+			for (const FPendingCookedPlatformData& Data : PendingCookedPlatformDatas)
+			{
+				UObject* Object = Data.Object.Get();
+				if (!Object)
+				{
+					continue;
+				}
+				PackagesOfPendingObjects.Add(Object->GetPackage());
+				bool bExpected = false;
+				if (Algo::Find(CompilationUsers, Object->GetClass()->GetFName()) && GShaderCompilingManager->IsCompiling())
+				{
+					bExpected = true;
+				}
+				bAllExpected &= bExpected;
+			}
+			FString Message = FString::Printf(TEXT("Cooker has been blocked from saving the current packages for %f seconds."), GCookProgressWarnBusyTime);
+			if (bAllExpected)
+			{
+				UE_LOG(LogCook, Display, TEXT("%s"), *Message);
+			}
+			else
+			{
+				UE_LOG(LogCook, Warning, TEXT("%s"), *Message);
+			}
+
 			UE_LOG(LogCook, Display, TEXT("%d packages in the savequeue: "), SaveQueue.Num());
 			int DisplayCount = 0;
 			const int DisplayMax = 10;
@@ -1310,7 +1342,7 @@ void UCookOnTheFlyServer::SetSaveBusy(bool bInBusy)
 
 			UE_LOG(LogCook, Display, TEXT("%d objects that have not yet returned true from IsCachedCookedPlatformDataLoaded:"), PackageDatas->GetPendingCookedPlatformDatas().Num());
 			DisplayCount = 0;
-			for (const FPendingCookedPlatformData& Data : PackageDatas->GetPendingCookedPlatformDatas())
+			for (const FPendingCookedPlatformData& Data : PendingCookedPlatformDatas)
 			{
 				if (Data.Object.IsValid())
 				{
