Re-instated bSkipCompilerGeneratedDefaults when re-instancing

CPFUO does not need to update compiler-generated default values as this causes compiler-generated values to be incorrectly overrwitten.

diff --git a/Engine/Source/Editor/UnrealEd/Private/Kismet2/KismetReinstanceUtilities.cpp b/Engine/Source/Editor/UnrealEd/Private/Kismet2/KismetReinstanceUtilities.cpp
--- a/Engine/Source/Editor/UnrealEd/Private/Kismet2/KismetReinstanceUtilities.cpp
+++ b/Engine/Source/Editor/UnrealEd/Private/Kismet2/KismetReinstanceUtilities.cpp
@@ -1957,6 +1957,7 @@ static void ReplaceObjectHelper(UObject*& OldObject, UClass* OldClass, UObject*&
 	InstancedPropertyUtils::FArchiveInstancedSubObjCollector  InstancedSubObjCollector(OldObject, InstancedPropertyMap);
 	UEngine::FCopyPropertiesForUnrelatedObjectsParams Options;
 	Options.bNotifyObjectReplacement = true;
+	Options.bSkipCompilerGeneratedDefaults = true;
 	UEditorEngine::CopyPropertiesForUnrelatedObjects(OldObject, NewUObject, Options);
 	InstancedPropertyUtils::FArchiveInsertInstancedSubObjects InstancedSubObjSpawner(NewUObject, InstancedPropertyMap);
 
