Block quick launch for launching the game for partitioned world if cook on fly or cook by the book in editor are used.

diff --git a/Engine/Source/Editor/UnrealEd/Private/PlayLevelViaLauncher.cpp b/Engine/Source/Editor/UnrealEd/Private/PlayLevelViaLauncher.cpp
--- a/Engine/Source/Editor/UnrealEd/Private/PlayLevelViaLauncher.cpp
+++ b/Engine/Source/Editor/UnrealEd/Private/PlayLevelViaLauncher.cpp
@@ -215,6 +215,24 @@ void UEditorEngine::StartPlayUsingLauncherSession(FRequestPlaySessionParams& InR
 		CurrentLauncherCookMode = ELauncherProfileCookModes::OnTheFly;
 		bIncrimentalCooking = false;
 	}
+
+	if (CurrentLauncherCookMode == ELauncherProfileCookModes::OnTheFlyInEditor ||
+		CurrentLauncherCookMode == ELauncherProfileCookModes::OnTheFly ||
+		CurrentLauncherCookMode == ELauncherProfileCookModes::ByTheBookInEditor)
+	{
+		// For now World Partition doesn't these cook modes
+		FWorldContext& EditorContext = GetEditorWorldContext();
+		if (EditorContext.World()->GetWorldPartition())
+		{
+			FString ErrorMsg = FString::Printf(TEXT("Error launching map %s : Quick launch %s doesn't yet support partitioned worlds."), *EditorContext.World()->GetOutermost()->GetName(), (CurrentLauncherCookMode == ELauncherProfileCookModes::ByTheBookInEditor) ? TEXT("cook by the book") : TEXT("cook on the fly"));
+			UE_LOG(LogPlayLevel, Error, TEXT("%s"), *ErrorMsg);
+			FMessageLog("EditorErrors").Error(FText::FromString(ErrorMsg));
+			FMessageLog("EditorErrors").Open();
+			CancelRequestPlaySession();
+			return;
+		}
+	}
+
 	LauncherProfile->SetCookMode(CurrentLauncherCookMode);
 	LauncherProfile->SetUnversionedCooking(!bIncrimentalCooking);
 	LauncherProfile->SetIncrementalCooking(bIncrimentalCooking);
