Removed direct use of GSelectedObjectAnnotation

 - Added a way for UEdGraphNode_Comment to override its global editor selection state.
 - Updated UControlRigGraph to use that method, rather than directly adjust GSelectedObjectAnnotation.

diff --git a/Engine/Plugins/Experimental/ControlRig/Source/ControlRigDeveloper/Private/Graph/ControlRigGraph.cpp b/Engine/Plugins/Experimental/ControlRig/Source/ControlRigDeveloper/Private/Graph/ControlRigGraph.cpp
--- a/Engine/Plugins/Experimental/ControlRig/Source/ControlRigDeveloper/Private/Graph/ControlRigGraph.cpp
+++ b/Engine/Plugins/Experimental/ControlRig/Source/ControlRigDeveloper/Private/Graph/ControlRigGraph.cpp
@@ -602,7 +602,7 @@ void UControlRigGraph::HandleModifiedEvent(ERigVMGraphNotifType InNotifType, URi
 				UEdGraphNode_Comment* EdNode = Cast<UEdGraphNode_Comment>(FindNodeForModelNodeName(ModelNode->GetFName()));
 				if (EdNode)
 				{
-					GSelectedObjectAnnotation.Set(EdNode);
+					EdNode->SetSelectionState(UEdGraphNode_Comment::ESelectionState::Selected);
 				}
 			}
 			break;
@@ -614,7 +614,7 @@ void UControlRigGraph::HandleModifiedEvent(ERigVMGraphNotifType InNotifType, URi
 				UEdGraphNode_Comment* EdNode = Cast<UEdGraphNode_Comment>(FindNodeForModelNodeName(ModelNode->GetFName()));
 				if (EdNode)
 				{
-					GSelectedObjectAnnotation.Clear(EdNode);
+					EdNode->SetSelectionState(UEdGraphNode_Comment::ESelectionState::Deselected);
 				}
 			}
 			break;
diff --git a/Engine/Source/Editor/UnrealEd/Private/EdGraphNode_Comment.cpp b/Engine/Source/Editor/UnrealEd/Private/EdGraphNode_Comment.cpp
--- a/Engine/Source/Editor/UnrealEd/Private/EdGraphNode_Comment.cpp
+++ b/Engine/Source/Editor/UnrealEd/Private/EdGraphNode_Comment.cpp
@@ -164,6 +164,20 @@ TSharedPtr<class INameValidatorInterface> UEdGraphNode_Comment::MakeNameValidato
 	return MakeShareable(new FDummyNameValidator(EValidatorResult::Ok));
 }
 
+bool UEdGraphNode_Comment::IsSelectedInEditor() const
+{
+	if (SelectionState == ESelectionState::Inherited)
+	{
+		return Super::IsSelectedInEditor();
+	}
+	return SelectionState == ESelectionState::Selected;
+}
+
+void UEdGraphNode_Comment::SetSelectionState(const ESelectionState InSelectionState)
+{
+	SelectionState = InSelectionState;
+}
+
 /////////////////////////////////////////////////////
 
 #undef LOCTEXT_NAMESPACE
diff --git a/Engine/Source/Editor/UnrealEd/Public/EdGraphNode_Comment.h b/Engine/Source/Editor/UnrealEd/Public/EdGraphNode_Comment.h
--- a/Engine/Source/Editor/UnrealEd/Public/EdGraphNode_Comment.h
+++ b/Engine/Source/Editor/UnrealEd/Public/EdGraphNode_Comment.h
@@ -66,6 +66,7 @@ public:
 	//~ Begin UObject Interface
 	UNREALED_API static void AddReferencedObjects(UObject* InThis, FReferenceCollector& Collector);
 	virtual void PostEditChangeProperty(FPropertyChangedEvent& PropertyChangedEvent) override;
+	virtual bool IsSelectedInEditor() const override;
 	//~ End UObject Interface
 
 	//~ Begin UEdGraphNode Interface
@@ -99,11 +100,18 @@ public:
 	/** Return the font size of the comment */
 	UNREALED_API virtual int32 GetFontSize() const { return FontSize; }
 
+	/** Override the default selection state of this graph node */
+	enum class ESelectionState : uint8 { Inherited, Selected, Deselected };
+	UNREALED_API void SetSelectionState(const ESelectionState InSelectionState);
+
 private:
 	/** Nodes currently within the region of the comment */
 	FCommentNodeSet	NodesUnderComment;
 
 	/** Constructing FText strings can be costly, so we cache the node's tooltip */
 	FNodeTextCache CachedTooltip;
+
+	/** Override the default selection state of this graph node */
+	ESelectionState SelectionState = ESelectionState::Inherited;
 };
 
diff --git a/Engine/Source/Runtime/CoreUObject/Public/UObject/Object.h b/Engine/Source/Runtime/CoreUObject/Public/UObject/Object.h
--- a/Engine/Source/Runtime/CoreUObject/Public/UObject/Object.h
+++ b/Engine/Source/Runtime/CoreUObject/Public/UObject/Object.h
@@ -420,7 +420,6 @@ protected:
 	enum class ETransactionAnnotationCreationMode : uint8 { DefaultInstance, FindOrCreate };
 	virtual TSharedPtr<ITransactionObjectAnnotation> FactoryTransactionAnnotation(const ETransactionAnnotationCreationMode InCreationMode) const { return nullptr; }
 
-private:
 	/**
 	 * Test the selection state of a UObject
 	 *
