Optimized transaction object snapshotting

Removed the mapping of objects to their save count, in favor of putting the save count into a new FObjectRecords struct, and mapping objects to that instead. This allows the records for an object to be found via the map, and avoids the remaining look-up overhead of searching the records array during object snapshotting.

The object records map keys against the FPersistentObjectRef rather than the raw-UObject, so that the record for sub-objects can be found correctly even if the sub-object is re-created.

diff --git a/Engine/Source/Editor/UnrealEd/Classes/Editor/Transactor.h b/Engine/Source/Editor/UnrealEd/Classes/Editor/Transactor.h
--- a/Engine/Source/Editor/UnrealEd/Classes/Editor/Transactor.h
+++ b/Engine/Source/Editor/UnrealEd/Classes/Editor/Transactor.h
@@ -36,85 +36,89 @@ class UNREALED_API FTransaction : public ITransaction
 	friend class FTransactionObjectEvent;
 
 protected:
-	// Record of an object.
-	class UNREALED_API FObjectRecord
+	/**
+		This type is necessary because the blueprint system is destroying and creating
+		CDOs at edit time (usually on compile, but also on load), but also stores user
+		entered data in the CDO. We "need"  changes to a CDO to persist across instances
+		because as we undo and redo we  need to apply changes to different instances of
+		the CDO - alternatively we could destroy and create the CDO as part of a transaction
+		(this alternative is the reason for the bunny ears around need).
+
+		DanO: My long term preference is for the editor to use a dynamic, mutable type
+		(rather than the CDO) to store editor data. The CDO can then be re-instanced (or not)
+		as runtime code requires.
+	*/
+	struct FPersistentObjectRef
 	{
 	public:
-
-		/** 
-			This type is necessary because the blueprint system is destroying and creating 
-			CDOs at edit time (usually on compile, but also on load), but also stores user 
-			entered data in the CDO. We "need"  changes to a CDO to persist across instances 
-			because as we undo and redo we  need to apply changes to different instances of 
-			the CDO - alternatively we could destroy and create the CDO as part of a transaction 
-			(this alternative is the reason for the bunny ears around need).
-
-			DanO: My long term preference is for the editor to use a dynamic, mutable type
-			(rather than the CDO) to store editor data. The CDO can then be re-instanced (or not)
-			as runtime code requires.
-		*/
-		struct FPersistentObjectRef
+		friend FArchive& operator<<(FArchive& Ar, FPersistentObjectRef& ReferencedObject)
 		{
-		public:
-			friend FArchive& operator<<(FArchive& Ar, FPersistentObjectRef& ReferencedObject)
-			{
-				Ar << (std::underlying_type_t<EReferenceType>&)ReferencedObject.ReferenceType;
-				Ar << ReferencedObject.RootObject;
-				Ar << ReferencedObject.SubObjectHierarchyIDs;
-				return Ar;
-			}
+			Ar << (std::underlying_type_t<EReferenceType>&)ReferencedObject.ReferenceType;
+			Ar << ReferencedObject.RootObject;
+			Ar << ReferencedObject.SubObjectHierarchyIDs;
+			return Ar;
+		}
 
-			friend bool operator==(const FPersistentObjectRef& LHS, const FPersistentObjectRef& RHS)
-			{
-				return LHS.ReferenceType == RHS.ReferenceType
-					&& LHS.RootObject == RHS.RootObject
-					&& (LHS.ReferenceType != EReferenceType::SubObject || LHS.SubObjectHierarchyIDs == RHS.SubObjectHierarchyIDs);
-			}
+		friend bool operator==(const FPersistentObjectRef& LHS, const FPersistentObjectRef& RHS)
+		{
+			return LHS.ReferenceType == RHS.ReferenceType
+				&& LHS.RootObject == RHS.RootObject
+				&& (LHS.ReferenceType != EReferenceType::SubObject || LHS.SubObjectHierarchyIDs == RHS.SubObjectHierarchyIDs);
+		}
 
-			friend bool operator!=(const FPersistentObjectRef& LHS, const FPersistentObjectRef& RHS)
-			{
-				return !(LHS == RHS);
-			}
+		friend bool operator!=(const FPersistentObjectRef& LHS, const FPersistentObjectRef& RHS)
+		{
+			return !(LHS == RHS);
+		}
 
-			FPersistentObjectRef() = default;
-			explicit FPersistentObjectRef(UObject* InObject);
+		friend uint32 GetTypeHash(const FPersistentObjectRef& InObjRef)
+		{
+			return HashCombine(GetTypeHash(InObjRef.ReferenceType), GetTypeHash(InObjRef.RootObject));
+		}
 
-			bool IsRootObjectReference() const
-			{
-				return ReferenceType == EReferenceType::RootObject;
-			}
+		FPersistentObjectRef() = default;
+		explicit FPersistentObjectRef(UObject* InObject);
 
-			bool IsSubObjectReference() const
-			{
-				return ReferenceType == EReferenceType::SubObject;
-			}
+		bool IsRootObjectReference() const
+		{
+			return ReferenceType == EReferenceType::RootObject;
+		}
 
-			UObject* Get() const;
+		bool IsSubObjectReference() const
+		{
+			return ReferenceType == EReferenceType::SubObject;
+		}
 
-			void AddReferencedObjects(FReferenceCollector& Collector);
+		UObject* Get() const;
 
-		private:
-			/** This enum represents all of the different special cases we are handling with this type */
-			enum class EReferenceType : uint8
-			{
-				Unknown,
-				RootObject,
-				SubObject,
-			};
+		void AddReferencedObjects(FReferenceCollector& Collector);
 
-			/** The reference type we're handling */
-			EReferenceType ReferenceType = EReferenceType::Unknown;
-			/** Stores the object pointer when ReferenceType==RootObject, and the outermost pointer of the sub-object chain when when ReferenceType==SubObject */
-			UObject* RootObject = nullptr;
-			/** Stores the sub-object name chain when ReferenceType==SubObject */
-			TArray<FName, TInlineAllocator<4>> SubObjectHierarchyIDs;
-
-			/** Cached pointers corresponding to RootObject when ReferenceType==SubObject (@note cache needs testing on access as it may have become stale) */
-			mutable TWeakObjectPtr<UObject> CachedRootObject;
-			/** Cache of pointers corresponding to the items within SubObjectHierarchyIDs when ReferenceType==SubObject (@note cache needs testing on access as it may have become stale) */
-			mutable TArray<TWeakObjectPtr<UObject>, TInlineAllocator<4>> CachedSubObjectHierarchy;
+	private:
+		/** This enum represents all of the different special cases we are handling with this type */
+		enum class EReferenceType : uint8
+		{
+			Unknown,
+			RootObject,
+			SubObject,
 		};
 
+		/** The reference type we're handling */
+		EReferenceType ReferenceType = EReferenceType::Unknown;
+		/** Stores the object pointer when ReferenceType==RootObject, and the outermost pointer of the sub-object chain when when ReferenceType==SubObject */
+		UObject* RootObject = nullptr;
+		/** Stores the sub-object name chain when ReferenceType==SubObject */
+		TArray<FName, TInlineAllocator<4>> SubObjectHierarchyIDs;
+
+		/** Cached pointers corresponding to RootObject when ReferenceType==SubObject (@note cache needs testing on access as it may have become stale) */
+		mutable TWeakObjectPtr<UObject> CachedRootObject;
+		/** Cache of pointers corresponding to the items within SubObjectHierarchyIDs when ReferenceType==SubObject (@note cache needs testing on access as it may have become stale) */
+		mutable TArray<TWeakObjectPtr<UObject>, TInlineAllocator<4>> CachedSubObjectHierarchy;
+	};
+
+	// Record of an object.
+	class UNREALED_API FObjectRecord
+	{
+	public:
 		struct FSerializedProperty
 		{
 			FSerializedProperty()
@@ -241,16 +245,16 @@ protected:
 		/** Custom change to apply to this object to undo this record.  Executing the undo will return an object that can be used to redo the change. */
 		TUniquePtr<FChange> CustomChange;
 		/** Array: If an array object, reference to script array */
-		FScriptArray*		Array;
+		FScriptArray*		Array = nullptr;
 		/** Array: Offset into the array */
-		int32				Index;
+		int32				Index = 0;
 		/** Array: How many items to record */
-		int32				Count;
+		int32				Count = 0;
 		/** @todo document this
 		 Array: Operation performed on array: 1 (add/insert), 0 (modify), -1 (remove)? */
-		int32				Oper;
+		int32				Oper = 0;
 		/** Array: Size of each item in the array */
-		int32				ElementSize;
+		int32				ElementSize = 0;
 		/** Array: DefaultConstructor for each item in the array */
 		STRUCT_DC			DefaultConstructor;
 		/** Array: Serializer to use for each item in the array */
@@ -258,13 +262,13 @@ protected:
 		/** Array: Destructor for each item in the array */
 		STRUCT_DTOR			Destructor;
 		/** True if object  has already been restored from data. False otherwise. */
-		bool				bRestored;
+		bool				bRestored = false;
 		/** True if object has been finalized and generated diff data */
-		bool				bFinalized;
+		bool				bFinalized = false;
 		/** True if object has been snapshot before */
-		bool				bSnapshot;
+		bool				bSnapshot = false;
 		/** True if record should serialize data as binary blob (more compact). False to use tagged serialization (more robust) */
-		bool				bWantsBinarySerialization;
+		bool				bWantsBinarySerialization = true;
 		/** The serialized object data */
 		FSerializedObject	SerializedObject;
 		/** The serialized object data that will be used when the transaction is flipped */
@@ -275,8 +279,7 @@ protected:
 		FTransactionObjectDeltaChange DeltaChange;
 
 		// Constructors.
-		FObjectRecord()
-		{}
+		FObjectRecord() = default;
 		FObjectRecord( FTransaction* Owner, UObject* InObject, TUniquePtr<FChange> InCustomChange, FScriptArray* InArray, int32 InIndex, int32 InCount, int32 InOper, int32 InElementSize, STRUCT_DC InDefaultConstructor, STRUCT_AR InSerializer, STRUCT_DTOR InDestructor );
 
 	private:
@@ -361,15 +364,16 @@ protected:
 			}
 			void Preload( UObject* InObject ) override
 			{
-				if( Owner )
+				if (Owner)
 				{
-					if ( Owner->ObjectMap.Contains( InObject ) )
+					if (const FObjectRecords* ObjectRecords = Owner->ObjectRecordsMap.Find(FPersistentObjectRef(InObject)))
 					{
-						for (int32 i = 0; i < Owner->Records.Num(); i++)
+						for (FObjectRecord* Record : ObjectRecords->Records)
 						{
-							if (Owner->Records[i].Object.Get() == InObject)
+							checkSlow(Record->Object.Get() == InObject);
+							if (!Record->CustomChange)
 							{
-								Owner->Records[i].Restore(Owner);
+								Record->Restore(Owner);
 							}
 						}
 					}
@@ -533,10 +537,32 @@ protected:
 		};
 	};
 
+	struct FObjectRecords
+	{
+		friend FArchive& operator<<(FArchive& Ar, FObjectRecords& ObjectRecords)
+		{
+			Ar << ObjectRecords.SaveCount;
+
+			if (Ar.IsLoading())
+			{
+				// Clear the map on load, as it will need to be rebuilt from the source array
+				ObjectRecords.Records.Empty();
+			}
+
+			return Ar;
+		}
+
+		TArray<FObjectRecord*, TInlineAllocator<1>> Records;
+		int32 SaveCount = 0;
+	};
+
 	// Transaction variables.
 	/** List of object records in this transaction */
-	TIndirectArray<FObjectRecord>	Records;
-	
+	TIndirectArray<FObjectRecord> Records;
+
+	/** Map of object records (non-array), for optimized look-up and to prevent an object being serialized to a transaction more than once. */
+	TMap<FPersistentObjectRef, FObjectRecords> ObjectRecordsMap;
+
 	/** Unique identifier for this transaction, used to track it during its lifetime */
 	FGuid		Id;
 
@@ -554,15 +580,13 @@ protected:
 	FString		Context;
 
 	/** The key object being edited in this transaction. For example the blueprint object. Can be NULL */
-	UObject*	PrimaryObject;
-
-	/** Used to prevent objects from being serialized to a transaction more than once. */
-	TMap<UObject*, int32>	ObjectMap;
+	UObject* PrimaryObject = nullptr;
 
 	/** If true, on apply flip the direction of iteration over object records. The only client for which this is false is FMatineeTransaction */
-	bool					bFlip;
+	bool		bFlip = true;
+
 	/** Used to track direction to iterate over transaction's object records. Typically -1 for Undo, 1 for Redo */
-	int32					Inc;
+	int32		Inc = -1;
 
 	struct FChangedObjectValue
 	{
@@ -591,10 +615,9 @@ public:
 		:	Id( FGuid::NewGuid() )
 		,	Title( InTitle )
 		,	Context( InContext )
-		,	PrimaryObject(nullptr)
 		,	bFlip(bInFlip)
-		,	Inc(-1)
-	{}
+	{
+	}
 
 	virtual ~FTransaction()
 	{
@@ -678,35 +701,45 @@ public:
 	/** Returns the description of each contained Object Record */
 	FText GetDescription() const
 	{			
-		if (Algo::FindByPredicate(Records, [](const FObjectRecord& Record)
+		FString ConcatenatedRecords;
+		for (const FObjectRecord& Record : Records)
 		{
-			return Record.CustomChange.IsValid();
-		}))
-        {
-			FString ConcatenatedRecords;
-			for (const FObjectRecord& Record : Records)
+			if (Record.CustomChange.IsValid())
 			{
-				if (Record.CustomChange.IsValid())
+				if (ConcatenatedRecords.Len())
 				{
-					if (ConcatenatedRecords.Len())
-					{
-						ConcatenatedRecords += TEXT("\n");
-					}
-					ConcatenatedRecords += Record.CustomChange->ToString();
+					ConcatenatedRecords += TEXT("\n");
 				}
+				ConcatenatedRecords += Record.CustomChange->ToString();
 			}
+		}
 
-			return FText::FromString(ConcatenatedRecords);
-        }
-
-		return GetTitle();
+		return ConcatenatedRecords.IsEmpty() 
+			? GetTitle() 
+			: FText::FromString(ConcatenatedRecords);
 	}
 
-
 	/** Serializes a reference to a transaction in a given archive. */
 	friend FArchive& operator<<( FArchive& Ar, FTransaction& T )
 	{
-		return Ar << T.Records << T.Id << T.Title << T.ObjectMap << T.Context << T.PrimaryObject;
+		Ar << T.Records << T.ObjectRecordsMap << T.Id << T.Title << T.Context << T.PrimaryObject;
+
+		if (Ar.IsLoading())
+		{
+			// Rebuild the LUT from the records array
+			for (FObjectRecord& Record : T.Records)
+			{
+				if (!Record.Array)
+				{
+					// We should have already added a map entry for this object when loading T.ObjectRecordsMap, so check that we have (otherwise something is out-of-sync)
+					FObjectRecords& ObjectRecords = T.ObjectRecordsMap.FindChecked(Record.Object);
+					// Note: This is only safe to do because T.Records is a TIndirectArray
+					ObjectRecords.Records.Add(&Record);
+				}
+			}
+		}
+
+		return Ar;
 	}
 
 	/** Serializes a reference to a transaction in a given archive. */
diff --git a/Engine/Source/Editor/UnrealEd/Private/EditorTransaction.cpp b/Engine/Source/Editor/UnrealEd/Private/EditorTransaction.cpp
--- a/Engine/Source/Editor/UnrealEd/Private/EditorTransaction.cpp
+++ b/Engine/Source/Editor/UnrealEd/Private/EditorTransaction.cpp
@@ -34,10 +34,6 @@ FTransaction::FObjectRecord::FObjectRecord(FTransaction* Owner, UObject* InObjec
 	,	DefaultConstructor	( InDefaultConstructor )
 	,	Serializer			( InSerializer )
 	,	Destructor			( InDestructor )
-	,	bRestored			( false )
-	,   bFinalized			( false )
-	,	bSnapshot			( false )
-	,	bWantsBinarySerialization ( true )
 {
 	// Blueprint compile-in-place can alter class layout so use tagged serialization for objects relying on a UBlueprint's Class
 	if (UBlueprintGeneratedClass* Class = Cast<UBlueprintGeneratedClass>(InObject->GetClass()))
@@ -346,14 +342,11 @@ void FTransaction::FObjectRecord::Diff( const FTransaction* Owner, const FSerial
 		NewSerializedObject.SerializedObjectIndices.MultiFind(InPropertyName, NewSerializedObjectIndices, true);
 
 		bool bAreObjectPointersIdentical = OldSerializedObjectIndices.Num() == NewSerializedObjectIndices.Num();
-		if (bAreObjectPointersIdentical)
+		for (int32 ObjIndex = 0; ObjIndex < OldSerializedObjectIndices.Num() && bAreObjectPointersIdentical; ++ObjIndex)
 		{
-			for (int32 ObjIndex = 0; ObjIndex < OldSerializedObjectIndices.Num() && bAreObjectPointersIdentical; ++ObjIndex)
-			{
-				const UObject* OldObjectPtr = OldSerializedObject.ReferencedObjects.IsValidIndex(OldSerializedObjectIndices[ObjIndex]) ? OldSerializedObject.ReferencedObjects[OldSerializedObjectIndices[ObjIndex]].Get() : nullptr;
-				const UObject* NewObjectPtr = NewSerializedObject.ReferencedObjects.IsValidIndex(NewSerializedObjectIndices[ObjIndex]) ? NewSerializedObject.ReferencedObjects[NewSerializedObjectIndices[ObjIndex]].Get() : nullptr;
-				bAreObjectPointersIdentical = OldObjectPtr == NewObjectPtr;
-			}
+			const FPersistentObjectRef& OldObjectRef = OldSerializedObject.ReferencedObjects.IsValidIndex(OldSerializedObjectIndices[ObjIndex]) ? OldSerializedObject.ReferencedObjects[OldSerializedObjectIndices[ObjIndex]] : FPersistentObjectRef();
+			const FPersistentObjectRef& NewObjectRef = NewSerializedObject.ReferencedObjects.IsValidIndex(NewSerializedObjectIndices[ObjIndex]) ? NewSerializedObject.ReferencedObjects[NewSerializedObjectIndices[ObjIndex]] : FPersistentObjectRef();
+			bAreObjectPointersIdentical = OldObjectRef == NewObjectRef;
 		}
 		return bAreObjectPointersIdentical;
 	};
@@ -367,14 +360,11 @@ void FTransaction::FObjectRecord::Diff( const FTransaction* Owner, const FSerial
 		NewSerializedObject.SerializedNameIndices.MultiFind(InPropertyName, NewSerializedNameIndices, true);
 
 		bool bAreNamesIdentical = OldSerializedNameIndices.Num() == NewSerializedNameIndices.Num();
-		if (bAreNamesIdentical)
+		for (int32 ObjIndex = 0; ObjIndex < OldSerializedNameIndices.Num() && bAreNamesIdentical; ++ObjIndex)
 		{
-			for (int32 ObjIndex = 0; ObjIndex < OldSerializedNameIndices.Num() && bAreNamesIdentical; ++ObjIndex)
-			{
-				const FName& OldName = OldSerializedObject.ReferencedNames.IsValidIndex(OldSerializedNameIndices[ObjIndex]) ? OldSerializedObject.ReferencedNames[OldSerializedNameIndices[ObjIndex]] : FName();
-				const FName& NewName = NewSerializedObject.ReferencedNames.IsValidIndex(NewSerializedNameIndices[ObjIndex]) ? NewSerializedObject.ReferencedNames[NewSerializedNameIndices[ObjIndex]] : FName();
-				bAreNamesIdentical = OldName == NewName;
-			}
+			const FName& OldName = OldSerializedObject.ReferencedNames.IsValidIndex(OldSerializedNameIndices[ObjIndex]) ? OldSerializedObject.ReferencedNames[OldSerializedNameIndices[ObjIndex]] : FName();
+			const FName& NewName = NewSerializedObject.ReferencedNames.IsValidIndex(NewSerializedNameIndices[ObjIndex]) ? NewSerializedObject.ReferencedNames[NewSerializedNameIndices[ObjIndex]] : FName();
+			bAreNamesIdentical = OldName == NewName;
 		}
 		return bAreNamesIdentical;
 	};
@@ -575,10 +565,18 @@ void FTransaction::RemoveRecords( int32 Count /* = 1  */ )
 {
 	if ( Count > 0 && Records.Num() >= Count )
 	{
-		// Remove anything from the ObjectMap which is about to be removed from the Records array
+		// Remove anything from the ObjectRecordMap which is about to be removed from the Records array
 		for (int32 Index = 0; Index < Count; Index++)
 		{
-			ObjectMap.Remove( Records[Records.Num() - Count + Index].Object.Get() );
+			FObjectRecord& Record = Records[Records.Num() - Count + Index];
+			if (FObjectRecords* ObjectRecords = ObjectRecordsMap.Find(Record.Object))
+			{
+				ObjectRecords->Records.RemoveSingle(&Record);
+				if (ObjectRecords->Records.Num() == 0)
+				{
+					ObjectRecordsMap.Remove(Record.Object);
+				}
+			}
 		}
 
 		Records.RemoveAt( Records.Num() - Count, Count );
@@ -591,10 +589,10 @@ void FTransaction::RemoveRecords( int32 Count /* = 1  */ )
 void FTransaction::DumpObjectMap(FOutputDevice& Ar) const
 {
 	Ar.Logf( TEXT("===== DumpObjectMap %s ==== "), *Title.ToString() );
-	for ( auto It = ObjectMap.CreateConstIterator(); It; ++It )
+	for ( auto It = ObjectRecordsMap.CreateConstIterator(); It; ++It )
 	{
-		const UObject* CurrentObject	= It.Key();
-		const int32 SaveCount				= It.Value();
+		const UObject* CurrentObject	= It.Key().Get();
+		const int32 SaveCount			= It.Value().SaveCount;
 		Ar.Logf( TEXT("%i\t: %s"), SaveCount, *CurrentObject->GetPathName() );
 	}
 	Ar.Logf( TEXT("=== EndDumpObjectMap %s === "), *Title.ToString() );
@@ -611,7 +609,7 @@ FArchive& operator<<( FArchive& Ar, FTransaction::FObjectRecord& R )
 	return Ar;
 }
 
-FTransaction::FObjectRecord::FPersistentObjectRef::FPersistentObjectRef(UObject* InObject)
+FTransaction::FPersistentObjectRef::FPersistentObjectRef(UObject* InObject)
 {
 	RootObject = InObject;
 	{
@@ -659,7 +657,7 @@ FTransaction::FObjectRecord::FPersistentObjectRef::FPersistentObjectRef(UObject*
 	checkSlow(Get() == InObject);
 }
 
-UObject* FTransaction::FObjectRecord::FPersistentObjectRef::Get() const
+UObject* FTransaction::FPersistentObjectRef::Get() const
 {
 	if (ReferenceType == EReferenceType::SubObject)
 	{
@@ -722,7 +720,7 @@ UObject* FTransaction::FObjectRecord::FPersistentObjectRef::Get() const
 	return RootObject;
 }
 
-void FTransaction::FObjectRecord::FPersistentObjectRef::AddReferencedObjects(FReferenceCollector& Collector)
+void FTransaction::FPersistentObjectRef::AddReferencedObjects(FReferenceCollector& Collector)
 {
 	UObject* Obj = RootObject;
 	Collector.AddReferencedObject(Obj);
@@ -809,7 +807,11 @@ void FTransaction::AddReferencedObjects( FReferenceCollector& Collector )
 	{
 		ObjectRecord.AddReferencedObjects( Collector );
 	}
-	Collector.AddReferencedObjects(ObjectMap);
+
+	for (TTuple<FPersistentObjectRef, FObjectRecords>& ObjectRecordsPair : ObjectRecordsMap)
+	{
+		ObjectRecordsPair.Key.AddReferencedObjects(Collector);
+	}
 }
 
 void FTransaction::SaveObject( UObject* Object )
@@ -817,17 +819,14 @@ void FTransaction::SaveObject( UObject* Object )
 	check(Object);
 	Object->CheckDefaultSubobjects();
 
-	int32* SaveCount = ObjectMap.Find(Object);
-	if ( !SaveCount )
+	FObjectRecords& ObjectRecords = ObjectRecordsMap.FindOrAdd(FPersistentObjectRef(Object));
+	if (ObjectRecords.Records.Num() == 0)
 	{
-		ObjectMap.Add(Object,1);
 		// Save the object.
-		Records.Add(new FObjectRecord( this, Object, nullptr, nullptr, 0, 0, 0, 0, nullptr, nullptr, nullptr));
-	}
-	else
-	{
-		++(*SaveCount);
+		FObjectRecord* UndoRecord = ObjectRecords.Records.Add_GetRef(new FObjectRecord(this, Object, nullptr, nullptr, 0, 0, 0, 0, nullptr, nullptr, nullptr));
+		Records.Add(UndoRecord);
 	}
+	++ObjectRecords.SaveCount;
 }
 
 void FTransaction::SaveArray( UObject* Object, FScriptArray* Array, int32 Index, int32 Count, int32 Oper, int32 ElementSize, STRUCT_DC DefaultConstructor, STRUCT_AR Serializer, STRUCT_DTOR Destructor )
@@ -852,19 +851,15 @@ void FTransaction::SaveArray( UObject* Object, FScriptArray* Array, int32 Index,
 	}
 }
 
-void FTransaction::StoreUndo( UObject* Object, TUniquePtr<FChange> UndoChange )
+void FTransaction::StoreUndo(UObject* Object, TUniquePtr<FChange> UndoChange)
 {
-	check( Object );
+	check(Object);
 	Object->CheckDefaultSubobjects();
 
-	int32* SaveCount = ObjectMap.Find( Object );
-	if( !SaveCount )
-	{
-		ObjectMap.Add( Object, 0 );
-	}
-
 	// Save the undo record
-	Records.Add(new FObjectRecord( this, Object, MoveTemp( UndoChange ), nullptr, 0, 0, 0, 0, nullptr, nullptr, nullptr));
+	FObjectRecords& ObjectRecords = ObjectRecordsMap.FindOrAdd(FPersistentObjectRef(Object));
+	FObjectRecord* UndoRecord = ObjectRecords.Records.Add_GetRef(new FObjectRecord(this, Object, MoveTemp(UndoChange), nullptr, 0, 0, 0, 0, nullptr, nullptr, nullptr));
+	Records.Add(UndoRecord);
 }
 
 void FTransaction::SetPrimaryObject(UObject* InObject)
@@ -877,16 +872,17 @@ void FTransaction::SetPrimaryObject(UObject* InObject)
 
 void FTransaction::SnapshotObject( UObject* InObject, TArrayView<const FProperty*> Properties )
 {
-	if (InObject && ObjectMap.Contains(InObject))
-	{
-		FObjectRecord* FoundObjectRecord = Algo::FindByPredicate(Records, [ObjRefToFind = FObjectRecord::FPersistentObjectRef(InObject)](const FObjectRecord& ObjRecord)
-		{
-			return ObjRecord.Object == ObjRefToFind;
-		});
+	check(InObject);
 
-		if (FoundObjectRecord)
+	if (const FObjectRecords* ObjectRecords = ObjectRecordsMap.Find(FPersistentObjectRef(InObject)))
+	{
+		for (FObjectRecord* Record : ObjectRecords->Records)
 		{
-			FoundObjectRecord->Snapshot(this, Properties);
+			checkSlow(Record->Object.Get() == InObject);
+			if (!Record->CustomChange)
+			{
+				Record->Snapshot(this, Properties);
+			}
 		}
 	}
 }
