Ensure the inactive world created from the Content Browser is purged when loading that world for editing

UEditorEngine::OnAssetLoaded calls UEditorEngine::InitializeNewlyCreatedInactiveWorld for any world loaded, which can result in a resident world asset that is only partially initialized (inactive worlds disable things like AI).

UEditorEngine::Map_Load finds this resident world and attempts to re-use it, however since bIsWorldInitialized is already set to true it skips running the normal editor world initialization (which enables things like AI), and the end result is an editable world with certain features uninitialized.

UEditorEngine::Map_Load does attempt to purge any resident worlds that have been left lingering, but had made special exceptions for the world that the editor is attempting to load (UE-9631, UE-17205). This change has that logic run on any lingering resident worlds that have already been initialized, thus ensuring that the world being edited will definitely be initialized with the correct set of enabled features.

diff --git a/Engine/Source/Editor/UnrealEd/Private/EditorServer.cpp b/Engine/Source/Editor/UnrealEd/Private/EditorServer.cpp
--- a/Engine/Source/Editor/UnrealEd/Private/EditorServer.cpp
+++ b/Engine/Source/Editor/UnrealEd/Private/EditorServer.cpp
@@ -2583,23 +2583,21 @@ bool UEditorEngine::Map_Load(const TCHAR* Str, FOutputDevice& Ar)
 					}
 
 					// If we are loading the same world again (reloading) then we must not specify that we want to keep this world in memory.
-					// Otherwise, try to keep the existing world in memory since there is not reason to reload it.
+					// Otherwise, try to keep the existing uninitialized world in memory since there is not reason to reload it.
 					UWorld* NewWorld = nullptr;
-					if (!bIsLoadingMapTemplate && ExistingWorld != nullptr && Context.World() != ExistingWorld)
+					if (ExistingWorld && !ExistingWorld->bIsWorldInitialized && Context.World() != ExistingWorld && !bIsLoadingMapTemplate)
 					{
 						NewWorld = ExistingWorld;
 					}
 					EditorDestroyWorld( Context, LocalizedLoadingMap, NewWorld );
 
-					// Unload all other map packages currently loaded, before opening a new map.
+					// Unload all other map packages currently loaded and initialized, before opening a new map.
 					// The world is only initialized correctly as part of the level loading process, so ensure that every map package needs loading.
 					TArray<UPackage*> WorldPackages;
 					for (TObjectIterator<UWorld> It; It; ++It)
 					{
 						UPackage* Package = Cast<UPackage>(It->GetOuter());
-						
-
-						if (Package && Package != GetTransientPackage() && Package->GetPathName() != LongTempFname)
+						if (Package && Package != GetTransientPackage() && It->bIsWorldInitialized)
 						{
 							WorldPackages.AddUnique(Package);
 						}
