Fixing input flow fallout from CL 15851576. Specifically - allow ViewportClients to pass input to mode toolkits again.

diff --git a/Engine/Source/Editor/UnrealEd/Private/EditorModeManager.cpp b/Engine/Source/Editor/UnrealEd/Private/EditorModeManager.cpp
--- a/Engine/Source/Editor/UnrealEd/Private/EditorModeManager.cpp
+++ b/Engine/Source/Editor/UnrealEd/Private/EditorModeManager.cpp
@@ -1441,36 +1441,44 @@ bool FEditorModeTools::ProcessCapturedMouseMoves( FEditorViewportClient* InViewp
 	return bHandled;
 }
 
-/** Notifies all active modes of keyboard input */
+/** Notifies all active modes of keyboard input via a viewport client */
 bool FEditorModeTools::InputKey(FEditorViewportClient* InViewportClient, FViewport* Viewport, FKey Key, EInputEvent Event, bool bRouteToToolsContext)
 {
-	bool bHandled = false;
-	if (bRouteToToolsContext)
-	{
-		bHandled |= InteractiveToolsContext->InputKey(InViewportClient, Viewport, Key, Event);
-	}
-	
+	bool bHandled = bRouteToToolsContext && InteractiveToolsContext->InputKey(InViewportClient, Viewport, Key, Event);
 
-	// Copy the modes and iterate of that since a key may remove the edit mode and change ActiveScriptableModes
-	TArray<UEdMode*> CopyActiveScriptableModes(ActiveScriptableModes);
-	for (UEdMode* Mode : CopyActiveScriptableModes)
+	// If the toolkit should process the command, it should not have been handled by ITF, or be tracked elsewhere.
+	const bool bPassToToolkitCommands = bRouteToToolsContext && !bHandled;
+	ForEachEdMode([&bHandled, bPassToToolkitCommands, Event, Key, InViewportClient, Viewport](UEdMode* Mode)
 	{
+		// First, always give the legacy viewport interface a chance to process they key press. This is to support any of the FModeTools that may still exist.
 		if (ILegacyEdModeViewportInterface* ViewportInterface = Cast<ILegacyEdModeViewportInterface>(Mode))
 		{
-			bHandled |= ViewportInterface->InputKey(InViewportClient, Viewport, Key, Event);
+			if (ViewportInterface->InputKey(InViewportClient, Viewport, Key, Event))
+			{
+				bHandled |= true;
+				return true;  // Skip passing to the mode's toolkit if the legacy mode interface handled the input.
+			}
 		}
-	}
+		
+		// Next, give the toolkit commands a chance to process the key press if the tools context did not handle the key press.
+		if (bPassToToolkitCommands && (Event != IE_Released) && Mode->UsesToolkits() && Mode->GetToolkit().IsValid())
+		{
+			bHandled |= Mode->GetToolkit().Pin()->GetToolkitCommands()->ProcessCommandBindings(Key, FSlateApplication::Get().GetModifierKeys(), (Event == EInputEvent::IE_Repeat));
+			return true;
+		}
+
+		return true;
+	});
 
-	// Finally, pass input up to selected actors if not in a tool mode
-	GetEditorSelectionSet()->ForEachSelectedObject<AActor>([Key, Event](AActor* ActorPtr)
+	// Finally, pass input to selected actors if nothing else handled the input
+	if (!bHandled)
+	{
+		GetEditorSelectionSet()->ForEachSelectedObject<AActor>([Key, Event](AActor* ActorPtr)
 		{
-			if (ActorPtr)
-			{
-				// Tell the object we've had a key press
-				ActorPtr->EditorKeyPressed(Key, Event);
-			}
+			ActorPtr->EditorKeyPressed(Key, Event);
 			return true;
 		});
+	}
 	return bHandled;
 }
 
