Removed direct use of GSelectedActorAnnotation

 - Added a way for UExporter to be given an override of the current selection (via FExportObjectInnerContext).
 - Fixed FSelectedActorExportObjectInnerContext not respecting the given selection correctly during export.
 - Updated FDataprepExportObjectInnerContext to use the selection override on export, rather than directly adjust GSelectedActorAnnotation.

diff --git a/Engine/Plugins/Enterprise/DataprepEditor/Source/DataprepEditor/Private/DataprepSnapshot.cpp b/Engine/Plugins/Enterprise/DataprepEditor/Source/DataprepEditor/Private/DataprepSnapshot.cpp
--- a/Engine/Plugins/Enterprise/DataprepEditor/Source/DataprepEditor/Private/DataprepSnapshot.cpp
+++ b/Engine/Plugins/Enterprise/DataprepEditor/Source/DataprepEditor/Private/DataprepSnapshot.cpp
@@ -525,12 +525,8 @@ public:
 
 				if ( bIsValidActor )
 				{
-					// Select actor so it will be processed during the copy
-					if(SelectedActors.Find(TestParentAsActor) == nullptr)
-					{
-						SelectedActors.Add(TestParentAsActor);
-						GSelectedActorAnnotation.Set(TestParentAsActor);
-					}
+					// Track actor so it will be processed during the copy
+					SelectedActors.Add(TestParentAsActor);
 
 					bObjectMustBeCopied = true;
 					break;
@@ -557,17 +553,14 @@ public:
 		}
 	}
 
-	~FDataprepExportObjectInnerContext()
+	virtual bool IsObjectSelected(const UObject* InObj) const override
 	{
-		// Deselect all actors we processed
-		for(AActor* SelectedActor : SelectedActors)
-		{
-			GSelectedActorAnnotation.Clear(SelectedActor);
-		}
+		const AActor* Actor = Cast<AActor>(InObj);
+		return Actor && SelectedActors.Contains(Actor);
 	}
 
 	/** Set of actors marked as selected so they get included in the copy */
-	TSet<AActor*> SelectedActors;
+	TSet<const AActor*> SelectedActors;
 };
 
 void FDataprepEditor::TakeSnapshot()
diff --git a/Engine/Source/Editor/UnrealEd/Private/EditorExporters.cpp b/Engine/Source/Editor/UnrealEd/Private/EditorExporters.cpp
--- a/Engine/Source/Editor/UnrealEd/Private/EditorExporters.cpp
+++ b/Engine/Source/Editor/UnrealEd/Private/EditorExporters.cpp
@@ -499,7 +499,7 @@ bool ULevelExporterT3D::ExportText( const FExportObjectInnerContext* Context, UO
 		}
 		// Ensure actor is not a group if grouping is disabled and that the actor is currently selected
 		if( Actor && !Actor->IsA(AGroupActor::StaticClass()) &&
-			( bAllActors || Actor->IsSelected() ) )
+			( bAllActors || IsObjectSelectedForExport(Context, Actor) ) )
 		{
 			if (Actor->ShouldExport())
 			{
@@ -528,10 +528,6 @@ bool ULevelExporterT3D::ExportText( const FExportObjectInnerContext* Context, UO
 				// Restore dynamic delegate bindings.
 				UBlueprintGeneratedClass::BindDynamicDelegates(Actor->GetClass(), Actor);
 			}
-			else if (GEditor)
-			{
-				GEditor->GetSelectedActors()->Deselect( Actor );
-			}
 		}
 	}
 
@@ -639,7 +635,7 @@ bool ULevelExporterSTL::ExportText( const FExportObjectInnerContext* Context, UO
 	{
 		// Landscape
 		ALandscape* Landscape = Cast<ALandscape>(Level->Actors[iActor]);
-		if( Landscape && ( !bSelectedOnly || Landscape->IsSelected() ) )
+		if( Landscape && ( !bSelectedOnly || IsObjectSelectedForExport(Context, Landscape) ) )
 		{
 			ULandscapeInfo* LandscapeInfo = Landscape ? Landscape->GetLandscapeInfo() : NULL;
 			if( Landscape && LandscapeInfo )
@@ -682,7 +678,7 @@ bool ULevelExporterSTL::ExportText( const FExportObjectInnerContext* Context, UO
 		// Static meshes
 
 		AStaticMeshActor* Actor = Cast<AStaticMeshActor>(Level->Actors[iActor]);
-		if( Actor && ( !bSelectedOnly || Actor->IsSelected() ) && Actor->GetStaticMeshComponent()->GetStaticMesh() && Actor->GetStaticMeshComponent()->GetStaticMesh()->HasValidRenderData() )
+		if( Actor && ( !bSelectedOnly || IsObjectSelectedForExport(Context, Actor) ) && Actor->GetStaticMeshComponent()->GetStaticMesh() && Actor->GetStaticMeshComponent()->GetStaticMesh()->HasValidRenderData() )
 		{
 			FStaticMeshLODResources& LODModel = Actor->GetStaticMeshComponent()->GetStaticMesh()->GetRenderData()->LODResources[0];
 			FIndexArrayView Indices = LODModel.IndexBuffer.GetArrayView();
@@ -1288,7 +1284,7 @@ bool ULevelExporterLOD::ExportText(const FExportObjectInnerContext* Context, UOb
 	{
 		AActor* Actor = *It;
 		// only export selected actors if the flag is set
-		if( !Actor || (bSelectedOnly && !Actor->IsSelected()))
+		if( !Actor || (bSelectedOnly && !IsObjectSelectedForExport(Context, Actor)))
 		{
 			continue;
 		}
@@ -1475,7 +1471,7 @@ bool ULevelExporterOBJ::ExportText(const FExportObjectInnerContext* Context, UOb
 	{
 		AActor* Actor = *It;
 		// only export selected actors if the flag is set
-		if( !Actor || (bSelectedOnly && !Actor->IsSelected()))
+		if( !Actor || (bSelectedOnly && !IsObjectSelectedForExport(Context, Actor)))
 		{
 			continue;
 		}
diff --git a/Engine/Source/Editor/UnrealEd/Public/DragAndDrop/ExportTextDragDropOp.h b/Engine/Source/Editor/UnrealEd/Public/DragAndDrop/ExportTextDragDropOp.h
--- a/Engine/Source/Editor/UnrealEd/Public/DragAndDrop/ExportTextDragDropOp.h
+++ b/Engine/Source/Editor/UnrealEd/Public/DragAndDrop/ExportTextDragDropOp.h
@@ -20,25 +20,12 @@ public:
 	FString ActorExportText;
 	int32 NumActors;
 
-	/** The widget decorator to use */
-	//virtual TSharedPtr<SWidget> GetDefaultDecorator() const override
-	//{
-		
-	//}
-
 	static TSharedRef<FExportTextDragDropOp> New(const TArray<AActor*>& InActors)
 	{
 		TSharedRef<FExportTextDragDropOp> Operation = MakeShareable(new FExportTextDragDropOp);
 
-		for(int32 i=0; i<InActors.Num(); i++)
-		{
-			AActor* Actor = InActors[i];
-
-			GEditor->SelectActor(Actor, true, true);
-		}
-
 		FStringOutputDevice Ar;
-		const FSelectedActorExportObjectInnerContext Context;
+		const FSelectedActorExportObjectInnerContext Context(InActors);
 		UExporter::ExportToOutputDevice( &Context, GWorld, NULL, Ar, TEXT("copy"), 0, PPF_DeepCompareInstances | PPF_ExportsNotFullyQualified);
 		Operation->ActorExportText = Ar;
 		Operation->NumActors = InActors.Num();
diff --git a/Engine/Source/Runtime/Engine/Classes/Exporters/Exporter.h b/Engine/Source/Runtime/Engine/Classes/Exporters/Exporter.h
--- a/Engine/Source/Runtime/Engine/Classes/Exporters/Exporter.h
+++ b/Engine/Source/Runtime/Engine/Classes/Exporters/Exporter.h
@@ -292,6 +292,11 @@ class UExporter : public UObject
 	 */
 	ENGINE_API void ExportObjectInner(const class FExportObjectInnerContext* Context, UObject* Object, FOutputDevice& Ar, uint32 PortFlags);
 
+	/**
+	 * Should the given object be considered selected by the current export?
+	 */
+	static ENGINE_API bool IsObjectSelectedForExport(const FExportObjectInnerContext* Context, const UObject* Object);
+
 protected:
 	/**
 	 * Allows the Exporter to export any extra information it would like about each instanced object.
diff --git a/Engine/Source/Runtime/Engine/Private/UnrealExporter.cpp b/Engine/Source/Runtime/Engine/Private/UnrealExporter.cpp
--- a/Engine/Source/Runtime/Engine/Private/UnrealExporter.cpp
+++ b/Engine/Source/Runtime/Engine/Private/UnrealExporter.cpp
@@ -551,6 +551,20 @@ FExportObjectInnerContext::FExportObjectInnerContext(const TArray<UObject*>& Obj
 }
 
 
+bool FExportObjectInnerContext::IsObjectSelected(const UObject* InObj) const
+{
+	return InObj->IsSelected();
+}
+
+
+bool UExporter::IsObjectSelectedForExport(const FExportObjectInnerContext* Context, const UObject* Object)
+{
+	return Context
+		? Context->IsObjectSelected(Object)
+		: Object->IsSelected();
+}
+
+
 void UExporter::ExportObjectInner(const FExportObjectInnerContext* Context, UObject* Object, FOutputDevice& Ar, uint32 PortFlags)
 {
 	// indent all the text in here
@@ -908,7 +922,7 @@ FSelectedActorExportObjectInnerContext::FSelectedActorExportObjectInnerContext()
 	{
 		AActor* Actor = (AActor*)*It;
 		checkSlow(Actor->IsA(AActor::StaticClass()));
-		AddActorInner(Actor);
+		AddSelectedActor(Actor);
 	}
 }
 
@@ -916,14 +930,22 @@ FSelectedActorExportObjectInnerContext::FSelectedActorExportObjectInnerContext(c
 	: FExportObjectInnerContext(false) //call the empty version of the base class
 {
 	// For each selected actor...
-	for (AActor* Actor : InSelectedActors)
+	for (const AActor* Actor : InSelectedActors)
 	{
-		AddActorInner(Actor);
+		AddSelectedActor(Actor);
 	}
 }
 
-void FSelectedActorExportObjectInnerContext::AddActorInner(const AActor* InActor)
+bool FSelectedActorExportObjectInnerContext::IsObjectSelected(const UObject* InObj) const
+{
+	const AActor* Actor = Cast<AActor>(InObj);
+	return Actor && SelectedActors.Contains(Actor);
+}
+
+void FSelectedActorExportObjectInnerContext::AddSelectedActor(const AActor* InActor)
 {
+	SelectedActors.Add(InActor);
+
 	ForEachObjectWithOuter(InActor, [this](UObject* InnerObj)
 	{
 		UObject* OuterObj = InnerObj->GetOuter();
diff --git a/Engine/Source/Runtime/Engine/Public/UnrealExporter.h b/Engine/Source/Runtime/Engine/Public/UnrealExporter.h
--- a/Engine/Source/Runtime/Engine/Public/UnrealExporter.h
+++ b/Engine/Source/Runtime/Engine/Public/UnrealExporter.h
@@ -8,6 +8,8 @@
 
 #include "CoreMinimal.h"
 
+class AActor;
+
 /**
  * Encapsulates a map from objects to their direct inners, used by UExporter::ExportObjectInner when exporting objects.
  * Should be recreated before new objects are created within objects that are to be exported!
@@ -35,6 +37,8 @@ public:
 	/**Empty Constructor for derived export contexts */
 	FExportObjectInnerContext(const bool bIgnoredValue) {};
 
+	virtual ~FExportObjectInnerContext() = default;
+
 protected:
 	friend class UExporter;
 
@@ -54,6 +58,11 @@ public:
 	{
 		return ObjectToInnerMap.Find(InObj);
 	}
+
+	/**
+	 * Should the given object be considered selected by the current export?
+	 */
+	virtual bool IsObjectSelected(const UObject* InObj) const;
 };
 
 #if WITH_EDITOR
@@ -61,10 +70,14 @@ class ENGINE_API FSelectedActorExportObjectInnerContext : public FExportObjectIn
 {
 public:
 	FSelectedActorExportObjectInnerContext();
-	explicit FSelectedActorExportObjectInnerContext(const TArray<class AActor*> InSelectedActors);
+	explicit FSelectedActorExportObjectInnerContext(const TArray<AActor*> InSelectedActors);
+
+	virtual bool IsObjectSelected(const UObject* InObj) const override;
 
 private:
-	void AddActorInner(const class AActor* InActor);
+	void AddSelectedActor(const AActor* InActor);
+
+	TSet<const AActor*> SelectedActors;
 };
 #endif
 
