Implement FArchive::TotalSize for FTransaction::FObjectRecord::FReader and FTransaction::FObjectRecord::FWriter

- Allows FCompressedBuffer to be used with the undo stack

#rnx

diff --git a/Engine/Source/Editor/UnrealEd/Classes/Editor/Transactor.h b/Engine/Source/Editor/UnrealEd/Classes/Editor/Transactor.h
--- a/Engine/Source/Editor/UnrealEd/Classes/Editor/Transactor.h
+++ b/Engine/Source/Editor/UnrealEd/Classes/Editor/Transactor.h
@@ -330,6 +330,7 @@ protected:
 
 			virtual int64 Tell() override {return Offset;}
 			virtual void Seek( int64 InPos ) override { Offset = InPos; }
+			virtual int64 TotalSize() override { return SerializedObject.Data.Num(); }
 
 		private:
 			void Serialize( void* SerData, int64 Num ) override
@@ -421,6 +422,8 @@ protected:
 				Offset = InPos; 
 			}
 
+			virtual int64 TotalSize() override { return SerializedObject.Data.Num(); }
+
 			virtual bool ShouldSkipProperty(const FProperty* InProperty) const override
 			{
 				return (PropertiesToSerialize.Num() > 0 && !PropertiesToSerialize.Contains(InProperty))
