Importing UDIM textures will no longer attempt to keep the source format as compressed JPEG (even if that option is enabled)

- UDIM textures currently cannot be compressed in any way (due to how the individual blocks are accessed) so need to avoid the jpeg data being returned as compressed even if the option is enabled in the config file.
- ImportImage now takes a bitfield of flags EImageImportFlags to control it's import options, rather than add additional bool parameters.
- In addition to the pre-existing option to allow/disallow non power of two textures, the bitfield has an option to allow ImportImage to return the texture data in it's original compressed format.
- ImportTextureUDIM does not set the EImageImportFlags::AllowReturnOfCompressedData flag.

diff --git a/Engine/Source/Editor/UnrealEd/Classes/Factories/TextureFactory.h b/Engine/Source/Editor/UnrealEd/Classes/Factories/TextureFactory.h
--- a/Engine/Source/Editor/UnrealEd/Classes/Factories/TextureFactory.h
+++ b/Engine/Source/Editor/UnrealEd/Classes/Factories/TextureFactory.h
@@ -206,7 +206,7 @@ private:
 	/** This variable is static because in StaticImportObject() the type of the factory is not known. */
 	static bool bSuppressImportOverwriteDialog;
 
-    /** force overwriting the existing texture without the dialog box */
+	/** Force overwriting the existing texture without the dialog box */
 	static bool bForceOverwriteExistingSettings;
 
 	/**
@@ -221,8 +221,20 @@ private:
 	*/
 	static bool IsImportResolutionValid(int32 Width, int32 Height, bool bAllowNonPowerOfTwo, FFeedbackContext* Warn);
 
+	/** Flags to be used when calling ImportImage */
+	enum class EImageImportFlags
+	{
+		/** No options selected */
+		None						= 0,
+		/** Allows textures to be imported with dimensions that are not to the power of two */
+		AllowNonPowerOfTwo			= 1 << 0,
+		/** Allows the return of texture data in it's original compressed format, if this occurs then FImportImage::RawDataCompressionFormat will contain the returned format. */
+		AllowReturnOfCompressedData	= 1 << 1
+	};
+	FRIEND_ENUM_CLASS_FLAGS(EImageImportFlags);
+
 	/** Import image file into generic image struct, may be easily copied to FTextureSource */
-	bool ImportImage(const uint8* Buffer, uint32 Length, FFeedbackContext* Warn, bool bAllowNonPowerOfTwo, FImportImage& OutImage);
+	bool ImportImage(const uint8* Buffer, uint32 Length, FFeedbackContext* Warn, EImageImportFlags Flags, FImportImage& OutImage);
 
 	/** used by CreateTexture() */
 	UTexture* ImportTexture(UClass* Class, UObject* InParent, FName Name, EObjectFlags Flags, const TCHAR* Type, const uint8*& Buffer, const uint8* BufferEnd, FFeedbackContext* Warn);
@@ -235,3 +247,5 @@ private:
 	/** Texture settings from the automated importer that should be applied to the new texture */
 	TSharedPtr<class FJsonObject> AutomatedImportSettings;
 };
+
+ENUM_CLASS_FLAGS(UTextureFactory::EImageImportFlags);
diff --git a/Engine/Source/Editor/UnrealEd/Private/Factories/EditorFactories.cpp b/Engine/Source/Editor/UnrealEd/Private/Factories/EditorFactories.cpp
--- a/Engine/Source/Editor/UnrealEd/Private/Factories/EditorFactories.cpp
+++ b/Engine/Source/Editor/UnrealEd/Private/Factories/EditorFactories.cpp
@@ -3003,10 +3003,12 @@ void* FImportImage::GetMipData(int32 InMipIndex)
 	return &RawData[Offset];
 }
 
-bool UTextureFactory::ImportImage(const uint8* Buffer, uint32 Length, FFeedbackContext* Warn, bool bAllowNonPowerOfTwo, FImportImage& OutImage)
+bool UTextureFactory::ImportImage(const uint8* Buffer, uint32 Length, FFeedbackContext* Warn, EImageImportFlags Flags, FImportImage& OutImage)
 {
 	TRACE_CPUPROFILER_EVENT_SCOPE(UTextureFactory::ImportImage)
 
+	const bool bAllowNonPowerOfTwo = EnumHasAllFlags(Flags, EImageImportFlags::AllowNonPowerOfTwo);
+
 	IImageWrapperModule& ImageWrapperModule = FModuleManager::LoadModuleChecked<IImageWrapperModule>(FName("ImageWrapper"));
 
 	//
@@ -3136,11 +3138,13 @@ bool UTextureFactory::ImportImage(const uint8* Buffer, uint32 Length, FFeedbackC
 			TextureFormat,
 			BitDepth < 16
 		);
-
-		// For now this option is opt in via the config files once there is no technical risk this will
-		// become the default path.
+		
 		bool bRetainJpegFormat = false;
-		GConfig->GetBool(TEXT("TextureImporter"), TEXT("RetainJpegFormat"), bRetainJpegFormat, GEditorIni);
+		if (EnumHasAllFlags(Flags, EImageImportFlags::AllowReturnOfCompressedData))
+		{
+			// For now this option is opt in via the config files once there is no technical risk this will become the default path.
+			GConfig->GetBool(TEXT("TextureImporter"), TEXT("RetainJpegFormat"), bRetainJpegFormat, GEditorIni);
+		}	
 
 		if (bRetainJpegFormat)
 		{
@@ -3558,16 +3562,17 @@ UTexture* UTextureFactory::ImportTextureUDIM(UClass* Class, UObject* InParent, F
 	ETextureSourceFormat Format = TSF_Invalid;
 	TextureCompressionSettings TCSettings = TC_MAX;
 	bool bSRGB = false;
+	
+	// UDIM requires each page to be power-of-2 and be uncompressed
+	const EImageImportFlags ImportFlags = EImageImportFlags::None;
+
 	for (const auto& It : UDIMIndexToFile)
 	{
 		const FString& TexturePath = It.Value;
 		if (FFileHelper::LoadFileToArray(TextureData, *TexturePath))
 		{
-			// UDIM requires each page to be power-of-2
-			const bool bAllowNonPowerOfTwo = false;
-
 			FImportImage Image;
-			if (ImportImage(TextureData.GetData(), TextureData.Num(), Warn, bAllowNonPowerOfTwo, Image))
+			if (ImportImage(TextureData.GetData(), TextureData.Num(), Warn, ImportFlags, Image))
 			{
 				if (Format == TSF_Invalid)
 				{
@@ -3649,11 +3654,18 @@ UTexture* UTextureFactory::ImportTexture(UClass* Class, UObject* InParent, FName
 	// Validate it.
 	const int32 Length = BufferEnd - Buffer;
 
+	// We want to allow certain texture types to potentially return the compressed data rather than raw uncompressed
+	EImageImportFlags ImportFlags = EImageImportFlags::AllowReturnOfCompressedData;
+	if (bAllowNonPowerOfTwo)
+	{
+		ImportFlags |= EImageImportFlags::AllowNonPowerOfTwo;
+	}
+
 	//
 	// Generic 2D Image
 	//
 	FImportImage Image;
-	if (ImportImage(Buffer, Length, Warn, bAllowNonPowerOfTwo, Image))
+	if (ImportImage(Buffer, Length, Warn, ImportFlags, Image))
 	{
 		UTexture2D* Texture = CreateTexture2D(InParent, Name, Flags);
 		if (Texture)
