Add trace scope instrumentation for Editor init and exit
Add trace scope instrumentation for Engine exit

diff --git a/Engine/Source/Editor/UnrealEd/Private/UnrealEdGlobals.cpp b/Engine/Source/Editor/UnrealEd/Private/UnrealEdGlobals.cpp
--- a/Engine/Source/Editor/UnrealEd/Private/UnrealEdGlobals.cpp
+++ b/Engine/Source/Editor/UnrealEd/Private/UnrealEdGlobals.cpp
@@ -89,6 +89,8 @@ void CheckAndMaybeGoToVRModeInternal(const bool bIsImmersive)
 
 int32 EditorInit( IEngineLoop& EngineLoop )
 {
+	TRACE_CPUPROFILER_EVENT_SCOPE(EditorInit);
+
 	// Create debug exec.	
 	GDebugToolExec = new FDebugToolExec;
 
@@ -203,6 +205,8 @@ int32 EditorInit( IEngineLoop& EngineLoop )
 
 int32 EditorReinit()
 {
+	TRACE_CPUPROFILER_EVENT_SCOPE(EditorReinit);
+
 	// Are we in immersive mode?
 	const bool bIsImmersive = FPaths::IsProjectFilePathSet() && FParse::Param(FCommandLine::Get(), TEXT("immersive"));
 	// Do final set up on the editor frame and show it
diff --git a/Engine/Source/Runtime/Launch/Private/LaunchEngineLoop.cpp b/Engine/Source/Runtime/Launch/Private/LaunchEngineLoop.cpp
--- a/Engine/Source/Runtime/Launch/Private/LaunchEngineLoop.cpp
+++ b/Engine/Source/Runtime/Launch/Private/LaunchEngineLoop.cpp
@@ -4256,6 +4256,7 @@ int32 FEngineLoop::Init()
 void FEngineLoop::Exit()
 {
 	STAT_ADD_CUSTOMMESSAGE_NAME( STAT_NamedMarker, TEXT( "EngineLoop.Exit" ) );
+	TRACE_CPUPROFILER_EVENT_SCOPE(FEngineLoop::Exit);
 	TRACE_BOOKMARK(TEXT("EngineLoop.Exit"));
 
 	GIsRunning	= 0;
