Make sure MiniMap builder is going to work for licensees

* Added -UseOnlyHLODs: Allows building the minimap using on HLODs rather than the whole world.
* Added -AutoSubmit: Submit the updated minimap at the end of the generation
* Capture minimap using only albedo (base color)
* Fixed issue with async mesh compilation causing the built minimap to be blank
* Upped minimap resolution to 4k

diff --git a/Engine/Source/Editor/UnrealEd/Private/WorldPartition/WorldPartitionMiniMapBuilder.cpp b/Engine/Source/Editor/UnrealEd/Private/WorldPartition/WorldPartitionMiniMapBuilder.cpp
--- a/Engine/Source/Editor/UnrealEd/Private/WorldPartition/WorldPartitionMiniMapBuilder.cpp
+++ b/Engine/Source/Editor/UnrealEd/Private/WorldPartition/WorldPartitionMiniMapBuilder.cpp
@@ -4,7 +4,20 @@
 
 #include "CoreMinimal.h"
 #include "Logging/LogMacros.h"
+#include "Misc/CommandLine.h"
 
+#include "Engine/World.h"
+#include "SourceControlHelpers.h"
+
+#include "ISourceControlModule.h"
+#include "ISourceControlProvider.h"
+#include "SourceControlOperations.h"
+
+#include "WorldPartition/ActorDescContainer.h"
+#include "WorldPartition/HLOD/HLODActor.h"
+#include "WorldPartition/HLOD/HLODActorDesc.h"
+#include "WorldPartition/WorldPartition.h"
+#include "WorldPartition/WorldPartitionHandle.h"
 #include "WorldPartition/WorldPartitionMiniMap.h"
 #include "WorldPartition/WorldPartitionMiniMapHelper.h"
 
@@ -13,6 +26,8 @@ DEFINE_LOG_CATEGORY_STATIC(LogWorldPartitionMiniMapBuilder, All, All);
 UWorldPartitionMiniMapBuilder::UWorldPartitionMiniMapBuilder(const FObjectInitializer& ObjectInitializer)
 	: Super(ObjectInitializer)
 {
+	bUseOnlyHLODs = FParse::Param(FCommandLine::Get(), TEXT("UseOnlyHLODs"));
+	bAutoSubmit = bUseOnlyHLODs = FParse::Param(FCommandLine::Get(), TEXT("AutoSubmit"));
 }
 
 bool UWorldPartitionMiniMapBuilder::Run(UWorld* World, FPackageSourceControlHelper& PackageHelper)
@@ -21,7 +36,34 @@ bool UWorldPartitionMiniMapBuilder::Run(UWorld* World, FPackageSourceControlHelp
 	if (!WorldMiniMap)
 	{
 		UE_LOG(LogWorldPartitionMiniMapBuilder, Error, TEXT("Failed to create Minimap. WorldPartitionMiniMap actor not found in the persistent level."));
-		return 1;
+		return false;
+	}
+
+	// Load HLOD actors first if using only HLODs to generate MiniMap
+	TArray<FWorldPartitionReference> HLODActors;
+	if (bUseOnlyHLODs)
+	{
+		UWorldPartition* WorldPartition = World->GetWorldPartition();
+
+		TSet<FGuid> AllSubActors;
+		for (UActorDescContainer::TIterator<AWorldPartitionHLOD> HLODIterator(WorldPartition); HLODIterator; ++HLODIterator)
+		{
+			AllSubActors.Append(HLODIterator->GetSubActors());
+		}
+
+		for (UActorDescContainer::TIterator<AWorldPartitionHLOD> HLODIterator(WorldPartition); HLODIterator; ++HLODIterator)
+		{
+			// Only include top level HLODs - If an HLOD actor isn't included as a subactor it means it is top level
+			if (!AllSubActors.Contains(HLODIterator->GetGuid()))
+			{
+				FWorldPartitionReference HLODActorRef(WorldPartition, HLODIterator->GetGuid());
+
+				AWorldPartitionHLOD* Actor = Cast<AWorldPartitionHLOD>(HLODActorRef->GetActor());
+				Actor->SetIsTemporarilyHiddenInEditor(false);
+
+				HLODActors.Add(HLODActorRef);
+			}
+		}
 	}
 
 	WorldMiniMap->MiniMapSize = this->MiniMapSize;
@@ -34,21 +76,39 @@ bool UWorldPartitionMiniMapBuilder::Run(UWorld* World, FPackageSourceControlHelp
 	if (!PackageHelper.Checkout(WorldMiniMapExternalPackage))
 	{
 		UE_LOG(LogWorldPartitionMiniMapBuilder, Error, TEXT("Error checking out package %s."), *WorldMiniMapExternalPackage->GetName());
-		return 1;
+		return false;
 	}
 	
 	if (!UPackage::SavePackage(WorldMiniMapExternalPackage, nullptr, RF_Standalone, *PackageFileName, GError, nullptr, false, true, SAVE_Async))
 	{
 		UE_LOG(LogWorldPartitionMiniMapBuilder, Error, TEXT("Error saving package %s."), *WorldMiniMapExternalPackage->GetName());
-		return 1;
+		return false;
 	}
 
 	if (!PackageHelper.AddToSourceControl(WorldMiniMapExternalPackage))
 	{
 		UE_LOG(LogWorldPartitionMiniMapBuilder, Error, TEXT("Error adding package %s to source control."), *WorldMiniMapExternalPackage->GetName());
-		return 1;
+		return false;
 	}
 
 	UPackage::WaitForAsyncFileWrites();
+
+	if (bAutoSubmit)
+	{
+		FText ChangelistDescription = FText::FromString(FString::Printf(TEXT("Rebuilt minimap for \"%s\" at %s"), *World->GetName(), *FEngineVersion::Current().ToString()));
+
+		TSharedRef<FCheckIn, ESPMode::ThreadSafe> CheckInOperation = ISourceControlOperation::Create<FCheckIn>();
+		CheckInOperation->SetDescription(ChangelistDescription);
+		if (ISourceControlModule::Get().GetProvider().Execute(CheckInOperation, PackageFileName) != ECommandResult::Succeeded)
+		{
+			UE_LOG(LogWorldPartitionMiniMapBuilder, Error, TEXT("Failed to submit minimap (%s) to source control."), *PackageFileName);
+			return false;
+		}
+		else
+		{
+			UE_LOG(LogWorldPartitionMiniMapBuilder, Display, TEXT("#### Submitted minimap (%s) to source control ####"), *PackageFileName);
+		}
+	}
+	
 	return true;
 }
diff --git a/Engine/Source/Editor/UnrealEd/Public/WorldPartition/WorldPartitionMiniMapBuilder.h b/Engine/Source/Editor/UnrealEd/Public/WorldPartition/WorldPartitionMiniMapBuilder.h
--- a/Engine/Source/Editor/UnrealEd/Public/WorldPartition/WorldPartitionMiniMapBuilder.h
+++ b/Engine/Source/Editor/UnrealEd/Public/WorldPartition/WorldPartitionMiniMapBuilder.h
@@ -13,10 +13,14 @@ class UWorldPartitionMiniMapBuilder : public UWorldPartitionBuilder
 public:
 	// UWorldPartitionBuilder interface begin
 	virtual bool RequiresCommandletRendering() const override { return true; }
-	virtual bool RequiresEntireWorldLoading() const override { return true; }
+	virtual bool RequiresEntireWorldLoading() const override { return !bUseOnlyHLODs; }
 	virtual bool Run(UWorld* World, FPackageSourceControlHelper& PackageHelper) override;
 	// UWorldPartitionBuilder interface end
 
 	UPROPERTY(config)
-	int32 MiniMapSize = 1024;
+	int32 MiniMapSize = 4096;
+
+private:
+	bool bUseOnlyHLODs;
+	bool bAutoSubmit;
 };
\ No newline at end of file
diff --git a/Engine/Source/Runtime/Engine/Private/WorldPartition/WorldPartitionMiniMapHelper.cpp b/Engine/Source/Runtime/Engine/Private/WorldPartition/WorldPartitionMiniMapHelper.cpp
--- a/Engine/Source/Runtime/Engine/Private/WorldPartition/WorldPartitionMiniMapHelper.cpp
+++ b/Engine/Source/Runtime/Engine/Private/WorldPartition/WorldPartitionMiniMapHelper.cpp
@@ -12,12 +12,12 @@
 #include "Engine/SceneCapture2D.h"
 #include "Engine/World.h"
 #include "ShaderCompiler.h"
-#include "TextureCompiler.h"
+#include "AssetCompilingManager.h"
 #include "DistanceFieldAtlas.h"
 #include "PackageSourceControlHelper.h"
 #include "ProfilingDebugging/ScopedTimers.h"
 #include "WorldPartition/WorldPartitionMiniMap.h"
-#include"WorldPartition/WorldPartition.h"
+#include "WorldPartition/WorldPartition.h"
 
 DEFINE_LOG_CATEGORY_STATIC(LogWorldPartitionMiniMapHelper, All, All);
 
@@ -50,7 +50,6 @@ AWorldPartitionMiniMap* FWorldPartitionMiniMapHelper::GetWorldPartitionMiniMap(U
 	}
 	
 	return nullptr;
-	
 }
 
 void FWorldPartitionMiniMapHelper::CaptureWorldMiniMapToTexture(UWorld* InWorld, UObject* InOuterForTexture, uint32 InMiniMapSize, UTexture2D*& InOutMiniMapTexture, FBox& OutWorldBounds)
@@ -63,16 +62,7 @@ void FWorldPartitionMiniMapHelper::CaptureWorldMiniMapToTexture(UWorld* InWorld,
 			GShaderCompilingManager->FinishAllCompilation();
 		}
 	};
-	auto WaitForTextureBuilding = []()
-	{
-		UE_SCOPED_TIMER(TEXT("Building Textures"), LogWorldPartitionMiniMapHelper);
-		auto RemainingTextures = FTextureCompilingManager::Get().GetNumRemainingTextures();
-		if (RemainingTextures > 0)
-		{
-			FTextureCompilingManager::Get().FinishAllCompilation();
-		}
-	};
-	auto WaitForDistanceMeshFieldBuilding = []() {
+	auto WaitForDistanceFieldBuilding = []() {
 		UE_SCOPED_TIMER(TEXT("Building Mesh Distance Fields"), LogWorldPartitionMiniMapHelper);
 		if (GDistanceFieldAsyncQueue && GDistanceFieldAsyncQueue->GetNumOutstandingTasks() > 0)
 		{
@@ -80,9 +70,9 @@ void FWorldPartitionMiniMapHelper::CaptureWorldMiniMapToTexture(UWorld* InWorld,
 		}
 	};
 
-	WaitForDistanceMeshFieldBuilding();
-	WaitForTextureBuilding();
+	WaitForDistanceFieldBuilding();
 	WaitForShaderCompilation();
+	FAssetCompilingManager::Get().FinishAllCompilation();
 	
 	//Calculate bounds of the World
 	OutWorldBounds = InWorld->GetWorldPartition()->GetEditorWorldBounds();
@@ -114,6 +104,7 @@ void FWorldPartitionMiniMapHelper::CaptureWorldMiniMapToTexture(UWorld* InWorld,
 	auto CaptureComponent = CaptureActor->GetCaptureComponent2D();
 	CaptureComponent->TextureTarget = RenderTargetTexture;
 	CaptureComponent->ProjectionType = ECameraProjectionMode::Orthographic;
+	CaptureComponent->CaptureSource = SCS_BaseColor;
 	CaptureComponent->OrthoWidth = ViewportWidth;
 	CaptureComponent->bUseCustomProjectionMatrix = true;
 	CaptureComponent->CustomProjectionMatrix = ProjectionMatrix;
