Ctrl+Shift+Tab now summons the tab switching dialog immediately.

diff --git a/Engine/Source/Editor/UnrealEd/Private/Toolkits/GlobalEditorCommonCommands.cpp b/Engine/Source/Editor/UnrealEd/Private/Toolkits/GlobalEditorCommonCommands.cpp
--- a/Engine/Source/Editor/UnrealEd/Private/Toolkits/GlobalEditorCommonCommands.cpp
+++ b/Engine/Source/Editor/UnrealEd/Private/Toolkits/GlobalEditorCommonCommands.cpp
@@ -28,8 +28,11 @@ FGlobalEditorCommonCommands::FGlobalEditorCommonCommands()
 
 void FGlobalEditorCommonCommands::RegisterCommands()
 {
-	UI_COMMAND(SummonControlTabNavigation, "Tab Navigation", "Summons a list of open assets and tabs", EUserInterfaceActionType::Button, FInputChord(EModifierKey::Control, EKeys::Tab));
-	UI_COMMAND(SummonControlTabNavigationAlternate, "Tab Navigation", "Summons a list of open assets and tabs", EUserInterfaceActionType::Button, FInputChord(EModifierKey::Command, EKeys::Tab));
+	UI_COMMAND(SummonControlTabNavigation, "Tab Navigation", "Summons a list of open assets and tabs, and navigates forwards in it.", EUserInterfaceActionType::Button, FInputChord(EModifierKey::Control, EKeys::Tab));
+	UI_COMMAND(SummonControlTabNavigationAlternate, "Tab Navigation", "Summons a list of open assets and tabs, and navigates forward in its.", EUserInterfaceActionType::Button, FInputChord(EModifierKey::Command, EKeys::Tab));
+
+	UI_COMMAND(SummonControlTabNavigationBackwards, "Tab Navigation Backwards", "Summons a list of open assets and tabs, and navigates backwards in it.", EUserInterfaceActionType::Button, FInputChord(EModifierKey::Control | EModifierKey::Shift, EKeys::Tab));
+	UI_COMMAND(SummonControlTabNavigationBackwardsAlternate, "Tab Navigation Backwards", "Summons a list of open assets and tabs, and navigates backwards in it.", EUserInterfaceActionType::Button, FInputChord(EModifierKey::Command | EModifierKey::Shift, EKeys::Tab));
 
 	UI_COMMAND(SummonOpenAssetDialog, "Open Asset...", "Summons an asset picker", EUserInterfaceActionType::Button, FInputChord(EModifierKey::Control, EKeys::P));
 	UI_COMMAND(SummonOpenAssetDialogAlternate, "Open Asset...", "Summons an asset picker", EUserInterfaceActionType::Button, FInputChord(EModifierKey::Alt | EModifierKey::Shift, EKeys::O));
@@ -53,6 +56,14 @@ void FGlobalEditorCommonCommands::MapActions(TSharedRef<FUICommandList>& Toolkit
 		Get().SummonControlTabNavigationAlternate,
 		FExecuteAction::CreateStatic(&FGlobalEditorCommonCommands::OnPressedCtrlTab, Get().SummonControlTabNavigationAlternate));
 
+	ToolkitCommands->MapAction(
+		Get().SummonControlTabNavigationBackwards,
+		FExecuteAction::CreateStatic(&FGlobalEditorCommonCommands::OnPressedCtrlTab, Get().SummonControlTabNavigationBackwards));
+
+	ToolkitCommands->MapAction(
+		Get().SummonControlTabNavigationBackwardsAlternate,
+		FExecuteAction::CreateStatic(&FGlobalEditorCommonCommands::OnPressedCtrlTab, Get().SummonControlTabNavigationBackwardsAlternate));
+
 	ToolkitCommands->MapAction(
 		Get().SummonOpenAssetDialog,
 		FExecuteAction::CreateStatic(&FGlobalEditorCommonCommands::OnSummonedAssetPicker));
diff --git a/Engine/Source/Editor/UnrealEd/Public/Toolkits/GlobalEditorCommonCommands.h b/Engine/Source/Editor/UnrealEd/Public/Toolkits/GlobalEditorCommonCommands.h
--- a/Engine/Source/Editor/UnrealEd/Public/Toolkits/GlobalEditorCommonCommands.h
+++ b/Engine/Source/Editor/UnrealEd/Public/Toolkits/GlobalEditorCommonCommands.h
@@ -32,6 +32,8 @@ public:
 	TSharedPtr<FUICommandInfo> FindInContentBrowser;
 	TSharedPtr<FUICommandInfo> SummonControlTabNavigation;
 	TSharedPtr<FUICommandInfo> SummonControlTabNavigationAlternate;
+	TSharedPtr<FUICommandInfo> SummonControlTabNavigationBackwards;
+	TSharedPtr<FUICommandInfo> SummonControlTabNavigationBackwardsAlternate;
 	TSharedPtr<FUICommandInfo> SummonOpenAssetDialog;
 	TSharedPtr<FUICommandInfo> SummonOpenAssetDialogAlternate;
 	TSharedPtr<FUICommandInfo> OpenDocumentation;
