Run the HLOD build on a single machine
* Much faster build (no need to wait for machines on Horde or to perform temp p4 operations on multiple builders)
* Update minimap in the same process
* Fixed the -DeleteHLODs option not working with always loaded HLOD actors

diff --git a/Engine/Source/Editor/UnrealEd/Private/WorldPartition/WorldPartitionHLODsBuilder.cpp b/Engine/Source/Editor/UnrealEd/Private/WorldPartition/WorldPartitionHLODsBuilder.cpp
--- a/Engine/Source/Editor/UnrealEd/Private/WorldPartition/WorldPartitionHLODsBuilder.cpp
+++ b/Engine/Source/Editor/UnrealEd/Private/WorldPartition/WorldPartitionHLODsBuilder.cpp
@@ -128,6 +128,10 @@ UWorldPartitionHLODsBuilder::UWorldPartitionHLODsBuilder(const FObjectInitialize
 	bDeleteHLODs = FParse::Param(FCommandLine::Get(), TEXT("DeleteHLODs"));
 	bSubmitHLODs = FParse::Param(FCommandLine::Get(), TEXT("SubmitHLODs"));
 
+	bSingleBuildStep = !bSetupHLODs && !bBuildHLODs && !bSubmitHLODs && !bDeleteHLODs;
+
+	bAutoSubmit = FParse::Param(FCommandLine::Get(), TEXT("AutoSubmit"));
+
 	bDistributedBuild = FParse::Param(FCommandLine::Get(), TEXT("DistributedBuild"));
 
 	FParse::Value(FCommandLine::Get(), TEXT("BuildManifest="), BuildManifest);
@@ -204,6 +208,12 @@ bool UWorldPartitionHLODsBuilder::ValidateParams() const
 		return false;
 	}
 
+	if (bAutoSubmit && !bSingleBuildStep)
+	{
+		UE_LOG(LogWorldPartitionHLODsBuilder, Error, TEXT("-AutoSubmit argument only valid when building HLODs in a single step, exiting..."), *BuildManifest);
+		return false;
+	}
+
 	if (IsDistributedBuild() && !ISourceControlModule::Get().GetProvider().IsEnabled())
 	{
 		UE_LOG(LogWorldPartitionHLODsBuilder, Error, TEXT("Distributed builds requires that a valid source control provider is enabled, exiting..."), *BuildManifest);
@@ -239,27 +249,32 @@ bool UWorldPartitionHLODsBuilder::Run(UWorld* World, FPackageSourceControlHelper
 
 	SourceControlHelper = new FSourceControlHelper(PackageHelper);
 
-	bool bRet;
-	if (bSetupHLODs)
-	{
-		bRet = SetupHLODActors(true);
-	}
-	else if (bBuildHLODs)
+	bool bRet = true;
+
+	if (bSingleBuildStep)
 	{
-		bRet = BuildHLODActors();
+		bRet = SetupHLODActors(false);
 	}
-	else if (bDeleteHLODs)
+	else
 	{
-		bRet = DeleteHLODActors();
+		if (bSetupHLODs)
+		{
+			bRet = SetupHLODActors(true);
+		}
+		else if (bBuildHLODs)
+		{
+			bRet = BuildHLODActors();
+		}
+		else if (bDeleteHLODs)
+		{
+			bRet = DeleteHLODActors();
+		}
 	}
-	else if (bSubmitHLODs)
+
+	if (bRet && (bSubmitHLODs || bAutoSubmit))
 	{
 		bRet = SubmitHLODActors();
 	}
-	else
-	{
-		bRet = SetupHLODActors(false);
-	}
 
 	WorldPartition = nullptr;
 	delete SourceControlHelper;
diff --git a/Engine/Source/Editor/UnrealEd/Public/WorldPartition/WorldPartitionHLODsBuilder.h b/Engine/Source/Editor/UnrealEd/Public/WorldPartition/WorldPartitionHLODsBuilder.h
--- a/Engine/Source/Editor/UnrealEd/Public/WorldPartition/WorldPartitionHLODsBuilder.h
+++ b/Engine/Source/Editor/UnrealEd/Public/WorldPartition/WorldPartitionHLODsBuilder.h
@@ -46,6 +46,8 @@ private:
 	bool bBuildHLODs;
 	bool bDeleteHLODs;
 	bool bSubmitHLODs;
+	bool bSingleBuildStep;
+	bool bAutoSubmit;
 	bool bDistributedBuild;
 	FString BuildManifest;
 	int32 BuilderIdx;
