WorldPartitionBuilder: Add support for Iterative Cell Loading

diff --git a/Engine/Plugins/Experimental/VirtualHeightfieldMesh/Source/VirtualHeightfieldMeshEditor/Private/WorldPartitionVirtualHeightfieldMeshBuilder.cpp b/Engine/Plugins/Experimental/VirtualHeightfieldMesh/Source/VirtualHeightfieldMeshEditor/Private/WorldPartitionVirtualHeightfieldMeshBuilder.cpp
--- a/Engine/Plugins/Experimental/VirtualHeightfieldMesh/Source/VirtualHeightfieldMeshEditor/Private/WorldPartitionVirtualHeightfieldMeshBuilder.cpp
+++ b/Engine/Plugins/Experimental/VirtualHeightfieldMesh/Source/VirtualHeightfieldMeshEditor/Private/WorldPartitionVirtualHeightfieldMeshBuilder.cpp
@@ -14,7 +14,7 @@ UWorldPartitionVirtualHeightfieldMeshBuilder::UWorldPartitionVirtualHeightfieldM
 {
 }
 
-bool UWorldPartitionVirtualHeightfieldMeshBuilder::Run(UWorld* World, FPackageSourceControlHelper& PackageHelper)
+bool UWorldPartitionVirtualHeightfieldMeshBuilder::RunInternal(UWorld* World, const FBox& Bounds, FPackageSourceControlHelper& PackageHelper)
 {
 	// Find UVirtualHeightfieldMeshComponent and update the associated UHeightfieldMinMaxTexture assets.
 	// todo[vhm]: Convert builder to sequentially load and build sections of the world.
diff --git a/Engine/Plugins/Experimental/VirtualHeightfieldMesh/Source/VirtualHeightfieldMeshEditor/Private/WorldPartitionVirtualHeightfieldMeshBuilder.h b/Engine/Plugins/Experimental/VirtualHeightfieldMesh/Source/VirtualHeightfieldMeshEditor/Private/WorldPartitionVirtualHeightfieldMeshBuilder.h
--- a/Engine/Plugins/Experimental/VirtualHeightfieldMesh/Source/VirtualHeightfieldMeshEditor/Private/WorldPartitionVirtualHeightfieldMeshBuilder.h
+++ b/Engine/Plugins/Experimental/VirtualHeightfieldMesh/Source/VirtualHeightfieldMeshEditor/Private/WorldPartitionVirtualHeightfieldMeshBuilder.h
@@ -13,7 +13,8 @@ class UWorldPartitionVirtualHeightfieldMeshBuilder : public UWorldPartitionBuild
 public:
 	// UWorldPartitionBuilder interface begin
 	virtual bool RequiresCommandletRendering() const override { return true; }
-	virtual bool RequiresEntireWorldLoading() const override { return true; }
-	virtual bool Run(UWorld* World, FPackageSourceControlHelper& PackageHelper) override;
+	virtual ELoadingMode GetLoadingMode() const override { return ELoadingMode::EntireWorld; }
+protected:
+	virtual bool RunInternal(UWorld* World, const FBox& Bounds, FPackageSourceControlHelper& PackageHelper) override;
 	// UWorldPartitionBuilder interface end
 };
\ No newline at end of file
diff --git a/Engine/Source/Editor/UnrealEd/Private/Commandlets/WorldPartitionBuilderCommandlet.cpp b/Engine/Source/Editor/UnrealEd/Private/Commandlets/WorldPartitionBuilderCommandlet.cpp
--- a/Engine/Source/Editor/UnrealEd/Private/Commandlets/WorldPartitionBuilderCommandlet.cpp
+++ b/Engine/Source/Editor/UnrealEd/Private/Commandlets/WorldPartitionBuilderCommandlet.cpp
@@ -166,12 +166,6 @@ int32 UWorldPartitionBuilderCommandlet::Main(const FString& Params)
 	UWorld* PrevGWorld = GWorld;
 	GWorld = World;
 
-	if (Builder->RequiresEntireWorldLoading())
-	{
-		const FBox LoadBox(FVector(-WORLD_MAX, -WORLD_MAX, -WORLD_MAX), FVector(WORLD_MAX, WORLD_MAX, WORLD_MAX));
-		WorldPartition->LoadEditorCells(LoadBox);
-	}
-
 	// Run builder
 	bool bResult = Builder->Run(World, *this);
 	
diff --git a/Engine/Source/Editor/UnrealEd/Private/WorldPartition/WorldPartitionBuilder.cpp b/Engine/Source/Editor/UnrealEd/Private/WorldPartition/WorldPartitionBuilder.cpp
--- a/Engine/Source/Editor/UnrealEd/Private/WorldPartition/WorldPartitionBuilder.cpp
+++ b/Engine/Source/Editor/UnrealEd/Private/WorldPartition/WorldPartitionBuilder.cpp
@@ -7,7 +7,11 @@
 #include "DistanceFieldAtlas.h"
 #include "MeshCardRepresentation.h"
 #include "StaticMeshCompiler.h"
-
+#include "Engine/World.h"
+#include "WorldPartition/WorldPartition.h"
+#include "WorldPartition/DataLayer/WorldDataLayers.h"
+#include "WorldPartition/DataLayer/DataLayer.h"
+#include "Math/IntVector.h"
 
 DEFINE_LOG_CATEGORY_STATIC(LogWorldPartitionBuilder, Log, All);
 
@@ -16,10 +20,120 @@ UWorldPartitionBuilder::UWorldPartitionBuilder(const FObjectInitializer& ObjectI
 	: Super(ObjectInitializer)
 {}
 
+bool UWorldPartitionBuilder::Run(UWorld* World, FPackageSourceControlHelper& PackageHelper)
+{
+	UWorldPartition* WorldPartition = World->GetWorldPartition();
+
+	// Properly Setup DataLayers for Builder
+	if (const AWorldDataLayers* WorldDataLayers = AWorldDataLayers::Get(World))
+	{
+		// Load Data Layers
+		bool bUpdateEditorCells = false;
+		WorldDataLayers->ForEachDataLayer([&bUpdateEditorCells, this](UDataLayer* DataLayer)
+		{
+			// Load all Non DynamicallyLoaded Data Layers + Initially Active Data Layers + Data Layers provided by builder
+			const bool bLoadedInEditor = (bLoadNonDynamicDataLayers && !DataLayer->IsDynamicallyLoaded()) || 
+										 (bLoadInitiallyActiveDataLayers && DataLayer->GetInitialState() == EDataLayerState::Activated) ||
+										 DataLayerLabels.Contains(DataLayer->GetDataLayerLabel());
+			if (DataLayer->IsDynamicallyLoadedInEditor() != bLoadedInEditor)
+			{
+				bUpdateEditorCells = true;
+				DataLayer->SetIsDynamicallyLoadedInEditor(bLoadedInEditor);
+				if (RequiresCommandletRendering() && bLoadedInEditor)
+				{
+					DataLayer->SetIsInitiallyVisible(true);
+				}
+			}
+			
+			UE_LOG(LogWorldPartitionBuilder, Display, TEXT("DataLayer '%s' Loaded: %d"), *UDataLayer::GetDataLayerText(DataLayer).ToString(), bLoadedInEditor ? 1 : 0);
+			
+			return true;
+		});
+
+		if (bUpdateEditorCells)
+		{
+			UE_LOG(LogWorldPartitionBuilder, Display, TEXT("DataLayer load state changed refreshing editor cells"));
+			WorldPartition->RefreshLoadedEditorCells();
+		}
+	}
+		
+	ELoadingMode LoadingMode = GetLoadingMode();
+	if (LoadingMode == ELoadingMode::IterativeCells)
+	{
+		// do partial loading loop that calls RunInternal
+		FBox EditorBounds = WorldPartition->GetEditorWorldBounds();
+
+		auto GetCellCoord = [](FVector InPos, int32 InCellSize)
+		{
+			return FIntVector(
+				FMath::FloorToInt(InPos.X / InCellSize),
+				FMath::FloorToInt(InPos.Y / InCellSize),
+				FMath::FloorToInt(InPos.Z / InCellSize)
+			);
+		};
+
+		const FIntVector MinCellCoords = GetCellCoord(EditorBounds.Min, IterativeCellSize);
+		const FIntVector MaxCellCoords = GetCellCoord(EditorBounds.Max, IterativeCellSize);
+
+		UE_LOG(LogWorldPartitionBuilder, Display, TEXT("Iterative Cell Mode"));
+		UE_LOG(LogWorldPartitionBuilder, Display, TEXT("Cell Size %d"), IterativeCellSize);
+		UE_LOG(LogWorldPartitionBuilder, Display, TEXT("Cell Overlap %d"), IterativeCellOverlapSize);
+		UE_LOG(LogWorldPartitionBuilder, Display, TEXT("WorldBounds: Min %s, Max %s"), *EditorBounds.Min.ToString(), *EditorBounds.Max.ToString());
+		UE_LOG(LogWorldPartitionBuilder, Display, TEXT("Iteration Count: %d"), (MaxCellCoords.Z-MinCellCoords.Z) * (MaxCellCoords.Y-MinCellCoords.Y) * (MaxCellCoords.X - MinCellCoords.X));
+		
+
+		FBox LoadedBounds(ForceInit);
+		for (int32 z = MinCellCoords.Z; z <= MaxCellCoords.Z; z++)
+		{
+			for (int32 y = MinCellCoords.Y; y <= MaxCellCoords.Y; y++)
+			{
+				for (int32 x = MinCellCoords.X; x <= MaxCellCoords.X; x++)
+				{
+					const FVector Min(x * IterativeCellSize, y * IterativeCellSize, z * IterativeCellSize);
+					const FVector Max = Min + FVector(IterativeCellSize);
+					FBox BoundsToLoad(Min, Max);
+					BoundsToLoad.ExpandBy(IterativeCellOverlapSize);
+
+					UE_LOG(LogWorldPartitionBuilder, Verbose, TEXT("Loading Bounds: Min %s, Max %s"), *BoundsToLoad.Min.ToString(), *BoundsToLoad.Max.ToString());
+					WorldPartition->LoadEditorCells(BoundsToLoad);
+					LoadedBounds += BoundsToLoad;
+
+
+					if (!RunInternal(World, BoundsToLoad, PackageHelper))
+					{
+						return false;
+					}
+
+					if (HasExceededMaxMemory())
+					{
+						WorldPartition->UnloadEditorCells(LoadedBounds);
+						// Reset Loaded Bounds
+						LoadedBounds.Init();
+						DoCollectGarbage();
+					}
+				}
+			}
+		}
+
+		return true;
+	}
+	else
+	{
+		FBox BoundsToLoad(ForceInit);
+		if (LoadingMode == ELoadingMode::EntireWorld)
+		{
+			BoundsToLoad += FBox(FVector(-WORLD_MAX, -WORLD_MAX, -WORLD_MAX), FVector(WORLD_MAX, WORLD_MAX, WORLD_MAX));
+			WorldPartition->LoadEditorCells(BoundsToLoad);
+		}
+
+		return RunInternal(World, BoundsToLoad, PackageHelper);
+	}
+}
+
 bool UWorldPartitionBuilder::HasExceededMaxMemory() const
 {
 	const uint64 MemoryMinFreePhysical = 1024ll * 1024 * 1024;
-	const uint64 MemoryMaxUsedPhysical = 16384ll * 1024 * 1024l;
+	const uint64 MemoryMaxUsedPhysical = 32768ll * 1024 * 1024l;
 	const FPlatformMemoryStats MemStats = FPlatformMemory::GetStats();
 	return (MemStats.AvailablePhysical < MemoryMinFreePhysical) || (MemStats.UsedPhysical >= MemoryMaxUsedPhysical);
 };
diff --git a/Engine/Source/Editor/UnrealEd/Private/WorldPartition/WorldPartitionHLODsBuilder.cpp b/Engine/Source/Editor/UnrealEd/Private/WorldPartition/WorldPartitionHLODsBuilder.cpp
--- a/Engine/Source/Editor/UnrealEd/Private/WorldPartition/WorldPartitionHLODsBuilder.cpp
+++ b/Engine/Source/Editor/UnrealEd/Private/WorldPartition/WorldPartitionHLODsBuilder.cpp
@@ -242,7 +242,7 @@ bool UWorldPartitionHLODsBuilder::PreWorldInitialization(FPackageSourceControlHe
 	return bRet;
 }
 
-bool UWorldPartitionHLODsBuilder::Run(UWorld* World, FPackageSourceControlHelper& PackageHelper)
+bool UWorldPartitionHLODsBuilder::RunInternal(UWorld* World, const FBox& Bounds, FPackageSourceControlHelper& PackageHelper)
 {
 	WorldPartition = World->GetWorldPartition();
 	check(WorldPartition);
diff --git a/Engine/Source/Editor/UnrealEd/Private/WorldPartition/WorldPartitionMiniMapBuilder.cpp b/Engine/Source/Editor/UnrealEd/Private/WorldPartition/WorldPartitionMiniMapBuilder.cpp
--- a/Engine/Source/Editor/UnrealEd/Private/WorldPartition/WorldPartitionMiniMapBuilder.cpp
+++ b/Engine/Source/Editor/UnrealEd/Private/WorldPartition/WorldPartitionMiniMapBuilder.cpp
@@ -30,7 +30,7 @@ UWorldPartitionMiniMapBuilder::UWorldPartitionMiniMapBuilder(const FObjectInitia
 	bAutoSubmit = bUseOnlyHLODs = FParse::Param(FCommandLine::Get(), TEXT("AutoSubmit"));
 }
 
-bool UWorldPartitionMiniMapBuilder::Run(UWorld* World, FPackageSourceControlHelper& PackageHelper)
+bool UWorldPartitionMiniMapBuilder::RunInternal(UWorld* World, const FBox& Bounds, FPackageSourceControlHelper& PackageHelper)
 {
 	AWorldPartitionMiniMap* WorldMiniMap = FWorldPartitionMiniMapHelper::GetWorldPartitionMiniMap(World,true);
 	if (!WorldMiniMap)
diff --git a/Engine/Source/Editor/UnrealEd/Private/WorldPartition/WorldPartitionNavigationDataBuilder.cpp b/Engine/Source/Editor/UnrealEd/Private/WorldPartition/WorldPartitionNavigationDataBuilder.cpp
--- a/Engine/Source/Editor/UnrealEd/Private/WorldPartition/WorldPartitionNavigationDataBuilder.cpp
+++ b/Engine/Source/Editor/UnrealEd/Private/WorldPartition/WorldPartitionNavigationDataBuilder.cpp
@@ -18,7 +18,7 @@ UWorldPartitionNavigationDataBuilder::UWorldPartitionNavigationDataBuilder(const
 {
 }
 
-bool UWorldPartitionNavigationDataBuilder::Run(UWorld* World, FPackageSourceControlHelper& PackageHelper)
+bool UWorldPartitionNavigationDataBuilder::RunInternal(UWorld* World, const FBox& Bounds, FPackageSourceControlHelper& PackageHelper)
 {
 	UWorldPartitionSubsystem* WorldPartitionSubsystem = World->GetSubsystem<UWorldPartitionSubsystem>();
 	check(WorldPartitionSubsystem);
diff --git a/Engine/Source/Editor/UnrealEd/Private/WorldPartition/WorldPartitionResaveActorsBuilder.cpp b/Engine/Source/Editor/UnrealEd/Private/WorldPartition/WorldPartitionResaveActorsBuilder.cpp
--- a/Engine/Source/Editor/UnrealEd/Private/WorldPartition/WorldPartitionResaveActorsBuilder.cpp
+++ b/Engine/Source/Editor/UnrealEd/Private/WorldPartition/WorldPartitionResaveActorsBuilder.cpp
@@ -16,7 +16,7 @@ UWorldPartitionResaveActorsBuilder::UWorldPartitionResaveActorsBuilder(const FOb
 	FParse::Value(FCommandLine::Get(), TEXT("ActorClass="), ActorClassName);
 }
 
-bool UWorldPartitionResaveActorsBuilder::Run(UWorld* World, FPackageSourceControlHelper& PackageHelper)
+bool UWorldPartitionResaveActorsBuilder::RunInternal(UWorld* World, const FBox& Bounds, FPackageSourceControlHelper& PackageHelper)
 {
 	UWorldPartition* WorldPartition = World->GetWorldPartition();
 	check(WorldPartition);
diff --git a/Engine/Source/Editor/UnrealEd/Public/WorldPartition/WorldPartitionBuilder.h b/Engine/Source/Editor/UnrealEd/Public/WorldPartition/WorldPartitionBuilder.h
--- a/Engine/Source/Editor/UnrealEd/Public/WorldPartition/WorldPartitionBuilder.h
+++ b/Engine/Source/Editor/UnrealEd/Public/WorldPartition/WorldPartitionBuilder.h
@@ -11,13 +11,29 @@ class UNREALED_API UWorldPartitionBuilder : public UObject
 	GENERATED_UCLASS_BODY()
 
 public:
+	enum ELoadingMode
+	{
+		Custom,
+		EntireWorld,
+		IterativeCells
+	};
+
 	virtual bool RequiresCommandletRendering() const PURE_VIRTUAL(UWorldPartitionBuilder::RequiresCommandletRendering, return false;);
-	virtual bool RequiresEntireWorldLoading() const PURE_VIRTUAL(UWorldPartitionBuilder::RequiresEntireWorldLoading, return false;);
-	virtual bool Run(UWorld* World, FPackageSourceControlHelper& PackageHelper) PURE_VIRTUAL(UWorldPartitionBuilder::Run, return false;);
+	virtual ELoadingMode GetLoadingMode() const PURE_VIRTUAL(UWorldPartitionBuilder::GetLoadingMode, return ELoadingMode::Custom;);
+	
+	bool Run(UWorld* World, FPackageSourceControlHelper& PackageHelper);
 
 	virtual bool PreWorldInitialization(FPackageSourceControlHelper& PackageHelper) { return true; }
 
 protected:
+	virtual bool RunInternal(UWorld* World, const FBox& Bounds, FPackageSourceControlHelper& PackageHelper) PURE_VIRTUAL(UWorldPartition::RunInternal, return false;);
+
 	bool HasExceededMaxMemory() const;
 	void DoCollectGarbage() const;
+
+	int32 IterativeCellSize = 102400;
+	int32 IterativeCellOverlapSize = 0;
+	TSet<FName> DataLayerLabels;
+	bool bLoadNonDynamicDataLayers = true;
+	bool bLoadInitiallyActiveDataLayers = true;
 };
\ No newline at end of file
diff --git a/Engine/Source/Editor/UnrealEd/Public/WorldPartition/WorldPartitionHLODsBuilder.h b/Engine/Source/Editor/UnrealEd/Public/WorldPartition/WorldPartitionHLODsBuilder.h
--- a/Engine/Source/Editor/UnrealEd/Public/WorldPartition/WorldPartitionHLODsBuilder.h
+++ b/Engine/Source/Editor/UnrealEd/Public/WorldPartition/WorldPartitionHLODsBuilder.h
@@ -13,12 +13,12 @@ class UWorldPartitionHLODsBuilder : public UWorldPartitionBuilder
 public:
 	// UWorldPartitionBuilder interface begin
 	virtual bool RequiresCommandletRendering() const override;
-	virtual bool RequiresEntireWorldLoading() const override { return false; }
+	virtual ELoadingMode GetLoadingMode() const override { return ELoadingMode::Custom; }
 	virtual bool PreWorldInitialization(FPackageSourceControlHelper& PackageHelper) override;
-	virtual bool Run(UWorld* World, FPackageSourceControlHelper& PackageHelper) override;
+protected:
+	virtual bool RunInternal(UWorld* World, const FBox& Bounds, FPackageSourceControlHelper& PackageHelper) override;
 	// UWorldPartitionBuilder interface end
 
-protected:
 	bool IsDistributedBuild() const { return bDistributedBuild; }
 	bool IsUsingBuildManifest() const { return !BuildManifest.IsEmpty(); }
 	bool ValidateParams() const;
diff --git a/Engine/Source/Editor/UnrealEd/Public/WorldPartition/WorldPartitionMiniMapBuilder.h b/Engine/Source/Editor/UnrealEd/Public/WorldPartition/WorldPartitionMiniMapBuilder.h
--- a/Engine/Source/Editor/UnrealEd/Public/WorldPartition/WorldPartitionMiniMapBuilder.h
+++ b/Engine/Source/Editor/UnrealEd/Public/WorldPartition/WorldPartitionMiniMapBuilder.h
@@ -13,8 +13,9 @@ class UWorldPartitionMiniMapBuilder : public UWorldPartitionBuilder
 public:
 	// UWorldPartitionBuilder interface begin
 	virtual bool RequiresCommandletRendering() const override { return true; }
-	virtual bool RequiresEntireWorldLoading() const override { return !bUseOnlyHLODs; }
-	virtual bool Run(UWorld* World, FPackageSourceControlHelper& PackageHelper) override;
+	virtual ELoadingMode GetLoadingMode() const override { return bUseOnlyHLODs ? ELoadingMode::Custom : ELoadingMode::EntireWorld; }
+protected:
+	virtual bool RunInternal(UWorld* World, const FBox& Bounds, FPackageSourceControlHelper& PackageHelper) override;
 	// UWorldPartitionBuilder interface end
 
 	UPROPERTY(config)
diff --git a/Engine/Source/Editor/UnrealEd/Public/WorldPartition/WorldPartitionNavigationDataBuilder.h b/Engine/Source/Editor/UnrealEd/Public/WorldPartition/WorldPartitionNavigationDataBuilder.h
--- a/Engine/Source/Editor/UnrealEd/Public/WorldPartition/WorldPartitionNavigationDataBuilder.h
+++ b/Engine/Source/Editor/UnrealEd/Public/WorldPartition/WorldPartitionNavigationDataBuilder.h
@@ -13,7 +13,9 @@ class UWorldPartitionNavigationDataBuilder : public UWorldPartitionBuilder
 public:
 	// UWorldPartitionBuilder interface begin
 	virtual bool RequiresCommandletRendering() const override { return false; }
-	virtual bool RequiresEntireWorldLoading() const override { return true; }
-	virtual bool Run(UWorld* World, FPackageSourceControlHelper& PackageHelper) override;
+	virtual ELoadingMode GetLoadingMode() const override { return ELoadingMode::EntireWorld; }
+
+protected:
+	virtual bool RunInternal(UWorld* World, const FBox& Bounds, FPackageSourceControlHelper& PackageHelper) override;
 	// UWorldPartitionBuilder interface end
 };
diff --git a/Engine/Source/Editor/UnrealEd/Public/WorldPartition/WorldPartitionResaveActorsBuilder.h b/Engine/Source/Editor/UnrealEd/Public/WorldPartition/WorldPartitionResaveActorsBuilder.h
--- a/Engine/Source/Editor/UnrealEd/Public/WorldPartition/WorldPartitionResaveActorsBuilder.h
+++ b/Engine/Source/Editor/UnrealEd/Public/WorldPartition/WorldPartitionResaveActorsBuilder.h
@@ -14,8 +14,8 @@ class UWorldPartitionResaveActorsBuilder : public UWorldPartitionBuilder
 public:
 	// UWorldPartitionBuilder interface begin
 	virtual bool RequiresCommandletRendering() const override { return false; }
-	virtual bool RequiresEntireWorldLoading() const override { return false; }
-	virtual bool Run(UWorld* World, FPackageSourceControlHelper& PackageHelper) override;
+	virtual ELoadingMode GetLoadingMode() const override { return ELoadingMode::Custom; }
+	virtual bool RunInternal(UWorld* World, const FBox& Bounds, FPackageSourceControlHelper& PackageHelper) override;
 	// UWorldPartitionBuilder interface end
 
 private:
