Removed selection annotations

These existed as an optimization for USelection to quickly know if an object was selected, without having to test its array. UTypedElementList has that optimization embedded, so at this point the selection annotations are just extra state to keep in-sync, and failure to do so has already lead to bugs.

This change removes them in the simplest way possible, by using a function pointer to callback into editor code from core/engine code (as core/engine cannot directly access the global selection state), though longer-term we probably want to remove IsSelected/IsSelectedInEditor from UObject and let specific systems provide that functionality if they need it (eg, AActor/UActorComponent could provide it by getting the selection set from the owner world, rather than the global state).

diff --git a/Engine/Source/Editor/LevelEditor/Private/Elements/Actor/ActorElementLevelEditorSelectionCustomization.cpp b/Engine/Source/Editor/LevelEditor/Private/Elements/Actor/ActorElementLevelEditorSelectionCustomization.cpp
--- a/Engine/Source/Editor/LevelEditor/Private/Elements/Actor/ActorElementLevelEditorSelectionCustomization.cpp
+++ b/Engine/Source/Editor/LevelEditor/Private/Elements/Actor/ActorElementLevelEditorSelectionCustomization.cpp
@@ -178,9 +178,6 @@ bool FActorElementLevelEditorSelectionCustomization::SelectActorElement(const TT
 	}
 
 	UE_LOG(LogActorLevelEditorSelection, Verbose, TEXT("Selected Actor: %s (%s)"), *Actor->GetActorLabel(), *Actor->GetClass()->GetName());
-
-	// Update the annotation state
-	GSelectedActorAnnotation.Set(Actor);
 	
 	// Bind the override delegates for the components on the selected actor
 	for (UActorComponent* Component : Actor->GetComponents())
@@ -246,9 +243,6 @@ bool FActorElementLevelEditorSelectionCustomization::DeselectActorElement(const
 	
 	UE_LOG(LogActorLevelEditorSelection, Verbose, TEXT("Deselected Actor: %s (%s)"), *Actor->GetActorLabel(), *Actor->GetClass()->GetName());
 	
-	// Update the annotation state
-	GSelectedActorAnnotation.Clear(Actor);
-	
 	// Deselect and unbind the override delegates for the components on the selected actor
 	{
 		FTypedElementListLegacySyncScopedBatch LegacySyncBatch(InSelectionSet, InSelectionOptions.AllowLegacyNotifications());
diff --git a/Engine/Source/Editor/LevelEditor/Private/Elements/Component/ComponentElementLevelEditorSelectionCustomization.cpp b/Engine/Source/Editor/LevelEditor/Private/Elements/Component/ComponentElementLevelEditorSelectionCustomization.cpp
--- a/Engine/Source/Editor/LevelEditor/Private/Elements/Component/ComponentElementLevelEditorSelectionCustomization.cpp
+++ b/Engine/Source/Editor/LevelEditor/Private/Elements/Component/ComponentElementLevelEditorSelectionCustomization.cpp
@@ -111,9 +111,6 @@ bool FComponentElementLevelEditorSelectionCustomization::SelectComponentElement(
 
 	UE_LOG(LogComponentLevelEditorSelection, Verbose, TEXT("Selected Component: %s (%s)"), *Component->GetPathName(), *Component->GetClass()->GetName());
 	
-	// Update the annotation state
-	GSelectedComponentAnnotation.Set(Component);
-	
 	// Make sure the override delegate is bound properly
 	if (USceneComponent* SceneComponent = Cast<USceneComponent>(Component))
 	{
@@ -149,9 +146,6 @@ bool FComponentElementLevelEditorSelectionCustomization::DeselectComponentElemen
 	}
 
 	UE_LOG(LogComponentLevelEditorSelection, Verbose, TEXT("Deselected Component: %s (%s)"), *Component->GetPathName(), *Component->GetClass()->GetName());
-
-	// Update the annotation state
-	GSelectedComponentAnnotation.Clear(Component);
 	
 	// Update the selection visualization
 	if (AActor* ComponentOwner = Component->GetOwner())
diff --git a/Engine/Source/Editor/UnrealEd/Private/EditorEngine.cpp b/Engine/Source/Editor/UnrealEd/Private/EditorEngine.cpp
--- a/Engine/Source/Editor/UnrealEd/Private/EditorEngine.cpp
+++ b/Engine/Source/Editor/UnrealEd/Private/EditorEngine.cpp
@@ -133,7 +133,6 @@
 #include "Editor/ActorPositioning.h"
 
 #include "Elements/Framework/TypedElementSelectionSet.h"
-#include "Elements/Object/ObjectElementEditorSelectionCustomization.h"
 #include "Elements/Actor/ActorElementEditorViewportInteractionCustomization.h"
 #include "Elements/Component/ComponentElementEditorViewportInteractionCustomization.h"
 
@@ -273,10 +272,22 @@ void InitSelectionSets()
 
 	GObjectSelection = USelection::CreateObjectSelection(GetTransientPackage(), TEXT("SelectedObjects"), RF_Transactional);
 	GObjectSelection->AddToRoot();
+	GObjectSelection->SetElementSelectionSet(NewObject<UTypedElementSelectionSet>(GObjectSelection, NAME_None, RF_Transactional));
 
-	UTypedElementSelectionSet* ObjectSelectionSet = NewObject<UTypedElementSelectionSet>(GObjectSelection, NAME_None, RF_Transactional);
-	ObjectSelectionSet->RegisterInterfaceCustomizationByTypeName(NAME_Object, MakeUnique<FObjectElementEditorSelectionCustomization>());
-	GObjectSelection->SetElementSelectionSet(ObjectSelectionSet);
+	GIsActorSelectedInEditor = [](const AActor* InActor)
+	{
+		return GActorSelection->IsSelected(InActor);
+	};
+
+	GIsComponentSelectedInEditor = [](const UActorComponent* InComponent)
+	{
+		return GComponentSelection->IsSelected(InComponent);
+	};
+
+	GIsObjectSelectedInEditor = [](const UObject* InObject)
+	{
+		return GObjectSelection->IsSelected(InObject);
+	};
 }
 
 void DestroySelectionSets()
@@ -298,6 +309,10 @@ void DestroySelectionSets()
 		}
 	}
 
+	GIsActorSelectedInEditor = nullptr;
+	GIsComponentSelectedInEditor = nullptr;
+	GIsObjectSelectedInEditor = nullptr;
+
 	GActorSelection = nullptr;
 	GComponentSelection = nullptr;
 	GObjectSelection = nullptr;
diff --git a/Engine/Source/Editor/UnrealEd/Private/Elements/Object/ObjectElementEditorSelectionCustomization.cpp b/Engine/Source/Editor/UnrealEd/Private/Elements/Object/ObjectElementEditorSelectionCustomization.cpp
deleted file mode 100644
diff --git a/Engine/Source/Editor/UnrealEd/Private/Elements/Object/ObjectElementEditorSelectionCustomization.h b/Engine/Source/Editor/UnrealEd/Private/Elements/Object/ObjectElementEditorSelectionCustomization.h
deleted file mode 100644
diff --git a/Engine/Source/Runtime/CoreUObject/Private/UObject/Obj.cpp b/Engine/Source/Runtime/CoreUObject/Private/UObject/Obj.cpp
--- a/Engine/Source/Runtime/CoreUObject/Private/UObject/Obj.cpp
+++ b/Engine/Source/Runtime/CoreUObject/Private/UObject/Obj.cpp
@@ -589,7 +589,7 @@ TSharedPtr<ITransactionObjectAnnotation> UObject::CreateAndRestoreTransactionAnn
 
 bool UObject::IsSelectedInEditor() const
 {
-	return !IsPendingKill() && GSelectedObjectAnnotation.Get(this);
+	return !IsPendingKill() && GIsObjectSelectedInEditor && GIsObjectSelectedInEditor(this);
 }
 
 #endif // WITH_EDITOR
@@ -4716,3 +4716,7 @@ bool UObject::IsSupportedForNetworking() const
 {
 	return IsFullNameStableForNetworking();
 }
+
+#if WITH_EDITOR
+TFunction<bool(const UObject*)> GIsObjectSelectedInEditor;
+#endif
diff --git a/Engine/Source/Runtime/CoreUObject/Private/UObject/UObjectGlobals.cpp b/Engine/Source/Runtime/CoreUObject/Private/UObject/UObjectGlobals.cpp
--- a/Engine/Source/Runtime/CoreUObject/Private/UObject/UObjectGlobals.cpp
+++ b/Engine/Source/Runtime/CoreUObject/Private/UObject/UObjectGlobals.cpp
@@ -34,7 +34,6 @@
 #include "UObject/UnrealType.h"
 #include "UObject/ObjectRedirector.h"
 #include "UObject/PackageResourceManager.h"
-#include "UObject/UObjectAnnotation.h"
 #include "Serialization/DuplicatedObject.h"
 #include "Serialization/DuplicatedDataReader.h"
 #include "Serialization/DuplicatedDataWriter.h"
@@ -79,9 +78,6 @@ static FAutoConsoleVariableRef CVarAllowUnversionedContentInEditor(
 	ECVF_Default
 );
 
-/** Object annotation used by the engine to keep track of which objects are selected */
-FUObjectAnnotationSparseBool GSelectedObjectAnnotation;
-
 DEFINE_STAT(STAT_InitProperties);
 DEFINE_STAT(STAT_ConstructObject);
 DEFINE_STAT(STAT_AllocateObject);
diff --git a/Engine/Source/Runtime/CoreUObject/Public/UObject/Object.h b/Engine/Source/Runtime/CoreUObject/Public/UObject/Object.h
--- a/Engine/Source/Runtime/CoreUObject/Public/UObject/Object.h
+++ b/Engine/Source/Runtime/CoreUObject/Public/UObject/Object.h
@@ -1598,3 +1598,7 @@ FORCEINLINE bool IsValid(const UObject *Test)
 	return Test && !Test->IsPendingKill();
 }
 
+#if WITH_EDITOR
+/** Callback for editor object selection. This must be in core instead of editor for UObject::IsSelectedInEditor to work */
+extern COREUOBJECT_API TFunction<bool(const UObject*)> GIsObjectSelectedInEditor;
+#endif
diff --git a/Engine/Source/Runtime/CoreUObject/Public/UObject/UObjectAnnotation.h b/Engine/Source/Runtime/CoreUObject/Public/UObject/UObjectAnnotation.h
--- a/Engine/Source/Runtime/CoreUObject/Public/UObject/UObjectAnnotation.h
+++ b/Engine/Source/Runtime/CoreUObject/Public/UObject/UObjectAnnotation.h
@@ -1308,9 +1308,3 @@ private:
 	TArray<TBitType> AnnotationArray;
 
 };
-
-
-
-// Definition is in UObjectGlobals.cpp
-extern COREUOBJECT_API FUObjectAnnotationSparseBool GSelectedObjectAnnotation;
-
diff --git a/Engine/Source/Runtime/Engine/Classes/Components/ActorComponent.h b/Engine/Source/Runtime/Engine/Classes/Components/ActorComponent.h
--- a/Engine/Source/Runtime/Engine/Classes/Components/ActorComponent.h
+++ b/Engine/Source/Runtime/Engine/Classes/Components/ActorComponent.h
@@ -13,7 +13,6 @@
 #include "UObject/ScriptMacros.h"
 #include "EdGraph/EdGraphPin.h"
 #include "Interfaces/Interface_AssetUserData.h"
-#include "UObject/UObjectAnnotation.h"
 #include "UObject/StructOnScope.h"
 #include "ComponentInstanceDataCache.h"
 #include "ActorComponent.generated.h"
@@ -51,9 +50,6 @@ private:
 #if WITH_EDITOR
 class SWidget;
 struct FMinimalViewInfo;
-
-/** Annotation for component selection.  This must be in engine instead of editor for ::IsSelected to work */
-extern ENGINE_API FUObjectAnnotationSparseBool GSelectedComponentAnnotation;
 #endif
 
 /** Information about how to update transform when something is moved */
@@ -1127,3 +1123,8 @@ FORCEINLINE_DEBUGGABLE AActor* UActorComponent::GetOwner() const
 	return OwnerPrivate;
 #endif
 }
+
+#if WITH_EDITOR
+/** Callback for editor component selection. This must be in engine instead of editor for UActorComponent::IsSelectedInEditor to work */
+extern ENGINE_API TFunction<bool(const UActorComponent*)> GIsComponentSelectedInEditor;
+#endif
diff --git a/Engine/Source/Runtime/Engine/Classes/GameFramework/Actor.h b/Engine/Source/Runtime/Engine/Classes/GameFramework/Actor.h
--- a/Engine/Source/Runtime/Engine/Classes/GameFramework/Actor.h
+++ b/Engine/Source/Runtime/Engine/Classes/GameFramework/Actor.h
@@ -101,11 +101,6 @@ DECLARE_EVENT_TwoParams(AActor, FActorOnPackagingModeChanged, AActor*, bool /* b
 DECLARE_DELEGATE_RetVal_ThreeParams(bool, FOnProcessEvent, AActor*, UFunction*, void*);
 #endif
 
-#if WITH_EDITOR
-/** Annotation for actor selection.  This must be in engine instead of editor for ::IsSelected to work */
-extern ENGINE_API FUObjectAnnotationSparseBool GSelectedActorAnnotation;
-#endif
-
 /**
  * TInlineComponentArray is simply a TArray that reserves a fixed amount of space on the stack
  * to try to avoid heap allocation when there are fewer than a specified number of elements expected in the result.
@@ -3927,6 +3922,11 @@ FORCEINLINE_DEBUGGABLE bool AActor::IsNetMode(ENetMode Mode) const
 #endif
 }
 
+#if WITH_EDITOR
+/** Callback for editor actor selection. This must be in engine instead of editor for AActor::IsSelectedInEditor to work */
+extern ENGINE_API TFunction<bool(const AActor*)> GIsActorSelectedInEditor;
+#endif
+
 DEFINE_ACTORDESC_TYPE(AActor, FWorldPartitionActorDesc);
 
 //////////////////////////////////////////////////////////////////////////
diff --git a/Engine/Source/Runtime/Engine/Private/Actor.cpp b/Engine/Source/Runtime/Engine/Private/Actor.cpp
--- a/Engine/Source/Runtime/Engine/Private/Actor.cpp
+++ b/Engine/Source/Runtime/Engine/Private/Actor.cpp
@@ -73,10 +73,6 @@ DECLARE_CYCLE_STAT(TEXT("PostActorConstruction"), STAT_PostActorConstruction, ST
 
 FMakeNoiseDelegate AActor::MakeNoiseDelegate = FMakeNoiseDelegate::CreateStatic(&AActor::MakeNoiseImpl);
 
-#if WITH_EDITOR
-FUObjectAnnotationSparseBool GSelectedActorAnnotation;
-#endif
-
 #if !UE_BUILD_SHIPPING
 FOnProcessEvent AActor::ProcessEventDelegate;
 #endif
@@ -3203,7 +3199,7 @@ void AActor::PostEditImport()
 
 bool AActor::IsSelectedInEditor() const
 {
-	return !IsPendingKill() && GSelectedActorAnnotation.Get(this);
+	return !IsPendingKill() && GIsActorSelectedInEditor && GIsActorSelectedInEditor(this);
 }
 
 bool AActor::SupportsExternalPackaging() const
@@ -5575,4 +5571,8 @@ void AActor::SetInstigator(APawn* InInstigator)
 	MARK_PROPERTY_DIRTY_FROM_NAME(AActor, Instigator, this);
 }
 
+#if WITH_EDITOR
+TFunction<bool(const AActor*)> GIsActorSelectedInEditor;
+#endif
+
 #undef LOCTEXT_NAMESPACE
diff --git a/Engine/Source/Runtime/Engine/Private/Components/ActorComponent.cpp b/Engine/Source/Runtime/Engine/Private/Components/ActorComponent.cpp
--- a/Engine/Source/Runtime/Engine/Private/Components/ActorComponent.cpp
+++ b/Engine/Source/Runtime/Engine/Private/Components/ActorComponent.cpp
@@ -75,10 +75,6 @@ FAutoConsoleVariableRef GTickComponentLatentActionsWithTheComponentCVar(
 /** Enable to log out all render state create, destroy and updatetransform events */
 #define LOG_RENDER_STATE 0
 
-#if WITH_EDITOR
-FUObjectAnnotationSparseBool GSelectedComponentAnnotation;
-#endif
-
 /** Static var indicating activity of reregister context */
 int32 FGlobalComponentReregisterContext::ActiveGlobalReregisterContextCount = 0;
 
@@ -900,7 +896,7 @@ void UActorComponent::PostEditUndo()
 
 bool UActorComponent::IsSelectedInEditor() const
 {
-	return !IsPendingKill() && GSelectedComponentAnnotation.Get(this);
+	return !IsPendingKill() && GIsComponentSelectedInEditor && GIsComponentSelectedInEditor(this);
 }
 
 void UActorComponent::ConsolidatedPostEditChange(const FPropertyChangedEvent& PropertyChangedEvent)
@@ -2233,4 +2229,8 @@ bool UActorComponent::NeedsInitialization() const
 	return HasAnyFlags(RF_NeedInitialization);
 }
 
+#if WITH_EDITOR
+TFunction<bool(const UActorComponent*)> GIsComponentSelectedInEditor;
+#endif
+
 #undef LOCTEXT_NAMESPACE
