Fixed crash due to actors being destroyed in a complex copy scenario

This is a band-aid, as actors shouldn't be destroyed when only being copied. It's likely that MoveActorsToLevel is causing the destruction somehow, but regardless, all of this code should be (and needs to be for typed elements) updated to avoid messing with the selection state.

diff --git a/Engine/Source/Editor/UnrealEd/Private/EditorServer.cpp b/Engine/Source/Editor/UnrealEd/Private/EditorServer.cpp
--- a/Engine/Source/Editor/UnrealEd/Private/EditorServer.cpp
+++ b/Engine/Source/Editor/UnrealEd/Private/EditorServer.cpp
@@ -3455,7 +3455,7 @@ void UEditorEngine::CopySelectedActorsToClipboard( UWorld* InWorld, bool bShould
 		}
 
 		// Take a note of the current selection, so it can be restored at the end of this process
-		TArray<AActor*> CurrentlySelectedActors;
+		TArray<TWeakObjectPtr<AActor>> CurrentlySelectedActors;
 		for ( FSelectionIterator It( GetSelectedActorIterator() ) ; It ; ++It )
 		{
 			AActor* Actor = static_cast<AActor*>( *It );
@@ -3582,9 +3582,12 @@ void UEditorEngine::CopySelectedActorsToClipboard( UWorld* InWorld, bool bShould
 
 		// Restore old selection
 		GEditor->SelectNone( false, true );
-		for (auto& Actor : CurrentlySelectedActors)
+		for (const TWeakObjectPtr<AActor>& Actor : CurrentlySelectedActors)
 		{
-			GEditor->SelectActor( Actor, true, false );
+			if (AActor* ActorPtr = Actor.Get())
+			{
+				GEditor->SelectActor(ActorPtr, true, false);
+			}
 		}
 	}
 }
