WP HLOD Builder - Ensure all files modified by the source control helper are taken into account
HLODs built in "single step" mode weren't getting submitted

diff --git a/Engine/Source/Editor/UnrealEd/Private/WorldPartition/WorldPartitionHLODsBuilder.cpp b/Engine/Source/Editor/UnrealEd/Private/WorldPartition/WorldPartitionHLODsBuilder.cpp
--- a/Engine/Source/Editor/UnrealEd/Private/WorldPartition/WorldPartitionHLODsBuilder.cpp
+++ b/Engine/Source/Editor/UnrealEd/Private/WorldPartition/WorldPartitionHLODsBuilder.cpp
@@ -493,6 +493,9 @@ bool UWorldPartitionHLODsBuilder::SubmitHLODActors()
 {
 	bool bRet = true;
 
+	// Ensure all files modified by the source control helper are taken into account
+	ModifiedFiles.Append(SourceControlHelper->GetModifiedFiles());
+
 	// Check in all modified files
 	if (ModifiedFiles.Num() > 0)
 	{
