BP reinstantiation fixes
-  Don't go through AActor::Rename in ReplaceActorHelper because we aren't changing outer  and we also don't want to reset loaders if the actor is using an external package. We really just want to rename that actor out of the way so we can spawn the new one in the exact same package, keeping the package name intact (which depends on the main package actor name).
- Add missing call to patch the new object into the old object linker's export map when reinstantiating; subsequent loads may import this entry and we need to make sure that it returns the new object instead of the old one (fix done by [REDACTED_PERSON]).

diff --git a/Engine/Source/Editor/UnrealEd/Private/Kismet2/KismetReinstanceUtilities.cpp b/Engine/Source/Editor/UnrealEd/Private/Kismet2/KismetReinstanceUtilities.cpp
--- a/Engine/Source/Editor/UnrealEd/Private/Kismet2/KismetReinstanceUtilities.cpp
+++ b/Engine/Source/Editor/UnrealEd/Private/Kismet2/KismetReinstanceUtilities.cpp
@@ -2076,14 +2076,13 @@ static void ReplaceActorHelper(AActor* OldActor, UClass* OldClass, UObject*& New
 	// Make sure to reuse the same external package if any
 	SpawnInfo.bCreateActorPackage = false;
 	SpawnInfo.OverridePackage = OldActor->GetExternalPackage();
-	if (SpawnInfo.OverridePackage)
-	{
-		OldActor->SetPackageExternal(false, false);
-	}
 
 	SpawnInfo.OverrideActorGuid = OldActor->GetActorGuid();
 
-	OldActor->Rename(nullptr, OldActor->GetOuter(), REN_DoNotDirty | REN_DontCreateRedirectors | REN_ForceNoResetLoaders);
+	// Don't go through AActor::Rename here because we aren't changing outers (the actor's level) and we also don't want to reset loaders
+	// if the actor is using an external package. We really just want to rename that actor out of the way so we can spawn the new one in
+	// the exact same package, keeping the package name intact.
+	OldActor->UObject::Rename(nullptr, OldActor->GetOuter(), REN_DoNotDirty | REN_DontCreateRedirectors | REN_ForceNoResetLoaders);
 
 	AActor* NewActor = nullptr;
 	{
@@ -2091,11 +2090,6 @@ static void ReplaceActorHelper(AActor* OldActor, UClass* OldClass, UObject*& New
 		NewActor = World->SpawnActor(SpawnClass, &Location, &Rotation, SpawnInfo);
 	}
 
-	if (SpawnInfo.OverridePackage)
-	{
-		OldActor->SetExternalPackage(SpawnInfo.OverridePackage);
-	}
-
 	if (OldActor->CurrentTransactionAnnotation.IsValid())
 	{
 		NewActor->CurrentTransactionAnnotation = OldActor->CurrentTransactionAnnotation;
@@ -2477,6 +2471,13 @@ void FBlueprintCompileReinstancer::ReplaceInstancesOfClass_Inner(TMap<UClass*, U
 		UObject** NewObject = OldToNewInstanceMap.Find(Obj);
 		if (NewObject && *NewObject)
 		{
+			if (Obj)
+			{
+				// Patch the new object into the old object linker's export map; subsequent loads may import
+				// this entry and we need to make sure that it returns the new object instead of the old one.
+				FLinkerLoad::PRIVATE_PatchNewObjectIntoExport(Obj, *NewObject);
+			}
+
 			if (UAnimInstance* AnimTree = Cast<UAnimInstance>(*NewObject))
 			{
 				// Initialising the anim instance isn't enough to correctly set up the skeletal mesh again in a
