Do not execute construction script or register components on reinstanced actors when they are reinstanced in a world that is not yet initialized.
Those steps will be taken care of as part of the normal world initialization flow

diff --git a/Engine/Source/Editor/UnrealEd/Private/Kismet2/KismetReinstanceUtilities.cpp b/Engine/Source/Editor/UnrealEd/Private/Kismet2/KismetReinstanceUtilities.cpp
--- a/Engine/Source/Editor/UnrealEd/Private/Kismet2/KismetReinstanceUtilities.cpp
+++ b/Engine/Source/Editor/UnrealEd/Private/Kismet2/KismetReinstanceUtilities.cpp
@@ -1316,28 +1316,32 @@ void FActorReplacementHelper::Finalize(const TMap<UObject*, UObject*>& OldToNewI
 	// because this is an editor context it's important to use this execution guard
 	FEditorScriptExecutionGuard ScriptGuard;
 
-	// run the construction script, which will use the properties we just copied over
-	bool bCanReRun = UBlueprint::IsBlueprintHierarchyErrorFree(NewActor->GetClass());
-	if (NewActor->CurrentTransactionAnnotation.IsValid() && bCanReRun)
+	// Only rerun construction script and component registration after reinstancing actors if the world is already initialized, otherwise that will be handled on world initialization
+	if (NewActor->GetWorld()->bIsWorldInitialized)
 	{
-		NewActor->CurrentTransactionAnnotation->ComponentInstanceData.FindAndReplaceInstances(OldToNewInstanceMap);
-		NewActor->RerunConstructionScripts();
-	}
-	else if (CachedActorData.IsValid())
-	{
-		CachedActorData->ComponentInstanceData.FindAndReplaceInstances(OldToNewInstanceMap);
-		const bool bErrorFree = NewActor->ExecuteConstruction(TargetWorldTransform, nullptr, &CachedActorData->ComponentInstanceData);
-		if (!bErrorFree)
+		// run the construction script, which will use the properties we just copied over
+		bool bCanReRun = UBlueprint::IsBlueprintHierarchyErrorFree(NewActor->GetClass());
+		if (NewActor->CurrentTransactionAnnotation.IsValid() && bCanReRun)
+		{
+			NewActor->CurrentTransactionAnnotation->ComponentInstanceData.FindAndReplaceInstances(OldToNewInstanceMap);
+			NewActor->RerunConstructionScripts();
+		}
+		else if (CachedActorData.IsValid())
+		{
+			CachedActorData->ComponentInstanceData.FindAndReplaceInstances(OldToNewInstanceMap);
+			const bool bErrorFree = NewActor->ExecuteConstruction(TargetWorldTransform, nullptr, &CachedActorData->ComponentInstanceData);
+			if (!bErrorFree)
+			{
+				// Save off the cached actor data for once the blueprint has been fixed so we can reapply it
+				NewActor->CurrentTransactionAnnotation = CachedActorData;
+			}
+		}
+		else
 		{
-			// Save off the cached actor data for once the blueprint has been fixed so we can reapply it
-			NewActor->CurrentTransactionAnnotation = CachedActorData;
+			FComponentInstanceDataCache DummyComponentData;
+			NewActor->ExecuteConstruction(TargetWorldTransform, nullptr, &DummyComponentData);
 		}
 	}
-	else
-	{
-		FComponentInstanceDataCache DummyComponentData;
-		NewActor->ExecuteConstruction(TargetWorldTransform, nullptr, &DummyComponentData);
-	}	
 
 	// The reinstancing case doesn't ever explicitly call Actor->FinishSpawning, we've handled the construction script
 	// portion above but still need the PostActorConstruction() case so BeginPlay gets routed correctly while in a BegunPlay world
@@ -2124,9 +2128,14 @@ static void ReplaceActorHelper(AActor* OldActor, UClass* OldClass, UObject*& New
 
 	// reset properties/streams
 	NewActor->ResetPropertiesForConstruction();
-	// register native components
-	NewActor->RegisterAllComponents();
 
+	// Only register the components if the world is already initialized
+	if (World->bIsWorldInitialized)
+	{
+		// register native components
+		NewActor->RegisterAllComponents();
+	}
+	
 	// 
 	// clean up the old actor (unselect it, remove it from the world, etc.)...
 
diff --git a/Engine/Source/Runtime/Engine/Private/Actor.cpp b/Engine/Source/Runtime/Engine/Private/Actor.cpp
--- a/Engine/Source/Runtime/Engine/Private/Actor.cpp
+++ b/Engine/Source/Runtime/Engine/Private/Actor.cpp
@@ -3371,8 +3371,8 @@ void AActor::PostSpawnInitialize(FTransform const& UserSpawnTransform, AActor* I
 	// Register the actor's default (native) components, but only if we have a native scene root. If we don't, it implies that there could be only non-scene components
 	// at the native class level. In that case, if this is a Blueprint instance, we need to defer native registration until after SCS execution can establish a scene root.
 	// Note: This API will also call PostRegisterAllComponents() on the actor instance. If deferred, PostRegisterAllComponents() won't be called until the root is set by SCS.
-	bHasDeferredComponentRegistration = (SceneRootComponent == nullptr && Cast<UBlueprintGeneratedClass>(GetClass()) != nullptr);
-	if (!bHasDeferredComponentRegistration && GetWorld())
+	bHasDeferredComponentRegistration = (SceneRootComponent == nullptr && Cast<UBlueprintGeneratedClass>(GetClass()) != nullptr) || !World || !World->bIsWorldInitialized;
+	if (!bHasDeferredComponentRegistration)
 	{
 		RegisterAllComponents();
 	}
