Backout changelist 16052362 as it makes actors lose their components when undoing.

diff --git a/Engine/Source/Editor/UnrealEd/Classes/Editor/Transactor.h b/Engine/Source/Editor/UnrealEd/Classes/Editor/Transactor.h
--- a/Engine/Source/Editor/UnrealEd/Classes/Editor/Transactor.h
+++ b/Engine/Source/Editor/UnrealEd/Classes/Editor/Transactor.h
@@ -36,89 +36,85 @@ class UNREALED_API FTransaction : public ITransaction
 	friend class FTransactionObjectEvent;
 
 protected:
-	/**
-		This type is necessary because the blueprint system is destroying and creating
-		CDOs at edit time (usually on compile, but also on load), but also stores user
-		entered data in the CDO. We "need"  changes to a CDO to persist across instances
-		because as we undo and redo we  need to apply changes to different instances of
-		the CDO - alternatively we could destroy and create the CDO as part of a transaction
-		(this alternative is the reason for the bunny ears around need).
-
-		DanO: My long term preference is for the editor to use a dynamic, mutable type
-		(rather than the CDO) to store editor data. The CDO can then be re-instanced (or not)
-		as runtime code requires.
-	*/
-	struct FPersistentObjectRef
+	// Record of an object.
+	class UNREALED_API FObjectRecord
 	{
 	public:
-		friend FArchive& operator<<(FArchive& Ar, FPersistentObjectRef& ReferencedObject)
-		{
-			Ar << (std::underlying_type_t<EReferenceType>&)ReferencedObject.ReferenceType;
-			Ar << ReferencedObject.RootObject;
-			Ar << ReferencedObject.SubObjectHierarchyIDs;
-			return Ar;
-		}
 
-		friend bool operator==(const FPersistentObjectRef& LHS, const FPersistentObjectRef& RHS)
+		/** 
+			This type is necessary because the blueprint system is destroying and creating 
+			CDOs at edit time (usually on compile, but also on load), but also stores user 
+			entered data in the CDO. We "need"  changes to a CDO to persist across instances 
+			because as we undo and redo we  need to apply changes to different instances of 
+			the CDO - alternatively we could destroy and create the CDO as part of a transaction 
+			(this alternative is the reason for the bunny ears around need).
+
+			DanO: My long term preference is for the editor to use a dynamic, mutable type
+			(rather than the CDO) to store editor data. The CDO can then be re-instanced (or not)
+			as runtime code requires.
+		*/
+		struct FPersistentObjectRef
 		{
-			return LHS.ReferenceType == RHS.ReferenceType
-				&& LHS.RootObject == RHS.RootObject
-				&& (LHS.ReferenceType != EReferenceType::SubObject || LHS.SubObjectHierarchyIDs == RHS.SubObjectHierarchyIDs);
-		}
+		public:
+			friend FArchive& operator<<(FArchive& Ar, FPersistentObjectRef& ReferencedObject)
+			{
+				Ar << (std::underlying_type_t<EReferenceType>&)ReferencedObject.ReferenceType;
+				Ar << ReferencedObject.RootObject;
+				Ar << ReferencedObject.SubObjectHierarchyIDs;
+				return Ar;
+			}
 
-		friend bool operator!=(const FPersistentObjectRef& LHS, const FPersistentObjectRef& RHS)
-		{
-			return !(LHS == RHS);
-		}
+			friend bool operator==(const FPersistentObjectRef& LHS, const FPersistentObjectRef& RHS)
+			{
+				return LHS.ReferenceType == RHS.ReferenceType
+					&& LHS.RootObject == RHS.RootObject
+					&& (LHS.ReferenceType != EReferenceType::SubObject || LHS.SubObjectHierarchyIDs == RHS.SubObjectHierarchyIDs);
+			}
 
-		friend uint32 GetTypeHash(const FPersistentObjectRef& InObjRef)
-		{
-			return HashCombine(GetTypeHash(InObjRef.ReferenceType), GetTypeHash(InObjRef.RootObject));
-		}
+			friend bool operator!=(const FPersistentObjectRef& LHS, const FPersistentObjectRef& RHS)
+			{
+				return !(LHS == RHS);
+			}
 
-		FPersistentObjectRef() = default;
-		explicit FPersistentObjectRef(UObject* InObject);
+			FPersistentObjectRef() = default;
+			explicit FPersistentObjectRef(UObject* InObject);
 
-		bool IsRootObjectReference() const
-		{
-			return ReferenceType == EReferenceType::RootObject;
-		}
+			bool IsRootObjectReference() const
+			{
+				return ReferenceType == EReferenceType::RootObject;
+			}
 
-		bool IsSubObjectReference() const
-		{
-			return ReferenceType == EReferenceType::SubObject;
-		}
+			bool IsSubObjectReference() const
+			{
+				return ReferenceType == EReferenceType::SubObject;
+			}
 
-		UObject* Get() const;
+			UObject* Get() const;
 
-		void AddReferencedObjects(FReferenceCollector& Collector);
+			void AddReferencedObjects(FReferenceCollector& Collector);
 
-	private:
-		/** This enum represents all of the different special cases we are handling with this type */
-		enum class EReferenceType : uint8
-		{
-			Unknown,
-			RootObject,
-			SubObject,
-		};
+		private:
+			/** This enum represents all of the different special cases we are handling with this type */
+			enum class EReferenceType : uint8
+			{
+				Unknown,
+				RootObject,
+				SubObject,
+			};
 
-		/** The reference type we're handling */
-		EReferenceType ReferenceType = EReferenceType::Unknown;
-		/** Stores the object pointer when ReferenceType==RootObject, and the outermost pointer of the sub-object chain when when ReferenceType==SubObject */
-		UObject* RootObject = nullptr;
-		/** Stores the sub-object name chain when ReferenceType==SubObject */
-		TArray<FName, TInlineAllocator<4>> SubObjectHierarchyIDs;
-
-		/** Cached pointers corresponding to RootObject when ReferenceType==SubObject (@note cache needs testing on access as it may have become stale) */
-		mutable TWeakObjectPtr<UObject> CachedRootObject;
-		/** Cache of pointers corresponding to the items within SubObjectHierarchyIDs when ReferenceType==SubObject (@note cache needs testing on access as it may have become stale) */
-		mutable TArray<TWeakObjectPtr<UObject>, TInlineAllocator<4>> CachedSubObjectHierarchy;
-	};
+			/** The reference type we're handling */
+			EReferenceType ReferenceType = EReferenceType::Unknown;
+			/** Stores the object pointer when ReferenceType==RootObject, and the outermost pointer of the sub-object chain when when ReferenceType==SubObject */
+			UObject* RootObject = nullptr;
+			/** Stores the sub-object name chain when ReferenceType==SubObject */
+			TArray<FName, TInlineAllocator<4>> SubObjectHierarchyIDs;
+
+			/** Cached pointers corresponding to RootObject when ReferenceType==SubObject (@note cache needs testing on access as it may have become stale) */
+			mutable TWeakObjectPtr<UObject> CachedRootObject;
+			/** Cache of pointers corresponding to the items within SubObjectHierarchyIDs when ReferenceType==SubObject (@note cache needs testing on access as it may have become stale) */
+			mutable TArray<TWeakObjectPtr<UObject>, TInlineAllocator<4>> CachedSubObjectHierarchy;
+		};
 
-	// Record of an object.
-	class UNREALED_API FObjectRecord
-	{
-	public:
 		struct FSerializedProperty
 		{
 			FSerializedProperty()
@@ -245,16 +241,16 @@ protected:
 		/** Custom change to apply to this object to undo this record.  Executing the undo will return an object that can be used to redo the change. */
 		TUniquePtr<FChange> CustomChange;
 		/** Array: If an array object, reference to script array */
-		FScriptArray*		Array = nullptr;
+		FScriptArray*		Array;
 		/** Array: Offset into the array */
-		int32				Index = 0;
+		int32				Index;
 		/** Array: How many items to record */
-		int32				Count = 0;
+		int32				Count;
 		/** @todo document this
 		 Array: Operation performed on array: 1 (add/insert), 0 (modify), -1 (remove)? */
-		int32				Oper = 0;
+		int32				Oper;
 		/** Array: Size of each item in the array */
-		int32				ElementSize = 0;
+		int32				ElementSize;
 		/** Array: DefaultConstructor for each item in the array */
 		STRUCT_DC			DefaultConstructor;
 		/** Array: Serializer to use for each item in the array */
@@ -262,13 +258,13 @@ protected:
 		/** Array: Destructor for each item in the array */
 		STRUCT_DTOR			Destructor;
 		/** True if object  has already been restored from data. False otherwise. */
-		bool				bRestored = false;
+		bool				bRestored;
 		/** True if object has been finalized and generated diff data */
-		bool				bFinalized = false;
+		bool				bFinalized;
 		/** True if object has been snapshot before */
-		bool				bSnapshot = false;
+		bool				bSnapshot;
 		/** True if record should serialize data as binary blob (more compact). False to use tagged serialization (more robust) */
-		bool				bWantsBinarySerialization = true;
+		bool				bWantsBinarySerialization;
 		/** The serialized object data */
 		FSerializedObject	SerializedObject;
 		/** The serialized object data that will be used when the transaction is flipped */
@@ -279,7 +275,8 @@ protected:
 		FTransactionObjectDeltaChange DeltaChange;
 
 		// Constructors.
-		FObjectRecord() = default;
+		FObjectRecord()
+		{}
 		FObjectRecord( FTransaction* Owner, UObject* InObject, TUniquePtr<FChange> InCustomChange, FScriptArray* InArray, int32 InIndex, int32 InCount, int32 InOper, int32 InElementSize, STRUCT_DC InDefaultConstructor, STRUCT_AR InSerializer, STRUCT_DTOR InDestructor );
 
 	private:
@@ -365,16 +362,15 @@ protected:
 			}
 			void Preload( UObject* InObject ) override
 			{
-				if (Owner)
+				if( Owner )
 				{
-					if (const FObjectRecords* ObjectRecords = Owner->ObjectRecordsMap.Find(FPersistentObjectRef(InObject)))
+					if ( Owner->ObjectMap.Contains( InObject ) )
 					{
-						for (FObjectRecord* Record : ObjectRecords->Records)
+						for (int32 i = 0; i < Owner->Records.Num(); i++)
 						{
-							checkSlow(Record->Object.Get() == InObject);
-							if (!Record->CustomChange)
+							if (Owner->Records[i].Object.Get() == InObject)
 							{
-								Record->Restore(Owner);
+								Owner->Records[i].Restore(Owner);
 							}
 						}
 					}
@@ -540,32 +536,10 @@ protected:
 		};
 	};
 
-	struct FObjectRecords
-	{
-		friend FArchive& operator<<(FArchive& Ar, FObjectRecords& ObjectRecords)
-		{
-			Ar << ObjectRecords.SaveCount;
-
-			if (Ar.IsLoading())
-			{
-				// Clear the map on load, as it will need to be rebuilt from the source array
-				ObjectRecords.Records.Empty();
-			}
-
-			return Ar;
-		}
-
-		TArray<FObjectRecord*, TInlineAllocator<1>> Records;
-		int32 SaveCount = 0;
-	};
-
 	// Transaction variables.
 	/** List of object records in this transaction */
-	TIndirectArray<FObjectRecord> Records;
-
-	/** Map of object records (non-array), for optimized look-up and to prevent an object being serialized to a transaction more than once. */
-	TMap<FPersistentObjectRef, FObjectRecords> ObjectRecordsMap;
-
+	TIndirectArray<FObjectRecord>	Records;
+	
 	/** Unique identifier for this transaction, used to track it during its lifetime */
 	FGuid		Id;
 
@@ -583,13 +557,15 @@ protected:
 	FString		Context;
 
 	/** The key object being edited in this transaction. For example the blueprint object. Can be NULL */
-	UObject* PrimaryObject = nullptr;
+	UObject*	PrimaryObject;
 
-	/** If true, on apply flip the direction of iteration over object records. The only client for which this is false is FMatineeTransaction */
-	bool		bFlip = true;
+	/** Used to prevent objects from being serialized to a transaction more than once. */
+	TMap<UObject*, int32>	ObjectMap;
 
+	/** If true, on apply flip the direction of iteration over object records. The only client for which this is false is FMatineeTransaction */
+	bool					bFlip;
 	/** Used to track direction to iterate over transaction's object records. Typically -1 for Undo, 1 for Redo */
-	int32		Inc = -1;
+	int32					Inc;
 
 	struct FChangedObjectValue
 	{
@@ -618,9 +594,10 @@ public:
 		:	Id( FGuid::NewGuid() )
 		,	Title( InTitle )
 		,	Context( InContext )
+		,	PrimaryObject(nullptr)
 		,	bFlip(bInFlip)
-	{
-	}
+		,	Inc(-1)
+	{}
 
 	virtual ~FTransaction()
 	{
@@ -704,45 +681,35 @@ public:
 	/** Returns the description of each contained Object Record */
 	FText GetDescription() const
 	{			
-		FString ConcatenatedRecords;
-		for (const FObjectRecord& Record : Records)
+		if (Algo::FindByPredicate(Records, [](const FObjectRecord& Record)
 		{
-			if (Record.CustomChange.IsValid())
+			return Record.CustomChange.IsValid();
+		}))
+        {
+			FString ConcatenatedRecords;
+			for (const FObjectRecord& Record : Records)
 			{
-				if (ConcatenatedRecords.Len())
+				if (Record.CustomChange.IsValid())
 				{
-					ConcatenatedRecords += TEXT("\n");
+					if (ConcatenatedRecords.Len())
+					{
+						ConcatenatedRecords += TEXT("\n");
+					}
+					ConcatenatedRecords += Record.CustomChange->ToString();
 				}
-				ConcatenatedRecords += Record.CustomChange->ToString();
 			}
-		}
 
-		return ConcatenatedRecords.IsEmpty() 
-			? GetTitle() 
-			: FText::FromString(ConcatenatedRecords);
+			return FText::FromString(ConcatenatedRecords);
+        }
+
+		return GetTitle();
 	}
 
+
 	/** Serializes a reference to a transaction in a given archive. */
 	friend FArchive& operator<<( FArchive& Ar, FTransaction& T )
 	{
-		Ar << T.Records << T.ObjectRecordsMap << T.Id << T.Title << T.Context << T.PrimaryObject;
-
-		if (Ar.IsLoading())
-		{
-			// Rebuild the LUT from the records array
-			for (FObjectRecord& Record : T.Records)
-			{
-				if (!Record.Array)
-				{
-					// We should have already added a map entry for this object when loading T.ObjectRecordsMap, so check that we have (otherwise something is out-of-sync)
-					FObjectRecords& ObjectRecords = T.ObjectRecordsMap.FindChecked(Record.Object);
-					// Note: This is only safe to do because T.Records is a TIndirectArray
-					ObjectRecords.Records.Add(&Record);
-				}
-			}
-		}
-
-		return Ar;
+		return Ar << T.Records << T.Id << T.Title << T.ObjectMap << T.Context << T.PrimaryObject;
 	}
 
 	/** Serializes a reference to a transaction in a given archive. */
diff --git a/Engine/Source/Editor/UnrealEd/Private/EditorTransaction.cpp b/Engine/Source/Editor/UnrealEd/Private/EditorTransaction.cpp
--- a/Engine/Source/Editor/UnrealEd/Private/EditorTransaction.cpp
+++ b/Engine/Source/Editor/UnrealEd/Private/EditorTransaction.cpp
@@ -34,6 +34,10 @@ FTransaction::FObjectRecord::FObjectRecord(FTransaction* Owner, UObject* InObjec
 	,	DefaultConstructor	( InDefaultConstructor )
 	,	Serializer			( InSerializer )
 	,	Destructor			( InDestructor )
+	,	bRestored			( false )
+	,   bFinalized			( false )
+	,	bSnapshot			( false )
+	,	bWantsBinarySerialization ( true )
 {
 	// Blueprint compile-in-place can alter class layout so use tagged serialization for objects relying on a UBlueprint's Class
 	if (UBlueprintGeneratedClass* Class = Cast<UBlueprintGeneratedClass>(InObject->GetClass()))
@@ -342,11 +346,14 @@ void FTransaction::FObjectRecord::Diff( const FTransaction* Owner, const FSerial
 		NewSerializedObject.SerializedObjectIndices.MultiFind(InPropertyName, NewSerializedObjectIndices, true);
 
 		bool bAreObjectPointersIdentical = OldSerializedObjectIndices.Num() == NewSerializedObjectIndices.Num();
-		for (int32 ObjIndex = 0; ObjIndex < OldSerializedObjectIndices.Num() && bAreObjectPointersIdentical; ++ObjIndex)
+		if (bAreObjectPointersIdentical)
 		{
-			const FPersistentObjectRef& OldObjectRef = OldSerializedObject.ReferencedObjects.IsValidIndex(OldSerializedObjectIndices[ObjIndex]) ? OldSerializedObject.ReferencedObjects[OldSerializedObjectIndices[ObjIndex]] : FPersistentObjectRef();
-			const FPersistentObjectRef& NewObjectRef = NewSerializedObject.ReferencedObjects.IsValidIndex(NewSerializedObjectIndices[ObjIndex]) ? NewSerializedObject.ReferencedObjects[NewSerializedObjectIndices[ObjIndex]] : FPersistentObjectRef();
-			bAreObjectPointersIdentical = OldObjectRef == NewObjectRef;
+			for (int32 ObjIndex = 0; ObjIndex < OldSerializedObjectIndices.Num() && bAreObjectPointersIdentical; ++ObjIndex)
+			{
+				const UObject* OldObjectPtr = OldSerializedObject.ReferencedObjects.IsValidIndex(OldSerializedObjectIndices[ObjIndex]) ? OldSerializedObject.ReferencedObjects[OldSerializedObjectIndices[ObjIndex]].Get() : nullptr;
+				const UObject* NewObjectPtr = NewSerializedObject.ReferencedObjects.IsValidIndex(NewSerializedObjectIndices[ObjIndex]) ? NewSerializedObject.ReferencedObjects[NewSerializedObjectIndices[ObjIndex]].Get() : nullptr;
+				bAreObjectPointersIdentical = OldObjectPtr == NewObjectPtr;
+			}
 		}
 		return bAreObjectPointersIdentical;
 	};
@@ -360,11 +367,14 @@ void FTransaction::FObjectRecord::Diff( const FTransaction* Owner, const FSerial
 		NewSerializedObject.SerializedNameIndices.MultiFind(InPropertyName, NewSerializedNameIndices, true);
 
 		bool bAreNamesIdentical = OldSerializedNameIndices.Num() == NewSerializedNameIndices.Num();
-		for (int32 ObjIndex = 0; ObjIndex < OldSerializedNameIndices.Num() && bAreNamesIdentical; ++ObjIndex)
+		if (bAreNamesIdentical)
 		{
-			const FName& OldName = OldSerializedObject.ReferencedNames.IsValidIndex(OldSerializedNameIndices[ObjIndex]) ? OldSerializedObject.ReferencedNames[OldSerializedNameIndices[ObjIndex]] : FName();
-			const FName& NewName = NewSerializedObject.ReferencedNames.IsValidIndex(NewSerializedNameIndices[ObjIndex]) ? NewSerializedObject.ReferencedNames[NewSerializedNameIndices[ObjIndex]] : FName();
-			bAreNamesIdentical = OldName == NewName;
+			for (int32 ObjIndex = 0; ObjIndex < OldSerializedNameIndices.Num() && bAreNamesIdentical; ++ObjIndex)
+			{
+				const FName& OldName = OldSerializedObject.ReferencedNames.IsValidIndex(OldSerializedNameIndices[ObjIndex]) ? OldSerializedObject.ReferencedNames[OldSerializedNameIndices[ObjIndex]] : FName();
+				const FName& NewName = NewSerializedObject.ReferencedNames.IsValidIndex(NewSerializedNameIndices[ObjIndex]) ? NewSerializedObject.ReferencedNames[NewSerializedNameIndices[ObjIndex]] : FName();
+				bAreNamesIdentical = OldName == NewName;
+			}
 		}
 		return bAreNamesIdentical;
 	};
@@ -565,18 +575,10 @@ void FTransaction::RemoveRecords( int32 Count /* = 1  */ )
 {
 	if ( Count > 0 && Records.Num() >= Count )
 	{
-		// Remove anything from the ObjectRecordMap which is about to be removed from the Records array
+		// Remove anything from the ObjectMap which is about to be removed from the Records array
 		for (int32 Index = 0; Index < Count; Index++)
 		{
-			FObjectRecord& Record = Records[Records.Num() - Count + Index];
-			if (FObjectRecords* ObjectRecords = ObjectRecordsMap.Find(Record.Object))
-			{
-				ObjectRecords->Records.RemoveSingle(&Record);
-				if (ObjectRecords->Records.Num() == 0)
-				{
-					ObjectRecordsMap.Remove(Record.Object);
-				}
-			}
+			ObjectMap.Remove( Records[Records.Num() - Count + Index].Object.Get() );
 		}
 
 		Records.RemoveAt( Records.Num() - Count, Count );
@@ -589,10 +591,10 @@ void FTransaction::RemoveRecords( int32 Count /* = 1  */ )
 void FTransaction::DumpObjectMap(FOutputDevice& Ar) const
 {
 	Ar.Logf( TEXT("===== DumpObjectMap %s ==== "), *Title.ToString() );
-	for ( auto It = ObjectRecordsMap.CreateConstIterator(); It; ++It )
+	for ( auto It = ObjectMap.CreateConstIterator(); It; ++It )
 	{
-		const UObject* CurrentObject	= It.Key().Get();
-		const int32 SaveCount			= It.Value().SaveCount;
+		const UObject* CurrentObject	= It.Key();
+		const int32 SaveCount				= It.Value();
 		Ar.Logf( TEXT("%i\t: %s"), SaveCount, *CurrentObject->GetPathName() );
 	}
 	Ar.Logf( TEXT("=== EndDumpObjectMap %s === "), *Title.ToString() );
@@ -609,7 +611,7 @@ FArchive& operator<<( FArchive& Ar, FTransaction::FObjectRecord& R )
 	return Ar;
 }
 
-FTransaction::FPersistentObjectRef::FPersistentObjectRef(UObject* InObject)
+FTransaction::FObjectRecord::FPersistentObjectRef::FPersistentObjectRef(UObject* InObject)
 {
 	RootObject = InObject;
 	{
@@ -657,7 +659,7 @@ FTransaction::FPersistentObjectRef::FPersistentObjectRef(UObject* InObject)
 	checkSlow(Get() == InObject);
 }
 
-UObject* FTransaction::FPersistentObjectRef::Get() const
+UObject* FTransaction::FObjectRecord::FPersistentObjectRef::Get() const
 {
 	if (ReferenceType == EReferenceType::SubObject)
 	{
@@ -720,7 +722,7 @@ UObject* FTransaction::FPersistentObjectRef::Get() const
 	return RootObject;
 }
 
-void FTransaction::FPersistentObjectRef::AddReferencedObjects(FReferenceCollector& Collector)
+void FTransaction::FObjectRecord::FPersistentObjectRef::AddReferencedObjects(FReferenceCollector& Collector)
 {
 	UObject* Obj = RootObject;
 	Collector.AddReferencedObject(Obj);
@@ -807,11 +809,7 @@ void FTransaction::AddReferencedObjects( FReferenceCollector& Collector )
 	{
 		ObjectRecord.AddReferencedObjects( Collector );
 	}
-
-	for (TTuple<FPersistentObjectRef, FObjectRecords>& ObjectRecordsPair : ObjectRecordsMap)
-	{
-		ObjectRecordsPair.Key.AddReferencedObjects(Collector);
-	}
+	Collector.AddReferencedObjects(ObjectMap);
 }
 
 void FTransaction::SaveObject( UObject* Object )
@@ -819,14 +817,17 @@ void FTransaction::SaveObject( UObject* Object )
 	check(Object);
 	Object->CheckDefaultSubobjects();
 
-	FObjectRecords& ObjectRecords = ObjectRecordsMap.FindOrAdd(FPersistentObjectRef(Object));
-	if (ObjectRecords.Records.Num() == 0)
+	int32* SaveCount = ObjectMap.Find(Object);
+	if ( !SaveCount )
 	{
+		ObjectMap.Add(Object,1);
 		// Save the object.
-		FObjectRecord* UndoRecord = ObjectRecords.Records.Add_GetRef(new FObjectRecord(this, Object, nullptr, nullptr, 0, 0, 0, 0, nullptr, nullptr, nullptr));
-		Records.Add(UndoRecord);
+		Records.Add(new FObjectRecord( this, Object, nullptr, nullptr, 0, 0, 0, 0, nullptr, nullptr, nullptr));
+	}
+	else
+	{
+		++(*SaveCount);
 	}
-	++ObjectRecords.SaveCount;
 }
 
 void FTransaction::SaveArray( UObject* Object, FScriptArray* Array, int32 Index, int32 Count, int32 Oper, int32 ElementSize, STRUCT_DC DefaultConstructor, STRUCT_AR Serializer, STRUCT_DTOR Destructor )
@@ -851,15 +852,19 @@ void FTransaction::SaveArray( UObject* Object, FScriptArray* Array, int32 Index,
 	}
 }
 
-void FTransaction::StoreUndo(UObject* Object, TUniquePtr<FChange> UndoChange)
+void FTransaction::StoreUndo( UObject* Object, TUniquePtr<FChange> UndoChange )
 {
-	check(Object);
+	check( Object );
 	Object->CheckDefaultSubobjects();
 
+	int32* SaveCount = ObjectMap.Find( Object );
+	if( !SaveCount )
+	{
+		ObjectMap.Add( Object, 0 );
+	}
+
 	// Save the undo record
-	FObjectRecords& ObjectRecords = ObjectRecordsMap.FindOrAdd(FPersistentObjectRef(Object));
-	FObjectRecord* UndoRecord = ObjectRecords.Records.Add_GetRef(new FObjectRecord(this, Object, MoveTemp(UndoChange), nullptr, 0, 0, 0, 0, nullptr, nullptr, nullptr));
-	Records.Add(UndoRecord);
+	Records.Add(new FObjectRecord( this, Object, MoveTemp( UndoChange ), nullptr, 0, 0, 0, 0, nullptr, nullptr, nullptr));
 }
 
 void FTransaction::SetPrimaryObject(UObject* InObject)
@@ -872,17 +877,16 @@ void FTransaction::SetPrimaryObject(UObject* InObject)
 
 void FTransaction::SnapshotObject( UObject* InObject, TArrayView<const FProperty*> Properties )
 {
-	check(InObject);
-
-	if (const FObjectRecords* ObjectRecords = ObjectRecordsMap.Find(FPersistentObjectRef(InObject)))
+	if (InObject && ObjectMap.Contains(InObject))
 	{
-		for (FObjectRecord* Record : ObjectRecords->Records)
+		FObjectRecord* FoundObjectRecord = Algo::FindByPredicate(Records, [ObjRefToFind = FObjectRecord::FPersistentObjectRef(InObject)](const FObjectRecord& ObjRecord)
 		{
-			checkSlow(Record->Object.Get() == InObject);
-			if (!Record->CustomChange)
-			{
-				Record->Snapshot(this, Properties);
-			}
+			return ObjRecord.Object == ObjRefToFind;
+		});
+
+		if (FoundObjectRecord)
+		{
+			FoundObjectRecord->Snapshot(this, Properties);
 		}
 	}
 }
