Add null check to handle missing data

diff --git a/Engine/Source/Editor/UnrealEd/Private/MaterialGraph.cpp b/Engine/Source/Editor/UnrealEd/Private/MaterialGraph.cpp
--- a/Engine/Source/Editor/UnrealEd/Private/MaterialGraph.cpp
+++ b/Engine/Source/Editor/UnrealEd/Private/MaterialGraph.cpp
@@ -42,7 +42,10 @@ void UMaterialGraph::RebuildGraph()
 	}
 	for (UMaterialExpressionComment* Comment : Material->EditorComments)
 	{
-		SubgraphCommentMap.FindOrAdd(Comment->SubgraphExpression).Add(Comment);
+		if (Comment)
+		{
+			SubgraphCommentMap.FindOrAdd(Comment->SubgraphExpression).Add(Comment);
+		}
 	}
 
 	RebuildGraphInternal(SubgraphExpressionMap, SubgraphCommentMap);
