Fix SavePackage ensure that can occur if there is an Asset object in a package that is not the primary Asset and is not standalone.
#rnx

diff --git a/Engine/Source/Editor/UnrealEd/Private/EditorEngine.cpp b/Engine/Source/Editor/UnrealEd/Private/EditorEngine.cpp
--- a/Engine/Source/Editor/UnrealEd/Private/EditorEngine.cpp
+++ b/Engine/Source/Editor/UnrealEd/Private/EditorEngine.cpp
@@ -4289,6 +4289,8 @@ FSavePackageResultStruct UEditorEngine::Save( UPackage* InOuter, UObject* InBase
 
 	FScopedSlowTask SlowTask(100, FText(), bSlowTask);
 
+	bool bIsCooking = TargetPlatform != nullptr;
+	TopLevelFlags = UE::SavePackageUtilities::NormalizeTopLevelFlags(TopLevelFlags, bIsCooking);
 	UObject* Base = InBase;
 	if (!Base && InOuter)
 	{
@@ -4306,10 +4308,18 @@ FSavePackageResultStruct UEditorEngine::Save( UPackage* InOuter, UObject* InBase
 			GetObjectsWithPackage(InOuter, PotentialAssets, false);
 			for (UObject* Object : PotentialAssets)
 			{
-				if (Object->IsAsset() && !UE::AssetRegistry::FFiltering::ShouldSkipAsset(Object))
+				if (Object->IsAsset() && !UE::AssetRegistry::FFiltering::ShouldSkipAsset(Object) &&
+					(TopLevelFlags == RF_NoFlags || Object->HasAnyFlags(TopLevelFlags)))
 				{
-					Base = Object;
-					break;
+					if (FAssetData::IsUAsset(Object))
+					{
+						Base = Object;
+						break;
+					}
+					if (!Base)
+					{
+						Base = Object;
+					}
 				}
 			}
 		}
diff --git a/Engine/Source/Runtime/CoreUObject/Private/UObject/SavePackage.cpp b/Engine/Source/Runtime/CoreUObject/Private/UObject/SavePackage.cpp
--- a/Engine/Source/Runtime/CoreUObject/Private/UObject/SavePackage.cpp
+++ b/Engine/Source/Runtime/CoreUObject/Private/UObject/SavePackage.cpp
@@ -2061,11 +2061,7 @@ FSavePackageResultStruct UPackage::Save(UPackage* InOuter, UObject* Base, EObjec
 #if WITH_EDITOR
 	TMap<UObject*, UObject*> ReplacedImportOuters;
 
-	// Add the external package flag when not cooking
-	if (TopLevelFlags != RF_NoFlags && !bIsCooking)
-	{
-		TopLevelFlags |= RF_HasExternalPackage;
-	}
+	TopLevelFlags = UE::SavePackageUtilities::NormalizeTopLevelFlags(TopLevelFlags, bIsCooking);
 
 	// if the in memory package filename is different the filename we are saving it to,
 	// regenerate a new persistent id for it.
diff --git a/Engine/Source/Runtime/CoreUObject/Private/UObject/SavePackage/SaveContext.h b/Engine/Source/Runtime/CoreUObject/Private/UObject/SavePackage/SaveContext.h
--- a/Engine/Source/Runtime/CoreUObject/Private/UObject/SavePackage/SaveContext.h
+++ b/Engine/Source/Runtime/CoreUObject/Private/UObject/SavePackage/SaveContext.h
@@ -94,11 +94,7 @@ public:
 			TargetPackagePath.SetHeaderExtension(EPackageExtension::EmptyString);
 		}
 
-		// if we aren't cooking and top level flags aren't empty, add RF_HasExternalPackage to them to catch external packages data
-		if (SaveArgs.TopLevelFlags != RF_NoFlags && !IsCooking())
-		{
-			SaveArgs.TopLevelFlags |= RF_HasExternalPackage;
-		}
+		SaveArgs.TopLevelFlags = UE::SavePackageUtilities::NormalizeTopLevelFlags(SaveArgs.TopLevelFlags, IsCooking());
 
 		bCanUseUnversionedPropertySerialization = CanUseUnversionedPropertySerialization(SaveArgs.TargetPlatform);
 		bTextFormat = FString(Filename).EndsWith(FPackageName::GetTextAssetPackageExtension()) || FString(Filename).EndsWith(FPackageName::GetTextMapPackageExtension());
diff --git a/Engine/Source/Runtime/CoreUObject/Private/UObject/SavePackage/SavePackageUtilities.cpp b/Engine/Source/Runtime/CoreUObject/Private/UObject/SavePackage/SavePackageUtilities.cpp
--- a/Engine/Source/Runtime/CoreUObject/Private/UObject/SavePackage/SavePackageUtilities.cpp
+++ b/Engine/Source/Runtime/CoreUObject/Private/UObject/SavePackage/SavePackageUtilities.cpp
@@ -762,6 +762,16 @@ void CallPostSaveRoot(UObject* Object, FObjectSaveContextData& ObjectSaveContext
 	Object->PostSaveRoot(FObjectPostSaveRootContext(ObjectSaveContext));
 }
 
+EObjectFlags NormalizeTopLevelFlags(EObjectFlags TopLevelFlags, bool bIsCooking)
+{
+	// if we aren't cooking and top level flags aren't empty, add RF_HasExternalPackage to them to catch external packages data
+	if (TopLevelFlags != RF_NoFlags && !bIsCooking)
+	{
+		TopLevelFlags |= RF_HasExternalPackage;
+	}
+	return TopLevelFlags;
+}
+
 }
 } // end namespace UE::SavePackageUtilities
 
diff --git a/Engine/Source/Runtime/CoreUObject/Public/UObject/SavePackage.h b/Engine/Source/Runtime/CoreUObject/Public/UObject/SavePackage.h
--- a/Engine/Source/Runtime/CoreUObject/Public/UObject/SavePackage.h
+++ b/Engine/Source/Runtime/CoreUObject/Public/UObject/SavePackage.h
@@ -124,5 +124,8 @@ namespace SavePackageUtilities
 	/** Call the PostSaveRoot function on the given object. */
 	COREUOBJECT_API void CallPostSaveRoot(UObject* Object, FObjectSaveContextData& ObjectSaveContext, bool bCleanupRequired);
 
+	/** Add any required TopLevelFlags based on the save parameters. */
+	COREUOBJECT_API EObjectFlags NormalizeTopLevelFlags(EObjectFlags TopLevelFlags, bool bIsCooking);
+
 }
 }
\ No newline at end of file
