Make the open content browser drawer command "Ctrl+Space" on mac instead of "Command+Space". This is to avoid a conflict with spotlight search

diff --git a/Engine/Source/Editor/UnrealEd/Private/Toolkits/GlobalEditorCommonCommands.cpp b/Engine/Source/Editor/UnrealEd/Private/Toolkits/GlobalEditorCommonCommands.cpp
--- a/Engine/Source/Editor/UnrealEd/Private/Toolkits/GlobalEditorCommonCommands.cpp
+++ b/Engine/Source/Editor/UnrealEd/Private/Toolkits/GlobalEditorCommonCommands.cpp
@@ -41,7 +41,12 @@ void FGlobalEditorCommonCommands::RegisterCommands()
 	UI_COMMAND(OpenConsoleCommandBox, "Open Console Command Box", "Opens an edit box where you can type in a console command", EUserInterfaceActionType::Button, FInputChord(EKeys::Tilde));
 	UI_COMMAND(OpenDocumentation, "Open Documentation...", "Opens documentation for this tool", EUserInterfaceActionType::Button, FInputChord(EKeys::F1));
 
+#if PLATFORM_MAC
+	// On mac command and ctrl are automatically swapped. Command + Space is spotlight search so we use ctrl+space on mac to avoid the conflict
+	UI_COMMAND(OpenContentBrowserDrawer, "Open Content Browser Drawer", "Opens the content browser drawer from the status bar and focuses the search field", EUserInterfaceActionType::Button, FInputChord(EModifierKey::Command, EKeys::SpaceBar));
+#else
 	UI_COMMAND(OpenContentBrowserDrawer, "Open Content Browser Drawer", "Opens the content browser drawer from the status bar and focuses the search field", EUserInterfaceActionType::Button, FInputChord(EModifierKey::Control, EKeys::SpaceBar));
+#endif
 }
 
 void FGlobalEditorCommonCommands::MapActions(TSharedRef<FUICommandList>& ToolkitCommands)
