Backing out 16053782 while the issues with LSBPs can be investigated

diff --git a/Engine/Source/Editor/UnrealEd/Private/EditorServer.cpp b/Engine/Source/Editor/UnrealEd/Private/EditorServer.cpp
--- a/Engine/Source/Editor/UnrealEd/Private/EditorServer.cpp
+++ b/Engine/Source/Editor/UnrealEd/Private/EditorServer.cpp
@@ -2583,21 +2583,23 @@ bool UEditorEngine::Map_Load(const TCHAR* Str, FOutputDevice& Ar)
 					}
 
 					// If we are loading the same world again (reloading) then we must not specify that we want to keep this world in memory.
-					// Otherwise, try to keep the existing uninitialized world in memory since there is not reason to reload it.
+					// Otherwise, try to keep the existing world in memory since there is not reason to reload it.
 					UWorld* NewWorld = nullptr;
-					if (ExistingWorld && !ExistingWorld->bIsWorldInitialized && Context.World() != ExistingWorld && !bIsLoadingMapTemplate)
+					if (!bIsLoadingMapTemplate && ExistingWorld != nullptr && Context.World() != ExistingWorld)
 					{
 						NewWorld = ExistingWorld;
 					}
 					EditorDestroyWorld( Context, LocalizedLoadingMap, NewWorld );
 
-					// Unload all other map packages currently loaded and initialized, before opening a new map.
+					// Unload all other map packages currently loaded, before opening a new map.
 					// The world is only initialized correctly as part of the level loading process, so ensure that every map package needs loading.
 					TArray<UPackage*> WorldPackages;
 					for (TObjectIterator<UWorld> It; It; ++It)
 					{
 						UPackage* Package = Cast<UPackage>(It->GetOuter());
-						if (Package && Package != GetTransientPackage() && It->bIsWorldInitialized)
+						
+
+						if (Package && Package != GetTransientPackage() && Package->GetPathName() != LongTempFname)
 						{
 							WorldPackages.AddUnique(Package);
 						}
