Simplify asset waiting logic

diff --git a/Engine/Plugins/Experimental/AutomationUtils/Source/AutomationUtils/Private/AutomationUtilsBlueprintLibrary.cpp b/Engine/Plugins/Experimental/AutomationUtils/Source/AutomationUtils/Private/AutomationUtilsBlueprintLibrary.cpp
--- a/Engine/Plugins/Experimental/AutomationUtils/Source/AutomationUtils/Private/AutomationUtilsBlueprintLibrary.cpp
+++ b/Engine/Plugins/Experimental/AutomationUtils/Source/AutomationUtils/Private/AutomationUtilsBlueprintLibrary.cpp
@@ -345,15 +345,15 @@ UAutomationUtilsBlueprintLibrary::UAutomationUtilsBlueprintLibrary(const FObject
 
 void UAutomationUtilsBlueprintLibrary::TakeGameplayAutomationScreenshot(const FString ScreenshotName, float MaxGlobalError, float MaxLocalError, FString MapNameOverride)
 {
+	FlushAsyncLoading();
+
 	//Finish Loading Before Screenshot
 	if (!FPlatformProperties::RequiresCookedData())
 	{
-		//Finish Compiling all shaders
-		GShaderCompilingManager->FinishAllCompilation();
+		//Finish Compiling all shaders and other async assets
+		FAssetCompilingManager::Get().FinishAllCompilation();
 	}
 
-	FlushAsyncLoading();
-
 	// Make sure we finish all level streaming
 	if (UGameEngine* GameEngine = Cast<UGameEngine>(GEngine))
 	{
diff --git a/Engine/Source/Editor/UnrealEd/Private/CookOnTheFlyServer.cpp b/Engine/Source/Editor/UnrealEd/Private/CookOnTheFlyServer.cpp
--- a/Engine/Source/Editor/UnrealEd/Private/CookOnTheFlyServer.cpp
+++ b/Engine/Source/Editor/UnrealEd/Private/CookOnTheFlyServer.cpp
@@ -8438,38 +8438,13 @@ uint32 UCookOnTheFlyServer::FullLoadAndSave(uint32& CookedPackageCount)
 		FTaskGraphInterface::Get().ProcessThreadUntilIdle(ENamedThreads::GameThread);
 	}
 
-	// Wait for all shaders to finish compiling
-	if (GShaderCompilingManager)
+	if (FAssetCompilingManager::Get().GetNumRemainingAssets())
 	{
-		UE_LOG(LogCook, Display, TEXT("Waiting for shader compilation..."));
-		UE_SCOPED_HIERARCHICAL_COOKTIMER(FullLoadAndSave_WaitForShaderCompilation);
-		while(GShaderCompilingManager->IsCompiling())
-		{
-			GShaderCompilingManager->ProcessAsyncResults(false, false);
-			FPlatformProcess::Sleep(0.5f);
-		}
-
-		// One last process to get the shaders that were compiled at the very end
-		GShaderCompilingManager->ProcessAsyncResults(false, false);
-	}
-
-	if (GDistanceFieldAsyncQueue)
-	{
-		UE_LOG(LogCook, Display, TEXT("Waiting for distance field async operations..."));
-		UE_SCOPED_HIERARCHICAL_COOKTIMER(FullLoadAndSave_WaitForDistanceField);
-		GDistanceFieldAsyncQueue->BlockUntilAllBuildsComplete();
+		UE_LOG(LogCook, Display, TEXT("Waiting for async compilation..."));
+		// Wait for all assets to finish compiling
+		FAssetCompilingManager::Get().FinishAllCompilation();
 	}
 
-	if (GCardRepresentationAsyncQueue)
-	{
-		UE_LOG(LogCook, Display, TEXT("Waiting for card representation async operations..."));
-		UE_SCOPED_HIERARCHICAL_COOKTIMER(FullLoadAndSave_WaitForCardRepresentation);
-		GCardRepresentationAsyncQueue->BlockUntilAllBuildsComplete();
-	}
-
-	// Wait for all assets to finish compiling
-	FAssetCompilingManager::Get().ProcessAsyncTasks(false);
-
 	// Wait for all platform data to be loaded
 	{
 		UE_LOG(LogCook, Display, TEXT("Waiting for cooked platform data..."));
diff --git a/Engine/Source/Runtime/Engine/Private/WorldPartition/WorldPartitionMiniMapHelper.cpp b/Engine/Source/Runtime/Engine/Private/WorldPartition/WorldPartitionMiniMapHelper.cpp
--- a/Engine/Source/Runtime/Engine/Private/WorldPartition/WorldPartitionMiniMapHelper.cpp
+++ b/Engine/Source/Runtime/Engine/Private/WorldPartition/WorldPartitionMiniMapHelper.cpp
@@ -54,23 +54,7 @@ AWorldPartitionMiniMap* FWorldPartitionMiniMapHelper::GetWorldPartitionMiniMap(U
 
 void FWorldPartitionMiniMapHelper::CaptureWorldMiniMapToTexture(UWorld* InWorld, UObject* InOuterForTexture, uint32 InMiniMapSize, UTexture2D*& InOutMiniMapTexture, FBox& OutWorldBounds)
 {
-	auto WaitForShaderCompilation = [](){
-		UE_SCOPED_TIMER(TEXT("Shader Compilation"), LogWorldPartitionMiniMapHelper);
-		if (GShaderCompilingManager && GShaderCompilingManager->IsCompiling())
-		{
-			GShaderCompilingManager->FinishAllCompilation();
-		}
-	};
-	auto WaitForDistanceFieldBuilding = []() {
-		UE_SCOPED_TIMER(TEXT("Building Mesh Distance Fields"), LogWorldPartitionMiniMapHelper);
-		if (GDistanceFieldAsyncQueue && GDistanceFieldAsyncQueue->GetNumOutstandingTasks() > 0)
-		{
-			GDistanceFieldAsyncQueue->BlockUntilAllBuildsComplete();
-		}
-	};
-
-	WaitForDistanceFieldBuilding();
-	WaitForShaderCompilation();
+	// Before capturing the scene, make sure all assets are finished compiling
 	FAssetCompilingManager::Get().FinishAllCompilation();
 
 	//Calculate bounds of the World
