Fix potential crash during reimport by adding a missing PreEditChange before modifying an existing texture during a reimport

diff --git a/Engine/Source/Editor/UnrealEd/Private/Factories/EditorFactories.cpp b/Engine/Source/Editor/UnrealEd/Private/Factories/EditorFactories.cpp
--- a/Engine/Source/Editor/UnrealEd/Private/Factories/EditorFactories.cpp
+++ b/Engine/Source/Editor/UnrealEd/Private/Factories/EditorFactories.cpp
@@ -4175,7 +4175,13 @@ UObject* UTextureFactory::FactoryCreateBinary
 		// Static texture needs to avoid having pending InitRHI() before enqueuing ReleaseRHI() to safely track access of the PlatformData on the renderthread.
 		ExistingTexture->WaitForPendingInitOrStreaming();
 	}
-	
+
+	// Make sure the changes are part of the transaction when reimporting over an existing texture
+	if (ExistingTexture)
+	{
+		ExistingTexture->PreEditChange(nullptr);
+	}
+
 	FTextureReferenceReplacer RefReplacer(ExistingTexture);
 
 	UTexture* Texture = nullptr;
