Fix notifications which are just duplicating widgets in order to show the default notification background. This is no longer necessary.

diff --git a/Engine/Plugins/VirtualProduction/Takes/Source/TakeRecorder/Private/Recorder/TakeRecorder.cpp b/Engine/Plugins/VirtualProduction/Takes/Source/TakeRecorder/Private/Recorder/TakeRecorder.cpp
--- a/Engine/Plugins/VirtualProduction/Takes/Source/TakeRecorder/Private/Recorder/TakeRecorder.cpp
+++ b/Engine/Plugins/VirtualProduction/Takes/Source/TakeRecorder/Private/Recorder/TakeRecorder.cpp
@@ -72,9 +72,8 @@ public:
 
 		ChildSlot
 		[
-			SNew(SBorder)
+			SNew(SBox)
 			.Padding(FMargin(15.0f))
-			.BorderImage(FCoreStyle::Get().GetBrush("NotificationList.ItemBackground"))
 			[
 				SNew(SVerticalBox)
 
diff --git a/Engine/Source/Editor/MovieSceneCaptureDialog/Private/MovieSceneCaptureDialogModule.cpp b/Engine/Source/Editor/MovieSceneCaptureDialog/Private/MovieSceneCaptureDialogModule.cpp
--- a/Engine/Source/Editor/MovieSceneCaptureDialog/Private/MovieSceneCaptureDialogModule.cpp
+++ b/Engine/Source/Editor/MovieSceneCaptureDialog/Private/MovieSceneCaptureDialogModule.cpp
@@ -251,9 +251,8 @@ public:
 
 		ChildSlot
 		[
-			SNew(SBorder)
+			SNew(SBox)
 			.Padding(FMargin(15.0f))
-			.BorderImage(FCoreStyle::Get().GetBrush("NotificationList.ItemBackground"))
 			[
 				SNew(SVerticalBox)
 
diff --git a/Engine/Source/Editor/StatusBar/Private/SStatusBar.cpp b/Engine/Source/Editor/StatusBar/Private/SStatusBar.cpp
--- a/Engine/Source/Editor/StatusBar/Private/SStatusBar.cpp
+++ b/Engine/Source/Editor/StatusBar/Private/SStatusBar.cpp
@@ -389,8 +389,6 @@ public:
 	{
 		return AsShared();
 	}
-
-	bool UseNotificationBackground() const override { return true; }
 private:
 	TAttribute<const FStatusBarProgress*> StatusBarProgress;
 	TSharedPtr<SProgressBar> ProgressBar;
diff --git a/Engine/Source/Editor/UnrealEd/Private/AssetGuideline.cpp b/Engine/Source/Editor/UnrealEd/Private/AssetGuideline.cpp
--- a/Engine/Source/Editor/UnrealEd/Private/AssetGuideline.cpp
+++ b/Engine/Source/Editor/UnrealEd/Private/AssetGuideline.cpp
@@ -35,238 +35,6 @@
 
 #define LOCTEXT_NAMESPACE "AssetGuideine"
 
-class SAssetGuidelineNotification : public SCompoundWidget, public INotificationWidget
-{
-public:
-	SLATE_BEGIN_ARGS(SAssetGuidelineNotification){}
-	
-		SLATE_ATTRIBUTE( FText, TitleText)
-		SLATE_ATTRIBUTE( FText, HyperlinkText )
-		SLATE_EVENT(FSimpleDelegate, Hyperlink )
-		SLATE_ATTRIBUTE( TArray<FNotificationButtonInfo>, ButtonDetails )
-
-	SLATE_END_ARGS()
-
-	void Construct(const FArguments& InArgs)
-	{
-		TitleText = InArgs._TitleText.Get();
-		HyperlinkText = InArgs._HyperlinkText.Get();
-		Hyperlink = InArgs._Hyperlink;
-		
-		ChildSlot
-		[
-			SNew(SBorder)
-			.Padding( FMargin(5) )
-			.BorderImage(FCoreStyle::Get().GetBrush("NotificationList.ItemBackground_Border"))
-			.BorderBackgroundColor(FColor(0,0,0,1))
-			[
-				ConstructInternals(InArgs)
-			]
-		];
-	}
-
-
-	/**
-	 * Returns the internals of the notification
-	 */
-	TSharedRef<SHorizontalBox> ConstructInternals( const FArguments& InArgs ) 
-	{
-		TSharedRef<SHorizontalBox> HorizontalBox = SNew(SHorizontalBox);
-
-		// Notification image
-		HorizontalBox->AddSlot()
-		.AutoWidth()
-		.Padding(10.f, 0.f, 0.f, 0.f)
-		.VAlign(VAlign_Center)
-		.HAlign(HAlign_Left)
-		[
-			SNew(SImage)
-			.Image(FCoreStyle::Get().GetBrush("NotificationList.DefaultMessage"))
-		];
-
-		{
-			FSlateFontInfo Font = FCoreStyle::Get().GetFontStyle(TEXT("NotificationList.FontBold"));
-			TSharedRef<SVerticalBox> TextAndInteractiveWidgetsBox = SNew(SVerticalBox);
-
-			HorizontalBox->AddSlot()
-			.AutoWidth()
-			.Padding(10.f, 0.f, 15.f, 0.f)
-			.VAlign(VAlign_Center)
-			.HAlign(HAlign_Left)
-			[
-				TextAndInteractiveWidgetsBox
-			];
-
-			// Build Text box
-			TextAndInteractiveWidgetsBox->AddSlot()
-			.AutoHeight()
-			[
-				SNew(SBox)
-				[
-					SNew(STextBlock)
-					.Text(this, &SAssetGuidelineNotification::GetTextFromState)
-					.Font(Font)
-				]
-			];
-
-			TSharedRef<SVerticalBox> InteractiveWidgetsBox = SNew(SVerticalBox);
-			TextAndInteractiveWidgetsBox->AddSlot()
-			.AutoHeight()
-			[
-				InteractiveWidgetsBox
-			];
-			
-			// Adds a hyperlink
-			InteractiveWidgetsBox->AddSlot()
-			.AutoHeight()
-			.VAlign(VAlign_Bottom)
-			.HAlign(HAlign_Right)
-			[
-				SNew(SBox)
-				.Padding(FMargin(0.0f, 2.0f, 0.0f, 2.0f))
-				.VAlign(VAlign_Center)
-				.HAlign(HAlign_Left)
-				.Visibility(this, &SAssetGuidelineNotification::GetInteractiveVisibility)
-				[
-					SNew(SHyperlink)
-					.Text(this, &SAssetGuidelineNotification::GetHyperlinkTextFromState)
-					.OnNavigate(this, &SAssetGuidelineNotification::OnHyperlinkClicked)
-				]
-			];
-
-			// Adds any buttons that were passed in.
-			{
-				TSharedRef<SHorizontalBox> ButtonsBox = SNew(SHorizontalBox);
-				for (int32 idx = 0; idx < InArgs._ButtonDetails.Get().Num(); idx++)
-				{
-					FNotificationButtonInfo Button = InArgs._ButtonDetails.Get()[idx];
-
-					ButtonsBox->AddSlot()
-					.AutoWidth()
-					.HAlign(HAlign_Left)
-					.VAlign(VAlign_Center)
-					.Padding(0.0f, 0.0f, 4.0f, 0.0f)
-					[
-						SNew(SButton)
-						.Text(Button.Text)
-						.ToolTipText(Button.ToolTip)
-						.OnClicked(this, &SAssetGuidelineNotification::OnButtonClicked, Button.Callback)
-						.Visibility( this, &SAssetGuidelineNotification::GetInteractiveVisibility )
-					];
-				}
-				InteractiveWidgetsBox->AddSlot()
-				.AutoHeight()
-				.Padding(0.0f, 2.0f, 0.0f, 0.0f)
-				.VAlign(VAlign_Center)
-				.HAlign(HAlign_Left)
-				[
-					ButtonsBox
-				];
-			}
-		}
-
-		// Build success/fail image
-		HorizontalBox->AddSlot()
-		.AutoWidth()
-		[
-			SNew(SBox)
-			.Padding(FMargin(8.f, 0.f, 10.f, 0.f))
-			.HAlign(HAlign_Center)
-			.VAlign(VAlign_Center)
-			.Visibility( this, &SAssetGuidelineNotification::GetSuccessFailImageVisibility )
-			[
-				SNew(SImage)
-				.Image( this, &SAssetGuidelineNotification::GetSuccessFailImage )
-			]
-		];
-
-		return HorizontalBox;
-	}
-
-	/* Gets text based on current notification state*/
-	FText GetTextFromState() const
-	{
-		switch (State)
-		{
-		case SNotificationItem::CS_Success: return LOCTEXT("RestartNeeded", "Plugins & project settings updated, but will be out of sync until restart.");
-		case SNotificationItem::CS_Fail: return LOCTEXT("ChangeFailure", "Failed to change plugins & project settings.");
-		}
-
-		return TitleText;
-	}
-
-	/* Gets text based on current notification state*/
-	FText GetHyperlinkTextFromState() const
-	{
-		// Make hyperlink text on sucess or fail empty, so that the box auto-resizes correctly.
-		switch (State)
-		{
-		case SNotificationItem::CS_Success: return FText::GetEmpty();
-		case SNotificationItem::CS_Fail: return FText::GetEmpty();
-		}
-
-		return HyperlinkText;
-	}
-
-	/* Used to determine whether interactive components are visible */
-	EVisibility GetInteractiveVisibility() const
-	{
-		switch ( State )
-		{
-		case SNotificationItem::CS_None: return EVisibility::Visible;
-		case SNotificationItem::CS_Pending: return EVisibility::Visible;
-		case SNotificationItem::CS_Success: return EVisibility::Hidden;
-		case SNotificationItem::CS_Fail: return EVisibility::Hidden;
-		default:
-			check( false );
-			return EVisibility::Visible;
-		}
-	}
-
-	/* Used as a wrapper for the callback, this means any code calling it does not require access to FReply type */
-	FReply OnButtonClicked(FSimpleDelegate InCallback)
-	{
-		InCallback.ExecuteIfBound();
-		return FReply::Handled();
-	}
-
-	/* Execute the delegate for the hyperlink, if bound */
-	void OnHyperlinkClicked() const
-	{
-		Hyperlink.ExecuteIfBound();
-	}
-
-	EVisibility GetSuccessFailImageVisibility() const
-	{
-		return (State == SNotificationItem::ECompletionState::CS_Success || State == SNotificationItem::ECompletionState::CS_Fail) ? EVisibility::Visible : EVisibility::Collapsed;
-	}
-
-	const FSlateBrush* GetSuccessFailImage() const
-	{
-		return State == SNotificationItem::ECompletionState::CS_Success ? FCoreStyle::Get().GetBrush("NotificationList.SuccessImage") : FCoreStyle::Get().GetBrush("NotificationList.FailImage");
-	}
-
-	/** Begin INotificationWidget Interface */
-	virtual void OnSetCompletionState(SNotificationItem::ECompletionState InState)
-	{
-		State = InState;
-	}
-
-	virtual TSharedRef< SWidget > AsWidget()
-	{
-		return AsShared();
-	}
-	/** End INotificationWidget Interface */
-
-private:
-	SNotificationItem::ECompletionState State;
-
-	FText TitleText;
-
-	FText HyperlinkText;
-	FSimpleDelegate Hyperlink;
-};
-
 bool UAssetGuideline::IsPostLoadThreadSafe() const
 {
 	return true;
@@ -336,26 +104,26 @@ void UAssetGuideline::PostLoad()
 
 	if (!NeededPlugins.IsEmpty() || !NeededProjectSettings.IsEmpty())
 	{
-		FText WarningHyperlinkText;
-		FText NeededItems;
+		FText SubText;
+		FText TitleText;
 		{
 			FText AssetName = FText::AsCultureInvariant(GetPackage() ? GetPackage()->GetFName().ToString() : GetFName().ToString());
 
-			FText MissingPlugins = FText::Format(LOCTEXT("MissingPlugins", "Needed plugins: \n{0}"), NeededPlugins.IsEmpty() ? FText::GetEmpty() : FText::AsCultureInvariant(NeededPlugins));
-			FText PluginWarning = FText::Format(LOCTEXT("PluginWarning", "Asset '{0}' needs the above plugins. Assets related to '{0}' may not display properly.\n	Attemping to save '{0}' or related assets may result in irreverisble modification due to missing plugins. \n"), AssetName);
+			FText MissingPlugins = NeededPlugins.IsEmpty() ? FText::GetEmpty() : FText::Format(LOCTEXT("MissingPlugins", "Needed Plugins:\n{0}"), FText::AsCultureInvariant(NeededPlugins));
+			FText PluginWarning = NeededPlugins.IsEmpty() ? FText::GetEmpty() : FText::Format(LOCTEXT("PluginWarning", "Asset '{0}' needs the plugins listed above. Releated assets may not display properly.\nAttemping to save this asset or related assets may result in irreverisble modification due to missing plugins."), AssetName);
 
-			FText MissingProjectSettings = FText::Format(LOCTEXT("MissingProjectSettings", "Needed project settings: \n{0}"), NeededProjectSettings.IsEmpty() ? FText::GetEmpty() : FText::AsCultureInvariant(NeededProjectSettings));
-			FText ProjectSettingWarning = FText::Format(LOCTEXT("ProjectSettingWarning", "Asset '{0}' needs the above project settings. Assets related to '{0}' may not display properly."), AssetName);
+			FText MissingProjectSettings = NeededProjectSettings.IsEmpty() ? FText::GetEmpty() : FText::Format(LOCTEXT("MissingProjectSettings", "Needed project settings: \n{0}"), FText::AsCultureInvariant(NeededProjectSettings));
+			FText ProjectSettingWarning = NeededProjectSettings.IsEmpty() ? FText::GetEmpty() : FText::Format(LOCTEXT("ProjectSettingWarning", "Asset '{0}' needs the project settings listed above. Releated assets may not display properly."), AssetName);
 
-			FFormatNamedArguments WarningHyperlinkArgs;
-			WarningHyperlinkArgs.Add("PluginHyperlink", NeededPlugins.IsEmpty() ? FText::GetEmpty() : FText::Format(FText::AsCultureInvariant("{0}{1}\n"), MissingPlugins, PluginWarning));
-			WarningHyperlinkArgs.Add("ProjectSettingHyperlink", NeededProjectSettings.IsEmpty() ? FText::GetEmpty() : FText::Format(FText::AsCultureInvariant("{0}{1}\n"), MissingProjectSettings, ProjectSettingWarning));
-			WarningHyperlinkText = FText::Format(LOCTEXT("WarningHyperLink", "{PluginHyperlink}{ProjectSettingHyperlink}"), WarningHyperlinkArgs);
+			FFormatNamedArguments SubTextArgs;
+			SubTextArgs.Add("PluginSubText", NeededPlugins.IsEmpty() ? FText::GetEmpty() : FText::Format(FText::AsCultureInvariant("{0}{1}\n"), MissingPlugins, PluginWarning));
+			SubTextArgs.Add("ProjectSettingSubText", NeededProjectSettings.IsEmpty() ? FText::GetEmpty() : FText::Format(FText::AsCultureInvariant("{0}{1}\n"), MissingProjectSettings, ProjectSettingWarning));
+			SubText = FText::Format(LOCTEXT("SubText", "{PluginSubText}{ProjectSettingSubText}"), SubTextArgs);
 
 			FText NeedPlugins = LOCTEXT("NeedPlugins", "Missing Plugins!");
 			FText NeedProjectSettings = LOCTEXT("NeedProjectSettings", "Missing Project Settings!");
 			FText NeedBothGuidelines = LOCTEXT("NeedBothGuidelines", "Missing Plugins & Project Settings!");
-			NeededItems = !NeededPlugins.IsEmpty() && !NeededProjectSettings.IsEmpty() ? NeedBothGuidelines : !NeededPlugins.IsEmpty() ? NeedPlugins : NeedProjectSettings;
+			TitleText = !NeededPlugins.IsEmpty() && !NeededProjectSettings.IsEmpty() ? NeedBothGuidelines : !NeededPlugins.IsEmpty() ? NeedPlugins : NeedProjectSettings;
 		}
 
 		auto WarningHyperLink = [](bool NeedPluginLink, bool NeedProjectSettingLink)
@@ -371,30 +139,92 @@ void UAssetGuideline::PostLoad()
 			}
 		};
 
-		FNotificationInfo Info(NeededItems);
+
+		TPromise<TWeakPtr<SNotificationItem>> TextNotificationPromise;
+		TPromise<TWeakPtr<SNotificationItem>> HyperlinkNotificationPromise;
+		auto GetTextFromState = [NotificationFuture = TextNotificationPromise.GetFuture().Share(), SubText]()
+		{
+			SNotificationItem::ECompletionState State = SNotificationItem::CS_None;
+			if (TSharedPtr<SNotificationItem> NotificationPtr = NotificationFuture.Get().Pin())
+			{
+				State = NotificationPtr->GetCompletionState();
+			}
+
+			switch (State)
+			{
+			case SNotificationItem::CS_Success: return LOCTEXT("RestartNeeded", "Plugins & project settings updated, but will be out of sync until restart.");
+			case SNotificationItem::CS_Fail: return LOCTEXT("ChangeFailure", "Failed to change plugins & project settings.");
+			}
+
+			return SubText;
+		};
+
+		FText HyperlinkText = FText::GetEmpty();
+		if (!NeededPlugins.IsEmpty())
+		{
+			HyperlinkText = LOCTEXT("PluginHyperlinkText", "Open Plugin Browser");
+		}
+		else if (!NeededProjectSettings.IsEmpty())
+		{
+			HyperlinkText = LOCTEXT("ProjectSettingsHyperlinkText", "Open Project Settings");
+		}
+
+		/* Gets text based on current notification state*/
+		auto GetHyperlinkTextFromState = [NotificationFuture = HyperlinkNotificationPromise.GetFuture().Share(), HyperlinkText]()
+		{
+			SNotificationItem::ECompletionState State = SNotificationItem::CS_None;
+			if (TSharedPtr<SNotificationItem> NotificationPtr = NotificationFuture.Get().Pin())
+			{
+				State = NotificationPtr->GetCompletionState();
+			}
+
+			// Make hyperlink text on success or fail empty, so that the box auto-resizes correctly.
+			switch (State)
+			{
+			case SNotificationItem::CS_Success: return FText::GetEmpty();
+			case SNotificationItem::CS_Fail: return FText::GetEmpty();
+			}
+
+			return HyperlinkText;
+		};
+
+
+		FNotificationInfo Info(TitleText);
 		Info.bFireAndForget = false;
+		Info.FadeOutDuration = 0.0f;
+		Info.ExpireDuration = 0.0f;
+		Info.WidthOverride = FOptionalSize();
 		Info.ButtonDetails.Add(FNotificationButtonInfo(
-			LOCTEXT("GuidelineEnableMissing", "Enable Missing..."), 
+			LOCTEXT("GuidelineEnableMissing", "Enable Missing"), 
 			LOCTEXT("GuidelineEnableMissingTT", "Attempt to automatically set missing plugins / project settings"), 
-			FSimpleDelegate::CreateUObject(this, &UAssetGuideline::EnableMissingGuidelines, IncorrectPlugins, IncorrectProjectSettings)));
+			FSimpleDelegate::CreateUObject(this, &UAssetGuideline::EnableMissingGuidelines, IncorrectPlugins, IncorrectProjectSettings),
+			SNotificationItem::CS_None));
+
 		Info.ButtonDetails.Add(FNotificationButtonInfo(
 			LOCTEXT("GuidelineDismiss", "Dismiss"),
 			LOCTEXT("GuidelineDismissTT", "Dismiss this notification."), 
-			FSimpleDelegate::CreateUObject(this, &UAssetGuideline::DismissNotifications)));
+			FSimpleDelegate::CreateUObject(this, &UAssetGuideline::DismissNotifications),
+			SNotificationItem::CS_None));
+
 		Info.ButtonDetails.Add(FNotificationButtonInfo(
-			LOCTEXT("GuidelineRemove", "Remove guideline from asset"),
-			LOCTEXT("GuidelineRemoveTT", "Remove asset guideline. Preventing this notifcation from showing up again."),
-			FSimpleDelegate::CreateUObject(this, &UAssetGuideline::RemoveAssetGuideline)));
-		Info.ContentWidget = SNew(SAssetGuidelineNotification)
-			.TitleText(NeededItems)
-			.HyperlinkText(WarningHyperlinkText)
-			.Hyperlink(FSimpleDelegate::CreateLambda(WarningHyperLink, !NeededPlugins.IsEmpty(), !NeededProjectSettings.IsEmpty()))
-			.ButtonDetails(Info.ButtonDetails);
+			LOCTEXT("GuidelineRemove", "Remove Guideline"),
+			LOCTEXT("GuidelineRemoveTT", "Remove guideline from this asset. Preventing this notifcation from showing up again."),
+			FSimpleDelegate::CreateUObject(this, &UAssetGuideline::RemoveAssetGuideline),
+			SNotificationItem::CS_None));
+
+		Info.Text = TitleText;
+		Info.SubText = MakeAttributeLambda(GetTextFromState);
+		
+		Info.HyperlinkText = MakeAttributeLambda(GetHyperlinkTextFromState);
+		Info.Hyperlink = FSimpleDelegate::CreateLambda(WarningHyperLink, !NeededPlugins.IsEmpty(), !NeededProjectSettings.IsEmpty());
+
 
 		NotificationPtr = FSlateNotificationManager::Get().AddNotification(Info);
 		if (TSharedPtr<SNotificationItem> NotificationPin = NotificationPtr.Pin())
 		{
-			NotificationPin->SetCompletionState(SNotificationItem::CS_Pending);
+			TextNotificationPromise.SetValue(NotificationPtr);
+			HyperlinkNotificationPromise.SetValue(NotificationPtr);
+			NotificationPin->SetCompletionState(SNotificationItem::CS_None);
 		}
 	}
 }
diff --git a/Engine/Source/Editor/UnrealEd/Private/AutoReimport/ReimportFeedbackContext.h b/Engine/Source/Editor/UnrealEd/Private/AutoReimport/ReimportFeedbackContext.h
--- a/Engine/Source/Editor/UnrealEd/Private/AutoReimport/ReimportFeedbackContext.h
+++ b/Engine/Source/Editor/UnrealEd/Private/AutoReimport/ReimportFeedbackContext.h
@@ -42,7 +42,6 @@ public:
 	virtual void StartSlowTask(const FText& Task, bool bShowCancelButton=false) override;
 	virtual void ProgressReported(const float TotalProgressInterp, FText DisplayMessage) override;
 	virtual void FinalizeSlowTask() override;
-	virtual bool UseNotificationBackground() const override { return true; }
 
 	/** True if we will suppress slow task messages from being shown on the UI, false if otherwise */
 	bool bSuppressSlowTaskMessages;
diff --git a/Engine/Source/Runtime/AppFramework/Private/Widgets/Testing/STestSuite.cpp b/Engine/Source/Runtime/AppFramework/Private/Widgets/Testing/STestSuite.cpp
--- a/Engine/Source/Runtime/AppFramework/Private/Widgets/Testing/STestSuite.cpp
+++ b/Engine/Source/Runtime/AppFramework/Private/Widgets/Testing/STestSuite.cpp
@@ -4717,9 +4717,8 @@ public:
 	{
 		ChildSlot
 			[
-				SNew(SBorder)
+				SNew(SBox)
 				.Padding(15.0f)
-				.BorderImage(FCoreStyle::Get().GetBrush("NotificationList.ItemBackground"))
 				[
 					SNew(SHorizontalBox)
 					+ SHorizontalBox::Slot()
diff --git a/Engine/Source/Runtime/Slate/Private/Framework/Notifications/SlateAsyncTaskNotificationImpl.cpp b/Engine/Source/Runtime/Slate/Private/Framework/Notifications/SlateAsyncTaskNotificationImpl.cpp
--- a/Engine/Source/Runtime/Slate/Private/Framework/Notifications/SlateAsyncTaskNotificationImpl.cpp
+++ b/Engine/Source/Runtime/Slate/Private/Framework/Notifications/SlateAsyncTaskNotificationImpl.cpp
@@ -150,12 +150,10 @@ void SSlateAsyncTaskNotificationWidget::Construct(const FArguments& InArgs, cons
 
 	ChildSlot
 	[
-		SNew(SBorder)
+		SNew(SBox)
 		.Padding(FMargin(15.0f))
-		.BorderImage(FCoreStyle::Get().GetBrush("NotificationList.ItemBackground"))
 		[
 			SNew(SHorizontalBox)
-
 			// Main Icon
 			+SHorizontalBox::Slot()
 			.AutoWidth()
diff --git a/Engine/Source/Runtime/Slate/Private/Widgets/Notifications/SNotificationList.cpp b/Engine/Source/Runtime/Slate/Private/Widgets/Notifications/SNotificationList.cpp
--- a/Engine/Source/Runtime/Slate/Private/Widgets/Notifications/SNotificationList.cpp
+++ b/Engine/Source/Runtime/Slate/Private/Widgets/Notifications/SNotificationList.cpp
@@ -37,12 +37,19 @@ public:
 	
 	
 	/** Sets the text for message element */
-	virtual void SetText( const TAttribute< FText >& InText ) override
+	virtual void SetText( const TAttribute<FText>& InText ) override
 	{
 		Text = InText;
 		MyTextBlock->SetText( Text );
 	}
 
+	/** Sets the text for message element */
+	virtual void SetSubText(const TAttribute<FText>& InSubText) override
+	{
+		SubText = InSubText;
+		MySubTextBlock->SetText(SubText);
+	}
+
 	virtual ECompletionState GetCompletionState() const override
 	{
 		return CompletionState;
@@ -277,9 +284,12 @@ public:
 
 protected:
 
-	/** The text displayed in this text block */
+	/** The main text displayed in this text block */
 	TAttribute<FText> Text;
 
+	/** The sub text displayed in this text block. Subtext uses a smaller font and */
+	TAttribute<FText> SubText;
+
 	/** The fade in duration for this element */
 	TAttribute<float> FadeInDuration;
 
@@ -297,6 +307,8 @@ protected:
 	/** The text displayed in this element */
 	TSharedPtr<STextBlock> MyTextBlock;
 
+	TSharedPtr<STextBlock> MySubTextBlock;
+
 	/** The completion state of this message */	
 	ECompletionState CompletionState;
 
@@ -336,7 +348,8 @@ public:
 	{}
 
 		/** The text displayed in this text block */
-		SLATE_ARGUMENT(FText, Text)
+		SLATE_ATTRIBUTE(FText, Text)
+		SLATE_ATTRIBUTE(FText, SubText)
 
 		/** Setup information for the buttons on the notification */ 
 		SLATE_ARGUMENT(TArray<FNotificationButtonInfo>, ButtonDetails)
@@ -378,6 +391,8 @@ public:
 		CompletionState = CS_None;
 
 		Text = InArgs._Text;
+		SubText = InArgs._SubText;
+
 		FadeInDuration = InArgs._FadeInDuration;
 		FadeOutDuration = InArgs._FadeOutDuration;
 		ExpireDuration = InArgs._ExpireDuration;
@@ -416,6 +431,16 @@ public:
 			.WrapTextAt(InArgs._WidthOverride.IsSet() ? InArgs._WidthOverride.Get()-50.0f : 0.0f )
 		];
 
+		InteractiveWidgetsBox->AddSlot()
+		.Padding(FMargin(0.0f, 10.0f, 0.0f, 2.0f))
+		.AutoHeight()
+		[
+			SAssignNew(MySubTextBlock, STextBlock)
+			.Text(SubText)
+			.Visibility_Lambda([this]() { return SubText.Get().IsEmpty() ? EVisibility::Collapsed : EVisibility::SelfHitTestInvisible; })
+			.WrapTextAt(InArgs._WidthOverride.IsSet() ? InArgs._WidthOverride.Get() - 50.0f : 0.0f)
+		];
+
 		{
 			InteractiveWidgetsBox->AddSlot()
 			.AutoHeight()
@@ -724,6 +749,7 @@ TSharedRef<SNotificationItem> SNotificationList::AddNotification(const FNotifica
 			// Create notification.
 			NewItem = SNew(SNotificationItemImpl)
 				.Text(Info.Text)
+				.SubText(Info.SubText)
 				.ButtonDetails(Info.ButtonDetails)
 				.Image((Info.Image != nullptr) ? Info.Image : CachedImage)
 				.FadeInDuration(Info.FadeInDuration)
diff --git a/Engine/Source/Runtime/Slate/Public/Widgets/Notifications/INotificationWidget.h b/Engine/Source/Runtime/Slate/Public/Widgets/Notifications/INotificationWidget.h
--- a/Engine/Source/Runtime/Slate/Public/Widgets/Notifications/INotificationWidget.h
+++ b/Engine/Source/Runtime/Slate/Public/Widgets/Notifications/INotificationWidget.h
@@ -11,5 +11,5 @@ class SLATE_API INotificationWidget
 public:
 	virtual void OnSetCompletionState(SNotificationItem::ECompletionState State) = 0;
 	virtual TSharedRef< SWidget > AsWidget() = 0;
-	virtual bool UseNotificationBackground() const { return false; }
+	virtual bool UseNotificationBackground() const { return true; }
 };
diff --git a/Engine/Source/Runtime/Slate/Public/Widgets/Notifications/SNotificationList.h b/Engine/Source/Runtime/Slate/Public/Widgets/Notifications/SNotificationList.h
--- a/Engine/Source/Runtime/Slate/Public/Widgets/Notifications/SNotificationList.h
+++ b/Engine/Source/Runtime/Slate/Public/Widgets/Notifications/SNotificationList.h
@@ -32,11 +32,14 @@ public:
 		CS_Fail,
 	};
 
-	/** Sets the text for message element */
-	virtual void SetText( const TAttribute< FText >& InText ) = 0;
+	/** Sets the text for the notification element */
+	virtual void SetText(const TAttribute< FText >& InText) = 0;
+	
+	/** Sets the subtext the notification item. Sub text is used for longer text and is a smaller font */
+	virtual void SetSubText(const TAttribute<FText>& InSubText) = 0;
 
 	/** Sets the text and delegate for the hyperlink */
-	virtual void SetHyperlink( const FSimpleDelegate& InHyperlink, const TAttribute< FText >& InHyperlinkText = TAttribute< FText >() ) = 0;
+	virtual void SetHyperlink(const FSimpleDelegate& InHyperlink, const TAttribute< FText >& InHyperlinkText = TAttribute< FText >()) = 0;
 
 	/** Sets the ExpireDuration */
 	virtual void SetExpireDuration(float ExpireDuration) = 0;
@@ -110,6 +113,7 @@ struct FNotificationInfo
 	FNotificationInfo(const FText& InText)
 		: ContentWidget(),
 		Text(InText),
+		SubText(),
 		ButtonDetails(),
 		Image(nullptr),
 		FadeInDuration(0.5f),
@@ -136,6 +140,7 @@ struct FNotificationInfo
 	FNotificationInfo(TSharedPtr<INotificationWidget> InContentWidget)
 		: ContentWidget(InContentWidget),
 		Text(),
+		SubText(),
 		ButtonDetails(),
 		Image(nullptr),
 		FadeInDuration(0.5f),
@@ -157,8 +162,11 @@ struct FNotificationInfo
 	/** If set, overrides the entire content of the notification with this widget */
 	TSharedPtr<INotificationWidget> ContentWidget;
 
-	/** The text displayed in this text block */
-	FText Text;
+	/** The text displayed in this notification. Suitable for short notifications and titles. */
+	TAttribute<FText> Text;
+
+	/** Optional subtext displayed in this notification. Subtext is smaller than the default text field and is better for long descriptions.*/
+	TAttribute<FText> SubText;
 
 	/** Setup information for the buttons on the notification */ 
 	TArray<FNotificationButtonInfo> ButtonDetails;
