OFPA: Clear the PKG_NewlyCreated flag on reinstanced actor if the package didn't have the flag before reinstancing occurs.
- This fixes case where a BP was compiled and then a deleted actor couldn't be saved

diff --git a/Engine/Source/Editor/UnrealEd/Private/Kismet2/KismetReinstanceUtilities.cpp b/Engine/Source/Editor/UnrealEd/Private/Kismet2/KismetReinstanceUtilities.cpp
--- a/Engine/Source/Editor/UnrealEd/Private/Kismet2/KismetReinstanceUtilities.cpp
+++ b/Engine/Source/Editor/UnrealEd/Private/Kismet2/KismetReinstanceUtilities.cpp
@@ -2084,6 +2084,8 @@ static void ReplaceActorHelper(AActor* OldActor, UClass* OldClass, UObject*& New
 	// the exact same package, keeping the package name intact.
 	OldActor->UObject::Rename(nullptr, OldActor->GetOuter(), REN_DoNotDirty | REN_DontCreateRedirectors | REN_ForceNoResetLoaders);
 
+	const bool bPackageNewlyCreated = OldActor->GetExternalPackage() && OldActor->GetExternalPackage()->HasAnyPackageFlags(PKG_NewlyCreated);
+	
 	AActor* NewActor = nullptr;
 	{
 		FMakeClassSpawnableOnScope TemporarilySpawnable(SpawnClass);
@@ -2096,6 +2098,19 @@ static void ReplaceActorHelper(AActor* OldActor, UClass* OldClass, UObject*& New
 	}
 
 	check(NewActor != nullptr);
+
+	// When Spawning an actor that has an external package the package can be PKG_NewlyCreated. 
+	// We need to remove this flag if the package didn't have that flag prior to the SpawnActor. 
+	// This means we are reinstancing an actor that was already saved on disk.
+	if (UPackage* ExternalPackage = NewActor->GetExternalPackage())
+	{
+		if (!bPackageNewlyCreated && ExternalPackage->HasAnyPackageFlags(PKG_NewlyCreated))
+		{
+			ExternalPackage->ClearPackageFlags(PKG_NewlyCreated);
+		}
+	}
+
+
 	NewUObject = NewActor;
 	// store the new actor for the second pass (NOTE: this detaches 
 	// OldActor from all child/parent attachments)
