Added "testprops struct" command to easier test the struct details view.

diff --git a/Engine/Source/Editor/UnrealEd/Private/EditorServer.cpp b/Engine/Source/Editor/UnrealEd/Private/EditorServer.cpp
--- a/Engine/Source/Editor/UnrealEd/Private/EditorServer.cpp
+++ b/Engine/Source/Editor/UnrealEd/Private/EditorServer.cpp
@@ -124,6 +124,7 @@
 #include "PropertyEditorModule.h"
 #include "IPropertyTable.h"
 #include "IDetailsView.h"
+#include "IStructureDetailsView.h"
 #include "AssetRegistryModule.h"
 #include "SnappingUtils.h"
 
@@ -5992,7 +5993,6 @@ bool UEditorEngine::HandleTestPropsCommand( const TCHAR* Str, FOutputDevice& Ar
 
 	FPropertyEditorModule& Module = FModuleManager::LoadModuleChecked<FPropertyEditorModule>( "PropertyEditor" );
 
-
 	if( FParse::Command(&Str,TEXT("TREE")) )
 	{
 		FDetailsViewArgs Args;
@@ -6041,6 +6041,24 @@ bool UEditorEngine::HandleTestPropsCommand( const TCHAR* Str, FOutputDevice& Ar
 			]
 		);
 	}
+	else if ( FParse::Command(&Str, TEXT("STRUCT")) )
+	{
+		FDetailsViewArgs DetailsArgs;
+
+		FStructureDetailsViewArgs StructArgs;
+		StructArgs.bShowObjects = true;
+		StructArgs.bShowInterfaces = true;
+
+		TSharedRef<IStructureDetailsView> StructDetails = Module.CreateStructureDetailView(DetailsArgs, StructArgs, MakeShared<FStructOnScope>(FPropertyEditorTestBasicStruct::StaticStruct()));
+
+		Window->SetContent(
+			SNew(SBorder)
+			.BorderImage(FEditorStyle::GetBrush("ToolPanel.GroupBorder"))
+			[
+				StructDetails->GetWidget().ToSharedRef()
+			]
+		);
+	}
 	else
 	{
 		//Details
@@ -6048,8 +6066,6 @@ bool UEditorEngine::HandleTestPropsCommand( const TCHAR* Str, FOutputDevice& Ar
 		Objects.Add( Object );
 
 		FDetailsViewArgs Args;
-		Args.bAllowSearch = true;
-		Args.bUpdatesFromSelection = false;
 		TSharedRef<IDetailsView> DetailsView = Module.CreateDetailView( Args );
 
 		Window->SetContent
