Fix clang shadow warnings

diff --git a/Engine/Source/Editor/UnrealEd/Private/AssetGuideline.cpp b/Engine/Source/Editor/UnrealEd/Private/AssetGuideline.cpp
--- a/Engine/Source/Editor/UnrealEd/Private/AssetGuideline.cpp
+++ b/Engine/Source/Editor/UnrealEd/Private/AssetGuideline.cpp
@@ -145,9 +145,9 @@ void UAssetGuideline::PostLoad()
 		auto GetTextFromState = [NotificationFuture = TextNotificationPromise.GetFuture().Share(), SubText]()
 		{
 			SNotificationItem::ECompletionState State = SNotificationItem::CS_None;
-			if (TSharedPtr<SNotificationItem> NotificationPtr = NotificationFuture.Get().Pin())
+			if (TSharedPtr<SNotificationItem> Notification = NotificationFuture.Get().Pin())
 			{
-				State = NotificationPtr->GetCompletionState();
+				State = Notification->GetCompletionState();
 			}
 
 			switch (State)
@@ -173,9 +173,9 @@ void UAssetGuideline::PostLoad()
 		auto GetHyperlinkTextFromState = [NotificationFuture = HyperlinkNotificationPromise.GetFuture().Share(), HyperlinkText]()
 		{
 			SNotificationItem::ECompletionState State = SNotificationItem::CS_None;
-			if (TSharedPtr<SNotificationItem> NotificationPtr = NotificationFuture.Get().Pin())
+			if (TSharedPtr<SNotificationItem> Notification = NotificationFuture.Get().Pin())
 			{
-				State = NotificationPtr->GetCompletionState();
+				State = Notification->GetCompletionState();
 			}
 
 			// Make hyperlink text on success or fail empty, so that the box auto-resizes correctly.
