Added WorldPartitionRenameCommandlet to rename a partitioned world.

diff --git a/Engine/Source/Editor/UnrealEd/Classes/Commandlets/WorldPartitionRenameCommandlet.h b/Engine/Source/Editor/UnrealEd/Classes/Commandlets/WorldPartitionRenameCommandlet.h
new file mode 100644
--- /dev/null
+++ b/Engine/Source/Editor/UnrealEd/Classes/Commandlets/WorldPartitionRenameCommandlet.h
@@ -0,0 +1,20 @@
+// Copyright Epic Games, Inc. All Rights Reserved.
+
+#pragma once
+
+#include "CoreMinimal.h"
+#include "UObject/ObjectMacros.h"
+#include "UObject/UObjectAnnotation.h"
+#include "Commandlets/Commandlet.h"
+#include "WorldPartitionRenameCommandlet.generated.h"
+
+UCLASS()
+class UNREALED_API UWorldPartitionRenameCommandlet : public UCommandlet
+{
+	GENERATED_UCLASS_BODY()
+
+public:
+	//~ Begin UCommandlet Interface
+	virtual int32 Main(const FString& Params) override;
+	//~ End UCommandlet Interface
+};
diff --git a/Engine/Source/Editor/UnrealEd/Private/Commandlets/WorldPartitionRenameCommandlet.cpp b/Engine/Source/Editor/UnrealEd/Private/Commandlets/WorldPartitionRenameCommandlet.cpp
new file mode 100644
--- /dev/null
+++ b/Engine/Source/Editor/UnrealEd/Private/Commandlets/WorldPartitionRenameCommandlet.cpp
@@ -0,0 +1,289 @@
+// Copyright Epic Games, Inc. All Rights Reserved.
+
+/*=============================================================================
+ UWorldPartitionRenameCommandlet.cpp: Commandlet used to rename a partitioned world
+=============================================================================*/
+
+#include "Commandlets/WorldPartitionRenameCommandlet.h"
+
+#include "Editor.h"
+#include "Engine/World.h"
+#include "SourceControlHelpers.h"
+#include "PackageSourceControlHelper.h"
+#include "WorldPartition/ActorDescList.h"
+#include "WorldPartition/WorldPartition.h"
+#include "WorldPartition/WorldPartitionHandle.h"
+#include "Misc/CommandLine.h"
+#include "Misc/PackageName.h"
+#include "Misc/Paths.h"
+
+DEFINE_LOG_CATEGORY_STATIC(LogWorldPartitionRenameCommandlet, All, All);
+
+struct FSoftPathFixupSerializer : public FArchiveUObject
+{
+	FSoftPathFixupSerializer(TMap<FString, FString>& InRemapSoftObjectPaths)
+	: RemapSoftObjectPaths(InRemapSoftObjectPaths)
+	{
+		this->SetIsSaving(true);
+	}
+
+	FArchive& operator<<(FSoftObjectPath& Value)
+	{
+		if (Value.IsNull())
+		{
+			return *this;
+		}
+
+		FString OriginalValue = Value.ToString();
+
+		auto GetSourceString = [this]()
+		{
+			FString DebugStackString;
+			for (const FName& DebugData: DebugDataStack)
+			{
+				DebugStackString += DebugData.ToString();
+				DebugStackString += TEXT(".");
+			}
+			DebugStackString.RemoveFromEnd(TEXT("."));
+			return DebugStackString;
+		};
+
+		if (FString* RemappedValue = RemapSoftObjectPaths.Find(OriginalValue))
+		{
+			Value.SetPath(*RemappedValue);
+		}
+		else if (Value.GetSubPathString().StartsWith(TEXT("PersistentLevel.")))
+		{
+			int32 DotPos = Value.GetSubPathString().Find(TEXT("."), ESearchCase::IgnoreCase, ESearchDir::FromStart);
+			if (DotPos != INDEX_NONE)
+			{
+				RemappedValue = RemapSoftObjectPaths.Find(Value.GetAssetPathName().ToString());
+				if (RemappedValue)
+				{
+					FString NewPath = *RemappedValue + ':' + Value.GetSubPathString();
+					Value.SetPath(NewPath);
+				}
+		}
+
+			FString NewValue = Value.ToString();
+			if (NewValue == OriginalValue)
+			{
+				Value.Reset();
+				UE_LOG(LogWorldPartitionRenameCommandlet, Warning, TEXT("Error remapping SoftObjectPath %s"), *OriginalValue);
+				UE_LOG(LogWorldPartitionRenameCommandlet, Warning, TEXT("  Source: %s"), *GetSourceString());
+			}
+		}
+
+		if (!Value.IsNull())
+		{
+			FString NewValue = Value.ToString();
+			if (NewValue != OriginalValue)
+			{
+				UE_LOG(LogWorldPartitionRenameCommandlet, Verbose, TEXT("Remapped SoftObjectPath %s to %s"), *OriginalValue, *NewValue);
+				UE_LOG(LogWorldPartitionRenameCommandlet, Verbose, TEXT("  Source: %s"), *GetSourceString());
+			}
+		}
+
+		return *this;
+	}
+
+private:
+	virtual void PushDebugDataString(const FName& DebugData) override
+	{
+		DebugDataStack.Add(DebugData);
+	}
+
+	virtual void PopDebugDataString() override
+	{
+		DebugDataStack.Pop();
+	}
+
+	TArray<FName> DebugDataStack;
+	TMap<FString, FString>& RemapSoftObjectPaths;
+};
+
+UWorldPartitionRenameCommandlet::UWorldPartitionRenameCommandlet(const FObjectInitializer& ObjectInitializer)
+	: Super(ObjectInitializer)
+{}
+
+int32 UWorldPartitionRenameCommandlet::Main(const FString& Params)
+{
+	TArray<UPackage*> PackagesToSave;
+	TArray<FString> PackagesToDelete;
+	FPackageSourceControlHelper PackageHelper;
+
+	TArray<FString> Tokens, Switches;
+	ParseCommandLine(*Params, Tokens, Switches);
+
+	const bool bDeleteSourceMap = Switches.Contains(TEXT("DeleteSourceMap"));
+
+	if (!Switches.Contains(TEXT("AllowCommandletRendering")))
+	{
+		UE_LOG(LogWorldPartitionRenameCommandlet, Error, TEXT("The option \"-AllowCommandletRendering\" is required."));
+		return 1;
+	}
+
+	// Validate old map
+	FString OldMapFullPath;
+	if (!FParse::Value(FCommandLine::Get(), TEXT("oldmap="), OldMapFullPath, false))
+	{
+		UE_LOG(LogWorldPartitionRenameCommandlet, Error, TEXT("Invalid old map name."));
+		return 1;
+	}
+
+	FString OldMapFullPathOnDisk;
+	if (!FPackageName::SearchForPackageOnDisk(OldMapFullPath, &OldMapFullPathOnDisk))
+	{
+		UE_LOG(LogWorldPartitionRenameCommandlet, Error, TEXT("Unknown map '%s'"), *OldMapFullPath);
+		return false;
+	}
+
+	if (OldMapFullPath != OldMapFullPathOnDisk)
+	{
+		UE_LOG(LogWorldPartitionRenameCommandlet, Error, TEXT("Full path is required for map name %s"), *OldMapFullPath);
+		return false;
+	}
+
+	// Validate new map
+	FString NewMapFullPath;
+	if (!FParse::Value(FCommandLine::Get(), TEXT("newmap="), NewMapFullPath, false))
+	{
+		UE_LOG(LogWorldPartitionRenameCommandlet, Error, TEXT("Invalid old map name."));
+		return 1;
+	}
+
+	// Load the map package
+	UPackage* MapPackage = LoadPackage(NULL, *OldMapFullPath, LOAD_None);
+	if (!MapPackage)
+	{
+		UE_LOG(LogWorldPartitionRenameCommandlet, Error, TEXT("Couldn't load package %s."), *OldMapFullPath);
+		return 1;
+	}
+
+	// Find the world in the given package
+	UWorld* World = UWorld::FindWorldInPackage(MapPackage);
+	if (!World)
+	{
+		UE_LOG(LogWorldPartitionRenameCommandlet, Error, TEXT("No world in specified package %s."), *OldMapFullPath);
+		return 1;
+	}
+
+	// Setup the world
+	World->WorldType = EWorldType::Editor;
+	World->AddToRoot();
+	if (!World->bIsWorldInitialized)
+	{
+		UWorld::InitializationValues IVS;
+		IVS.RequiresHitProxies(false);
+		IVS.ShouldSimulatePhysics(false);
+		IVS.EnableTraceCollision(false);
+		IVS.CreateNavigation(false);
+		IVS.CreateAISystem(false);
+		IVS.AllowAudioPlayback(false);
+		IVS.CreatePhysicsScene(true);
+
+		World->InitWorld(UWorld::InitializationValues(IVS));
+		World->PersistentLevel->UpdateModelComponents();
+		World->UpdateWorldComponents(true /*bRerunConstructionScripts*/, false /*bCurrentLevelOnly*/);
+	}
+
+	// Make sure the world is partitioned
+	UWorldPartition* WorldPartition = World->GetWorldPartition();
+	if (!WorldPartition)
+	{
+		UE_LOG(LogWorldPartitionRenameCommandlet, Error, TEXT("Commandlet only works on partitioned maps."));
+		return 1;
+	}
+
+	// Init world
+	FWorldContext& WorldContext = GEditor->GetEditorWorldContext(true /*bEnsureIsGWorld*/);
+	WorldContext.SetCurrentWorld(World);
+	GWorld = World;
+
+	// Soft object paths remappings
+	TMap<FString, FString> RemapSoftObjectPaths;
+
+	// Load all actors
+	TArray<FWorldPartitionReference> ActorReferences;
+	for (FActorDescList::TIterator<> ActorDescIterator(WorldPartition); ActorDescIterator; ++ActorDescIterator)
+	{
+		ActorReferences.Emplace(WorldPartition, ActorDescIterator->GetGuid());
+
+		if (bDeleteSourceMap)
+		{
+			const FString PackageFileName = SourceControlHelpers::PackageFilename(ActorDescIterator->GetActor()->GetPackage());
+			PackagesToDelete.Add(PackageFileName);
+		}
+	}
+
+	if (bDeleteSourceMap)
+	{
+		const FString PackageFileName = SourceControlHelpers::PackageFilename(World->GetPackage());
+		PackagesToDelete.Add(PackageFileName);
+	}
+
+	// Rename world
+	const FString OldWorldPath = FSoftObjectPath(World).ToString();
+	const FString NewWorldName = FPackageName::GetShortName(NewMapFullPath);
+
+	ResetLoaders(World->GetPackage());
+	World->GetPackage()->Rename(*NewMapFullPath, nullptr, REN_NonTransactional | REN_DontCreateRedirectors | REN_ForceNoResetLoaders);
+	World->Rename(*NewWorldName, nullptr, REN_NonTransactional | REN_DontCreateRedirectors | REN_ForceNoResetLoaders);
+
+	const FString NewWorldPath = FSoftObjectPath(World).ToString();	
+	RemapSoftObjectPaths.Add(OldWorldPath, NewWorldPath);
+
+	// Remap soft object paths
+	FSoftPathFixupSerializer FixupSerializer(RemapSoftObjectPaths);
+	ForEachObjectWithPackage(World->GetPackage(), [&](UObject* Object) { Object->Serialize(FixupSerializer); return true; }, true, RF_NoFlags, EInternalObjectFlags::PendingKill);
+	for (FActorDescList::TIterator<> ActorDescIterator(WorldPartition); ActorDescIterator; ++ActorDescIterator)
+	{
+		ForEachObjectWithPackage(ActorDescIterator->GetActor()->GetPackage(), [&](UObject* Object) { Object->Serialize(FixupSerializer); return true; }, true, RF_NoFlags, EInternalObjectFlags::PendingKill);
+		PackagesToSave.Add(ActorDescIterator->GetActor()->GetPackage());
+	}
+
+	PackagesToSave.Add(World->GetPackage());
+
+	for (const FString& PackageToDelete: PackagesToDelete)
+	{
+		if (!PackageHelper.Delete(PackageToDelete))
+		{
+			return 1;
+		}
+	}
+
+	for (UPackage* PackageToSave : PackagesToSave)
+	{
+		const FString PackageFileName = SourceControlHelpers::PackageFilename(PackageToSave);
+
+		if (FPaths::FileExists(PackageFileName))
+		{
+			UE_LOG(LogWorldPartitionRenameCommandlet, Error, TEXT("File %s already exists"), *PackageFileName);
+			return 1;
+		}
+
+		if (!PackageHelper.Checkout(PackageToSave))
+		{
+			return 1;
+		}
+
+		if (!UPackage::SavePackage(PackageToSave, nullptr, RF_Standalone, *PackageFileName, GError, nullptr, false, true, SAVE_Async))
+		{
+			return 1;
+		}
+
+		if(!PackageHelper.AddToSourceControl(PackageToSave))
+		{
+			return 1;
+		}
+	}
+
+	UPackage::WaitForAsyncFileWrites();
+
+	// Cleanup
+	World->RemoveFromRoot();
+	WorldContext.SetCurrentWorld(nullptr);
+	GWorld = nullptr;
+
+	return 0;
+}
\ No newline at end of file
