Optimisation to the legacy texture import and improvement to texture wrapper

Minor change to FImageWrapperBase. GetRaw and GetCompressed now consume the array with the same name instead of having to do a copy of it.
I changed the api IImageWrapper::GetCompressed to return a TArray64<uint8> instead of returning a const TArray64<uint8>&.
Added the format RGBAF to the struct ERGBFormat. Changed the engine code using the EXR image wrapper to reflect that.
The EXR image wrapper now avoid doing an unessary copy of the compressed image when calling compress.

Improvement to the performence of the function UTextureFactory::ImportImage. We now use the magic bytes of the file for certains format to skip some tests.

Here is some performance metrics I captured on my desktop (6 core, 12 threads XEON)

Importing a folder of tiff files (22 files, 4.16 GB Total)
Before: 66.152738 seconds
After: 43.609245 seconds

diff --git a/Engine/Plugins/Experimental/PanoramicCapture/Source/PanoramicCapture/Private/SceneCapturer.cpp b/Engine/Plugins/Experimental/PanoramicCapture/Source/PanoramicCapture/Private/SceneCapturer.cpp
--- a/Engine/Plugins/Experimental/PanoramicCapture/Source/PanoramicCapture/Private/SceneCapturer.cpp
+++ b/Engine/Plugins/Experimental/PanoramicCapture/Source/PanoramicCapture/Private/SceneCapturer.cpp
@@ -984,7 +984,7 @@ TArray<FLinearColor> USceneCapturer::SaveAtlas(FString Folder, const TArray<FLin
 		FString FrameStringUnprojected = FString::Printf(TEXT("%s_%05d_Unprojected.png"), *Folder, CurrentFrameCount);
 		FString AtlasNameUnprojected = OutputDir / Timestamp / FrameStringUnprojected;
 		ImageWrapper->SetRaw(SurfaceData.GetData(), SurfaceData.GetAllocatedSize(), UnprojectedAtlasWidth, UnprojectedAtlasHeight, ERGBFormat::BGRA, 32);
-		const TArray64<uint8>& PNGDataUnprojected = ImageWrapper->GetCompressed(100);
+		const TArray64<uint8> PNGDataUnprojected = ImageWrapper->GetCompressed(100);
 		FFileHelper::SaveArrayToFile(PNGDataUnprojected, *AtlasNameUnprojected);
 		ImageWrapper.Reset();
 	}
@@ -1063,7 +1063,7 @@ void USceneCapturer::CaptureComponent(int32 CurrentHorizontalStep, int32 Current
 			}
 
 			ImageWrapper->SetRaw(SurfaceDataWhole.GetData(), SurfaceDataWhole.GetAllocatedSize(), CaptureWidth, CaptureHeight, ERGBFormat::BGRA, 32);
-			const TArray64<uint8>& PNGData = ImageWrapper->GetCompressed(100);
+			const TArray64<uint8> PNGData = ImageWrapper->GetCompressed(100);
 
 			FFileHelper::SaveArrayToFile(PNGData, *CaptureName);
 			ImageWrapper.Reset();
@@ -1080,7 +1080,7 @@ void USceneCapturer::CaptureComponent(int32 CurrentHorizontalStep, int32 Current
 
 			TSharedPtr<IImageWrapper> ImageWrapper = ImageWrapperModule.CreateImageWrapper(EImageFormat::PNG);
 			ImageWrapper->SetRaw(SurfaceData.GetData(), SurfaceData.GetAllocatedSize(), StripWidth, StripHeight, ERGBFormat::BGRA, 32);
-			const TArray64<uint8>& PNGData = ImageWrapper->GetCompressed(100);
+			const TArray64<uint8> PNGData = ImageWrapper->GetCompressed(100);
 
 			FFileHelper::SaveArrayToFile(PNGData, *CaptureName);
 			ImageWrapper.Reset();
@@ -1267,7 +1267,7 @@ void USceneCapturer::Tick( float DeltaTime )
 				// write
 				TSharedPtr<IImageWrapper> ImageWrapper = ImageWrapperModule.CreateImageWrapper(EImageFormat::PNG);
 				ImageWrapper->SetRaw(CombinedAtlas8bit.GetData(), CombinedAtlas8bit.GetAllocatedSize(), SphericalAtlasWidth, SphericalAtlasHeight * EyeCount, ERGBFormat::BGRA, 8);
-				const TArray64<uint8>& ImageData = ImageWrapper->GetCompressed(100);
+				const TArray64<uint8> ImageData = ImageWrapper->GetCompressed(100);
 				FFileHelper::SaveArrayToFile(ImageData, *AtlasName);
 				ImageWrapper.Reset();
 				FMessageLog(StereoPanoramaLogName).Message(EMessageSeverity::Info, LOCTEXT("Done", "Done!"));
@@ -1282,8 +1282,8 @@ void USceneCapturer::Tick( float DeltaTime )
 
 				// write
 				TSharedPtr<IImageWrapper> ImageWrapper = ImageWrapperModule.CreateImageWrapper(EImageFormat::EXR);
-				ImageWrapper->SetRaw(CombinedAtlas.GetData(), CombinedAtlas.GetAllocatedSize(), SphericalAtlasWidth, SphericalAtlasHeight * EyeCount, ERGBFormat::RGBA, 32);
-				const TArray64<uint8>& ImageData = ImageWrapper->GetCompressed((int32)EImageCompressionQuality::Default);
+				ImageWrapper->SetRaw(CombinedAtlas.GetData(), CombinedAtlas.GetAllocatedSize(), SphericalAtlasWidth, SphericalAtlasHeight * EyeCount, ERGBFormat::RGBAF, 32);
+				const TArray64<uint8> ImageData = ImageWrapper->GetCompressed((int32)EImageCompressionQuality::Default);
 				FFileHelper::SaveArrayToFile(ImageData, *AtlasName);
 				ImageWrapper.Reset();
 				FMessageLog(StereoPanoramaLogName).Message(EMessageSeverity::Info, LOCTEXT("Done", "Done!"));
diff --git a/Engine/Plugins/Experimental/RemoteSession/Source/RemoteSession/Private/Channels/RemoteSessionImageChannel.cpp b/Engine/Plugins/Experimental/RemoteSession/Source/RemoteSession/Private/Channels/RemoteSessionImageChannel.cpp
--- a/Engine/Plugins/Experimental/RemoteSession/Source/RemoteSession/Private/Channels/RemoteSessionImageChannel.cpp
+++ b/Engine/Plugins/Experimental/RemoteSession/Source/RemoteSession/Private/Channels/RemoteSessionImageChannel.cpp
@@ -72,7 +72,7 @@ void FRemoteSessionImageChannel::FImageSender::SendRawImageToClients(int32 Width
 			ImageWrapper->SetRaw(ImageData, AllocatedImageDataSize, Width, Height, ERGBFormat::BGRA, 8);
 
 			const int32 CurrentQuality = QualityMasterSetting > 0 ? QualityMasterSetting : CompressQuality.Load();
-			const TArray64<uint8>& JPGData = ImageWrapper->GetCompressed(CurrentQuality);
+			const TArray64<uint8> JPGData = ImageWrapper->GetCompressed(CurrentQuality);
 
 			TBackChannelSharedPtr<FBackChannelOSCMessage> Msg = MakeShared<FBackChannelOSCMessage, ESPMode::ThreadSafe>(TEXT("/Screen"));
 	
diff --git a/Engine/Plugins/Importers/USDImporter/Source/USDUtilities/Private/USDShadeConversion.cpp b/Engine/Plugins/Importers/USDImporter/Source/USDUtilities/Private/USDShadeConversion.cpp
--- a/Engine/Plugins/Importers/USDImporter/Source/USDUtilities/Private/USDShadeConversion.cpp
+++ b/Engine/Plugins/Importers/USDImporter/Source/USDUtilities/Private/USDShadeConversion.cpp
@@ -942,7 +942,7 @@ namespace UsdShadeConversionImpl
 				}
 
 				PNGImageWrapper->SetRaw( Samples.GetData(), Samples.GetAllocatedSize(), FinalSize.X, FinalSize.Y, ERGBFormat::BGRA, 8 );
-				const TArray64<uint8>& PNGData = PNGImageWrapper->GetCompressed( 100 );
+				const TArray64<uint8> PNGData = PNGImageWrapper->GetCompressed( 100 );
 
 				bool bWroteFile = FFileHelper::SaveArrayToFile( PNGData, *TextureFilePath );
 				if ( bWroteFile )
diff --git a/Engine/Plugins/Media/PixelStreaming/Source/PixelStreaming/Private/PixelStreamingModule.cpp b/Engine/Plugins/Media/PixelStreaming/Source/PixelStreaming/Private/PixelStreamingModule.cpp
--- a/Engine/Plugins/Media/PixelStreaming/Source/PixelStreaming/Private/PixelStreamingModule.cpp
+++ b/Engine/Plugins/Media/PixelStreaming/Source/PixelStreaming/Private/PixelStreamingModule.cpp
@@ -400,7 +400,7 @@ void FPixelStreamingModule::SendJpeg(TArray<FColor> RawData, const FIntRect& Rec
 	{
 		// Compress to a JPEG of the maximum possible quality.
 		int32 Quality = CVarFreezeFrameQuality.GetValueOnAnyThread();
-		const TArray64<uint8>& JpegBytes = ImageWrapper->GetCompressed(Quality);
+		const TArray64<uint8> JpegBytes = ImageWrapper->GetCompressed(Quality);
 		Streamer->SendFreezeFrame(JpegBytes);
 	}
 	else
diff --git a/Engine/Plugins/MovieScene/MovieRenderPipeline/Source/MovieRenderPipelineRenderPasses/Private/MoviePipelineEXROutput.cpp b/Engine/Plugins/MovieScene/MovieRenderPipeline/Source/MovieRenderPipelineRenderPasses/Private/MoviePipelineEXROutput.cpp
--- a/Engine/Plugins/MovieScene/MovieRenderPipeline/Source/MovieRenderPipelineRenderPasses/Private/MoviePipelineEXROutput.cpp
+++ b/Engine/Plugins/MovieScene/MovieRenderPipeline/Source/MovieRenderPipelineRenderPasses/Private/MoviePipelineEXROutput.cpp
@@ -254,6 +254,7 @@ static FString GetChannelName(const FString& InLayerName, const int32 InChannelI
 	switch (InFormat)
 	{
 		case ERGBFormat::RGBA:
+		case ERGBFormat::RGBAF:
 		{
 			ChannelNames = RGBAChannelNames;
 		}
diff --git a/Engine/Plugins/Runtime/Oculus/OculusVR/Source/OculusHMD/Private/OculusSceneCaptureCubemap.cpp b/Engine/Plugins/Runtime/Oculus/OculusVR/Source/OculusHMD/Private/OculusSceneCaptureCubemap.cpp
--- a/Engine/Plugins/Runtime/Oculus/OculusVR/Source/OculusHMD/Private/OculusSceneCaptureCubemap.cpp
+++ b/Engine/Plugins/Runtime/Oculus/OculusVR/Source/OculusHMD/Private/OculusSceneCaptureCubemap.cpp
@@ -142,7 +142,7 @@ void UOculusSceneCaptureCubemap::Tick(float DeltaTime)
 	}
 
 	ImageWrapper->SetRaw(WholeCubemapData.GetData(), WholeCubemapData.GetAllocatedSize(), CaptureBoxSideRes * 6, CaptureBoxSideRes, ERGBFormat::BGRA, 8);
-	const TArray64<uint8>& PNGData = ImageWrapper->GetCompressed(100);
+	const TArray64<uint8> PNGData = ImageWrapper->GetCompressed(100);
 
 	const FString Filename = OutputDir + FString::Printf(TEXT("/Cubemap-%d-%s.png"), CaptureBoxSideRes, *FDateTime::Now().ToString(TEXT("%m.%d-%H.%M.%S")));
 
diff --git a/Engine/Source/Developer/ScreenShotComparisonTools/Private/ImageComparer.cpp b/Engine/Source/Developer/ScreenShotComparisonTools/Private/ImageComparer.cpp
--- a/Engine/Source/Developer/ScreenShotComparisonTools/Private/ImageComparer.cpp
+++ b/Engine/Source/Developer/ScreenShotComparisonTools/Private/ImageComparer.cpp
@@ -98,7 +98,7 @@ public:
 		{
 			if ( ImageWriter->SetRaw(Image.GetData(), Image.Num(), Width, Height, ERGBFormat::RGBA, 8) )
 			{
-				const TArray64<uint8>& PngData = ImageWriter->GetCompressed();
+				TArray64<uint8> PngData = ImageWriter->GetCompressed();
 
 				if ( FFileHelper::SaveArrayToFile(PngData, *TempDeltaFile) )
 				{
diff --git a/Engine/Source/Developer/TextureFormatASTC/Private/TextureFormatASTC.cpp b/Engine/Source/Developer/TextureFormatASTC/Private/TextureFormatASTC.cpp
--- a/Engine/Source/Developer/TextureFormatASTC/Private/TextureFormatASTC.cpp
+++ b/Engine/Source/Developer/TextureFormatASTC/Private/TextureFormatASTC.cpp
@@ -176,7 +176,7 @@ static bool CompressSliceToASTC(
 	// Compress and retrieve the PNG data to write out to disk
 	TSharedPtr<IImageWrapper> ImageWrapper = ImageWrapperModule.CreateImageWrapper(EImageFormat::PNG);
 	ImageWrapper->SetRaw(SourceData, SizeX * SizeY * 4, SizeX, SizeY, ERGBFormat::RGBA, 8);
-	const TArray64<uint8>& FileData = ImageWrapper->GetCompressed();
+	const TArray64<uint8> FileData = ImageWrapper->GetCompressed();
 	int64 FileDataSize = FileData.Num();
 
 	FGuid Guid;
diff --git a/Engine/Source/Editor/LandscapeEditor/Private/LandscapeFileFormatPng.cpp b/Engine/Source/Editor/LandscapeEditor/Private/LandscapeFileFormatPng.cpp
--- a/Engine/Source/Editor/LandscapeEditor/Private/LandscapeFileFormatPng.cpp
+++ b/Engine/Source/Editor/LandscapeEditor/Private/LandscapeFileFormatPng.cpp
@@ -143,7 +143,7 @@ void FLandscapeHeightmapFileFormat_Png::Export(const TCHAR* HeightmapFilename, F
 
 	if (ImageWrapper->SetRaw(Data.GetData(), Data.Num() * 2, DataResolution.Width, DataResolution.Height, ERGBFormat::Gray, 16))
 	{
-		const TArray64<uint8>& TempData = ImageWrapper->GetCompressed();
+		const TArray64<uint8> TempData = ImageWrapper->GetCompressed();
 		FFileHelper::SaveArrayToFile(TempData, HeightmapFilename);
 	}
 }
@@ -244,7 +244,7 @@ void FLandscapeWeightmapFileFormat_Png::Export(const TCHAR* WeightmapFilename, F
 
 	if (ImageWrapper->SetRaw(Data.GetData(), Data.Num(), DataResolution.Width, DataResolution.Height, ERGBFormat::Gray, 8))
 	{
-		const TArray64<uint8>& TempData = ImageWrapper->GetCompressed();
+		const TArray64<uint8> TempData = ImageWrapper->GetCompressed();
 		FFileHelper::SaveArrayToFile(TempData, WeightmapFilename);
 	}
 }
diff --git a/Engine/Source/Editor/UnrealEd/Private/Factories/EditorFactories.cpp b/Engine/Source/Editor/UnrealEd/Private/Factories/EditorFactories.cpp
--- a/Engine/Source/Editor/UnrealEd/Private/Factories/EditorFactories.cpp
+++ b/Engine/Source/Editor/UnrealEd/Private/Factories/EditorFactories.cpp
@@ -3011,228 +3011,240 @@ bool UTextureFactory::ImportImage(const uint8* Buffer, uint32 Length, FFeedbackC
 
 	IImageWrapperModule& ImageWrapperModule = FModuleManager::LoadModuleChecked<IImageWrapperModule>(FName("ImageWrapper"));
 
+	// Use the magic bytes when possible to avoid calling inefficient code to check if the image is of the right format
+	EImageFormat ImageFormat = ImageWrapperModule.DetectImageFormat(Buffer, int64(Length));
+
 	//
 	// PNG
 	//
-	TSharedPtr<IImageWrapper> PngImageWrapper = ImageWrapperModule.CreateImageWrapper(EImageFormat::PNG);
-	if (PngImageWrapper.IsValid() && PngImageWrapper->SetCompressed(Buffer, Length))
+	if (ImageFormat == EImageFormat::PNG)
 	{
-		if (!IsImportResolutionValid(PngImageWrapper->GetWidth(), PngImageWrapper->GetHeight(), bAllowNonPowerOfTwo, Warn))
+		TSharedPtr<IImageWrapper> PngImageWrapper = ImageWrapperModule.CreateImageWrapper(EImageFormat::PNG);
+		if (PngImageWrapper.IsValid() && PngImageWrapper->SetCompressed(Buffer, Length))
 		{
-			return false;
-		}
-
-		// Select the texture's source format
-		ETextureSourceFormat TextureFormat = TSF_Invalid;
-		int32 BitDepth = PngImageWrapper->GetBitDepth();
-		ERGBFormat Format = PngImageWrapper->GetFormat();
-
-		if (Format == ERGBFormat::Gray)
-		{
-			if (BitDepth <= 8)
+			if (!IsImportResolutionValid(PngImageWrapper->GetWidth(), PngImageWrapper->GetHeight(), bAllowNonPowerOfTwo, Warn))
 			{
-				TextureFormat = TSF_G8;
-				Format = ERGBFormat::Gray;
-				BitDepth = 8;
+				return false;
 			}
-			else if (BitDepth == 16)
+
+			// Select the texture's source format
+			ETextureSourceFormat TextureFormat = TSF_Invalid;
+			int32 BitDepth = PngImageWrapper->GetBitDepth();
+			ERGBFormat Format = PngImageWrapper->GetFormat();
+
+			if (Format == ERGBFormat::Gray)
 			{
-				TextureFormat = TSF_G16;
-				Format = ERGBFormat::Gray;
-				BitDepth = 16;
+				if (BitDepth <= 8)
+				{
+					TextureFormat = TSF_G8;
+					Format = ERGBFormat::Gray;
+					BitDepth = 8;
+				}
+				else if (BitDepth == 16)
+				{
+					TextureFormat = TSF_G16;
+					Format = ERGBFormat::Gray;
+					BitDepth = 16;
+				}
 			}
-		}
-		else if (Format == ERGBFormat::RGBA || Format == ERGBFormat::BGRA)
-		{
-			if (BitDepth <= 8)
+			else if (Format == ERGBFormat::RGBA || Format == ERGBFormat::BGRA)
 			{
-				TextureFormat = TSF_BGRA8;
-				Format = ERGBFormat::BGRA;
-				BitDepth = 8;
+				if (BitDepth <= 8)
+				{
+					TextureFormat = TSF_BGRA8;
+					Format = ERGBFormat::BGRA;
+					BitDepth = 8;
+				}
+				else if (BitDepth == 16)
+				{
+					TextureFormat = TSF_RGBA16;
+					Format = ERGBFormat::RGBA;
+					BitDepth = 16;
+				}
 			}
-			else if (BitDepth == 16)
+
+			if (TextureFormat == TSF_Invalid)
 			{
-				TextureFormat = TSF_RGBA16;
-				Format = ERGBFormat::RGBA;
-				BitDepth = 16;
+				Warn->Logf(ELogVerbosity::Error, TEXT("PNG file contains data in an unsupported format."));
+				return false;
 			}
-		}
-
-		if (TextureFormat == TSF_Invalid)
-		{
-			Warn->Logf(ELogVerbosity::Error, TEXT("PNG file contains data in an unsupported format."));
-			return false;
-		}
 
-		OutImage.Init2DWithParams(
-			PngImageWrapper->GetWidth(),
-			PngImageWrapper->GetHeight(),
-			TextureFormat,
-			BitDepth < 16
-		);
+			OutImage.Init2DWithParams(
+				PngImageWrapper->GetWidth(),
+				PngImageWrapper->GetHeight(),
+				TextureFormat,
+				BitDepth < 16
+			);
 
-		if (PngImageWrapper->GetRaw(Format, BitDepth, OutImage.RawData))
-		{
-			bool bFillPNGZeroAlpha = true;
-			GConfig->GetBool(TEXT("TextureImporter"), TEXT("FillPNGZeroAlpha"), bFillPNGZeroAlpha, GEditorIni);
+			if (PngImageWrapper->GetRaw(Format, BitDepth, OutImage.RawData))
+			{
+				bool bFillPNGZeroAlpha = true;
+				GConfig->GetBool(TEXT("TextureImporter"), TEXT("FillPNGZeroAlpha"), bFillPNGZeroAlpha, GEditorIni);
 
-			if (bFillPNGZeroAlpha)
+				if (bFillPNGZeroAlpha)
+				{
+					// Replace the pixels with 0.0 alpha with a color value from the nearest neighboring color which has a non-zero alpha
+					FillZeroAlphaPNGData(OutImage.SizeX, OutImage.SizeY, OutImage.Format, OutImage.RawData.GetData());
+				}
+			}
+			else
 			{
-				// Replace the pixels with 0.0 alpha with a color value from the nearest neighboring color which has a non-zero alpha
-				FillZeroAlphaPNGData(OutImage.SizeX, OutImage.SizeY, OutImage.Format, OutImage.RawData.GetData());
+				Warn->Logf(ELogVerbosity::Error, TEXT("Failed to decode PNG."));
+				return false;
 			}
-		}
-		else
-		{
-			Warn->Logf(ELogVerbosity::Error, TEXT("Failed to decode PNG."));
-			return false;
-		}
 
-		return true;
+			return true;
+		}
 	}
 
 	//
 	// JPEG
 	//
-	TSharedPtr<IImageWrapper> JpegImageWrapper = ImageWrapperModule.CreateImageWrapper(EImageFormat::JPEG);
-	if (JpegImageWrapper.IsValid() && JpegImageWrapper->SetCompressed(Buffer, Length))
+	if (ImageFormat == EImageFormat::JPEG)
 	{
-		if (!IsImportResolutionValid(JpegImageWrapper->GetWidth(), JpegImageWrapper->GetHeight(), bAllowNonPowerOfTwo, Warn))
+		TSharedPtr<IImageWrapper> JpegImageWrapper = ImageWrapperModule.CreateImageWrapper(EImageFormat::JPEG);
+		if (JpegImageWrapper.IsValid() && JpegImageWrapper->SetCompressed(Buffer, Length))
 		{
-			return false;
-		}
+			if (!IsImportResolutionValid(JpegImageWrapper->GetWidth(), JpegImageWrapper->GetHeight(), bAllowNonPowerOfTwo, Warn))
+			{
+				return false;
+			}
 
-		// Select the texture's source format
-		ETextureSourceFormat TextureFormat = TSF_Invalid;
-		int32 BitDepth = JpegImageWrapper->GetBitDepth();
-		ERGBFormat Format = JpegImageWrapper->GetFormat();
+			// Select the texture's source format
+			ETextureSourceFormat TextureFormat = TSF_Invalid;
+			int32 BitDepth = JpegImageWrapper->GetBitDepth();
+			ERGBFormat Format = JpegImageWrapper->GetFormat();
 
-		if (Format == ERGBFormat::Gray)
-		{
-			if (BitDepth <= 8)
+			if (Format == ERGBFormat::Gray)
 			{
-				TextureFormat =  TSF_G8;
-				Format = ERGBFormat::Gray;
-				BitDepth = 8;
+				if (BitDepth <= 8)
+				{
+					TextureFormat = TSF_G8;
+					Format = ERGBFormat::Gray;
+					BitDepth = 8;
+				}
 			}
-		}
-		else if (Format == ERGBFormat::RGBA)
-		{
-			if (BitDepth <= 8)
+			else if (Format == ERGBFormat::RGBA)
 			{
-				TextureFormat = TSF_BGRA8;
-				Format = ERGBFormat::BGRA;
-				BitDepth = 8;
+				if (BitDepth <= 8)
+				{
+					TextureFormat = TSF_BGRA8;
+					Format = ERGBFormat::BGRA;
+					BitDepth = 8;
+				}
 			}
-		}
 
-		if (TextureFormat == TSF_Invalid)
-		{
-			Warn->Logf(ELogVerbosity::Error, TEXT("JPEG file contains data in an unsupported format."));
-			return false;
-		}
+			if (TextureFormat == TSF_Invalid)
+			{
+				Warn->Logf(ELogVerbosity::Error, TEXT("JPEG file contains data in an unsupported format."));
+				return false;
+			}
 
-		OutImage.Init2DWithParams(
-			JpegImageWrapper->GetWidth(),
-			JpegImageWrapper->GetHeight(),
-			TextureFormat,
-			BitDepth < 16
-		);
-		
-		bool bRetainJpegFormat = false;
-		if (EnumHasAllFlags(Flags, EImageImportFlags::AllowReturnOfCompressedData))
-		{
-			// For now this option is opt in via the config files once there is no technical risk this will become the default path.
-			GConfig->GetBool(TEXT("TextureImporter"), TEXT("RetainJpegFormat"), bRetainJpegFormat, GEditorIni);
-		}	
+			OutImage.Init2DWithParams(
+				JpegImageWrapper->GetWidth(),
+				JpegImageWrapper->GetHeight(),
+				TextureFormat,
+				BitDepth < 16
+			);
 
-		if (bRetainJpegFormat)
-		{
-			OutImage.RawData.Append(Buffer, Length);
-			OutImage.RawDataCompressionFormat = ETextureSourceCompressionFormat::TSCF_JPEG;
-		}
-		else
-		{
-			if (!JpegImageWrapper->GetRaw(Format, BitDepth, OutImage.RawData))
+			bool bRetainJpegFormat = false;
+			if (EnumHasAllFlags(Flags, EImageImportFlags::AllowReturnOfCompressedData))
+			{
+				// For now this option is opt in via the config files once there is no technical risk this will become the default path.
+				GConfig->GetBool(TEXT("TextureImporter"), TEXT("RetainJpegFormat"), bRetainJpegFormat, GEditorIni);
+			}
+
+			if (bRetainJpegFormat)
+			{
+				OutImage.RawData.Append(Buffer, Length);
+				OutImage.RawDataCompressionFormat = ETextureSourceCompressionFormat::TSCF_JPEG;
+			}
+			else if (!JpegImageWrapper->GetRaw(Format, BitDepth, OutImage.RawData))
 			{
 				Warn->Logf(ELogVerbosity::Error, TEXT("Failed to decode JPEG."));
 				return false;
 			}
-		}
 
-		return true;
+			return true;
+		}
 	}
 
 	//
 	// EXR
 	//
-	TSharedPtr<IImageWrapper> ExrImageWrapper = ImageWrapperModule.CreateImageWrapper(EImageFormat::EXR);
-	if (ExrImageWrapper.IsValid() && ExrImageWrapper->SetCompressed(Buffer, Length))
+	if (ImageFormat == EImageFormat::EXR)
 	{
-		int32 Width = ExrImageWrapper->GetWidth();
-		int32 Height = ExrImageWrapper->GetHeight();
-
-		if (!IsImportResolutionValid(Width, Height, bAllowNonPowerOfTwo, Warn))
+		TSharedPtr<IImageWrapper> ExrImageWrapper = ImageWrapperModule.CreateImageWrapper(EImageFormat::EXR);
+		if (ExrImageWrapper.IsValid() && ExrImageWrapper->SetCompressed(Buffer, Length))
 		{
-			return false;
-		}
+			int32 Width = ExrImageWrapper->GetWidth();
+			int32 Height = ExrImageWrapper->GetHeight();
 
-		// Select the texture's source format
-		ETextureSourceFormat TextureFormat = TSF_Invalid;
-		int32 BitDepth = ExrImageWrapper->GetBitDepth();
-		ERGBFormat Format = ExrImageWrapper->GetFormat();
+			if (!IsImportResolutionValid(Width, Height, bAllowNonPowerOfTwo, Warn))
+			{
+				return false;
+			}
 
-		if (Format == ERGBFormat::RGBA && BitDepth == 16)
-		{
-			TextureFormat = TSF_RGBA16F;
-			Format = ERGBFormat::BGRA;
-		}
+			// Select the texture's source format
+			ETextureSourceFormat TextureFormat = TSF_Invalid;
+			int32 BitDepth = ExrImageWrapper->GetBitDepth();
+			ERGBFormat Format = ExrImageWrapper->GetFormat();
 
-		if (TextureFormat == TSF_Invalid)
-		{
-			Warn->Logf(ELogVerbosity::Error, TEXT("EXR file contains data in an unsupported format."));
-			return false;
-		}
+			if (Format == ERGBFormat::RGBAF && BitDepth == 16)
+			{
+				TextureFormat = TSF_RGBA16F;
+			}
 
-		OutImage.Init2DWithParams(
-			Width,
-			Height,
-			TextureFormat,
-			false
-		);
-		OutImage.CompressionSettings = TC_HDR;
+			if (TextureFormat == TSF_Invalid)
+			{
+				Warn->Logf(ELogVerbosity::Error, TEXT("EXR file contains data in an unsupported format."));
+				return false;
+			}
 
-		if (!ExrImageWrapper->GetRaw(Format, BitDepth, OutImage.RawData))
-		{
-			Warn->Logf(ELogVerbosity::Error, TEXT("Failed to decode EXR."));
-			return false;
-		}
+			OutImage.Init2DWithParams(
+				Width,
+				Height,
+				TextureFormat,
+				false
+			);
+			OutImage.CompressionSettings = TC_HDR;
 
-		return true;
+			if (!ExrImageWrapper->GetRaw(Format, BitDepth, OutImage.RawData))
+			{
+				Warn->Logf(ELogVerbosity::Error, TEXT("Failed to decode EXR."));
+				return false;
+			}
+
+			return true;
+		}
 	}
 
 	//
 	// BMP
 	//
-	TSharedPtr<IImageWrapper> BmpImageWrapper = ImageWrapperModule.CreateImageWrapper(EImageFormat::BMP);
-	if (BmpImageWrapper.IsValid() && BmpImageWrapper->SetCompressed(Buffer, Length))
+	if (ImageFormat == EImageFormat::BMP)
 	{
-		// Check the resolution of the imported texture to ensure validity
-		if (!IsImportResolutionValid(BmpImageWrapper->GetWidth(), BmpImageWrapper->GetHeight(), bAllowNonPowerOfTwo, Warn))
+		TSharedPtr<IImageWrapper> BmpImageWrapper = ImageWrapperModule.CreateImageWrapper(EImageFormat::BMP);
+		if (BmpImageWrapper.IsValid() && BmpImageWrapper->SetCompressed(Buffer, Length))
 		{
-			return false;
-		}
+			// Check the resolution of the imported texture to ensure validity
+			if (!IsImportResolutionValid(BmpImageWrapper->GetWidth(), BmpImageWrapper->GetHeight(), bAllowNonPowerOfTwo, Warn))
+			{
+				return false;
+			}
 
-		OutImage.Init2DWithParams(
-			BmpImageWrapper->GetWidth(),
-			BmpImageWrapper->GetHeight(),
-			TSF_BGRA8,
-			false
-		);
+			OutImage.Init2DWithParams(
+				BmpImageWrapper->GetWidth(),
+				BmpImageWrapper->GetHeight(),
+				TSF_BGRA8,
+				false
+			);
 
-		return BmpImageWrapper->GetRaw(BmpImageWrapper->GetFormat(), BmpImageWrapper->GetBitDepth(), OutImage.RawData);
+			return BmpImageWrapper->GetRaw(BmpImageWrapper->GetFormat(), BmpImageWrapper->GetBitDepth(), OutImage.RawData);
+		}
 	}
 
+
 	//
 	// PCX
 	//
@@ -3348,6 +3360,7 @@ bool UTextureFactory::ImportImage(const uint8* Buffer, uint32 Length, FFeedbackC
 
 		return true;
 	}
+
 	//
 	// TGA
 	//
@@ -3528,6 +3541,9 @@ bool UTextureFactory::ImportImage(const uint8* Buffer, uint32 Length, FFeedbackC
 		}
 	}
 
+	//
+	// TIFF
+	//
 	FTiffLoadHelper TiffLoaderHelper;
 	if (TiffLoaderHelper.IsValid())
 	{
diff --git a/Engine/Source/Runtime/Engine/Private/ImageUtils.cpp b/Engine/Source/Runtime/Engine/Private/ImageUtils.cpp
--- a/Engine/Source/Runtime/Engine/Private/ImageUtils.cpp
+++ b/Engine/Source/Runtime/Engine/Private/ImageUtils.cpp
@@ -801,7 +801,7 @@ bool FImageUtils::ExportRenderTarget2DAsPNG(UTextureRenderTarget2D* TexRT, FArch
 
 		PNGImageWrapper->SetRaw(RawData.GetData(), RawData.GetAllocatedSize(), Size.X, Size.Y, ERGBFormat::BGRA, 8);
 
-		const TArray64<uint8>& PNGData = PNGImageWrapper->GetCompressed(100);
+		const TArray64<uint8> PNGData = PNGImageWrapper->GetCompressed(100);
 
 		Ar.Serialize((void*)PNGData.GetData(), PNGData.GetAllocatedSize());
 	}
@@ -830,7 +830,7 @@ ENGINE_API bool FImageUtils::ExportRenderTarget2DAsEXR(UTextureRenderTarget2D* T
 
 		EXRImageWrapper->SetRaw(RawData.GetData(), RawData.GetAllocatedSize(), Size.X, Size.Y, RGBFormat, BitsPerPixel);
 
-		const TArray64<uint8>& Data = EXRImageWrapper->GetCompressed(100);
+		const TArray64<uint8> Data = EXRImageWrapper->GetCompressed(100);
 
 		Ar.Serialize((void*)Data.GetData(), Data.GetAllocatedSize());
 
diff --git a/Engine/Source/Runtime/Engine/Private/Texture.cpp b/Engine/Source/Runtime/Engine/Private/Texture.cpp
--- a/Engine/Source/Runtime/Engine/Private/Texture.cpp
+++ b/Engine/Source/Runtime/Engine/Private/Texture.cpp
@@ -1184,7 +1184,7 @@ void FTextureSource::Compress()
 		ERGBFormat RawFormat = (Format == TSF_G8 || Format == TSF_G16) ? ERGBFormat::Gray : ERGBFormat::RGBA;
 		if ( ImageWrapper.IsValid() && ImageWrapper->SetRaw( BulkDataPtr, BulkData.GetBulkDataSize(), SizeX, SizeY, RawFormat, (Format == TSF_G16 || Format == TSF_RGBA16) ? 16 : 8 ) )
 		{
-			const TArray64<uint8>& CompressedData = ImageWrapper->GetCompressed();
+			const TArray64<uint8> CompressedData = ImageWrapper->GetCompressed();
 			if ( CompressedData.Num() > 0 )
 			{
 				BulkDataPtr = (uint8*)BulkData.Realloc(CompressedData.Num());
diff --git a/Engine/Source/Runtime/Engine/Private/VT/VirtualTextureDataBuilder.cpp b/Engine/Source/Runtime/Engine/Private/VT/VirtualTextureDataBuilder.cpp
--- a/Engine/Source/Runtime/Engine/Private/VT/VirtualTextureDataBuilder.cpp
+++ b/Engine/Source/Runtime/Engine/Private/VT/VirtualTextureDataBuilder.cpp
@@ -225,6 +225,7 @@ struct FPixelDataRectangle
 			ImageWrapper->SetRaw(Data, BytesPerPixel * Width * Height, Width, Height, ERGBFormat::RGBA, 16);
 			break;
 		case TSF_RGBA16F:
+			// This will probably result in bad image. (Should we convert the data or export it in a format that support float(ERGBFormat::RGBAF)?)
 			ImageWrapper->SetRaw(Data, BytesPerPixel * Width * Height, Width, Height, ERGBFormat::RGBA, 16);
 			break;
 		case TSF_RGBA8:
@@ -245,7 +246,7 @@ struct FPixelDataRectangle
 		FArchive* Ar = FileManager->CreateFileWriter(*Filename);
 		if (Ar != nullptr)
 		{
-			const TArray<uint8>& CompressedData = ImageWrapper->GetCompressed((int32)EImageCompressionQuality::Uncompressed);
+			TArray64<uint8> CompressedData = ImageWrapper->GetCompressed((int32)EImageCompressionQuality::Uncompressed);
 			Ar->Serialize((void *)CompressedData.GetData(), CompressedData.Num());
 			delete Ar;
 		}
diff --git a/Engine/Source/Runtime/ImageWrapper/Private/Formats/ExrImageWrapper.cpp b/Engine/Source/Runtime/ImageWrapper/Private/Formats/ExrImageWrapper.cpp
--- a/Engine/Source/Runtime/ImageWrapper/Private/Formats/ExrImageWrapper.cpp
+++ b/Engine/Source/Runtime/ImageWrapper/Private/Formats/ExrImageWrapper.cpp
@@ -241,6 +241,7 @@ namespace
 		{
 			case ERGBFormat::RGBA:
 			case ERGBFormat::BGRA:
+			case ERGBFormat::RGBAF:
 				return 4;
 			case ERGBFormat::Gray:
 				return 1;
@@ -263,6 +264,7 @@ const char* FExrImageWrapper::GetRawChannelName(int ChannelIndex) const
 	switch (RawFormat)
 	{
 		case ERGBFormat::RGBA:
+		case ERGBFormat::RGBAF:
 		{
 			ChannelNames = RGBAChannelNames;
 		}
@@ -346,8 +348,7 @@ void FExrImageWrapper::CompressRaw(const sourcetype* SrcData, bool bIgnoreAlpha)
 		FileLength = MemFile.tellp();
 	}
 
-	CompressedData.AddUninitialized(FileLength);
-	FMemory::Memcpy(CompressedData.GetData(), MemFile.Data.GetData(), FileLength);
+	CompressedData = MoveTemp(MemFile.Data);
 
 	const double DeltaTime = FPlatformTime::Seconds() - StartTime;
 	UE_LOG(LogImageWrapper, Verbose, TEXT("Compressed image in %.3f seconds"), DeltaTime);
@@ -456,8 +457,7 @@ bool FExrImageWrapper::SetCompressed( const void* InCompressedData, int64 InComp
 	Width = dim.x;
 	Height = dim.y;
 
-	// ideally we can specify float here
-	Format = ERGBFormat::RGBA;
+	Format = ERGBFormat::RGBAF;
 
 	return true;
 }
diff --git a/Engine/Source/Runtime/ImageWrapper/Private/ImageWrapperBase.cpp b/Engine/Source/Runtime/ImageWrapper/Private/ImageWrapperBase.cpp
--- a/Engine/Source/Runtime/ImageWrapper/Private/ImageWrapperBase.cpp
+++ b/Engine/Source/Runtime/ImageWrapper/Private/ImageWrapperBase.cpp
@@ -45,12 +45,12 @@ void FImageWrapperBase::SetError(const TCHAR* ErrorMessage)
 /* IImageWrapper structors
  *****************************************************************************/
 
-const TArray64<uint8>& FImageWrapperBase::GetCompressed(int32 Quality)
+TArray64<uint8> FImageWrapperBase::GetCompressed(int32 Quality)
 {
 	LastError.Empty();
 	Compress(Quality);
 
-	return CompressedData;
+	return MoveTemp(CompressedData);
 }
 
 
diff --git a/Engine/Source/Runtime/ImageWrapper/Private/ImageWrapperBase.h b/Engine/Source/Runtime/ImageWrapper/Private/ImageWrapperBase.h
--- a/Engine/Source/Runtime/ImageWrapper/Private/ImageWrapperBase.h
+++ b/Engine/Source/Runtime/ImageWrapper/Private/ImageWrapperBase.h
@@ -74,7 +74,7 @@ public:
 
 	//~ IImageWrapper interface
 
-	virtual const TArray64<uint8>& GetCompressed(int32 Quality = 0) override;
+	virtual TArray64<uint8> GetCompressed(int32 Quality = 0) override;
 
 	virtual int32 GetBitDepth() const override
 	{
diff --git a/Engine/Source/Runtime/ImageWrapper/Private/ImageWrapperModule.cpp b/Engine/Source/Runtime/ImageWrapper/Private/ImageWrapperModule.cpp
--- a/Engine/Source/Runtime/ImageWrapper/Private/ImageWrapperModule.cpp
+++ b/Engine/Source/Runtime/ImageWrapper/Private/ImageWrapperModule.cpp
@@ -60,53 +60,53 @@ public:
 
 	virtual TSharedPtr<IImageWrapper> CreateImageWrapper(const EImageFormat InFormat) override
 	{
-		FImageWrapperBase* ImageWrapper = NULL;
+		TSharedPtr<IImageWrapper> ImageWrapper;
 
 		// Allocate a helper for the format type
 		switch(InFormat)
 		{
 #if WITH_UNREALPNG
 		case EImageFormat::PNG:
-			ImageWrapper = new FPngImageWrapper();
+			ImageWrapper = MakeShared<FPngImageWrapper>();
 			break;
 #endif	// WITH_UNREALPNG
 
 #if WITH_UNREALJPEG
 		case EImageFormat::JPEG:
-			ImageWrapper = new FJpegImageWrapper();
+			ImageWrapper = MakeShared<FJpegImageWrapper>();
 			break;
 
 		case EImageFormat::GrayscaleJPEG:
-			ImageWrapper = new FJpegImageWrapper(1);
+			ImageWrapper = MakeShared<FJpegImageWrapper>(1);
 			break;
 #endif	//WITH_UNREALJPEG
 
 		case EImageFormat::BMP:
-			ImageWrapper = new FBmpImageWrapper();
+			ImageWrapper = MakeShared<FBmpImageWrapper>();
 			break;
 
 		case EImageFormat::ICO:
-			ImageWrapper = new FIcoImageWrapper();
+			ImageWrapper = MakeShared<FIcoImageWrapper>();
 			break;
 
 #if WITH_UNREALEXR
 		case EImageFormat::EXR:
-			ImageWrapper = new FExrImageWrapper();
+			ImageWrapper = MakeShared<FExrImageWrapper>();
 			break;
 #endif
 		case EImageFormat::ICNS:
-			ImageWrapper = new FIcnsImageWrapper();
+			ImageWrapper = MakeShared<FIcnsImageWrapper>();
 			break;
 
 		case EImageFormat::TGA:
-			ImageWrapper = new FTgaImageWrapper();
+			ImageWrapper = MakeShared<FTgaImageWrapper>();
 			break;
 
 		default:
 			break;
 		}
 
-		return MakeShareable(ImageWrapper);
+		return ImageWrapper;
 	}
 
 	virtual EImageFormat DetectImageFormat(const void* CompressedData, int64 CompressedSize) override
diff --git a/Engine/Source/Runtime/ImageWrapper/Public/IImageWrapper.h b/Engine/Source/Runtime/ImageWrapper/Public/IImageWrapper.h
--- a/Engine/Source/Runtime/ImageWrapper/Public/IImageWrapper.h
+++ b/Engine/Source/Runtime/ImageWrapper/Public/IImageWrapper.h
@@ -47,9 +47,18 @@ enum class EImageFormat : int8
 enum class ERGBFormat : int8
 {
 	Invalid = -1,
+
+	// Red, Green, Blue and Alpha
 	RGBA =  0,
+
+	// Blue, Green, Red and Alpha
 	BGRA =  1,
+
+	// Gray scale
 	Gray =  2,
+
+	// Red, Green, Blue and Alpha using IEEE Floating-Point Arithmetic (see IEEE754). The format is always binary.
+	RGBAF = 3,
 };
 
 
@@ -102,13 +111,15 @@ public:
 
 	/**
 	 * Gets the compressed data.
-	 *
+	 * (Note: It may consume the data set in the SetCompressed function if it was set before)
+	 * 
 	 * @return Array of the compressed data.
 	 */
-	virtual const TArray64<uint8>& GetCompressed(int32 Quality = 0) = 0;
+	virtual TArray64<uint8> GetCompressed(int32 Quality = 0) = 0;
 
 	/**  
 	 * Gets the raw data.
+	 * (Note: It may consume the data set in the SetRaw function if it was set before)
 	 *
 	 * @param InFormat How we want to manipulate the RGB data.
 	 * @param InBitDepth The output bit-depth per channel, normally 8.
@@ -120,6 +131,7 @@ public:
 	/**
 	 * Gets the raw data in a TArray. Only use this if you're certain that the image is less than 2 GB in size.
 	 * Prefer using the overload which takes a TArray64 in general.
+	 * (Note: It may consume the data set in the SetRaw function if it was set before)
 	 *
 	 * @param InFormat How we want to manipulate the RGB data.
 	 * @param InBitDepth The output bit-depth per channel, normally 8.
diff --git a/Engine/Source/Runtime/ImageWriteQueue/Private/ImageWriteTask.cpp b/Engine/Source/Runtime/ImageWriteQueue/Private/ImageWriteTask.cpp
--- a/Engine/Source/Runtime/ImageWriteQueue/Private/ImageWriteTask.cpp
+++ b/Engine/Source/Runtime/ImageWriteQueue/Private/ImageWriteTask.cpp
@@ -181,7 +181,7 @@ bool FImageWriteTask::WriteToDisk()
 			{
 				if (InitializeWrapper(ImageWrapper, Format))
 				{
-					const TArray64<uint8>& CompressedFile = ImageWrapper->GetCompressed(CompressionQuality);
+					const TArray64<uint8> CompressedFile = ImageWrapper->GetCompressed(CompressionQuality);
 					uint64 TotalNumberOfBytes, NumberOfFreeBytes;
 					if (FPlatformMisc::GetDiskTotalAndFreeSpace(FPaths::GetPath(Filename), TotalNumberOfBytes, NumberOfFreeBytes))
 					{
diff --git a/Engine/Source/Runtime/ImageWriteQueue/Public/ImagePixelData.h b/Engine/Source/Runtime/ImageWriteQueue/Public/ImagePixelData.h
--- a/Engine/Source/Runtime/ImageWriteQueue/Public/ImagePixelData.h
+++ b/Engine/Source/Runtime/ImageWriteQueue/Public/ImagePixelData.h
@@ -232,7 +232,7 @@ template<> struct TImagePixelDataTraits<FColor>
 
 template<> struct TImagePixelDataTraits<FFloat16Color>
 {
-	static const ERGBFormat PixelLayout = ERGBFormat::RGBA;
+	static const ERGBFormat PixelLayout = ERGBFormat::RGBAF;
 	static const EImagePixelType PixelType = EImagePixelType::Float16;
 
 	enum { BitDepth = 16, NumChannels = 4 };
@@ -240,7 +240,7 @@ template<> struct TImagePixelDataTraits<FFloat16Color>
 
 template<> struct TImagePixelDataTraits<FLinearColor>
 {
-	static const ERGBFormat PixelLayout = ERGBFormat::RGBA;
+	static const ERGBFormat PixelLayout = ERGBFormat::RGBAF;
 	static const EImagePixelType PixelType = EImagePixelType::Float32;
 
 	enum { BitDepth = 32, NumChannels = 4 };
diff --git a/Engine/Source/Runtime/MediaUtils/Private/MediaRecorder.cpp b/Engine/Source/Runtime/MediaUtils/Private/MediaRecorder.cpp
--- a/Engine/Source/Runtime/MediaUtils/Private/MediaRecorder.cpp
+++ b/Engine/Source/Runtime/MediaUtils/Private/MediaRecorder.cpp
@@ -268,7 +268,7 @@ void FMediaRecorder::TickRecording()
 				if (Sample->GetFormat() == EMediaTextureSampleFormat::FloatRGBA)
 				{
 					PixelType = EImagePixelType::Float16;
-					PixelLayout = ERGBFormat::RGBA;
+					PixelLayout = ERGBFormat::RGBAF;
 					BitDepth = 16;
 				}
 
