Removed the class count from USelection, in favor of using the counter on the underlying element list

diff --git a/Engine/Source/Editor/UnrealEd/Private/Selection.cpp b/Engine/Source/Editor/UnrealEd/Private/Selection.cpp
--- a/Engine/Source/Editor/UnrealEd/Private/Selection.cpp
+++ b/Engine/Source/Editor/UnrealEd/Private/Selection.cpp
@@ -93,7 +93,6 @@ USelection* USelection::CreateComponentSelection(UObject* InOuter, FName InName,
 void USelection::Initialize(TSharedRef<ISelectionElementBridge>&& InSelectionElementBridge)
 {
 	SelectionElementBridge = MoveTemp(InSelectionElementBridge);
-	SyncSelectedClasses();
 }
 
 void USelection::SetElementSelectionSet(UTypedElementSelectionSet* InElementSelectionSet)
@@ -202,107 +201,30 @@ UObject* USelection::GetObjectForElementHandle(const FTypedElementHandle& InElem
 void USelection::OnElementListSyncEvent(const UTypedElementList* InElementList, FTypedElementListLegacySync::ESyncType InSyncType, const FTypedElementHandle& InElementHandle, bool bIsWithinBatchOperation)
 {
 	check(InElementList == ElementSelectionSet->GetElementList());
-	switch (InSyncType)
-	{
-	case FTypedElementListLegacySync::ESyncType::Added:
-		if (UObject* Object = GetObjectForElementHandle(InElementHandle))
-		{
-			OnObjectSelected(Object, !bIsWithinBatchOperation);
-		}
-		break;
-
-	case FTypedElementListLegacySync::ESyncType::Removed:
-		if (UObject* Object = GetObjectForElementHandle(InElementHandle))
-		{
-			OnObjectDeselected(Object, !bIsWithinBatchOperation);
-		}
-		break;
-
-	case FTypedElementListLegacySync::ESyncType::BatchComplete:
-		OnSelectedChanged(/*bSyncState*/false, !bIsWithinBatchOperation);
-		break;
-
-	default:
-		OnSelectedChanged(/*bSyncState*/true, !bIsWithinBatchOperation);
-		break;
-	}
-}
-
-void USelection::OnObjectSelected(UObject* InObject, const bool bNotify)
-{
-	check(InObject);
-
-	if (FSelectedClassInfo* SelectedClassInfo = SelectedClasses.Find(InObject->GetClass()))
-	{
-		++SelectedClassInfo->SelectionCount;
-	}
-	else
-	{
-		// 1 Object with a new class type has been selected
-		SelectedClasses.Add(FSelectedClassInfo(InObject->GetClass(), 1));
-	}
 
+	const bool bNotify = !bIsWithinBatchOperation;
 	if (bNotify)
 	{
-		// Call this after the item has been added from the selection set.
-		USelection::SelectObjectEvent.Broadcast(InObject);
-	}
-}
-
-void USelection::OnObjectDeselected(UObject* InObject, const bool bNotify)
-{
-	check(InObject);
-
-	FSetElementId Id = SelectedClasses.FindId(InObject->GetClass());
-	if (Id.IsValidId())
-	{
-		FSelectedClassInfo& ClassInfo = SelectedClasses[Id];
-		// One less object of this class is selected;
-		--ClassInfo.SelectionCount;
-		// If no more objects of the selected class exists, remove it
-		if (ClassInfo.SelectionCount == 0)
+		switch (InSyncType)
 		{
-			SelectedClasses.Remove(Id);
-		}
-	}
-
-	if (bNotify)
-	{
-		// Call this after the item has been removed from the selection set.
-		USelection::SelectObjectEvent.Broadcast(InObject);
-	}
-}
-
-void USelection::OnSelectedChanged(const bool bSyncState, const bool bNotify)
-{
-	if (bSyncState)
-	{
-		SyncSelectedClasses();
-	}
-	
-	if (bNotify)
-	{
-		USelection::SelectionChangedEvent.Broadcast(this);
-	}
-}
-
-void USelection::SyncSelectedClasses()
-{
-	SelectedClasses.Reset();
-
-	for (int32 Idx = 0; Idx < Num(); ++Idx)
-	{
-		if (UObject* Object = GetSelectedObject(Idx))
-		{
-			if (FSelectedClassInfo* SelectedClassInfo = SelectedClasses.Find(Object->GetClass()))
-			{
-				++SelectedClassInfo->SelectionCount;
-			}
-			else
+		// Emit a general object selection changed notification for the following events
+		case FTypedElementListLegacySync::ESyncType::Added:
+		case FTypedElementListLegacySync::ESyncType::Removed:
+			if (UObject* Object = GetObjectForElementHandle(InElementHandle))
 			{
-				// 1 Object with a new class type has been selected
-				SelectedClasses.Add(FSelectedClassInfo(Object->GetClass(), 1));
+				USelection::SelectObjectEvent.Broadcast(Object);
 			}
+			break;
+
+		// Emit a general selection changed notification for the following events
+		case FTypedElementListLegacySync::ESyncType::Modified:
+		case FTypedElementListLegacySync::ESyncType::Cleared:
+		case FTypedElementListLegacySync::ESyncType::BatchComplete:
+			USelection::SelectionChangedEvent.Broadcast(this);
+			break;
+
+		default:
+			break;
 		}
 	}
 }
@@ -411,6 +333,12 @@ bool USelection::IsSelected(const UObject* InObject) const
 	return false;
 }
 
+bool USelection::IsClassSelected(UClass* Class) const
+{
+	return ElementSelectionSet
+		&& TypedElementListObjectUtil::HasObjectsOfExactClass(ElementSelectionSet->GetElementList(), Class);
+}
+
 void USelection::Serialize(FArchive& Ar)
 {
 	if (ElementSelectionSet)
diff --git a/Engine/Source/Editor/UnrealEd/Public/Selection.h b/Engine/Source/Editor/UnrealEd/Public/Selection.h
--- a/Engine/Source/Editor/UnrealEd/Public/Selection.h
+++ b/Engine/Source/Editor/UnrealEd/Public/Selection.h
@@ -22,36 +22,6 @@ class UNREALED_API USelection : public UObject
 	GENERATED_BODY()
 
 private:
-	/** Contains info about each class and how many objects of that class are selected */
-	struct FSelectedClassInfo
-	{
-		/** The selected class */
-		const UClass* Class;
-		/** How many objects of that class are selected */
-		int32 SelectionCount;
-
-		FSelectedClassInfo(const UClass* InClass)
-			: Class(InClass)
-			, SelectionCount(0)
-		{}
-
-		FSelectedClassInfo(const UClass* InClass, int32 InSelectionCount)
-			: Class(InClass)
-			, SelectionCount(InSelectionCount)
-		{}
-		bool operator==(const FSelectedClassInfo& Info) const
-		{
-			return Class == Info.Class;
-		}
-
-		friend uint32 GetTypeHash(const FSelectedClassInfo& Info)
-		{
-			return GetTypeHash(Info.Class);
-		}
-	};
-
-	typedef TSet<FSelectedClassInfo> ClassArray;
-
 	template<typename SelectionFilter>
 	friend class TSelectionIterator;
 
@@ -75,12 +45,6 @@ public:
 	/** Called when the assigned typed element selection pointer set for a selection is changed */
 	static FOnSelectionElementSelectionPtrChanged SelectionElementSelectionPtrChanged;
 
-	typedef ClassArray::TIterator TClassIterator;
-	typedef ClassArray::TConstIterator TClassConstIterator;
-
-	TClassIterator			ClassItor()				{ return TClassIterator( SelectedClasses ); }
-	TClassConstIterator		ClassConstItor() const	{ return TClassConstIterator( SelectedClasses ); }
-
 	/**
 	 * Set the element selection set instance for this selection set.
 	 * @note Also sets the element list instance.
@@ -297,11 +261,7 @@ public:
 		return Count;
 	}
 
-	bool IsClassSelected(UClass* Class) const
-	{
-		const FSelectedClassInfo* Info = SelectedClasses.Find(Class);
-		return Info && Info->SelectionCount > 0;
-	}
+	bool IsClassSelected(UClass* Class) const;
 
 	//~ Begin UObject Interface
 	virtual void Serialize(FArchive& Ar) override;
@@ -367,11 +327,6 @@ private:
 
 	void OnElementListSyncEvent(const UTypedElementList* InElementList, FTypedElementListLegacySync::ESyncType InSyncType, const FTypedElementHandle& InElementHandle, bool bIsWithinBatchOperation);
 
-	void OnObjectSelected(UObject* InObject, const bool bNotify);
-	void OnObjectDeselected(UObject* InObject, const bool bNotify);
-	void OnSelectedChanged(const bool bSyncState, const bool bNotify);
-	void SyncSelectedClasses();
-
 	/** Bridge from UObjects to their corresponding typed elements. */
 	TSharedPtr<ISelectionElementBridge> SelectionElementBridge;
 
@@ -379,9 +334,6 @@ private:
 	UPROPERTY()
 	TObjectPtr<UTypedElementSelectionSet> ElementSelectionSet = nullptr;
 
-	/** Tracks the most recently selected actor classes.  Used for UnrealEd menus. */
-	ClassArray SelectedClasses;
-
 private:
 	// Hide IsSelected(), as calling IsSelected() on a selection set almost always indicates
 	// an error where the caller should use IsSelected(UObject* InObject).
