Do not consider the level visibility when re-instancing BP actors if the world the level is owned by is not yet initialized

diff --git a/Engine/Source/Editor/UnrealEd/Private/Kismet2/KismetReinstanceUtilities.cpp b/Engine/Source/Editor/UnrealEd/Private/Kismet2/KismetReinstanceUtilities.cpp
--- a/Engine/Source/Editor/UnrealEd/Private/Kismet2/KismetReinstanceUtilities.cpp
+++ b/Engine/Source/Editor/UnrealEd/Private/Kismet2/KismetReinstanceUtilities.cpp
@@ -1353,7 +1353,9 @@ void FActorReplacementHelper::Finalize(const TMap<UObject*, UObject*>& OldToNewI
 	bool bIsInHiddenLevel = false;
 	if (ULevel* Level = NewActor->GetLevel())
 	{
-		bIsInHiddenLevel = !Level->bIsVisible;
+		// consider non visible level as hidden in editor only if this is post world initialization
+		UWorld* World = Level->GetWorld();
+		bIsInHiddenLevel = !Level->bIsVisible && World && World->bIsWorldInitialized;
 	}
 
 	if (bIsInHiddenLevel)
