WorldPartition cook on the fly support features:
1) Create FPackageData for the generated packages before creating the packages. This supports exposing the generated packages to the rest of the cooker without needing to create the UPackages up front.
2) Keep the FGeneratorPackage alive for the lifetime of the cook. This supports populating the generated packages later on demand rather than populating them up front.
3) Mark Generated PackageDatas as generated and give them a pointer to their generator. This supports finding the generator to populate them when their cook is requested and they have not been populated yet.
#rnx

diff --git a/Engine/Source/Editor/UnrealEd/Classes/CookOnTheSide/CookOnTheFlyServer.h b/Engine/Source/Editor/UnrealEd/Classes/CookOnTheSide/CookOnTheFlyServer.h
--- a/Engine/Source/Editor/UnrealEd/Classes/CookOnTheSide/CookOnTheFlyServer.h
+++ b/Engine/Source/Editor/UnrealEd/Classes/CookOnTheSide/CookOnTheFlyServer.h
@@ -122,6 +122,7 @@ namespace Cook
 	struct FPendingCookedPlatformData;
 	struct FPlatformManager;
 	struct FTickStackData;
+	struct FPopulatePackageContext;
 }
 }
 
@@ -314,13 +315,19 @@ private:
 	/** Try to save all packages in the SaveQueue until it's time to break. Report the number of requests that were completed (either skipped or successfully saved or failed to save) */
 	void PumpSaves(UE::Cook::FTickStackData& StackData, uint32 DesiredQueueLength, int32& OutNumPushed, bool& bOutBusy);
 	/**
-	 * Inspect the given package and add it to the SaveQueue if the cooker should save it.
+	 * Inspect the given package and queue it for saving if necessary.
 	 *
-	 * @param Package				The package to be added to the SaveQueue.
-	 * @param bIsGeneratedPackage	If true, ignore the CookByTheBookOptions::bSkipHardReferences flag when in CookByTheBook mode.
+	 * @param Package				The package to be considered for saving.
 	 * @return						Returns the PackageData for this queued package.
 	 */
-	UE::Cook::FPackageData* QueueDiscoveredPackage(UPackage* Package, bool bIsGeneratedPackage = false);
+	UE::Cook::FPackageData* QueueDiscoveredPackage(UPackage* Package);
+	/**
+	 * Inspect the given package and queue it for saving if necessary.
+	 *
+	 * @param PackageData			The PackageData to be considered for saving.
+	 * @param bLoadReady			If true send on to LoadReady, otherwise add it at the Request stage.
+	 */
+	void QueueDiscoveredPackageData(UE::Cook::FPackageData& PackageData, bool bLoadReady = false);
 
 	/** Check whether the package filter has change, and if so call QueueDiscoveredPackage again on each existing package. */
 	void UpdatePackageFilter();
@@ -1097,6 +1104,9 @@ private:
 	*/
 	void SplitPackage(UE::Cook::FGeneratorPackage* Generator, bool& bCompleted, bool& bOutError);
 
+	/** Try calling the splitter's populate to create the package */
+	UPackage* TryPopulateGeneratedPackage(UE::Cook::FPopulatePackageContext& Context);
+
 	uint32 FullLoadAndSave(uint32& CookedPackageCount);
 
 	uint32		StatLoadedPackageCount = 0;
diff --git a/Engine/Source/Editor/UnrealEd/Private/CookOnTheFlyServer.cpp b/Engine/Source/Editor/UnrealEd/Private/CookOnTheFlyServer.cpp
--- a/Engine/Source/Editor/UnrealEd/Private/CookOnTheFlyServer.cpp
+++ b/Engine/Source/Editor/UnrealEd/Private/CookOnTheFlyServer.cpp
@@ -1710,7 +1710,7 @@ void UCookOnTheFlyServer::AddDependenciesToLoadQueue(UE::Cook::FPackageData& Pac
 			PackageData.SetIsVisited(true);
 
 			AssetDependenciesScratch.Reset();
-			// TODO EditorOnly References: We only load Game dependencies, because if we explictly load an EditorOnly dependency, that causes StaticLoadObjectInternal to SetLoadedByEditorPropertiesOnly(false), and we do not want to impact that value with our preloading of required packages 
+			// TODO EditorOnly References: We only load Game dependencies, because if we explicitly load an EditorOnly dependency, that causes StaticLoadObjectInternal to SetLoadedByEditorPropertiesOnly(false), and we do not want to impact that value with our preloading of required packages 
 			if (AssetRegistry->GetDependencies(PackageData.GetPackageName(), AssetDependenciesScratch, UE::AssetRegistry::EDependencyCategory::Package, UE::AssetRegistry::EDependencyQuery::Hard | UE::AssetRegistry::EDependencyQuery::Game))
 			{
 				Dependencies.Reserve(AssetDependenciesScratch.Num());
@@ -1894,64 +1894,130 @@ void UCookOnTheFlyServer::PumpPreloadStarts()
 	}
 }
 
+namespace UE::Cook
+{
+struct FPopulatePackageContext
+{
+	explicit FPopulatePackageContext(FGeneratorPackage* InGeneratorStruct, UPackage* InOwnerPackage)
+		: GeneratorStruct(InGeneratorStruct)
+		, OwnerPackage(InOwnerPackage)
+	{
+	}
+
+	// Input variables
+	FGeneratorPackage* GeneratorStruct = nullptr;
+	const UPackage* OwnerPackage = nullptr;
+	const FGeneratorPackage::FGeneratedStruct* GeneratedStruct = nullptr;
+
+	// Cache of initialized-on-demand variables
+	UObject* OwnerObject = nullptr;
+	FString GeneratorPackagePath;
+	FString GeneratorPackageShortName;
+	FString GeneratedCookedRootPath;
+};
+
+}
 void UCookOnTheFlyServer::LoadPackageInQueue(UE::Cook::FPackageData& PackageData, uint32& ResultFlags, int32& OutNumPushed)
 {
+	using namespace UE::Cook;
+
 	UPackage* LoadedPackage = nullptr;
 	OutNumPushed = 0;
 
 	FName PackageFileName(PackageData.GetFileName());
-	bool bLoadFullySuccessful = LoadPackageForCooking(PackageData, LoadedPackage);
-	if (!bLoadFullySuccessful)
+	if (PackageData.IsLoadFromDisk())
 	{
-		ResultFlags |= COSR_ErrorLoadingPackage;
-		UE_LOG(LogCook, Verbose, TEXT("Not cooking package %s"), *PackageFileName.ToString());
-		RejectPackageToLoad(PackageData, TEXT("failed to load"));
-		return;
-	}
-	check(LoadedPackage != nullptr && LoadedPackage->IsFullyLoaded());
+		bool bLoadFullySuccessful = LoadPackageForCooking(PackageData, LoadedPackage);
+		if (!bLoadFullySuccessful)
+		{
+			ResultFlags |= COSR_ErrorLoadingPackage;
+			UE_LOG(LogCook, Verbose, TEXT("Not cooking package %s"), *PackageFileName.ToString());
+			RejectPackageToLoad(PackageData, TEXT("failed to load"));
+			return;
+		}
+		check(LoadedPackage != nullptr && LoadedPackage->IsFullyLoaded());
 
-	FName NewPackageFileName(GetPackageNameCache().GetCachedStandardFileName(LoadedPackage));
-	if (LoadedPackage->GetFName() != PackageData.GetPackageName())
-	{
-		// The PackageName is not the name that we loaded. This can happen due to CoreRedirects.
-		// We refuse to cook requests for packages that no longer exist in PumpExternalRequests, but it is possible
-		// that a CoreRedirect exists from a (externally requested or requested as a reference) package that still exists.
-		// Mark the original PackageName as cooked for all platforms and send a request to cook the new FileName
-		check(NewPackageFileName != PackageFileName);
+		FName NewPackageFileName(GetPackageNameCache().GetCachedStandardFileName(LoadedPackage));
+		if (LoadedPackage->GetFName() != PackageData.GetPackageName())
+		{
+			// The PackageName is not the name that we loaded. This can happen due to CoreRedirects.
+			// We refuse to cook requests for packages that no longer exist in PumpExternalRequests, but it is possible
+			// that a CoreRedirect exists from a (externally requested or requested as a reference) package that still exists.
+			// Mark the original PackageName as cooked for all platforms and send a request to cook the new FileName
+			check(NewPackageFileName != PackageFileName);
 
-		UE_LOG(LogCook, Verbose, TEXT("Request for %s received going to save %s"), *PackageFileName.ToString(), *NewPackageFileName.ToString());
-		UE::Cook::FPackageData& OtherPackageData = PackageDatas->AddPackageDataByPackageNameChecked(LoadedPackage->GetFName());
-		OtherPackageData.UpdateRequestData(PackageData.GetRequestedPlatforms(), PackageData.GetIsUrgent(), UE::Cook::FCompletionCallback());
+			UE_LOG(LogCook, Verbose, TEXT("Request for %s received going to save %s"), *PackageFileName.ToString(), *NewPackageFileName.ToString());
+			FPackageData& OtherPackageData = PackageDatas->AddPackageDataByPackageNameChecked(LoadedPackage->GetFName());
+			OtherPackageData.UpdateRequestData(PackageData.GetRequestedPlatforms(), PackageData.GetIsUrgent(), FCompletionCallback());
 
-		PackageData.AddCookedPlatforms(PlatformManager->GetSessionPlatforms(), true);
-		RejectPackageToLoad(PackageData, TEXT("is redirected to another filename"));
-		return;
-	}
+			PackageData.AddCookedPlatforms(PlatformManager->GetSessionPlatforms(), true);
+			RejectPackageToLoad(PackageData, TEXT("is redirected to another filename"));
+			return;
+		}
 
-	if (NewPackageFileName != PackageFileName)
+		if (NewPackageFileName != PackageFileName)
+		{
+			// This case should never happen since we are checking for the existence of the file in PumpExternalRequests
+			UE_LOG(LogCook, Warning, TEXT("Unexpected change in FileName when loading a requested package. \"%s\" changed to \"%s\"."),
+				*PackageFileName.ToString(), *NewPackageFileName.ToString());
+
+			UE_LOG(LogCook, Verbose, TEXT("Request for %s received going to save %s"), *PackageFileName.ToString(), *NewPackageFileName.ToString());
+			// The package FileName is not the FileName that we loaded
+			//  sounds unpossible.... but it is due to searching for files and such
+			PackageDatas->UpdateFileName(LoadedPackage->GetFName());
+			// Register the FileName that we loaded as pointing to this PackageData
+			PackageDatas->RegisterFileNameAlias(PackageData, PackageFileName);
+		}
+	}
+	else
 	{
-		// This case should never happen since we are checking for the existence of the file in PumpExternalRequests
-		UE_LOG(LogCook, Warning, TEXT("Unexpected change in FileName when loading a requested package. \"%s\" changed to \"%s\"."),
-			*PackageFileName.ToString(), *NewPackageFileName.ToString());
+		check(PackageData.IsGenerated());
+		FGeneratorPackage* GeneratorStruct = PackageData.GetGeneratedOwner();
+		if (!GeneratorStruct)
+		{
+			UE_LOG(LogCook, Error, TEXT("Package %s is an out-of-date generated package with a no-longer-available generator. It can not be loaded."), *PackageFileName.ToString());
+			RejectPackageToLoad(PackageData, TEXT("is an orphaned generated package"));
+			return;
+		}
+		check(GeneratorStruct->IsDeferredPopulate());
+		FGeneratorPackage::FGeneratedStruct* GeneratedStruct = GeneratorStruct->FindGeneratedStruct(&PackageData);
+		if (!GeneratedStruct)
+		{
+			UE_LOG(LogCook, Error, TEXT("Package %s is a generated package but its generator no longer has a record of it. It can not be loaded."), *PackageFileName.ToString());
+			RejectPackageToLoad(PackageData, TEXT("is an orphaned generated package"));
+			return;
+		}
+		const FPackageData& OwnerPackageData = GeneratorStruct->GetOwner();
 
-		UE_LOG(LogCook, Verbose, TEXT("Request for %s received going to save %s"), *PackageFileName.ToString(), *NewPackageFileName.ToString());
-		// The package FileName is not the FileName that we loaded
-		//  sounds unpossible.... but it is due to searching for files and such
-		PackageDatas->UpdateFileName(LoadedPackage->GetFName());
-		// Register the FileName that we loaded as pointing to this PackageData
-		PackageDatas->RegisterFileNameAlias(PackageData, PackageFileName);
+		UPackage* OwnerPackage = LoadPackage(nullptr, *OwnerPackageData.GetFileName().ToString(), LOAD_None);
+		if (!OwnerPackage)
+		{
+			ResultFlags |= COSR_ErrorLoadingPackage;
+			UE_LOG(LogCook, Error, TEXT("Package %s is a generated package and we could not load it's generator package %s. It can not be loaded."),
+				*PackageFileName.ToString(), *OwnerPackageData.GetFileName().ToString());
+			RejectPackageToLoad(PackageData, TEXT("is a generated package which could not load its generator"));
+			return;
+		}
+		FPopulatePackageContext Context(GeneratorStruct, OwnerPackage);
+		Context.GeneratedStruct = GeneratedStruct;
+		LoadedPackage = TryPopulateGeneratedPackage(Context);
+		if (!LoadedPackage)
+		{
+			RejectPackageToLoad(PackageData, TEXT("is a generated package which could not be populated"));
+			return;
+		}
 	}
 
 	if (PackageData.HasAllCookedPlatforms(PackageData.GetRequestedPlatforms(), true))
 	{
 		// Already cooked. This can happen if we needed to load a package that was previously cooked and garbage collected because it is a loaddependency of a new request.
 		// Send the package back to idle, nothing further to do with it.
-		PackageData.SendToState(UE::Cook::EPackageState::Idle, UE::Cook::ESendFlags::QueueAdd);
+		PackageData.SendToState(EPackageState::Idle, ESendFlags::QueueAdd);
 	}
 	else
 	{
 		PackageData.SetPackage(LoadedPackage);
-		PackageData.SendToState(UE::Cook::EPackageState::Save, UE::Cook::ESendFlags::QueueAdd);
+		PackageData.SendToState(EPackageState::Save, ESendFlags::QueueAdd);
 		++OutNumPushed;
 	}
 }
@@ -1976,7 +2042,7 @@ void UCookOnTheFlyServer::RejectPackageToLoad(UE::Cook::FPackageData& PackageDat
 
 //////////////////////////////////////////////////////////////////////////
 
-UE::Cook::FPackageData* UCookOnTheFlyServer::QueueDiscoveredPackage(UPackage* Package, bool bIsGeneratedPackage)
+UE::Cook::FPackageData* UCookOnTheFlyServer::QueueDiscoveredPackage(UPackage* Package)
 {
 	check(Package != nullptr);
 
@@ -1986,24 +2052,35 @@ UE::Cook::FPackageData* UCookOnTheFlyServer::QueueDiscoveredPackage(UPackage* Pa
 		return nullptr;	// if we have name none that means we are in core packages or something...
 	}
 	UE::Cook::FPackageData& PackageData = PackageDatas->FindOrAddPackageData(Package->GetFName(), FileName);
+	QueueDiscoveredPackageData(PackageData, true /* bIsLoadReady */);
+	return &PackageData;
+}
 
+void UCookOnTheFlyServer::QueueDiscoveredPackageData(UE::Cook::FPackageData& PackageData, bool bLoadReady)
+{
 	const TArray<const ITargetPlatform*>& TargetPlatforms = PlatformManager->GetSessionPlatforms();
 	if (PackageData.HasAllCookedPlatforms(TargetPlatforms, true /* bIncludeFailed */))
 	{
 		// All SessionPlatforms have already been cooked for the package, so we don't need to save it again
-		return nullptr;
+		return;
 	}
 
-	if (!PackageData.IsInProgress() && 
-		(bIsGeneratedPackage || !IsCookByTheBookMode() || !CookByTheBookOptions->bSkipHardReferences))
+	if (!PackageData.IsInProgress() &&
+		(PackageData.IsGenerated() || !IsCookByTheBookMode() || !CookByTheBookOptions->bSkipHardReferences))
 	{
-		// Send this package into the LoadReadyQueue to fully load it and send it on to the SaveQueue
 		PackageData.SetRequestData(TargetPlatforms, /*bIsUrgent*/ false, UE::Cook::FCompletionCallback());
-		PackageData.SendToState(UE::Cook::EPackageState::LoadReady, UE::Cook::ESendFlags::QueueNone);
-		// Send it to the front of the LoadReadyQueue since it is mostly loaded already
-		PackageDatas->GetLoadReadyQueue().AddFront(&PackageData);
+		if (bLoadReady)
+		{
+			// Send this package into the LoadReadyQueue to fully load it and send it on to the SaveQueue
+			PackageData.SendToState(UE::Cook::EPackageState::LoadReady, UE::Cook::ESendFlags::QueueRemove);
+			// Send it to the front of the LoadReadyQueue since it is mostly loaded already
+			PackageDatas->GetLoadReadyQueue().AddFront(&PackageData);
+		}
+		else
+		{
+			PackageData.SendToState(UE::Cook::EPackageState::Request, UE::Cook::ESendFlags::QueueAddAndRemove);
+		}
 	}
-	return &PackageData;
 }
 
 void UCookOnTheFlyServer::UpdatePackageFilter()
@@ -2156,95 +2233,164 @@ UE::Cook::FGeneratorPackage* UCookOnTheFlyServer::CreateGeneratorPackage(UE::Coo
 		return nullptr;
 	}
 
-	// Initialize the instance with the data object
-	SplitterInstance->SetDataObject(SplitDataObject);
-
 	// Create a FGeneratorPackage helper object using this CookPackageSplitter instance and return it
 	PackageData.CreateGeneratorPackage(SplitDataObject, SplitterInstance, GeneratedUncookedRootPath);
 
 	return PackageData.GetGeneratorPackage();
 }
 
-void UCookOnTheFlyServer::SplitPackage(UE::Cook::FGeneratorPackage* Generator, bool& bOutCompleted, bool& bOutError)
+void UCookOnTheFlyServer::SplitPackage(UE::Cook::FGeneratorPackage* GeneratorStruct, bool& bOutCompleted, bool& bOutError)
 {
+	using namespace UE::Cook;
+
 	bOutCompleted = false;
 	bOutError = false;
-	check(Generator);
+	check(GeneratorStruct);
 
-	const UPackage* GeneratorPackage = Generator->GetOwner().GetPackage();
-	const FString GeneratorPackagePath = FPackageName::GetLongPackagePath(GeneratorPackage->GetPathName());
-	const FString GeneratorPackageShortName = FPackageName::GetShortName(GeneratorPackage->GetName());
-	const FString GeneratedCookedRootPath = FPaths::RemoveDuplicateSlashes(FString::Printf(TEXT("/%s/%s/%s/"), *GeneratorPackagePath, *GeneratorPackageShortName, UE::Cook::GeneratedPackageSubPath));
-
-	// Start Splitting
-	while (const ICookPackageSplitter::FGeneratedPackage* PackageToGenerate = Generator->GetNextPackageToGenerate())
+	ICookPackageSplitter* Splitter = GeneratorStruct->GetCookPackageSplitterInstance();
+	if (Splitter->UseDeferredPopulate())
 	{
-		// Make package name to fit Mounted GeneratedUncookedRootPath
-		const FString GeneratedPackageName = FPaths::RemoveDuplicateSlashes(FString::Printf(TEXT("/%s/%s"), *Generator->GetGeneratedUncookedRootPath(), *PackageToGenerate->RelativePath));
-
-		// Create package
-		check(!FindObject<UPackage>(nullptr, *GeneratedPackageName));
-		UPackage* GeneratedPackage = CreatePackage(*GeneratedPackageName);
-
-		// Populate package using CookPackageSplitterInstance and pass GeneratedPackage's cooked name for it to properly setup any internal reference to this package (SoftObjectPaths or others)
-		const FString GeneratedPackageCookedName = FPaths::RemoveDuplicateSlashes(FString::Printf(TEXT("/%s/%s"), *GeneratedCookedRootPath, *PackageToGenerate->RelativePath));
-		if (!Generator->GetCookPackageSplitterInstance()->TryPopulatePackage(GeneratedPackage, PackageToGenerate->RelativePath, GeneratedPackageCookedName))
+		while (const FGeneratorPackage::FGeneratedStruct* GeneratedStruct = GeneratorStruct->GetNextPackageToGenerate())
 		{
-			UE_LOG(LogCook, Error, TEXT("Error populating generated package %s for object %s."), *PackageToGenerate->RelativePath, *Generator->GetSplitDataObject()->GetFullName());
-			bOutError = true;
-			return;
+			FPackageData* GeneratedPackageData = GeneratedStruct->PackageData;
+			if (!GeneratedPackageData)
+			{
+				UE_LOG(LogCook, Error, TEXT("PackageSplitter was not given valid PackageData to queue for generated package: no package data. Splitter=%s, Generated=%s."),
+					*GeneratorStruct->GetSplitDataObjectName().ToString(), *GeneratedStruct->GeneratorData.RelativePath);
+				bOutError = true;
+				return;
+			}
+			GeneratedPackageData->ClearCookedPlatformData();
+			QueueDiscoveredPackageData(*GeneratedPackageData);
 		}
+	}
+	else
+	{
+		UPackage* Owner = GeneratorStruct->GetOwner().GetPackage();
+		FPopulatePackageContext Context(GeneratorStruct, Owner);
 
-		// Save package into the uncooked intermediate directory
-		const FString GeneratedPackageExtension = GeneratedPackage->ContainsMap() ? FPackageName::GetMapPackageExtension() : FPackageName::GetAssetPackageExtension();
-		const FString GeneratedPackageUncookedFileName = FPaths::RemoveDuplicateSlashes(FPackageName::LongPackageNameToFilename(GeneratedPackage->GetName(), GeneratedPackageExtension));
-		FSavePackageResultStruct SaveResult = GEditor->Save(GeneratedPackage, nullptr, RF_Standalone, *GeneratedPackageUncookedFileName, GError, nullptr, /*bForceByteSwapping*/ false, /*bWarnOfLongFilename*/ true, SAVE_None);
-		if (SaveResult.Result != ESavePackageResult::Success)
+		// Start Splitting
+		while (const FGeneratorPackage::FGeneratedStruct* GeneratedStruct = GeneratorStruct->GetNextPackageToGenerate())
 		{
-			UE_LOG(LogCook, Error, TEXT("Error saving generated package %s for object %s."), *GeneratedPackageUncookedFileName, *Generator->GetSplitDataObject()->GetFullName());
-			bOutError = true;
-			return;
-		}
+			FPackageData* GeneratedPackageData = GeneratedStruct->PackageData;
+			if (!GeneratedPackageData)
+			{
+				UE_LOG(LogCook, Error, TEXT("PackageSplitter was not given valid PackageData to queue for generated package: no package data. Splitter=%s, Generated=%s."),
+					*GeneratorStruct->GetSplitDataObjectName().ToString(), *GeneratedStruct->GeneratorData.RelativePath);
+				bOutError = true;
+				return;
+			}
+			GeneratedPackageData->ClearCookedPlatformData();
 
-		// Update Asset Registry with the generated package
-		AssetRegistry->ScanFilesSynchronous({ GeneratedPackageUncookedFileName }, true);
+			Context.GeneratedStruct = GeneratedStruct;
+			UPackage* GeneratedPackage = TryPopulateGeneratedPackage(Context);
+			if (!GeneratedPackage)
+			{
+				bOutError = true;
+				return;
+			}
 
-		// Make sure PackageNameCache successfully caches the generated package
-		if (GetPackageNameCache().GetCachedStandardFileName(GeneratedPackage).IsNone())
-		{
-			UE_LOG(LogCook, Error, TEXT("Error caching generated package name %s for object %s."), *GeneratedPackageUncookedFileName, *Generator->GetSplitDataObject()->GetFullName());
-			bOutError = true;
-			return;
-		}
+			// Save package into the uncooked intermediate directory
+			FSavePackageResultStruct SaveResult = GEditor->Save(GeneratedPackage, nullptr, RF_Standalone,
+				*GeneratedPackageData->GetFileName().ToString(), GError, nullptr, /*bForceByteSwapping*/ false,
+				/*bWarnOfLongFilename*/ true, SAVE_None);
+			if (SaveResult.Result != ESavePackageResult::Success)
+			{
+				UE_LOG(LogCook, Error, TEXT("PackageSplitter could not save uncooked generated package. Splitter=%s, Generated=%s."),
+					*GeneratorStruct->GetSplitDataObjectName().ToString(), *GeneratedPackageData->GetFileName().ToString());
+				bOutError = true;
+				return;
+			}
 
-		// Add PackageData for GeneratedPackage
-		UE::Cook::FPackageData* GeneratedPackageData = PackageDatas->TryAddPackageDataByFileName(FName(GeneratedPackageUncookedFileName));
-		if (!GeneratedPackageData)
-		{
-			UE_LOG(LogCook, Error, TEXT("Error adding generated package data %s for object %s."), *GeneratedPackageUncookedFileName, *Generator->GetSplitDataObject()->GetFullName());
-			bOutError = true;
-			return;
+			// Update Asset Registry with the generated package
+			AssetRegistry->ScanFilesSynchronous({ GeneratedPackageData->GetFileName().ToString() }, true);
+
+			// Queue the generated package
+			QueueDiscoveredPackageData(*GeneratedPackageData, true /* bLoadReady */);
+
+			if (HasExceededMaxMemory())
+			{
+				return;
+			}
 		}
+	}
+
+	bOutCompleted = true;
+}
 
-		// Set override cooked FileName to be used when saving the cooked generated package
-		const FString GeneratePackageCookedFileName = FPaths::RemoveDuplicateSlashes(FString::Printf(TEXT("%s/%s/%s/%s%s"), 
-													  *FPaths::GetPath(Generator->GetOwner().GetFileName().ToString()), *GeneratorPackageShortName, UE::Cook::GeneratedPackageSubPath, *PackageToGenerate->RelativePath, *GeneratedPackageExtension));
-		GeneratedPackageData->SetCookedFileName(FPackageNameCache::GetStandardFileName(GeneratePackageCookedFileName));
+UPackage* UCookOnTheFlyServer::TryPopulateGeneratedPackage(UE::Cook::FPopulatePackageContext& Context)
+{
+	using namespace UE::Cook;
 
-		// Queue the generated package
-		QueueDiscoveredPackage(GeneratedPackage, /*bIsGeneratedPackage*/true);
-		check(GeneratedPackageData->IsInProgress());
+	FGeneratorPackage* GeneratorStruct = Context.GeneratorStruct;
+	const UPackage* OwnerPackage = Context.OwnerPackage;
+	const ICookPackageSplitter::FGeneratedPackage& GeneratorData = Context.GeneratedStruct->GeneratorData;
+	UE::Cook::FPackageData* GeneratedPackageData = Context.GeneratedStruct->PackageData;
+	const FString GeneratedPackageName = GeneratedPackageData->GetPackageName().ToString();
 
-		if (HasExceededMaxMemory())
+	if (Context.OwnerObject == nullptr)
+	{
+		Context.OwnerObject = GeneratorStruct->FindSplitDataObject();
+		if (!Context.OwnerObject)
 		{
-			return;
+			UE_LOG(LogCook, Error, TEXT("PopulateGeneratedPacakge could not find the original splitting object. Generated package can not be created. Splitter=%s, Generated=%s."),
+				*GeneratorStruct->GetSplitDataObjectName().ToString(), *GeneratedPackageName);
+			return nullptr;
 		}
 	}
 
-	Generator->Finalize();
-	bOutCompleted = true;
+	if (Context.GeneratorPackagePath.IsEmpty())
+	{
+		Context.GeneratorPackagePath = FPackageName::GetLongPackagePath(Context.OwnerPackage->GetPathName());
+		Context.GeneratorPackageShortName = FPackageName::GetShortName(Context.OwnerPackage->GetName());
+		Context.GeneratedCookedRootPath = FPaths::RemoveDuplicateSlashes(FString::Printf(TEXT("/%s/%s/%s/"),
+			*Context.GeneratorPackagePath, *Context.GeneratorPackageShortName, UE::Cook::GeneratedPackageSubPath));
+	}
+
+	check(GeneratorData.GetCreateAsMap().IsSet()); // Should have been confirmed in TryGenerateList and not allowed to get here if not set
+	check(GeneratedPackageData); // Caller already checked this
+
+	// Make package name to fit Mounted GeneratedUncookedRootPath
+	const FString GeneratedPackageCookedName = FPaths::RemoveDuplicateSlashes(FString::Printf(TEXT("/%s/%s"), *Context.GeneratedCookedRootPath, *GeneratorData.RelativePath));
+	const FString GeneratedPackageExtension = FString(FPathViews::GetExtension(GeneratedPackageData->GetFileName().ToString(), true /* bIncludeDot */));
+
+	// Create package
+	if (FindObject<UPackage>(nullptr, *GeneratedPackageName))
+	{
+		UE_LOG(LogCook, Error, TEXT("PackageSplitter found an existing copy of a package it was trying to generate; this is not yet supported. Splitter=%s, Generated=%s."),
+			*GeneratorStruct->GetSplitDataObjectName().ToString(), *GeneratedPackageName);
+		return nullptr;
+	}
+	UPackage* GeneratedPackage = CreatePackage(*GeneratedPackageName);
+
+	// Populate package using CookPackageSplitterInstance and pass GeneratedPackage's cooked name for it to
+	// properly setup any internal reference to this package (SoftObjectPaths or others)
+	if (!GeneratorStruct->GetCookPackageSplitterInstance()->TryPopulatePackage(OwnerPackage, Context.OwnerObject, GeneratedPackage,
+		GeneratorData.RelativePath, GeneratedPackageCookedName))
+	{
+		UE_LOG(LogCook, Error, TEXT("PackageSplitter returned false from TryPopulatePackage. Splitter=%s, Generated=%s."),
+			*GeneratorStruct->GetSplitDataObjectName().ToString(), *GeneratedPackageName);
+		return nullptr;
+	}
+	bool bPackageIsMap = GeneratedPackage->ContainsMap();
+	if (bPackageIsMap != *GeneratorData.GetCreateAsMap())
+	{
+		UE_LOG(LogCook, Error, TEXT("PackageSplitter specified generated package is %s in GetGenerateList results, but then in TryPopulatePackage created it as %s. Splitter=%s, Generated=%s."),
+			(*GeneratorData.GetCreateAsMap() ? TEXT("map") : TEXT("uasset")), (bPackageIsMap ? TEXT("map") : TEXT("uasset")),
+			*GeneratorStruct->GetSplitDataObjectName().ToString(), *GeneratedPackageName);
+		return nullptr;
+	}
+
+	// Set override cooked FileName to be used when saving the cooked generated package
+	const FString GeneratePackageCookedFileName = FPaths::RemoveDuplicateSlashes(FString::Printf(TEXT("%s/%s/%s/%s%s"),
+		*FPaths::GetPath(GeneratorStruct->GetOwner().GetFileName().ToString()), *Context.GeneratorPackageShortName,
+		UE::Cook::GeneratedPackageSubPath, *GeneratorData.RelativePath, *GeneratedPackageExtension));
+	GeneratedPackageData->SetCookedFileName(FPackageNameCache::GetStandardFileName(GeneratePackageCookedFileName));
+
+	return GeneratedPackage;
 }
 
+
 bool UCookOnTheFlyServer::BeginPrepareSave(UE::Cook::FPackageData& PackageData, UE::Cook::FCookerTimer& Timer, bool bIsPreCaching)
 {
 	if (PackageData.GetCookedPlatformDataCalled())
@@ -2278,28 +2424,45 @@ bool UCookOnTheFlyServer::BeginPrepareSave(UE::Cook::FPackageData& PackageData,
 	check(PackageData.GetState() == UE::Cook::EPackageState::Save);
 	PackageData.CreateObjectCache();
 
-	int32& CookedPlatformDataNextIndex = PackageData.GetCookedPlatformDataNextIndex();
-	if (CookedPlatformDataNextIndex == -1)
+	if (!PackageData.HasCompletedGeneration())
 	{
 		// Check for Splitting the package; this needs to happen before we make any of the BeginCacheForCookedPlatformData calls in the package
 		bool bError = false;
 
-		UE::Cook::FGeneratorPackage* Generator = PackageData.GetGeneratorPackage();
-		if (!Generator)
+		UE::Cook::FGeneratorPackage* Generator = nullptr;
+		if (!PackageData.HasInitializedGeneratorSave())
 		{
+			PackageData.SetInitializedGeneratorSave(true);
+			PackageData.DestroyGeneratorPackage();
 			UObject* SplitDataObject = nullptr;
-			if (UE::Cook::Private::FRegisteredCookPackageSplitter* Splitter = TryGetRegisteredCookPackageSplitter(PackageData, SplitDataObject, bError))
-			{
-				// Found a splitter, create generator package
-				Generator = CreateGeneratorPackage(PackageData, SplitDataObject, Splitter);
-				bError = !Generator;
-			}
+			UE::Cook::Private::FRegisteredCookPackageSplitter* Splitter = TryGetRegisteredCookPackageSplitter(PackageData, SplitDataObject, bError);
 			if (bError)
 			{
-				// Any error marks the package data for PumpSaves to handle this as an error and to put it in idle state
 				PackageData.SetHasBeginPrepareSaveFailed(true);
 				return false;
 			}
+
+			if (Splitter)
+			{
+				// Found a splitter, create generator package
+				Generator = CreateGeneratorPackage(PackageData, SplitDataObject, Splitter);
+				if (!Generator)
+				{
+					PackageData.SetHasBeginPrepareSaveFailed(true);
+					return false;
+				}
+
+				// Call the splitter to generate the list
+				if (!Generator->TryGenerateList(Package, SplitDataObject, GetPackageNameCache(), *PackageDatas))
+				{
+					PackageData.SetHasBeginPrepareSaveFailed(true);
+					return false;
+				}
+			}
+		}
+		else
+		{
+			Generator = PackageData.GetGeneratorPackage();
 		}
 		if (Generator)
 		{
@@ -2320,9 +2483,16 @@ bool UCookOnTheFlyServer::BeginPrepareSave(UE::Cook::FPackageData& PackageData,
 			{
 				return false;
 			}
-			PackageData.DestroyGeneratorPackage();
+			UObject* SplitObject = Generator->FindSplitDataObject();
+			if (!SplitObject)
+			{
+				UE_LOG(LogCook, Error, TEXT("Could not find SplitDataObject for %s"), *Package->GetName());
+				PackageData.SetHasBeginPrepareSaveFailed(true);
+				return false;
+			}
+			Generator->GetCookPackageSplitterInstance()->PreSaveGeneratorPackage(Package, SplitObject);
 		}
-		CookedPlatformDataNextIndex = 0;
+		PackageData.SetCompletedGeneration(true);
 	}
 
 	// Note that we cache cooked data for all requested platforms, rather than only for the requested platforms that have not cooked yet.  This allows
@@ -2331,6 +2501,7 @@ bool UCookOnTheFlyServer::BeginPrepareSave(UE::Cook::FPackageData& PackageData,
 	// Currently this does not cause significant cost since saving new platforms with some platforms already saved is a rare operation.
 
 	TArray<FWeakObjectPtr>& CachedObjectsInOuter = PackageData.GetCachedObjectsInOuter();
+	int32& CookedPlatformDataNextIndex = PackageData.GetCookedPlatformDataNextIndex();
 	FWeakObjectPtr* CachedObjectsInOuterData = CachedObjectsInOuter.GetData();
 	const TArray<const ITargetPlatform*>& TargetPlatforms = PackageData.GetRequestedPlatforms();
 	int NumPlatforms = TargetPlatforms.Num();
@@ -2429,32 +2600,6 @@ void UCookOnTheFlyServer::ReleaseCookedPlatformData(UE::Cook::FPackageData& Pack
 {
 	using namespace UE::Cook;
 
-	ON_SCOPE_EXIT
-	{
-		if (CurrentCookMode == ECookMode::CookByTheBook)
-		{
-			// For each object for which data is cached we can call FinishedCookedPlatformDataCache
-			// we can only safely call this when we are finished caching the object completely.
-			// this doesn't ever happen for cook in editor or cook on the fly mode
-			for (FWeakObjectPtr& WeakPtr : PackageData.GetCachedObjectsInOuter())
-			{
-				UObject* Obj = WeakPtr.Get();
-				if (Obj)
-				{
-					Obj->WillNeverCacheCookedPlatformDataAgain();
-				}
-			}
-
-			UPackage* Package = PackageData.GetPackage();
-			if (Package && Package->LinkerLoad)
-			{
-				// Loaders and their handles can have large buffers held in process memory and in the system file cache from the
-				// data that was loaded.  Keeping this for the lifetime of the cook is costly, so we try and unload it here.
-				Package->LinkerLoad->FlushCache();
-			}
-		}
-	};
-
 	if (!PackageData.GetCookedPlatformDataStarted())
 	{
 		PackageData.CheckCookedPlatformDataEmpty();
@@ -2465,7 +2610,7 @@ void UCookOnTheFlyServer::ReleaseCookedPlatformData(UE::Cook::FPackageData& Pack
 	if (PackageData.GetCookedPlatformDataComplete())
 	{
 		// Since we have completed CookedPlatformData, we know we called BeginCacheForCookedPlatformData on all objects in the package, and none are pending
-		if (!IsCookingInEditor()) // ClearAllCachedCookedPlatformData calls are only used when not in editor
+		if (!IsCookingInEditor()) // ClearAllCachedCookedPlatformData and WillNeverCacheCookedPlatformDataAgain calls are only used when not in editor
 		{
 			UE_SCOPED_HIERARCHICAL_COOKTIMER(ClearAllCachedCookedPlatformData);
 			for (FWeakObjectPtr& WeakPtr: PackageData.GetCachedObjectsInOuter())
@@ -2474,75 +2619,106 @@ void UCookOnTheFlyServer::ReleaseCookedPlatformData(UE::Cook::FPackageData& Pack
 				if (Object)
 				{
 					Object->ClearAllCachedCookedPlatformData();
+					if (CurrentCookMode == ECookMode::CookByTheBook)
+					{
+						Object->WillNeverCacheCookedPlatformDataAgain();
+					}
 				}
 			}
 		}
-		PackageData.ClearCookedPlatformData();
-		return;
 	}
-
-	// This is an exceptional flow handling case; we are releasing the CookedPlatformData before we called SavePackage
-	// Note that even after we return from this function, some objects with pending IsCachedCookedPlatformDataLoaded calls may still exist for this Package in PackageDatas->GetPendingCookedPlatformDatas(),
-	// and this PackageData may therefore still have GetNumPendingCookedPlatformData > 0
-	if (!IsCookingInEditor()) // ClearAllCachedCookedPlatformData calls are only used when not in editor.
+	else 
 	{
-		int32 NumPlatforms = PackageData.GetRequestedPlatforms().Num();
-		if (NumPlatforms > 0) // Shouldn't happen because PumpSaves checks for this, but avoid a divide by 0 if it does.
+		// This is an exceptional flow handling case; we are releasing the CookedPlatformData before we called SavePackage
+		// Note that even after we return from this function, some objects with pending IsCachedCookedPlatformDataLoaded calls may still exist for this Package in PackageDatas->GetPendingCookedPlatformDatas(),
+		// and this PackageData may therefore still have GetNumPendingCookedPlatformData > 0
+		if (!IsCookingInEditor()) // ClearAllCachedCookedPlatformData calls are only used when not in editor.
 		{
-			// We have only called BeginCacheForCookedPlatformData on Object,Platform pairs up to GetCookedPlatformDataNextIndex.
-			// Further, some of those calls might still be pending.
-
-			// Find all pending BeginCacheForCookedPlatformData for this FPackageData
-			TMap<UObject*, TArray<FPendingCookedPlatformData*>> PendingObjects;
-			for (FPendingCookedPlatformData& PendingCookedPlatformData : PackageDatas->GetPendingCookedPlatformDatas())
+			int32 NumPlatforms = PackageData.GetRequestedPlatforms().Num();
+			if (NumPlatforms > 0) // Shouldn't happen because PumpSaves checks for this, but avoid a divide by 0 if it does.
 			{
-				if (&PendingCookedPlatformData.PackageData == &PackageData && !PendingCookedPlatformData.PollIsComplete())
-				{
-					UObject* Object = PendingCookedPlatformData.Object.Get();
-					check(Object); // Otherwise PollIsComplete would have returned true
-					check(!PendingCookedPlatformData.bHasReleased); // bHasReleased should be false since PollIsComplete returned false
-					PendingObjects.FindOrAdd(Object).Add(&PendingCookedPlatformData);
-				}
-			}
+				// We have only called BeginCacheForCookedPlatformData on Object,Platform pairs up to GetCookedPlatformDataNextIndex.
+				// Further, some of those calls might still be pending.
 
-			// Iterate over all objects in the FPackageData up to GetCookedPlatformDataNextIndex
-			TArray<FWeakObjectPtr>& CachedObjects = PackageData.GetCachedObjectsInOuter();
-			int32 NumIndexes = PackageData.GetCookedPlatformDataNextIndex();
-			check(NumIndexes <= NumPlatforms * CachedObjects.Num());
-			// GetCookedPlatformDataNextIndex is a value in an inline iteration over the two-dimensional array of Objects x Platforms, in Object-major order.
-			// We take the ceiling of NextIndex/NumPlatforms to get the number of objects.
-			int32 NumObjects = (NumIndexes + NumPlatforms - 1) / NumPlatforms;
-			for (int32 ObjectIndex = 0; ObjectIndex < NumObjects; ++ObjectIndex)
-			{
-				UObject* Object = CachedObjects[ObjectIndex].Get();
-				if (!Object)
-				{
-					continue;
-				}
-				TArray<FPendingCookedPlatformData*>* PendingDatas = PendingObjects.Find(Object);
-				if (!PendingDatas || PendingDatas->Num() == 0)
+				// Find all pending BeginCacheForCookedPlatformData for this FPackageData
+				TMap<UObject*, TArray<FPendingCookedPlatformData*>> PendingObjects;
+				for (FPendingCookedPlatformData& PendingCookedPlatformData : PackageDatas->GetPendingCookedPlatformDatas())
 				{
-					// No pending BeginCacheForCookedPlatformData calls for this object; clear it now.
-					Object->ClearAllCachedCookedPlatformData();
+					if (&PendingCookedPlatformData.PackageData == &PackageData && !PendingCookedPlatformData.PollIsComplete())
+					{
+						UObject* Object = PendingCookedPlatformData.Object.Get();
+						check(Object); // Otherwise PollIsComplete would have returned true
+						check(!PendingCookedPlatformData.bHasReleased); // bHasReleased should be false since PollIsComplete returned false
+						PendingObjects.FindOrAdd(Object).Add(&PendingCookedPlatformData);
+					}
 				}
-				else
+
+				// Iterate over all objects in the FPackageData up to GetCookedPlatformDataNextIndex
+				TArray<FWeakObjectPtr>& CachedObjects = PackageData.GetCachedObjectsInOuter();
+				int32 NumIndexes = PackageData.GetCookedPlatformDataNextIndex();
+				check(NumIndexes <= NumPlatforms * CachedObjects.Num());
+				// GetCookedPlatformDataNextIndex is a value in an inline iteration over the two-dimensional array of Objects x Platforms, in Object-major order.
+				// We take the ceiling of NextIndex/NumPlatforms to get the number of objects.
+				int32 NumObjects = (NumIndexes + NumPlatforms - 1) / NumPlatforms;
+				for (int32 ObjectIndex = 0; ObjectIndex < NumObjects; ++ObjectIndex)
 				{
-					// For any pending Objects, we add a CancelManager to the FPendingCookedPlatformData to call ClearAllCachedCookedPlatformData when the pending Object,Platform pairs for that object complete.
-					FPendingCookedPlatformDataCancelManager* CancelManager = new FPendingCookedPlatformDataCancelManager();
-					CancelManager->NumPendingPlatforms = PendingDatas->Num();
-					for (FPendingCookedPlatformData* PendingCookedPlatformData : *PendingDatas)
+					UObject* Object = CachedObjects[ObjectIndex].Get();
+					if (!Object)
+					{
+						continue;
+					}
+					TArray<FPendingCookedPlatformData*>* PendingDatas = PendingObjects.Find(Object);
+					if (!PendingDatas || PendingDatas->Num() == 0)
+					{
+						// No pending BeginCacheForCookedPlatformData calls for this object; clear it now.
+						Object->ClearAllCachedCookedPlatformData();
+					}
+					else
 					{
-						// We never start a new package until after the previous cancel finished, so all of the FPendingCookedPlatformData for the PlatformData we are cancelling can not have been cancelled before.  We would leak the CancelManager if we overwrote it here.
-						check(PendingCookedPlatformData->CancelManager == nullptr);
-						// If bHasReleaased on the PendingCookedPlatformData were already true, we would leak the CancelManager because the PendingCookedPlatformData would never call Release on it.
-						check(!PendingCookedPlatformData->bHasReleased);
-						PendingCookedPlatformData->CancelManager = CancelManager;
+						// For any pending Objects, we add a CancelManager to the FPendingCookedPlatformData to call ClearAllCachedCookedPlatformData when the pending Object,Platform pairs for that object complete.
+						FPendingCookedPlatformDataCancelManager* CancelManager = new FPendingCookedPlatformDataCancelManager();
+						CancelManager->NumPendingPlatforms = PendingDatas->Num();
+						for (FPendingCookedPlatformData* PendingCookedPlatformData : *PendingDatas)
+						{
+							// We never start a new package until after the previous cancel finished, so all of the FPendingCookedPlatformData for the PlatformData we are cancelling can not have been cancelled before.  We would leak the CancelManager if we overwrote it here.
+							check(PendingCookedPlatformData->CancelManager == nullptr);
+							// If bHasReleaased on the PendingCookedPlatformData were already true, we would leak the CancelManager because the PendingCookedPlatformData would never call Release on it.
+							check(!PendingCookedPlatformData->bHasReleased);
+							PendingCookedPlatformData->CancelManager = CancelManager;
+						}
 					}
 				}
 			}
 		}
 	}
+
+	FGeneratorPackage* GeneratorPackage = PackageData.GetGeneratorPackage();
+	if (GeneratorPackage && PackageData.HasCompletedGeneration())
+	{
+		UObject* SplitObject = GeneratorPackage->FindSplitDataObject();
+		if (!SplitObject || !PackageData.GetPackage())
+		{
+			UE_LOG(LogCook, Error, TEXT("PackageSplitter: Could not find %s. This prevents us from restoring the package to its precooked state. Splitter=%s."),
+				(!PackageData.GetPackage() ? TEXT("UPackage") : TEXT("SplitDataObject")), *GeneratorPackage->GetSplitDataObjectName().ToString());
+		}
+		else
+		{
+			GeneratorPackage->GetCookPackageSplitterInstance()->PostSaveGeneratorPackage(PackageData.GetPackage(), SplitObject);
+		}
+	}
+
 	PackageData.ClearCookedPlatformData();
+
+	if (CurrentCookMode == ECookMode::CookByTheBook)
+	{
+		UPackage* Package = PackageData.GetPackage();
+		if (Package && Package->LinkerLoad)
+		{
+			// Loaders and their handles can have large buffers held in process memory and in the system file cache from the
+			// data that was loaded.  Keeping this for the lifetime of the cook is costly, so we try and unload it here.
+			Package->LinkerLoad->FlushCache();
+		}
+	}
 }
 
 void UCookOnTheFlyServer::TickCancels()
@@ -3542,7 +3718,7 @@ void UCookOnTheFlyServer::PreGarbageCollect()
 	{
 		// When a package is in the process of generating other packages,
 		// we need to keep it loaded (and all objects outered to it)
-		if (PackageData->GetGeneratorPackage())
+		if (PackageData->IsGenerating())
 		{
 			GCKeepPackages.Add(PackageData->GetPackage());
 		}
@@ -3643,11 +3819,11 @@ void UCookOnTheFlyServer::PostGarbageCollect()
 	{
 		if (PackageData->IsSaveInvalidated())
 		{
-			if (PackageData->GetGeneratorPackage())
+			if (PackageData->IsGenerating())
 			{
 				// Generator packages cannot be demoted as that would break their generation data which must be preserved
 				// We're relying on BeginCacheForCookedPlatformData not having started; this is currently guaranteed when PackageData has a valid GeneratorPackage.
-				check(PackageData->GetCookedPlatformDataNextIndex() == -1);
+				check(PackageData->GetCookedPlatformDataNextIndex() == 0);
 				// Recalculate object cache
 				PackageData->ClearObjectCache();
 				PackageData->CreateObjectCache();
@@ -8108,7 +8284,9 @@ bool UCookOnTheFlyServer::GetAllPackageFilenamesFromAssetRegistry( const FString
 				const FName PackageName = Packages[AssetIndex]->PackageName;
 
 				FString StandardFilename;
-				if (!GetPackageNameCache().CalculateCacheData(PackageName, StandardFilename, OutPackageFilenames[AssetIndex]))
+				// TODO ICookPackageSplitter: Need to handle GeneratedPackages that exist in the cooked AssetRegistry we are reading, but do not exist
+				// in the workspace domain and so are not found when we try to calculate their cache data here.
+				if (!GetPackageNameCache().TryCalculateCacheData(PackageName, StandardFilename, OutPackageFilenames[AssetIndex]))
 				{
 					if (bVerifyPackagesExist)
 					{
diff --git a/Engine/Source/Editor/UnrealEd/Private/Cooker/CookPackageData.cpp b/Engine/Source/Editor/UnrealEd/Private/Cooker/CookPackageData.cpp
--- a/Engine/Source/Editor/UnrealEd/Private/Cooker/CookPackageData.cpp
+++ b/Engine/Source/Editor/UnrealEd/Private/Cooker/CookPackageData.cpp
@@ -27,8 +27,11 @@ namespace Cook
 
 
 	FPackageData::FPackageData(FPackageDatas& PackageDatas, const FName& InPackageName, const FName& InFileName)
-		: PackageName(InPackageName), FileName(InFileName), PackageDatas(PackageDatas), PreloadableFileFormat(EPackageFormat::Binary)
-		, bIsUrgent(0), bIsVisited(0), bIsPreloadAttempted(0), bIsPreloaded(0), bHasSaveCache(0), bHasBeginPrepareSaveFailed(0), bCookedPlatformDataStarted(0), bCookedPlatformDataCalled(0), bCookedPlatformDataComplete(0), bMonitorIsCooked(0)
+		: GeneratedOwner(nullptr), PackageName(InPackageName), FileName(InFileName), PackageDatas(PackageDatas)
+		, PreloadableFileFormat(EPackageFormat::Binary), bIsUrgent(0), bIsVisited(0), bIsPreloadAttempted(0)
+		, bIsPreloaded(0), bHasSaveCache(0), bHasBeginPrepareSaveFailed(0), bCookedPlatformDataStarted(0)
+		, bCookedPlatformDataCalled(0), bCookedPlatformDataComplete(0), bMonitorIsCooked(0)
+		, bInitializedGeneratorSave(0), bCompletedGeneration(0), bGenerated(0)
 	{
 		SetState(EPackageState::Idle);
 		SendToState(EPackageState::Idle, ESendFlags::QueueAdd);
@@ -36,10 +39,16 @@ namespace Cook
 
 	FPackageData::~FPackageData()
 	{
+		ClearReferences(); // Should have been called earlier, but call it here in case it was missed
 		ClearCookedPlatforms(); // We need to send OnCookedPlatformRemoved message to the monitor, so it is not valid to destruct without calling ClearCookedPlatforms
 		SendToState(EPackageState::Idle, ESendFlags::QueueNone); // Update the monitor's counters and call exit functions
 	}
 
+	void FPackageData::ClearReferences()
+	{
+		DestroyGeneratorPackage();
+	}
+
 	const FName& FPackageData::GetPackageName() const
 	{
 		return PackageName;
@@ -595,8 +604,9 @@ namespace Cook
 	{
 		check(GetPackage() != nullptr && GetPackage()->IsFullyLoaded());
 
-		check(!GetHasBeginPrepareSaveFailed())
-		check(!GetGeneratorPackage())
+		check(!GetHasBeginPrepareSaveFailed());
+		check(!HasInitializedGeneratorSave());
+		check(!HasCompletedGeneration());
 		CheckObjectCacheEmpty();
 		CheckCookedPlatformDataEmpty();
 	}
@@ -606,7 +616,6 @@ namespace Cook
 		PackageDatas.GetCookOnTheFlyServer().ReleaseCookedPlatformData(*this);
 		ClearObjectCache();
 		SetHasBeginPrepareSaveFailed(false);
-		DestroyGeneratorPackage();
 	}
 
 	void FPackageData::OnEnterInProgress()
@@ -678,6 +687,14 @@ namespace Cook
 			SetIsPreloadAttempted(true);
 			return true;
 		}
+		if (IsGenerated() && (!GeneratorPackage || GeneratorPackage->IsDeferredPopulate()))
+		{
+			// Orphaned generator packages will fail to load, so no need to preload them
+			// Deferred populate generated packages are loaded from their generator, not from disk
+			ClearPreload();
+			SetIsPreloadAttempted(true);
+			return true;
+		}
 		if (!PreloadableFile.Get())
 		{
 			TStringBuilder<NAME_SIZE> FileNameString;
@@ -858,19 +875,23 @@ namespace Cook
 
 	void FPackageData::CheckCookedPlatformDataEmpty() const
 	{
-		check(GetCookedPlatformDataNextIndex() == -1);
+		check(GetCookedPlatformDataNextIndex() == 0);
 		check(!GetCookedPlatformDataStarted());
 		check(!GetCookedPlatformDataCalled());
 		check(!GetCookedPlatformDataComplete());
+		check(!HasInitializedGeneratorSave());
+		check(!HasCompletedGeneration());
 	}
 
 	void FPackageData::ClearCookedPlatformData()
 	{
-		CookedPlatformDataNextIndex = -1;
+		CookedPlatformDataNextIndex = 0;
 		// Note that GetNumPendingCookedPlatformData is not cleared; it persists across Saves and CookSessions
 		SetCookedPlatformDataStarted(false);
 		SetCookedPlatformDataCalled(false);
 		SetCookedPlatformDataComplete(false);
+		SetInitializedGeneratorSave(false);
+		SetCompletedGeneration(false);
 	}
 
 	void FPackageData::OnRemoveSessionPlatform(const ITargetPlatform* Platform)
@@ -905,12 +926,39 @@ namespace Cook
 				});
 	}
 
+	void FPackageData::SetGeneratedOwner(FGeneratorPackage* InGeneratedOwner)
+	{
+		check(IsGenerated());
+		check(!(GeneratedOwner && InGeneratedOwner));
+		GeneratedOwner = InGeneratedOwner;
+	}
+
+	bool FPackageData::IsLoadFromDisk() const
+	{
+		// Non-generated packages always load from disk
+		if (!IsGenerated())
+		{
+			return true;
+		}
+
+		// Packages generated from a splitter that does not support deferred populate load from their generated
+		// package in the intermediate content directory
+		if (GeneratedOwner && !GeneratedOwner->IsDeferredPopulate())
+		{
+			return true;
+		}
+
+		// Orphaned generated packages should be treated as failed to load
+		// Deferred populate generated packages are loaded from their Generator package
+		return false;
+	}
+
 	bool FPackageData::GeneratorPackageRequiresGC() const
 	{
 		// We consider that if a FPackageData has valid GeneratorPackage helper object,
 		// this means that COTFS's process of generating packages was not completed 
 		// either due to an error or because it has exceeded a maximum memory threshold. 
-		return GetGeneratorPackage() && !GetHasBeginPrepareSaveFailed();
+		return IsGenerating() && !GetHasBeginPrepareSaveFailed();
 	}
 
 	UE::Cook::FGeneratorPackage* FPackageData::CreateGeneratorPackage(const UObject* InSplitDataObject, ICookPackageSplitter* InCookPackageSplitterInstance, const FString& InGeneratedUncookedRootPath)
@@ -924,18 +972,78 @@ namespace Cook
 	// FGeneratorPackage
 
 	FGeneratorPackage::FGeneratorPackage(UE::Cook::FPackageData& InOwner, const UObject* InSplitDataObject, ICookPackageSplitter* InCookPackageSplitterInstance, const FString& InGeneratedUncookedRootPath)
-	: bIsFinalized(false)
-	, Owner(InOwner)
-	, SplitDataObject(InSplitDataObject)
+	: Owner(InOwner)
+	, SplitDataObjectName(*InSplitDataObject->GetFullName())
 	, GeneratedUncookedRootPath(InGeneratedUncookedRootPath)
 	{
 		check(PackageToGenerateNextIndex == 0);
 		check(InCookPackageSplitterInstance);
 		CookPackageSplitterInstance.Reset(InCookPackageSplitterInstance);
-		PackagesToGenerate = CookPackageSplitterInstance->GetGenerateList();
 	}
 
-	const ICookPackageSplitter::FGeneratedPackage* FGeneratorPackage::GetNextPackageToGenerate()
+	FGeneratorPackage::~FGeneratorPackage()
+	{
+		ClearGeneratedPackages();
+	}
+
+	void FGeneratorPackage::ClearGeneratedPackages()
+	{
+		for (FGeneratedStruct& GeneratedStruct : PackagesToGenerate)
+		{
+			if (GeneratedStruct.PackageData)
+			{
+				check(GeneratedStruct.PackageData->GetGeneratedOwner() == this);
+				GeneratedStruct.PackageData->SetGeneratedOwner(nullptr);
+				GeneratedStruct.PackageData = nullptr;
+			}
+		}
+	}
+
+	bool FGeneratorPackage::TryGenerateList(UPackage* OwnerPackage, UObject* OwnerObject, const FPackageNameCache& PackageNameCache, FPackageDatas& PackageDatas)
+	{
+		TArray<ICookPackageSplitter::FGeneratedPackage> GeneratorDatas = CookPackageSplitterInstance->GetGenerateList(OwnerPackage, OwnerObject);
+		PackagesToGenerate.Reset(GeneratorDatas.Num());
+		for (ICookPackageSplitter::FGeneratedPackage& ReturnedGeneratorData : GeneratorDatas)
+		{
+			FGeneratedStruct& GeneratedStruct = PackagesToGenerate.Emplace_GetRef();
+			GeneratedStruct.GeneratorData = MoveTemp(ReturnedGeneratorData);
+			ICookPackageSplitter::FGeneratedPackage& GeneratorData = GeneratedStruct.GeneratorData;
+			FString PackageName = FPaths::RemoveDuplicateSlashes(FString::Printf(TEXT("/%s/%s"),
+				*GetGeneratedUncookedRootPath(), *GeneratorData.RelativePath));
+			FName PackageFName(*PackageName);
+
+			if (!GeneratorData.GetCreateAsMap().IsSet())
+			{
+				UE_LOG(LogCook, Error, TEXT("PackageSplitter did not specify whether CreateAsMap is true for generated package. Splitter=%s, Generated=%s."),
+					*this->GetSplitDataObjectName().ToString(), *PackageName);
+				return false;
+			}
+			const FName FileName = PackageNameCache.GetCachedStandardFileName(PackageFName, false /* bRequireExists */,
+				*GeneratorData.GetCreateAsMap());
+			if (FileName.IsNone())
+			{
+				UE_LOG(LogCook, Error, TEXT("PackageSplitter could not find mounted filename for generated packagepath. Splitter=%s, Generated=%s."),
+					*this->GetSplitDataObjectName().ToString(), *PackageName);
+				return false;
+			}
+			UE::Cook::FPackageData& PackageData = PackageDatas.FindOrAddPackageData(PackageFName, FileName);
+			if (IAssetRegistry::Get()->GetAssetPackageData(PackageFName) && !PackageData.IsGenerated())
+			{
+				UE_LOG(LogCook, Warning, TEXT("PackageSplitter specified a generated package that already exists in the workspace domain. Splitter=%s, Generated=%s."),
+					*this->GetSplitDataObjectName().ToString(), *PackageName);
+				return false;
+			}
+			GeneratedStruct.PackageData = &PackageData;
+			PackageData.SetGenerated(true);
+			// No package should be generated by two different splitters. If an earlier run of this splitter generated the package, the package's owner
+			// should have been reset to null when we called ClearGeneratedPackages between then and now
+			check(PackageData.GetGeneratedOwner() == nullptr); 
+			PackageData.SetGeneratedOwner(this);
+		}
+		return true;
+	}
+
+	const FGeneratorPackage::FGeneratedStruct* FGeneratorPackage::GetNextPackageToGenerate()
 	{
 		if (PackagesToGenerate.IsValidIndex(PackageToGenerateNextIndex))
 		{
@@ -944,14 +1052,29 @@ namespace Cook
 		return nullptr;
 	}
 
-	void FGeneratorPackage::Finalize()
+	FGeneratorPackage::FGeneratedStruct* FGeneratorPackage::FindGeneratedStruct(FPackageData* PackageData)
+	{
+		for (FGeneratedStruct& GeneratedStruct : PackagesToGenerate)
+		{
+			if (GeneratedStruct.PackageData == PackageData)
+			{
+				return &GeneratedStruct;
+			}
+		}
+		return nullptr;
+	}
+
+	UObject* FGeneratorPackage::FindSplitDataObject() const
 	{
-		if (ensure(!bIsFinalized))
+		FString ObjectPath = GetSplitDataObjectName().ToString();
+
+		// SplitDataObjectName is a FullObjectPath; strip off the leading <ClassName> in "<ClassName> <Package>.<Object>:<SubObject>"
+		int32 ClassDelimiterIndex = -1;
+		if (ObjectPath.FindChar(' ', ClassDelimiterIndex))
 		{
-			check(PackageToGenerateNextIndex == PackagesToGenerate.Num());
-			GetCookPackageSplitterInstance()->FinalizeGeneratorPackage();
-			bIsFinalized = true;
+			ObjectPath.RightChopInline(ClassDelimiterIndex + 1);
 		}
+		return FindObject<UObject>(nullptr, *ObjectPath);
 	}
 
 	//////////////////////////////////////////////////////////////////////////
@@ -1404,6 +1527,10 @@ namespace Cook
 		PackageNameToPackageData.Empty();
 		FileNameToPackageData.Empty();
 		for (FPackageData* PackageData : PackageDatas)
+		{
+			PackageData->ClearReferences();
+		}
+		for (FPackageData* PackageData : PackageDatas)
 		{
 			delete PackageData;
 		}
@@ -1598,4 +1725,4 @@ namespace Cook
 #endif
 
 }
-}
\ No newline at end of file
+}
diff --git a/Engine/Source/Editor/UnrealEd/Private/Cooker/CookPackageData.h b/Engine/Source/Editor/UnrealEd/Private/Cooker/CookPackageData.h
--- a/Engine/Source/Editor/UnrealEd/Private/Cooker/CookPackageData.h
+++ b/Engine/Source/Editor/UnrealEd/Private/Cooker/CookPackageData.h
@@ -87,6 +87,11 @@ namespace Cook
 	public:
 		FPackageData(FPackageDatas& PackageDatas, const FName& InPackageName, const FName& InFileName);
 		~FPackageData();
+		/**
+		 * ClearReferences is called on every PackageData before any packageDatas are deleted,
+		 * so references are still valid during ClearReferences
+		 */
+		void ClearReferences();
 
 		FPackageData(const FPackageData& other) = delete;
 		FPackageData(FPackageData&& other) = delete;
@@ -281,6 +286,26 @@ namespace Cook
 		void DestroyGeneratorPackage() { GeneratorPackage.Reset(); }
 		/** Get GeneratorPackage helper object. */
 		FGeneratorPackage* GetGeneratorPackage() const { return GeneratorPackage.Get(); }
+		/** Get whether the PackageData has completed any necessary Generator steps and is ready for BeginCache calls */
+		bool HasCompletedGeneration() const { return static_cast<bool>(bCompletedGeneration); }
+		/** Set whether the PackageData has completed any necessary Generator steps and is ready for BeginCache calls */
+		void SetCompletedGeneration(bool Value) { bCompletedGeneration = Value != 0; }
+		/** Return whether the package is a Generator package and has started but not completed generation */
+		bool IsGenerating() const { return GetGeneratorPackage() != nullptr && HasInitializedGeneratorSave() && !HasCompletedGeneration(); }
+		/** Get whether the PackageData has completed the check for whether it is a Generator, used during Generation */
+		bool HasInitializedGeneratorSave() const { return static_cast<bool>(bInitializedGeneratorSave); }
+		/** Set whether the PackageData has completed the check for whether it is a Generator, used during Generation */
+		void SetInitializedGeneratorSave(bool Value) { bInitializedGeneratorSave = Value != 0; }
+		/** Return whether the PackageData is a generated package created by a owning generator PackageData */
+		bool IsGenerated() const { return static_cast<bool>(bGenerated); }
+		/** Set whether the PackageData is a generated package created by a owning generator PackageData */
+		void SetGenerated(bool Value) { bGenerated = Value != 0; }
+		/** Set the owning generator PackageData. Only valid to call if SetGenerated(true) has been called */
+		void SetGeneratedOwner(FGeneratorPackage* InGeneratedOwner);
+		/** Return the owning generator PackageData. Will be null if not a generated package or if orphaned */
+		FGeneratorPackage* GetGeneratedOwner() const { return GeneratedOwner; }
+		/** Return whether PumpLoads should load the package from disk or from the special generated path */
+		bool IsLoadFromDisk() const;
 
 		/** Return whether a GC is required by a generator package as soon as possible. */
 		bool GeneratorPackageRequiresGC() const;
@@ -341,6 +366,7 @@ namespace Cook
 		void OnEnterHasPackage();
 
 		TUniquePtr<FGeneratorPackage> GeneratorPackage;
+		FGeneratorPackage* GeneratedOwner;
 		TArray<const ITargetPlatform*> RequestedPlatforms;
 		TArray<const ITargetPlatform*> CookedPlatforms; // Platform part of the CookedPlatforms set. Always the same length as CookSucceeded.
 		TArray<bool> CookSucceeded; // Success flag part of the CookedPlatforms set. Always the same length as CookedPlatforms.
@@ -365,7 +391,7 @@ namespace Cook
 		};
 		FTrackedPreloadableFilePtr PreloadableFile;
 		int32 NumPendingCookedPlatformData = 0;
-		int32 CookedPlatformDataNextIndex = -1;
+		int32 CookedPlatformDataNextIndex = 0;
 		std::atomic<EPackageFormat> PreloadableFileFormat;
 
 		uint32 State : int32(EPackageState::BitCount);
@@ -379,6 +405,9 @@ namespace Cook
 		uint32 bCookedPlatformDataCalled : 1;
 		uint32 bCookedPlatformDataComplete : 1;
 		uint32 bMonitorIsCooked : 1;
+		uint32 bInitializedGeneratorSave : 1;
+		uint32 bCompletedGeneration : 1;
+		uint32 bGenerated : 1;
 	};
 
 	/**
@@ -388,31 +417,47 @@ namespace Cook
 	struct FGeneratorPackage
 	{
 	public:
+		struct FGeneratedStruct
+		{
+			ICookPackageSplitter::FGeneratedPackage GeneratorData;
+			FPackageData* PackageData = nullptr;
+		};
+
 		/** Store the provided CookPackageSplitter and prepare the packages to generate. */
 		FGeneratorPackage(UE::Cook::FPackageData& InOwner, const UObject* InSplitDataObject, ICookPackageSplitter* InCookPackageSplitterInstance, const FString& InGeneratedUncookedRootPath);
+		~FGeneratorPackage();
+		/** Clear references to owned generated packages, and mark those packages as orphaned */
+		void ClearGeneratedPackages();
+
+		/** Call the Splitter's GetGenerateList and create the PackageDatas */
+		bool TryGenerateList(UPackage* OwnerPackage, UObject* OwnerObject, const FPackageNameCache& PackageNameCache, FPackageDatas& PackageDatas);
 		/** Return next package to generate. */
-		const ICookPackageSplitter::FGeneratedPackage* GetNextPackageToGenerate();
-		/** Finalize generator package. */
-		void Finalize();
+		const FGeneratedStruct* GetNextPackageToGenerate();
 
-		/** Return whether generator package is finalized. */
-		bool IsFinalized() const { return bIsFinalized; }
 		/** Return CookPackageSplitter. */
 		ICookPackageSplitter* GetCookPackageSplitterInstance() const { return CookPackageSplitterInstance.Get(); }
 		/** Return owner FPackageData. */
 		const UE::Cook::FPackageData& GetOwner() const { return Owner; }
-		/** Return the SplitDataObject. */
-		const UObject* GetSplitDataObject() const { return SplitDataObject; }
+		/** Return the SplitDataObject's FullObjectPath. */
+		const FName GetSplitDataObjectName() const { return SplitDataObjectName; }
 		/** Returns root path of generated packages. */
 		const FString& GetGeneratedUncookedRootPath() const { return GeneratedUncookedRootPath; }
+		/** Splitter capability for whether the generated packages should populate from owner rather than load from disk */
+		bool IsDeferredPopulate() const { return CookPackageSplitterInstance->UseDeferredPopulate(); }
+		/** Return the FGeneratedStruct with the given PackageData, or null if not found. */
+		FGeneratedStruct* FindGeneratedStruct(FPackageData* PackageData);
+		/**
+		 * Find again the split object from its name, or return null if no longer in memory.
+		 * It may have been GC'd and reloaded since the last time we used it.
+		 */
+		UObject* FindSplitDataObject() const;
 
 	private:
-		bool bIsFinalized;
 		const UE::Cook::FPackageData& Owner;
-		const UObject* SplitDataObject;
+		FName SplitDataObjectName;
 		const FString GeneratedUncookedRootPath;
 		TUniquePtr<ICookPackageSplitter> CookPackageSplitterInstance; // Cached CookPackageSplitter
-		TArray<ICookPackageSplitter::FGeneratedPackage> PackagesToGenerate; // Generator's packages to generate
+		TArray<FGeneratedStruct> PackagesToGenerate; // Generator's packages to generate
 		int32 PackageToGenerateNextIndex = 0; // Index of next package to generate
 	};
 
diff --git a/Engine/Source/Editor/UnrealEd/Private/Cooker/PackageNameCache.h b/Engine/Source/Editor/UnrealEd/Private/Cooker/PackageNameCache.h
--- a/Engine/Source/Editor/UnrealEd/Private/Cooker/PackageNameCache.h
+++ b/Engine/Source/Editor/UnrealEd/Private/Cooker/PackageNameCache.h
@@ -16,7 +16,7 @@ struct FPackageNameCache
 
 	FString			GetCachedStandardFileNameString(const UPackage* Package) const;
 
-	FName			GetCachedStandardFileName(const FName& PackageName) const;
+	FName			GetCachedStandardFileName(const FName& PackageName, bool bRequireExists=true, bool bCreateAsMap=false) const;
 	FName			GetCachedStandardFileName(const UPackage* Package) const;
 
 	const FName*	GetCachedPackageNameFromStandardFileName(const FName& NormalizedFileName, bool bExactMatchRequired=true, FName* FoundFileName = nullptr) const;
@@ -27,7 +27,8 @@ struct FPackageNameCache
 
 	void			AppendCacheResults(TArray<TTuple<FName, FString>>&& PackageToStandardFileNames) const;
 
-	bool			CalculateCacheData(FName PackageName, FString& OutStandardFilename, FName& OutStandardFileFName) const;
+	bool			TryCalculateCacheData(FName PackageName, FString& OutStandardFilename, FName& OutStandardFileFName,
+										bool bRequireExists = true, bool bCreateAsMap = false) const;
 
 	bool			ContainsPackageName(FName PackageName) const;
 	
@@ -61,7 +62,7 @@ private:
 	};
 
 	bool DoesPackageExist(const FName& PackageName, FString* OutFilename) const;
-	const FCachedPackageFilename& Cache(const FName& PackageName) const;
+	const FCachedPackageFilename& Cache(const FName& PackageName, bool bRequireExists, bool bCreateAsMap) const;
 
 	mutable IAssetRegistry* AssetRegistry = nullptr;
 
@@ -69,9 +70,9 @@ private:
 	mutable TMap<FName, FName>					PackageFilenameToPackageFNameCache;
 };
 
-inline FName FPackageNameCache::GetCachedStandardFileName(const FName& PackageName) const
+inline FName FPackageNameCache::GetCachedStandardFileName(const FName& PackageName, bool bRequireExists, bool bCreateAsMap) const
 {
-	return Cache(PackageName).StandardFileName;
+	return Cache(PackageName, bRequireExists, bCreateAsMap).StandardFileName;
 }
 
 inline bool FPackageNameCache::HasCacheForPackageName(const FName& PackageName) const
@@ -82,13 +83,13 @@ inline bool FPackageNameCache::HasCacheForPackageName(const FName& PackageName)
 inline FString FPackageNameCache::GetCachedStandardFileNameString(const UPackage* Package) const
 {
 	// check( Package->GetName() == Package->GetFName().ToString() );
-	return Cache(Package->GetFName()).StandardFileNameString;
+	return Cache(Package->GetFName(), true /* bRequireExists */, false /* bCreateAsMap */).StandardFileNameString;
 }
 
 inline FName FPackageNameCache::GetCachedStandardFileName(const UPackage* Package) const
 {
 	// check( Package->GetName() == Package->GetFName().ToString() );
-	return Cache(Package->GetFName()).StandardFileName;
+	return Cache(Package->GetFName(), true /* bRequireExists */, false /* bCreateAsMap */).StandardFileName;
 }
 
 inline bool FPackageNameCache::ClearPackageFileNameCacheForPackage(const UPackage* Package) const
@@ -179,18 +180,34 @@ inline bool FPackageNameCache::DoesPackageExist(const FName& PackageName, FStrin
 	return true;
 }
 
-inline bool FPackageNameCache::CalculateCacheData(FName PackageName, FString& OutStandardFilename, FName& OutStandardFileFName) const
+inline bool FPackageNameCache::TryCalculateCacheData(FName PackageName, FString& OutStandardFilename, FName& OutStandardFileFName,
+	bool bRequireExists, bool bCreateAsMap) const
 {
 	FString FilenameOnDisk;
+	bool bFound = false;
 	if (DoesPackageExist(PackageName, &FilenameOnDisk))
+	{
+		bFound = true;
+	}
+	else if (!bRequireExists)
+	{
+		FString Extension = bCreateAsMap ? FPackageName::GetMapPackageExtension() : FPackageName::GetAssetPackageExtension();
+		if (FPackageName::TryConvertLongPackageNameToFilename(PackageName.ToString(), FilenameOnDisk, Extension))
+		{
+			bFound = true;
+		}
+	}
+	if (bFound)
 	{
 		OutStandardFilename = FPaths::ConvertRelativePathToFull(FilenameOnDisk);
-
 		FPaths::MakeStandardFilename(OutStandardFilename);
 		OutStandardFileFName = FName(*OutStandardFilename);
 		return true;
 	}
-	return false;
+	else
+	{
+		return false;
+	}
 }
 
 inline bool FPackageNameCache::ContainsPackageName(FName PackageName) const
@@ -202,7 +219,8 @@ inline IAssetRegistry* FPackageNameCache::GetAssetRegistry() const
 	return AssetRegistry;
 }
 
-inline const FPackageNameCache::FCachedPackageFilename& FPackageNameCache::Cache(const FName& PackageName) const
+inline const FPackageNameCache::FCachedPackageFilename& FPackageNameCache::Cache(const FName& PackageName,
+	bool bRequireExists, bool bCreateAsMap) const
 {
 	check(IsInGameThread());
 
@@ -210,7 +228,10 @@ inline const FPackageNameCache::FCachedPackageFilename& FPackageNameCache::Cache
 
 	if (Cached != nullptr)
 	{
-		return *Cached;
+		if (!Cached->StandardFileName.IsNone() || bRequireExists)
+		{
+			return *Cached;
+		}
 	}
 
 	// cache all the things, like it's your birthday!
@@ -218,9 +239,10 @@ inline const FPackageNameCache::FCachedPackageFilename& FPackageNameCache::Cache
 	FString FileNameString;
 	FName FileName = NAME_None;
 
-	CalculateCacheData(PackageName, FileNameString, FileName);
-
-	PackageFilenameToPackageFNameCache.Add(FileName, PackageName);
+	if (TryCalculateCacheData(PackageName, FileNameString, FileName, bRequireExists, bCreateAsMap))
+	{
+		PackageFilenameToPackageFNameCache.Add(FileName, PackageName);
+	}
 
 	return PackageFilenameCache.Emplace(PackageName, FCachedPackageFilename(MoveTemp(FileNameString), FileName));
 }
@@ -249,7 +271,7 @@ inline const FName* FPackageNameCache::GetCachedPackageNameFromStandardFileName(
 		PackageName = FName(*PotentialLongPackageName);
 	}
 
-	const FCachedPackageFilename& CachedFilename = Cache(PackageName);
+	const FCachedPackageFilename& CachedFilename = Cache(PackageName, true /* bRequireExists */, false /* bCreateAsMap */);
 
 	if (bExactMatchRequired)
 	{
diff --git a/Engine/Source/Editor/UnrealEd/Public/CookPackageSplitter.h b/Engine/Source/Editor/UnrealEd/Public/CookPackageSplitter.h
--- a/Engine/Source/Editor/UnrealEd/Public/CookPackageSplitter.h
+++ b/Engine/Source/Editor/UnrealEd/Public/CookPackageSplitter.h
@@ -7,6 +7,7 @@
 #include "CoreMinimal.h"
 
 #include "Containers/List.h"
+#include "Misc/Optional.h"
 #include "UObject/Object.h"
 #include "UObject/ObjectMacros.h"
 
@@ -24,6 +25,11 @@ public:
 	{
 		FString RelativePath;		// Generated package relative to the Parent/_Generated_ root
 		TArray<FName> Dependencies; // LongPackageNames that the generated package references
+		/* GetGenerateList must specify true if the package will be a map (.umap, contains a UWorld or ULevel), else false */
+		void SetCreateAsMap(bool bInCreateAsMap) { bCreateAsMap = bInCreateAsMap; }
+		const TOptional<bool>& GetCreateAsMap() const { return bCreateAsMap; }
+	private:
+		TOptional<bool> bCreateAsMap;
 	};
 	
 	/** 
@@ -32,11 +38,16 @@ public:
 	 */
 	static bool ShouldSplit(UObject* SplitData) { return false; }
 	
-	/** Sets the SplitDataClass instance on which the CookPackageSplitter instance should work. */
-	virtual void SetDataObject(UObject* SplitData) = 0;
+	/**
+	 * If true, this Splitter can handle populating generated packages long after the generator package is finalized and cooked.
+	 * The cooker will add requests for the packages returned from GetGenerateList and will populate and save them later.
+	 * If false the cooker will populate each package before calling PreSaveGeneratorPackage, and will save them
+	 * without further reference to the splitter after population.
+	 */
+	virtual bool UseDeferredPopulate() { return false; }
 
 	/** Return the list of packages to generate. */
-	virtual TArray<FGeneratedPackage> GetGenerateList() = 0;
+	virtual TArray<FGeneratedPackage> GetGenerateList(const UPackage* OwnerPackage, const UObject* OwnerObject) = 0;
 
 	/**
 	 * Try to populate a generated package.
@@ -46,19 +57,28 @@ public:
 	 * Note that TryPopulatePackage will not be called for every package on
 	 * an iterative cook; it will only be called for the packages with changed dependencies. 
 	 *
+	 * @param OwnerPackage				The parent package being split
+	 * @param OwnerObject				The SplitDataClass instance that this CookPackageSplitter instance was created for
 	 * @param GeneratedPackage			The package the cooker generated for the given FGeneratedPackage
 	 * @param RelativePath				The relative path from the given FGeneratedPackage
 	 * @param GeneratedPackageCookName	The name that will be given for the cooked package (differs from the uncooked package)
 	 * @return							True if successfully populates,  false on error (this will cause a cook error).
 	 */
-	virtual bool TryPopulatePackage(UPackage* GeneratedPackage, const FStringView& RelativePath, const FStringView& GeneratedPackageCookName) = 0;
+	virtual bool TryPopulatePackage(const UPackage* OwnerPackage, const UObject* OwnerObject, UPackage* GeneratedPackage, const FStringView& RelativePath, const FStringView& GeneratedPackageCookName) = 0;
+
+	/**
+	 * Called before saving the parent generator package. Depending on the value of UseDeferredPopulate, this may be
+	 * before or after the TryPopulatePackage calls.
+	 * Make any required adjustments to the parent package before it is saved into the target domain.
+	 */
+	virtual void PreSaveGeneratorPackage(UPackage* OwnerPackage, UObject* OwnerObject) {}
 
 	/**
-	 * Called after calling all TryPopulatePackage.
-	 * Make any remaining adjustments to the parent package before returning and allowing
-	 * it to be saved into the target domain.
+	 * Called after saving the parent generator package. Undo any required adjustments to the parent package that
+	 * were made in PreSaveGeneratorPackage, so that the package is once again ready for use in the editor or in
+	 * future GetGenerateList or TryPopulatePackage calls
 	 */
-	virtual void FinalizeGeneratorPackage() {}
+	virtual void PostSaveGeneratorPackage(UPackage* OwnerPackage, UObject* OwnerObject) {}
 };
 
 namespace UE
diff --git a/Engine/Source/Runtime/Engine/Private/WorldPartition/WorldPartitionCookPackageSplitter.cpp b/Engine/Source/Runtime/Engine/Private/WorldPartition/WorldPartitionCookPackageSplitter.cpp
--- a/Engine/Source/Runtime/Engine/Private/WorldPartition/WorldPartitionCookPackageSplitter.cpp
+++ b/Engine/Source/Runtime/Engine/Private/WorldPartition/WorldPartitionCookPackageSplitter.cpp
@@ -20,16 +20,26 @@ bool FWorldPartitionCookPackageSplitter::ShouldSplit(UObject* SplitData)
 	return World && !!World->GetWorldPartition();
 }
 
-void FWorldPartitionCookPackageSplitter::SetDataObject(UObject* SplitData)
+UWorld* FWorldPartitionCookPackageSplitter::ValidateDataObject(UObject* SplitData)
 {
-	PartitionedWorld = CastChecked<UWorld>(SplitData);
+	UWorld* PartitionedWorld = CastChecked<UWorld>(SplitData);
 	check(PartitionedWorld);
 	check(PartitionedWorld->PersistentLevel);
 	check(PartitionedWorld->PersistentLevel->GetWorldPartition());
+	return PartitionedWorld;
 }
 
-TArray<ICookPackageSplitter::FGeneratedPackage> FWorldPartitionCookPackageSplitter::GetGenerateList()
+const UWorld* FWorldPartitionCookPackageSplitter::ValidateDataObject(const UObject* SplitData)
 {
+	return ValidateDataObject(const_cast<UObject*>(SplitData));
+}
+
+TArray<ICookPackageSplitter::FGeneratedPackage> FWorldPartitionCookPackageSplitter::GetGenerateList(const UPackage* OwnerPackage, const UObject* OwnerObject)
+{
+	// TODO: Make WorldPartition functions const so we can honor the constness of the OwnerObject in this API function
+	const UWorld* ConstPartitionedWorld = ValidateDataObject(OwnerObject);
+	UWorld* PartitionedWorld = const_cast<UWorld*>(ConstPartitionedWorld);
+
 	// World is not initialized
 	ensure(!PartitionedWorld->bIsWorldInitialized);
 
@@ -44,20 +54,28 @@ TArray<ICookPackageSplitter::FGeneratedPackage> FWorldPartitionCookPackageSplitt
 	PackagesToGenerate.Reserve(WorldPartitionGeneratedPackages.Num());
 	for (const FString& PackageName : WorldPartitionGeneratedPackages)
 	{
-		PackagesToGenerate.Emplace_GetRef().RelativePath = PackageName;
+		ICookPackageSplitter::FGeneratedPackage& GeneratedPackage = PackagesToGenerate.Emplace_GetRef();
+		GeneratedPackage.RelativePath = PackageName;
+		// all packages we generate get a ULevel from CreateEmptyLevelForRuntimeCell and are hence maps
+		GeneratedPackage.SetCreateAsMap(true);
 		// @todo_ow: Set dependencies once we get iterative cooking working
 	}
 	return PackagesToGenerate;
 }
 
-bool FWorldPartitionCookPackageSplitter::TryPopulatePackage(UPackage* GeneratedPackage, const FStringView& RelativePath, const FStringView& GeneratedPackageCookName)
+bool FWorldPartitionCookPackageSplitter::TryPopulatePackage(const UPackage* OwnerPackage, const UObject* OwnerObject, UPackage* GeneratedPackage, const FStringView& RelativePath, const FStringView& GeneratedPackageCookName)
 {
+	// TODO: Make PopulateGeneratedPackageForCook const so we can honor the constness of the OwnerObject in this API function
+	const UWorld* ConstPartitionedWorld = ValidateDataObject(OwnerObject);
+	UWorld* PartitionedWorld = const_cast<UWorld*>(ConstPartitionedWorld);
 	UWorldPartition* WorldPartition = PartitionedWorld->PersistentLevel->GetWorldPartition();
-	return WorldPartition->PopulateGeneratedPackageForCook(GeneratedPackage, FString(RelativePath), FString(GeneratedPackageCookName));
+	bool bResult = WorldPartition->PopulateGeneratedPackageForCook(GeneratedPackage, FString(RelativePath), FString(GeneratedPackageCookName));
+	return bResult;
 }
 
-void FWorldPartitionCookPackageSplitter::FinalizeGeneratorPackage()
+void FWorldPartitionCookPackageSplitter::PreSaveGeneratorPackage(UPackage* OwnerPackage, UObject* OwnerObject)
 {
+	UWorld* PartitionedWorld = ValidateDataObject(OwnerObject);
 	UWorldPartition* WorldPartition = PartitionedWorld->PersistentLevel->GetWorldPartition();
 	WorldPartition->FinalizeGeneratedPackageForCook();
 }
diff --git a/Engine/Source/Runtime/Engine/Private/WorldPartition/WorldPartitionCookPackageSplitter.h b/Engine/Source/Runtime/Engine/Private/WorldPartition/WorldPartitionCookPackageSplitter.h
--- a/Engine/Source/Runtime/Engine/Private/WorldPartition/WorldPartitionCookPackageSplitter.h
+++ b/Engine/Source/Runtime/Engine/Private/WorldPartition/WorldPartitionCookPackageSplitter.h
@@ -12,15 +12,16 @@ class FWorldPartitionCookPackageSplitter : public ICookPackageSplitter
 public:
 	//~ Begin of ICookPackageSplitter
 	static bool ShouldSplit(UObject* SplitData);
+	virtual bool UseDeferredPopulate() { return false; }
 	virtual ~FWorldPartitionCookPackageSplitter() {}
-	virtual void SetDataObject(UObject* SplitData) override;
-	virtual TArray<ICookPackageSplitter::FGeneratedPackage> GetGenerateList() override;
-	virtual bool TryPopulatePackage(UPackage* GeneratedPackage, const FStringView& RelativePath, const FStringView& GeneratedPackageCookName) override;
-	virtual void FinalizeGeneratorPackage() override;
+	virtual TArray<ICookPackageSplitter::FGeneratedPackage> GetGenerateList(const UPackage* OwnerPackage, const UObject* OwnerObject) override;
+	virtual bool TryPopulatePackage(const UPackage* OwnerPackage, const UObject* OwnerObject, UPackage* GeneratedPackage, const FStringView& RelativePath, const FStringView& GeneratedPackageCookName) override;
+	virtual void PreSaveGeneratorPackage(UPackage* OwnerPackage, UObject* OwnerObject) override;
 	//~ End of ICookPackageSplitter
 
 private:
-	UWorld* PartitionedWorld;
+	const UWorld* ValidateDataObject(const UObject* SplitData);
+	UWorld* ValidateDataObject(UObject* SplitData);
 };
 
 #endif
\ No newline at end of file
