Removed redundant reselection code on BP compile

The re-instancing logic in the selection set already handles this, and this extra code can sometimes trip an ensure

diff --git a/Engine/Source/Editor/UnrealEd/Private/Kismet2/KismetReinstanceUtilities.cpp b/Engine/Source/Editor/UnrealEd/Private/Kismet2/KismetReinstanceUtilities.cpp
--- a/Engine/Source/Editor/UnrealEd/Private/Kismet2/KismetReinstanceUtilities.cpp
+++ b/Engine/Source/Editor/UnrealEd/Private/Kismet2/KismetReinstanceUtilities.cpp
@@ -1233,7 +1233,6 @@ struct FActorReplacementHelper
 		: NewActor(InNewActor)
 		, TargetWorldTransform(FTransform::Identity)
 		, AttachmentData( MoveTemp(InAttachmentData) )
-		, bSelectNewActor(OldActor->IsSelected())
 	{
 		CachedActorData = StaticCastSharedPtr<AActor::FActorTransactionAnnotation>(OldActor->FindOrCreateTransactionAnnotation());
 		TArray<AActor*> AttachedActors;
@@ -1298,7 +1297,6 @@ private:
 	AActor*          NewActor;
 	FTransform       TargetWorldTransform;
 	FActorAttachmentData AttachmentData;
-	bool             bSelectNewActor;
 
 	/** Holds actor component data, etc. that we use to apply */
 	TSharedPtr<AActor::FActorTransactionAnnotation> CachedActorData;
@@ -1364,11 +1362,6 @@ void FActorReplacementHelper::Finalize(const TMap<UObject*, UObject*>& OldToNewI
 		NewActor->MarkComponentsRenderStateDirty();
 	}
 
-	if (bSelectNewActor && GEditor)
-	{
-		GEditor->SelectActor(NewActor, /*bInSelected =*/true, /*bNotify =*/true);
-	}
-
 	TMap<UObject*, UObject*> ConstructedComponentReplacementMap;
 	for (UActorComponent* NewActorComponent : NewActor->GetComponents())
 	{
