Adding TRACE_BOOKMARK for map loading and editor startup.

diff --git a/Engine/Source/Editor/UnrealEd/Private/FileHelpers.cpp b/Engine/Source/Editor/UnrealEd/Private/FileHelpers.cpp
--- a/Engine/Source/Editor/UnrealEd/Private/FileHelpers.cpp
+++ b/Engine/Source/Editor/UnrealEd/Private/FileHelpers.cpp
@@ -2671,6 +2671,8 @@ bool FEditorFileUtils::LoadMap(const FString& InFilename, bool LoadAsTemplate, b
 	const double MapLoadTime = FStudioAnalytics::GetAnalyticSeconds() - LoadStartTime;
 	UE_LOG(LogFileHelpers, Log, TEXT("Loading map '%s' took %.3f"), *FPaths::GetBaseFilename(Filename), MapLoadTime);
 
+	TRACE_BOOKMARK(TEXT("LoadMap"));
+
 	FStudioAnalytics::FireEvent_Loading(TEXT("LoadMap"), MapLoadTime, {
 		FAnalyticsEventAttribute(TEXT("MapName"), FPaths::GetBaseFilename(Filename))
 	});
diff --git a/Engine/Source/Editor/UnrealEd/Private/UnrealEdMisc.cpp b/Engine/Source/Editor/UnrealEd/Private/UnrealEdMisc.cpp
--- a/Engine/Source/Editor/UnrealEd/Private/UnrealEdMisc.cpp
+++ b/Engine/Source/Editor/UnrealEd/Private/UnrealEdMisc.cpp
@@ -535,6 +535,8 @@ void FUnrealEdMisc::OnInit()
 	const double TotalEditorStartupTime = (FStudioAnalytics::GetAnalyticSeconds() - GStartTime);
 	UE_LOG(LogUnrealEdMisc, Log, TEXT("Total Editor Startup Time, took %.3f"), TotalEditorStartupTime);
 
+	TRACE_BOOKMARK(TEXT("Editor Startup"));
+
 	FStudioAnalytics::FireEvent_Loading(TEXT("TotalEditorStartup"), TotalEditorStartupTime);
 
 	GShaderCompilingManager->PrintStats(true);
