WorldPartitionRenameCommandlet - Replace map name in cells names referenced by HLOD actors

diff --git a/Engine/Source/Editor/UnrealEd/Private/Commandlets/WorldPartitionRenameCommandlet.cpp b/Engine/Source/Editor/UnrealEd/Private/Commandlets/WorldPartitionRenameCommandlet.cpp
--- a/Engine/Source/Editor/UnrealEd/Private/Commandlets/WorldPartitionRenameCommandlet.cpp
+++ b/Engine/Source/Editor/UnrealEd/Private/Commandlets/WorldPartitionRenameCommandlet.cpp
@@ -13,6 +13,8 @@
 #include "WorldPartition/ActorDescList.h"
 #include "WorldPartition/WorldPartition.h"
 #include "WorldPartition/WorldPartitionHandle.h"
+#include "WorldPartition/HLOD/HLODActor.h"
+#include "WorldPartition/HLOD/HLODActorDesc.h"
 #include "Misc/CommandLine.h"
 #include "Misc/PackageName.h"
 #include "Misc/Paths.h"
@@ -224,6 +226,7 @@ int32 UWorldPartitionRenameCommandlet::Main(const FString& Params)
 
 	// Rename world
 	const FString OldWorldPath = FSoftObjectPath(World).ToString();
+	const FString OldWorldName = FPackageName::GetShortName(OldMapFullPath);
 	const FString NewWorldName = FPackageName::GetShortName(NewMapFullPath);
 
 	ResetLoaders(World->GetPackage());
@@ -242,6 +245,15 @@ int32 UWorldPartitionRenameCommandlet::Main(const FString& Params)
 		PackagesToSave.Add(ActorDescIterator->GetActor()->GetPackage());
 	}
 
+	// Fixup cell names referenced by HLODs
+	for (FActorDescList::TIterator<AWorldPartitionHLOD> HLODActorDescIterator(WorldPartition); HLODActorDescIterator; ++HLODActorDescIterator)
+	{
+		AWorldPartitionHLOD* HLODActor = CastChecked<AWorldPartitionHLOD>(HLODActorDescIterator->GetActor());
+		FString OldCellName = HLODActor->GetCellName().ToString();
+		FString NewCellName = OldCellName.Replace(*OldWorldName, *NewWorldName);
+		HLODActor->SetCellName(FName(NewCellName));
+	}
+
 	PackagesToSave.Add(World->GetPackage());
 
 	for (const FString& PackageToDelete: PackagesToDelete)
