EditorDomain - Add in-process save, enable asynchronous EditorDomain loading, and add some cleanup for EditorDomainSave.
#rnx

diff --git a/Engine/Source/Editor/UnrealEd/Private/CookOnTheFlyServer.cpp b/Engine/Source/Editor/UnrealEd/Private/CookOnTheFlyServer.cpp
--- a/Engine/Source/Editor/UnrealEd/Private/CookOnTheFlyServer.cpp
+++ b/Engine/Source/Editor/UnrealEd/Private/CookOnTheFlyServer.cpp
@@ -4450,7 +4450,6 @@ void UCookOnTheFlyServer::Initialize( ECookMode::Type DesiredCookMode, ECookInit
 
 	if (IsCookByTheBookMode() && !IsCookingInEditor() &&
 		FPlatformMisc::SupportsMultithreadedFileHandles() && // Preloading moves file handles between threads
-		!FEditorDomain::Get() && // EDITOR_DOMAIN_TODO EditorDomain does not yet support opening packages from async threads, because AssetRegistry is not asynchronous
 		!GAllowCookedDataInEditorBuilds // // Use of preloaded files is not yet implemented when GAllowCookedDataInEditorBuilds is on, see FLinkerLoad::CreateLoader
 		)
 	{
diff --git a/Engine/Source/Editor/UnrealEd/Private/EditorDomain/EditorDomain.cpp b/Engine/Source/Editor/UnrealEd/Private/EditorDomain/EditorDomain.cpp
--- a/Engine/Source/Editor/UnrealEd/Private/EditorDomain/EditorDomain.cpp
+++ b/Engine/Source/Editor/UnrealEd/Private/EditorDomain/EditorDomain.cpp
@@ -15,6 +15,7 @@
 #include "HAL/CriticalSection.h"
 #include "Misc/CommandLine.h"
 #include "Misc/ConfigCacheIni.h"
+#include "Misc/CoreDelegates.h"
 #include "Misc/PackagePath.h"
 #include "Misc/PackageSegment.h"
 #include "Misc/ScopeLock.h"
@@ -23,6 +24,7 @@
 #include "Templates/RefCounting.h"
 #include "Templates/UniquePtr.h"
 #include "UObject/PackageResourceManagerFile.h"
+#include "UObject/UObjectIterator.h"
 
 DEFINE_LOG_CATEGORY(LogEditorDomain);
 
@@ -43,6 +45,10 @@ public:
 			GConfig->GetBool(TEXT("CookSettings"), TEXT("EditorDomainEnabled"), bEditorDomainEnabled, GEditorIni);
 			if (bEditorDomainEnabled)
 			{
+				// Set values for config settings EditorDomain depends on
+				GAllowUnversionedContentInEditor = 1;
+
+				// Create the editor domain and return it as the package resource manager
 				check(FEditorDomain::RegisteredEditorDomain == nullptr);
 				FEditorDomain::RegisteredEditorDomain = new FEditorDomain();
 				return FEditorDomain::RegisteredEditorDomain;
@@ -58,18 +64,25 @@ FEditorDomain::FEditorDomain()
 {
 	Locks = TRefCountPtr<FLocks>(new FLocks(*this));
 	Workspace.Reset(MakePackageResourceManagerFile());
-	SaveClient.Reset(new FEditorDomainSaveClient());
+	GConfig->GetBool(TEXT("CookSettings"), TEXT("EditorDomainExternalSave"), bExternalSave, GEditorIni);
+	if (bExternalSave)
+	{
+		SaveClient.Reset(new FEditorDomainSaveClient());
+	}
 	AssetRegistry = IAssetRegistry::Get();
 	// We require calling SearchAllAssets, because we rely on being able to call WaitOnAsset
 	// without needing to call ScanPathsSynchronous
 	AssetRegistry->SearchAllAssets(false /* bSynchronousSearch */);
 
 	bEditorDomainReadEnabled = !FParse::Param(FCommandLine::Get(), TEXT("noeditordomainread"));
+
+	FCoreDelegates::OnPostEngineInit.AddRaw(this, &FEditorDomain::OnPostEngineInit);
 }
 
 FEditorDomain::~FEditorDomain()
 {
 	FScopeLock ScopeLock(&Locks->Lock);
+	FCoreDelegates::OnPostEngineInit.RemoveAll(this);
 	Locks->Owner = nullptr;
 	AssetRegistry = nullptr;
 	Workspace.Reset();
@@ -107,7 +120,7 @@ FEditorDomain::FLocks::FLocks(FEditorDomain& InOwner)
 {
 }
 
-bool FEditorDomain::TryGetPackageSource(const FPackagePath& PackagePath, TRefCountPtr<FPackageSource>& OutSource)
+bool FEditorDomain::TryFindOrAddPackageSource(const FPackagePath& PackagePath, TRefCountPtr<FPackageSource>& OutSource)
 {
 	// Called within Locks.Lock
 	using namespace UE::EditorDomain;
@@ -158,6 +171,35 @@ bool FEditorDomain::TryGetPackageSource(const FPackagePath& PackagePath, TRefCou
 	}
 }
 
+TRefCountPtr<FEditorDomain::FPackageSource> FEditorDomain::FindPackageSource(const FPackagePath& PackagePath)
+{
+	// Called within Locks.Lock
+	using namespace UE::EditorDomain;
+
+	FName PackageName = PackagePath.GetPackageFName();
+	if (!PackageName.IsNone())
+	{
+		TRefCountPtr<FPackageSource>* PackageSource = PackageSources.Find(PackageName);
+		if (PackageSource)
+		{
+			return *PackageSource;
+		}
+	}
+
+	return TRefCountPtr<FPackageSource>();
+}
+
+void FEditorDomain::MarkNeedsLoadFromWorkspace(const FPackagePath& PackagePath, TRefCountPtr<FPackageSource>& PackageSource)
+{
+	PackageSource->Source = FEditorDomain::EPackageSource::Workspace;
+	if (bExternalSave)
+	{
+		SaveClient->RequestSave(PackagePath);
+	}
+	// Otherwise, we will note the need for save in OnEndLoad
+
+}
+
 int64 FEditorDomain::FileSize(const FPackagePath& PackagePath, EPackageSegment PackageSegment,
 	FPackagePath* OutUpdatedPath)
 {
@@ -173,7 +215,7 @@ int64 FEditorDomain::FileSize(const FPackagePath& PackagePath, EPackageSegment P
 	{
 		FScopeLock ScopeLock(&Locks->Lock);
 		TRefCountPtr<FPackageSource> PackageSource;
-		if (!TryGetPackageSource(PackagePath, PackageSource) || PackageSource->Source == EPackageSource::Workspace)
+		if (!TryFindOrAddPackageSource(PackagePath, PackageSource) || PackageSource->Source == EPackageSource::Workspace)
 		{
 			return Workspace->FileSize(PackagePath, PackageSegment, OutUpdatedPath);
 		}
@@ -199,8 +241,7 @@ int64 FEditorDomain::FileSize(const FPackagePath& PackagePath, EPackageSegment P
 				if (Locks->Owner)
 				{
 					FEditorDomain& EditorDomain = *Locks->Owner;
-					PackageSource->Source = EPackageSource::Workspace;
-					EditorDomain.SaveClient->RequestSave(PackagePath);
+					EditorDomain.MarkNeedsLoadFromWorkspace(PackagePath, PackageSource);
 					FileSize = EditorDomain.Workspace->FileSize(PackagePath, PackageSegment, OutUpdatedPath);
 				}
 				else
@@ -229,7 +270,7 @@ FOpenPackageResult FEditorDomain::OpenReadPackage(const FPackagePath& PackagePat
 		return Workspace->OpenReadPackage(PackagePath, PackageSegment, OutUpdatedPath);
 	}
 	TRefCountPtr<FPackageSource> PackageSource;
-	if (!TryGetPackageSource(PackagePath, PackageSource) ||	(PackageSource->Source == EPackageSource::Workspace))
+	if (!TryFindOrAddPackageSource(PackagePath, PackageSource) ||	(PackageSource->Source == EPackageSource::Workspace))
 	{
 		return Workspace->OpenReadPackage(PackagePath, PackageSegment, OutUpdatedPath);
 	}
@@ -273,7 +314,7 @@ IAsyncReadFileHandle* FEditorDomain::OpenAsyncReadPackage(const FPackagePath& Pa
 	}
 
 	TRefCountPtr<FPackageSource> PackageSource;
-	if (!TryGetPackageSource(PackagePath, PackageSource) ||
+	if (!TryFindOrAddPackageSource(PackagePath, PackageSource) ||
 		(PackageSource->Source == EPackageSource::Workspace))
 	{
 		return Workspace->OpenAsyncReadPackage(PackagePath, PackageSegment);
@@ -350,5 +391,99 @@ void FEditorDomain::IteratePackagesStatInLocalOnlyDirectory(FStringView RootDir,
 
 void FEditorDomain::Tick(float DeltaTime)
 {
-	SaveClient->Tick(DeltaTime);
+	if (bExternalSave)
+	{
+		SaveClient->Tick(DeltaTime);
+	}
 }
+
+void FEditorDomain::OnEndLoad(TConstArrayView<UPackage*> LoadedPackages)
+{
+	if (bExternalSave)
+	{
+		return;
+	}
+	TArray<UPackage*> PackagesToSave;
+	{
+		FScopeLock ScopeLock(&Locks->Lock);
+		if (!bHasPassedPostEngineInit)
+		{
+			return;
+		}
+		PackagesToSave.Reserve(LoadedPackages.Num());
+		for (UPackage* Package : LoadedPackages)
+		{
+			PackagesToSave.Add(Package);
+		}
+		FilterKeepPackagesToSave(PackagesToSave);
+	}
+
+	for (UPackage* Package : PackagesToSave)
+	{
+		if (!UE::EditorDomain::TrySavePackage(Package))
+		{
+			UE_LOG(LogEditorDomain, Warning, TEXT("Could not save package %s into EditorDomain."), *Package->GetName());
+		}
+	}
+}
+
+void FEditorDomain::OnPostEngineInit()
+{
+	TArray<UPackage*> PackagesToSave;
+	{
+		FScopeLock ScopeLock(&Locks->Lock);
+		bHasPassedPostEngineInit = true;
+		if (bExternalSave)
+		{
+			return;
+		}
+
+		FString PackageName;
+		for (TObjectIterator<UPackage> It; It; ++It)
+		{
+			UPackage* Package = *It;
+			Package->GetName(PackageName);
+			if (Package->IsFullyLoaded() && !FPackageName::IsScriptPackage(PackageName))
+			{
+				PackagesToSave.Add(Package);
+			}
+		}
+		FilterKeepPackagesToSave(PackagesToSave);
+	}
+
+	for (UPackage* Package : PackagesToSave)
+	{
+		if (!UE::EditorDomain::TrySavePackage(Package))
+		{
+			UE_LOG(LogEditorDomain, Warning, TEXT("Could not save package %s into EditorDomain."), *Package->GetName());
+		}
+	}
+}
+
+void FEditorDomain::FilterKeepPackagesToSave(TArray<UPackage*>& InOutPackagesToSave)
+{
+	FPackagePath PackagePath;
+	for (int32 Index = 0; Index < InOutPackagesToSave.Num(); )
+	{
+		UPackage* Package = InOutPackagesToSave[Index];
+		bool bKeep = false;
+		if (FPackagePath::TryFromPackageName(Package->GetFName(), PackagePath))
+		{
+			TRefCountPtr<FPackageSource> PackageSource = FindPackageSource(PackagePath);
+			if (PackageSource && PackageSource->NeedsEditorDomainSave())
+			{
+				PackageSource->bHasSaved = true;
+				bKeep = true;
+			}
+		}
+		if (bKeep)
+		{
+			++Index;
+		}
+		else
+		{
+			InOutPackagesToSave.RemoveAtSwap(Index);
+		}
+	}
+}
+
diff --git a/Engine/Source/Editor/UnrealEd/Private/EditorDomain/EditorDomain.h b/Engine/Source/Editor/UnrealEd/Private/EditorDomain/EditorDomain.h
--- a/Engine/Source/Editor/UnrealEd/Private/EditorDomain/EditorDomain.h
+++ b/Engine/Source/Editor/UnrealEd/Private/EditorDomain/EditorDomain.h
@@ -74,6 +74,7 @@ public:
 		FPackageSegmentStatVisitor Callback) override;
 	virtual void IteratePackagesStatInLocalOnlyDirectory(FStringView RootDir, 
 		FPackageSegmentStatVisitor Callback) override;
+	virtual void OnEndLoad(TConstArrayView<UPackage*> LoadedPackages) override;
 
 	// FTickableEditorObject interface
 	virtual void Tick(float DeltaTime) override;
@@ -107,14 +108,27 @@ private:
 	{
 		UE::EditorDomain::FPackageDigest Digest;
 		EPackageSource Source = EPackageSource::Undecided;
+		bool bHasSaved = false;
+		bool NeedsEditorDomainSave() const
+		{
+			return !bHasSaved && Source == EPackageSource::Workspace;
+		}
 	};
 
 	/** Disallow copy constructors */
 	FEditorDomain(const FEditorDomain& Other) = delete;
 	FEditorDomain(FEditorDomain&& Other) = delete;
 
-	/** Read the PackageSource data (domain&digest) from the cache, or from the asset registry if not in cache. */
-	bool TryGetPackageSource(const FPackagePath& PackagePath, TRefCountPtr<FPackageSource>& OutSource);
+	/** Read the PackageSource data (domain&digest) from PackageSources, or from the asset registry if not in PackageSources. */
+	bool TryFindOrAddPackageSource(const FPackagePath& PackagePath, TRefCountPtr<FPackageSource>& OutSource);
+	/** Return the PackageSource data in PackageSources, if it exists */
+	TRefCountPtr<FPackageSource> FindPackageSource(const FPackagePath& PackagePath);
+	/** Mark that we had to load the Package from the workspace domain, and schedule its save into the EditorDomain. */
+	void MarkNeedsLoadFromWorkspace(const FPackagePath& PackagePath, TRefCountPtr<FPackageSource>& PackageSource);
+	/** Callback for PostEngineInit, to handle saving of packages which we could not save before then. */
+	void OnPostEngineInit();
+	/** For each of the now-loaded packages, if we had to load from workspace domain, save into the editor domain. */
+	void FilterKeepPackagesToSave(TArray<UPackage*>& InOutLoadedPackages);
 
 	/** Subsystem used to request the save of missing packages into the EditorDomain from a separate process. */
 	TUniquePtr<FEditorDomainSaveClient> SaveClient;
@@ -128,6 +142,10 @@ private:
 	TMap<FName, TRefCountPtr<FPackageSource>> PackageSources;
 	/** True by default, set to false when reading is disabled for testing. */
 	bool bEditorDomainReadEnabled = true;
+	/** If true, use an out-of-process EditorDomainSaveServer for saves, else save in process in EndLoad */
+	bool bExternalSave = false;
+	/** Marker for whether our PostEngineInit callback has been called */
+	bool bHasPassedPostEngineInit = false;
 
 	static FEditorDomain* RegisteredEditorDomain;
 
diff --git a/Engine/Source/Editor/UnrealEd/Private/EditorDomain/EditorDomainArchive.cpp b/Engine/Source/Editor/UnrealEd/Private/EditorDomain/EditorDomainArchive.cpp
--- a/Engine/Source/Editor/UnrealEd/Private/EditorDomain/EditorDomainArchive.cpp
+++ b/Engine/Source/Editor/UnrealEd/Private/EditorDomain/EditorDomainArchive.cpp
@@ -213,11 +213,10 @@ void FEditorDomainReadArchive::OnCacheRequestComplete(UE::DerivedData::FCacheGet
 		{
 			FEditorDomain& EditorDomain(*EditorDomainLocks->Owner);
 			IPackageResourceManager& Workspace = *EditorDomain.Workspace;
-			EditorDomain.SaveClient->RequestSave(PackagePath);
 			FOpenPackageResult Result = Workspace.OpenReadPackage(PackagePath, EPackageSegment::Header);
 			if (Result.Archive)
 			{
-				PackageSource->Source = FEditorDomain::EPackageSource::Workspace;
+				EditorDomain.MarkNeedsLoadFromWorkspace(PackagePath, PackageSource);
 				InnerArchive = MoveTemp(Result.Archive);
 				AsyncSource = ESource::Archive;
 				Size = InnerArchive->TotalSize();
@@ -439,12 +438,11 @@ void FEditorDomainAsyncReadFileHandle::OnCacheRequestComplete(UE::DerivedData::F
 		{
 			FEditorDomain& EditorDomain(*EditorDomainLocks->Owner);
 			IPackageResourceManager& Workspace = *EditorDomain.Workspace;
-			EditorDomain.SaveClient->RequestSave(PackagePath);
 			IAsyncReadFileHandle* Result = Workspace.OpenAsyncReadPackage(PackagePath, EPackageSegment::Header);
 			check(Result);
 			AsyncSource = ESource::Archive;
 			InnerArchive.Reset(Result);
-			PackageSource->Source = FEditorDomain::EPackageSource::Workspace;
+			EditorDomain.MarkNeedsLoadFromWorkspace(PackagePath, PackageSource);
 		}
 		else
 		{
diff --git a/Engine/Source/Editor/UnrealEd/Private/EditorDomain/EditorDomainSave.cpp b/Engine/Source/Editor/UnrealEd/Private/EditorDomain/EditorDomainSave.cpp
--- a/Engine/Source/Editor/UnrealEd/Private/EditorDomain/EditorDomainSave.cpp
+++ b/Engine/Source/Editor/UnrealEd/Private/EditorDomain/EditorDomainSave.cpp
@@ -5,9 +5,7 @@
 #include "AssetRegistry/IAssetRegistry.h"
 #include "Async/Async.h"
 #include "Commandlets/EditorDomainSaveCommandlet.h"
-#include "DerivedDataCache.h"
 #include "DerivedDataCacheInterface.h"
-#include "DerivedDataCacheRecord.h"
 #include "Dom/JsonObject.h"
 #include "EditorDomain/EditorDomainUtils.h"
 #include "HAL/FileManager.h"
@@ -15,11 +13,9 @@
 #include "HAL/RunnableThread.h"
 #include "IPAddress.h"
 #include "Misc/DateTime.h"
-#include "Misc/FileHelper.h"
 #include "Misc/PackagePath.h"
 #include "Misc/Paths.h"
 #include "Serialization/Archive.h"
-#include "Serialization/CompactBinaryWriter.h"
 #include "Serialization/JsonReader.h"
 #include "Serialization/JsonSerializer.h"
 #include "Serialization/JsonWriter.h"
@@ -237,7 +233,7 @@ bool FEditorDomainSaveServer::TryInitialize()
 	{
 		CreatorProcessId = 0;
 	}
-	bIsResidentServer = FParse::Param(CommandLine, TEXT("editordomainserverresident"));
+	bIsResidentServer = FParse::Param(CommandLine, TEXT("EditorDomainSaveResident"));
 
 	ProcessLock = CreateProcessLock();
 	if (!ProcessLock)
@@ -707,27 +703,6 @@ bool FEditorDomainSaveServer::TrySavePackage(const FPackagePath& PackagePath, FS
 		return false;
 	}
 
-	FPackageDigest PackageDigest;
-	EPackageDigestResult FindHashResult = GetPackageDigest(*AssetRegistry, PackageName, PackageDigest);
-	switch (FindHashResult)
-	{
-	case EPackageDigestResult::Success:
-		break;
-	case EPackageDigestResult::FileDoesNotExist:
-		OutErrorMessage = FString::Printf(TEXT("Package %s no longer exists in the Workspace Domain. Skipping its save."),
-			*PackagePath.GetDebugName());
-		return false;
-	case EPackageDigestResult::WrongThread:
-		OutErrorMessage = FString::Printf(TEXT("Package %s failed because GetPackageDigest returned WrongThread; ")
-			TEXT("this should not be possible in the EditorDomainSaveServer since it only works from main thread. ")
-			TEXT("Skipping its save."), *PackagePath.GetDebugName());
-		return false;
-	default:
-		check(false);
-		OutErrorMessage = TEXT("");
-		return false;
-	}
-
 	// EDITOR_DOMAIN_TODO: Check whether the package already exists in the cache before spending time on it
 	UPackage* Package = LoadPackage(nullptr, PackagePath, 0);
 	if (!Package)
@@ -737,45 +712,11 @@ bool FEditorDomainSaveServer::TrySavePackage(const FPackagePath& PackagePath, FS
 		return false;
 	}
 
-	// EDITOR_DOMAIN_TODO: Need to extend SavePackage to allow saving to an archive
-	FString TempFilename = FPaths::Combine(FPaths::ProjectIntermediateDir(), FGuid::NewGuid().ToString());
-	ON_SCOPE_EXIT{ IFileManager::Get().Delete(*TempFilename); };
-
-	uint32 SaveFlags = SAVE_NoError | // Do not crash the SaveServer on an error
-		SAVE_BulkDataByReference; // EditorDomain saves reference bulkdata from the WorkspaceDomain rather than duplicating it
-	FSavePackageResultStruct Result = GEditor->Save(Package, nullptr, RF_Standalone, *TempFilename, GError,
-		nullptr /* Conform */, false /* bForceByteSwapping */, true /* bWarnOfLongFilename */, SaveFlags);
-	if (Result.Result != ESavePackageResult::Success)
+	if (!UE::EditorDomain::TrySavePackage(Package))
 	{
-		OutErrorMessage = FString::Printf(TEXT("Package %s could not be saved to temporary file %s: ResultCode = %d"),
-			*PackagePath.GetDebugName(), *TempFilename, Result.Result);
+		OutErrorMessage = FString::Printf(TEXT("Could not save package %s."), *PackagePath.GetDebugName());
 		return false;
 	}
-
-	FSharedBuffer PackageBuffer;
-	{
-		TArray64<uint8> TempBytes;
-		bool bBytesRead = FFileHelper::LoadFileToArray(TempBytes, *TempFilename);
-		IFileManager::Get().Delete(*TempFilename);
-		if (!bBytesRead)
-		{
-			OutErrorMessage = FString::Printf(TEXT("Package %s could not be read from temporary file %s."),
-				*PackagePath.GetDebugName(), *TempFilename, Result.Result);
-			return false;
-		}
-		PackageBuffer = MakeSharedBufferFromArray(MoveTemp(TempBytes));
-	}
-
-	UE::DerivedData::ICache& Cache = GetDerivedDataCacheRef().GetCache();
-	UE::DerivedData::FCacheRecordBuilder RecordBuilder = Cache.CreateRecord(GetEditorDomainPackageKey(PackageDigest));
-	TCbWriter<256> MetaData;
-	MetaData.BeginObject();
-	MetaData << "FileSize" << PackageBuffer.GetSize();
-	MetaData.EndObject();
-	RecordBuilder.SetMeta(MetaData.Save().AsObject());
-	RecordBuilder.SetValue(PackageBuffer);
-	UE::DerivedData::FCacheRecord Record = RecordBuilder.Build();
-	Cache.Put(MakeArrayView(&Record, 1), PackagePath.GetDebugName());
 	return true;
 }
 
@@ -916,7 +857,7 @@ void FEditorDomainSaveClient::KickCommunication()
 	if (!TrySendBatchRequest())
 	{
 		bAsyncActive = true;
-		if (!InitializeCommunication())
+		if (!TryInitializeCommunication())
 		{
 			Requests.Empty();
 		}
@@ -971,11 +912,15 @@ bool FEditorDomainSaveClient::TrySendBatchRequest()
 	}
 }
 
-bool FEditorDomainSaveClient::InitializeCommunication()
+bool FEditorDomainSaveClient::TryInitializeCommunication()
 {
 	// Called within ThreadingLock
 	using namespace UE::EditorDomainSave;
-
+	if (IsInitialized())
+	{
+		return true;
+	}
+	
 	if (!ISocketSubsystem::Get())
 	{
 		UE_LOG(LogEditorDomainSave, Error,
@@ -999,6 +944,11 @@ bool FEditorDomainSaveClient::InitializeCommunication()
 	return true;
 }
 
+bool FEditorDomainSaveClient::IsInitialized() const
+{
+	return ProcessLock != nullptr;
+}
+
 void FEditorDomainSaveClient::RunCommunication()
 {
 	using namespace UE::EditorDomainSave;
diff --git a/Engine/Source/Editor/UnrealEd/Private/EditorDomain/EditorDomainSave.h b/Engine/Source/Editor/UnrealEd/Private/EditorDomain/EditorDomainSave.h
--- a/Engine/Source/Editor/UnrealEd/Private/EditorDomain/EditorDomainSave.h
+++ b/Engine/Source/Editor/UnrealEd/Private/EditorDomain/EditorDomainSave.h
@@ -180,7 +180,9 @@ private:
 	/** If connected, send a batch of names to the server. Return true on success or empty. */
 	bool TrySendBatchRequest();
 	/** Construct on-demand resources and flags. Returns true iff saving is enabled. */
-	bool InitializeCommunication();
+	bool TryInitializeCommunication();
+	/** Return whether initialization has already been successfully run. */
+	bool IsInitialized() const;
 	/** Called on an async thread. Try to connect until successful, send messages until empty, return. */
 	void RunCommunication();
 	/** Called from the public interface thread. Pump one connection step, send one message, return. */
diff --git a/Engine/Source/Editor/UnrealEd/Private/EditorDomain/EditorDomainUtils.cpp b/Engine/Source/Editor/UnrealEd/Private/EditorDomain/EditorDomainUtils.cpp
--- a/Engine/Source/Editor/UnrealEd/Private/EditorDomain/EditorDomainUtils.cpp
+++ b/Engine/Source/Editor/UnrealEd/Private/EditorDomain/EditorDomainUtils.cpp
@@ -6,12 +6,19 @@
 #include "AssetRegistry/AssetData.h"
 #include "AssetRegistry/IAssetRegistry.h"
 #include "Containers/Array.h"
+#include "DerivedDataCache.h"
 #include "DerivedDataCacheInterface.h"
 #include "DerivedDataCacheKey.h"
+#include "DerivedDataCacheRecord.h"
+#include "Editor.h"
+#include "HAL/FileManager.h"
+#include "Misc/FileHelper.h"
 #include "Misc/PackagePath.h"
 #include "Misc/StringBuilder.h"
 #include "Serialization/CompactBinaryWriter.h"
 #include "UObject/ObjectVersion.h"
+#include "UObject/Package.h"
+#include "UObject/UObjectHash.h"
 
 namespace UE
 {
@@ -56,10 +63,6 @@ FPackageDigest GetPackageDigest(const FAssetPackageData& PackageData, FName Pack
 EPackageDigestResult GetPackageDigest(IAssetRegistry& AssetRegistry, FName PackageName,
 	FPackageDigest& OutPackageDigest)
 {
-	if (!IsInGameThread())
-	{
-		return EPackageDigestResult::WrongThread;
-	}
 	AssetRegistry.WaitForPackage(PackageName.ToString());
 	TOptional<FAssetPackageData> PackageData = AssetRegistry.GetAssetPackageDataCopy(PackageName);
 	if (!PackageData)
@@ -87,5 +90,85 @@ UE::DerivedData::FRequest RequestEditorDomainPackage(const FPackagePath& Package
 		UE::DerivedData::ECachePolicy::QueryLocal, CachePriority, MoveTemp(Callback));
 }
 
+bool TrySavePackage(UPackage* Package)
+{
+	FPackageDigest PackageDigest;
+	EPackageDigestResult FindHashResult = GetPackageDigest(*IAssetRegistry::Get(), Package->GetFName(), PackageDigest);
+	switch (FindHashResult)
+	{
+	case EPackageDigestResult::Success:
+		break;
+	case EPackageDigestResult::FileDoesNotExist:
+		return false;
+	case EPackageDigestResult::WrongThread:
+		return false;
+	default:
+		check(false);
+		return false;
+	}
+
+	// EDITOR_DOMAIN_TODO: Need to extend SavePackage to allow saving to an archive
+	FString TempFilename = FPaths::Combine(FPaths::ProjectIntermediateDir(), FGuid::NewGuid().ToString());
+	ON_SCOPE_EXIT{ IFileManager::Get().Delete(*TempFilename); };
+
+	uint32 SaveFlags = SAVE_NoError | // Do not crash the SaveServer on an error
+		SAVE_BulkDataByReference; // EditorDomain saves reference bulkdata from the WorkspaceDomain rather than duplicating it
+
+	bool bEditorDomainSaveUnversioned = false;
+	GConfig->GetBool(TEXT("CookSettings"), TEXT("EditorDomainSaveUnversioned"), bEditorDomainSaveUnversioned, GEditorIni);
+	if (bEditorDomainSaveUnversioned)
+	{
+		// With some exceptions, EditorDomain packages are saved unversioned; 
+		// editors request the appropriate version of the EditorDomain package matching their serialization version
+		bool bSaveUnversioned = true;
+		TArray<UObject*> PackageObjects;
+		GetObjectsWithPackage(Package, PackageObjects);
+		for (UObject* Object : PackageObjects)
+		{
+			UClass* Class = Object ? Object->GetClass() : nullptr;
+			if (Class && Class->HasAnyClassFlags(CLASS_CompiledFromBlueprint))
+			{
+				// TODO: Revisit this once we track package schemas
+				// Packages with Blueprint class instances can not be saved unversioned,
+				// as the Blueprint class's layout can change during the editor's lifetime,
+				// and we don't currently have a way to keep track of the changing package schema
+				bSaveUnversioned = false;
+			}
+		}
+		SaveFlags |= bSaveUnversioned ? SAVE_Unversioned : 0;
+	}
+
+	FSavePackageResultStruct Result = GEditor->Save(Package, nullptr, RF_Standalone, *TempFilename, GError,
+		nullptr /* Conform */, false /* bForceByteSwapping */, true /* bWarnOfLongFilename */, SaveFlags);
+	if (Result.Result != ESavePackageResult::Success)
+	{
+		return false;
+	}
+
+	FSharedBuffer PackageBuffer;
+	{
+		TArray64<uint8> TempBytes;
+		bool bBytesRead = FFileHelper::LoadFileToArray(TempBytes, *TempFilename);
+		IFileManager::Get().Delete(*TempFilename);
+		if (!bBytesRead)
+		{
+			return false;
+		}
+		PackageBuffer = MakeSharedBufferFromArray(MoveTemp(TempBytes));
+	}
+
+	UE::DerivedData::ICache& Cache = GetDerivedDataCacheRef().GetCache();
+	UE::DerivedData::FCacheRecordBuilder RecordBuilder = Cache.CreateRecord(GetEditorDomainPackageKey(PackageDigest));
+	TCbWriter<256> MetaData;
+	MetaData.BeginObject();
+	MetaData << "FileSize" << PackageBuffer.GetSize();
+	MetaData.EndObject();
+	RecordBuilder.SetMeta(MetaData.Save().AsObject());
+	RecordBuilder.SetValue(PackageBuffer);
+	UE::DerivedData::FCacheRecord Record = RecordBuilder.Build();
+	Cache.Put(MakeArrayView(&Record, 1), Package->GetName());
+	return true;
+}
+
 }
 }
diff --git a/Engine/Source/Editor/UnrealEd/Private/EditorDomain/EditorDomainUtils.h b/Engine/Source/Editor/UnrealEd/Private/EditorDomain/EditorDomainUtils.h
--- a/Engine/Source/Editor/UnrealEd/Private/EditorDomain/EditorDomainUtils.h
+++ b/Engine/Source/Editor/UnrealEd/Private/EditorDomain/EditorDomainUtils.h
@@ -5,12 +5,14 @@
 #include "Containers/UnrealString.h"
 #include "DerivedDataCache.h"
 #include "EditorDomain/EditorDomain.h"
+#include "Misc/ConfigCacheIni.h"
 
 DECLARE_LOG_CATEGORY_EXTERN(LogEditorDomainSave, Log, All);
 
 class FAssetPackageData;
 class FPackagePath;
 class IAssetRegistry;
+class UPackage;
 
 namespace UE
 {
@@ -39,5 +41,8 @@ UE::DerivedData::FRequest RequestEditorDomainPackage(const FPackagePath& Package
 	const FPackageDigest& PackageDigest, UE::DerivedData::EPriority CachePriority,
 	UE::DerivedData::FOnCacheGetComplete&& Callback);
 
+/** Save the given package into the EditorDomain. */
+bool TrySavePackage(UPackage* Package);
+
 }
 }
diff --git a/Engine/Source/Runtime/CoreUObject/Private/UObject/UObjectGlobals.cpp b/Engine/Source/Runtime/CoreUObject/Private/UObject/UObjectGlobals.cpp
--- a/Engine/Source/Runtime/CoreUObject/Private/UObject/UObjectGlobals.cpp
+++ b/Engine/Source/Runtime/CoreUObject/Private/UObject/UObjectGlobals.cpp
@@ -78,6 +78,8 @@ static FAutoConsoleVariableRef CVarAllowUnversionedContentInEditor(
 	ECVF_Default
 );
 
+void EndLoad(FUObjectSerializeContext* LoadContext, TArray<UPackage*>* OutLoadedPackages);
+
 COREUOBJECT_API bool GetAllowNativeComponentClassOverrides()
 {
 	static const bool bAllowNativeComponentClassOverrides = []()
@@ -1186,6 +1188,7 @@ UPackage* LoadPackageInternal(UPackage* InOuter, const FPackagePath& PackagePath
 
 	// Declare here so that the linker does not get destroyed before ResetLoaders is called
 	FLinkerLoad* Linker = nullptr;
+	TArray<UPackage*> LoadedPackages;
 	{
 		// Keep track of start time.
 		const double StartTime = FPlatformTime::Seconds();
@@ -1228,7 +1231,7 @@ UPackage* LoadPackageInternal(UPackage* InOuter, const FPackagePath& PackagePath
 
 		auto EndLoadAndCopyLocalizationGatherFlag = [&]
 		{
-			EndLoad(Linker->GetSerializeContext());
+			EndLoad(Linker->GetSerializeContext(), &LoadedPackages);
 			// Set package-requires-localization flags from archive after loading. This reinforces flagging of packages that haven't yet been resaved.
 			Result->ThisRequiresLocalizationGather(Linker->RequiresLocalizationGather());
 		};
@@ -1408,6 +1411,14 @@ UPackage* LoadPackageInternal(UPackage* InOuter, const FPackagePath& PackagePath
 		Result->SetFlags(RF_WasLoaded);
 	}
 
+	if (GIsEditor)
+	{
+		if (IsInGameThread() && !IsInAsyncLoadingThread() && GGameThreadLoadCounter == 0)
+		{
+			IPackageResourceManager::Get().OnEndLoad(LoadedPackages);
+		}
+	}
+
 	return Result;
 }
 
@@ -1568,8 +1579,12 @@ struct FCompareUObjectByLinkerAndOffset
 //
 // End loading packages.
 //
-void EndLoad(FUObjectSerializeContext* LoadContext)
+void EndLoad(FUObjectSerializeContext* LoadContext, TArray<UPackage*>* OutLoadedPackages)
 {
+	if (OutLoadedPackages)
+	{
+		OutLoadedPackages->Reset();
+	}
 	check(LoadContext);
 
 	if (IsInAsyncLoadingThread())
@@ -1590,6 +1605,7 @@ void EndLoad(FUObjectSerializeContext* LoadContext)
 	int32 NumObjectsLoaded = 0, NumObjectsFound = 0;
 	TSet<UObject*> AssetsLoaded;
 #endif
+	TSet<UPackage*> LoadedPackages;
 
 	while (LoadContext->DecrementBeginLoadCount() == 0 && (LoadContext->HasLoadedObjects() || LoadContext->HasPendingImportsOrForcedExports()))
 	{
@@ -1649,6 +1665,10 @@ void EndLoad(FUObjectSerializeContext* LoadContext)
 				for (int32 i = 0; i < ObjLoaded.Num(); i++)
 				{
 					UObject* Obj = ObjLoaded[i];
+					if (OutLoadedPackages)
+					{
+						LoadedPackages.Add(Obj->GetPackage());
+					}
 					if (Obj->GetLinker())
 					{
 						LoadedLinkers.Add(Obj->GetLinker());
@@ -1823,8 +1843,23 @@ void EndLoad(FUObjectSerializeContext* LoadContext)
 			LoadContext->DetachFromLinkers();
 		}
 	}
+
+	if (OutLoadedPackages)
+	{
+		OutLoadedPackages->Reserve(LoadedPackages.Num());
+		for (UPackage* Package : LoadedPackages)
+		{
+			OutLoadedPackages->Add(Package);
+		}
+	}
 }
 
+void EndLoad(FUObjectSerializeContext* LoadContext)
+{
+	EndLoad(LoadContext, nullptr);
+}
+
+
 /*-----------------------------------------------------------------------------
 	Object name functions.
 -----------------------------------------------------------------------------*/
diff --git a/Engine/Source/Runtime/CoreUObject/Public/UObject/PackageResourceManager.h b/Engine/Source/Runtime/CoreUObject/Public/UObject/PackageResourceManager.h
--- a/Engine/Source/Runtime/CoreUObject/Public/UObject/PackageResourceManager.h
+++ b/Engine/Source/Runtime/CoreUObject/Public/UObject/PackageResourceManager.h
@@ -3,6 +3,7 @@
 #pragma once
 
 #include "Containers/Array.h"
+#include "Containers/ArrayView.h"
 #include "Containers/Map.h"
 #include "Containers/StringFwd.h"
 #include "Containers/StringView.h"
@@ -23,6 +24,7 @@ class IAsyncReadRequest;
 class IAsyncReadFileHandle;
 class IMappedFileHandle;
 class IPackageResourceManager;
+class UPackage;
 
 DECLARE_LOG_CATEGORY_EXTERN(LogPackageResourceManager, Log, All);
 DECLARE_DELEGATE_RetVal(IPackageResourceManager*, FSetPackageResourceManager);
@@ -267,6 +269,9 @@ public:
 	 */
 	COREUOBJECT_API void IteratePackagesStatInLocalOnlyDirectory(FStringView RootDir, FPackagePathStatVisitor Callback);
 
+	/** Hook for notification of packages loaded during an outermost BeginLoad/EndLoad */
+	COREUOBJECT_API virtual void OnEndLoad(TConstArrayView<UPackage*> LoadedPackages) {}
+
 	/**
 	 * Call during engine startup to choose the appropriate PackageResourceManager for the current configuration and
 	 * construct it and bind it to be returned from Get
