WorldPartitionRenameCommandlet: Drop a redirector where the original world was.

diff --git a/Engine/Source/Editor/UnrealEd/Private/Commandlets/WorldPartitionRenameCommandlet.cpp b/Engine/Source/Editor/UnrealEd/Private/Commandlets/WorldPartitionRenameCommandlet.cpp
--- a/Engine/Source/Editor/UnrealEd/Private/Commandlets/WorldPartitionRenameCommandlet.cpp
+++ b/Engine/Source/Editor/UnrealEd/Private/Commandlets/WorldPartitionRenameCommandlet.cpp
@@ -117,8 +117,6 @@ int32 UWorldPartitionRenameCommandlet::Main(const FString& Params)
 	TArray<FString> Tokens, Switches;
 	ParseCommandLine(*Params, Tokens, Switches);
 
-	const bool bDeleteSourceMap = Switches.Contains(TEXT("DeleteSourceMap"));
-
 	if (!Switches.Contains(TEXT("AllowCommandletRendering")))
 	{
 		UE_LOG(LogWorldPartitionRenameCommandlet, Error, TEXT("The option \"-AllowCommandletRendering\" is required."));
@@ -211,20 +209,13 @@ int32 UWorldPartitionRenameCommandlet::Main(const FString& Params)
 	{
 		ActorReferences.Emplace(WorldPartition, ActorDescIterator->GetGuid());
 
-		if (bDeleteSourceMap)
-		{
-			const FString PackageFileName = SourceControlHelpers::PackageFilename(ActorDescIterator->GetActor()->GetPackage());
-			PackagesToDelete.Add(PackageFileName);
-		}
-	}
-
-	if (bDeleteSourceMap)
-	{
-		const FString PackageFileName = SourceControlHelpers::PackageFilename(World->GetPackage());
+		const FString PackageFileName = SourceControlHelpers::PackageFilename(ActorDescIterator->GetActor()->GetPackage());
 		PackagesToDelete.Add(PackageFileName);
 	}
 
 	// Rename world
+	const FString OriginalWorldName = World->GetName();
+	const FString OriginalWorldPackage = World->GetPackage()->GetName();
 	const FString OldWorldPath = FSoftObjectPath(World).ToString();
 	const FString OldWorldName = FPackageName::GetShortName(OldMapFullPath);
 	const FString NewWorldName = FPackageName::GetShortName(NewMapFullPath);
@@ -233,10 +224,10 @@ int32 UWorldPartitionRenameCommandlet::Main(const FString& Params)
 	World->GetPackage()->Rename(*NewMapFullPath, nullptr, REN_NonTransactional | REN_DontCreateRedirectors | REN_ForceNoResetLoaders);
 	World->Rename(*NewWorldName, nullptr, REN_NonTransactional | REN_DontCreateRedirectors | REN_ForceNoResetLoaders);
 
+	// Remap soft object paths
 	const FString NewWorldPath = FSoftObjectPath(World).ToString();	
 	RemapSoftObjectPaths.Add(OldWorldPath, NewWorldPath);
-
-	// Remap soft object paths
+	
 	FSoftPathFixupSerializer FixupSerializer(RemapSoftObjectPaths);
 	ForEachObjectWithPackage(World->GetPackage(), [&](UObject* Object) { Object->Serialize(FixupSerializer); return true; }, true, RF_NoFlags, EInternalObjectFlags::PendingKill);
 	for (FActorDescList::TIterator<> ActorDescIterator(WorldPartition); ActorDescIterator; ++ActorDescIterator)
@@ -256,6 +247,15 @@ int32 UWorldPartitionRenameCommandlet::Main(const FString& Params)
 
 	PackagesToSave.Add(World->GetPackage());
 
+	// Replace old map package with a redirector to the new map package
+	UPackage* RedirectorPackage = CreatePackage(*OriginalWorldPackage);
+	RedirectorPackage->ThisContainsMap();
+
+	UObjectRedirector* Redirector = NewObject<UObjectRedirector>(RedirectorPackage, *OriginalWorldName, RF_Standalone | RF_Public);
+	Redirector->DestinationObject = World;
+
+	PackagesToSave.Add(RedirectorPackage);
+
 	for (const FString& PackageToDelete: PackagesToDelete)
 	{
 		if (!PackageHelper.Delete(PackageToDelete))
@@ -268,15 +268,17 @@ int32 UWorldPartitionRenameCommandlet::Main(const FString& Params)
 	{
 		const FString PackageFileName = SourceControlHelpers::PackageFilename(PackageToSave);
 
-		if (FPaths::FileExists(PackageFileName))
+		if (!PackageHelper.Checkout(PackageToSave))
 		{
-			UE_LOG(LogWorldPartitionRenameCommandlet, Error, TEXT("File %s already exists"), *PackageFileName);
 			return 1;
 		}
 
-		if (!PackageHelper.Checkout(PackageToSave))
+		if (FPaths::FileExists(PackageFileName))
 		{
-			return 1;
+			if (!PackageHelper.Delete(PackageFileName))
+			{
+				return 1;
+			}
 		}
 
 		if (!UPackage::SavePackage(PackageToSave, nullptr, RF_Standalone, *PackageFileName, GError, nullptr, false, true, SAVE_Async))
