EditorDomain: Fixed a deadlock with async cache access

#rnx

diff --git a/Engine/Source/Editor/UnrealEd/Private/EditorDomain/EditorDomain.cpp b/Engine/Source/Editor/UnrealEd/Private/EditorDomain/EditorDomain.cpp
--- a/Engine/Source/Editor/UnrealEd/Private/EditorDomain/EditorDomain.cpp
+++ b/Engine/Source/Editor/UnrealEd/Private/EditorDomain/EditorDomain.cpp
@@ -276,7 +276,13 @@ FOpenPackageResult FEditorDomain::OpenReadPackage(const FPackagePath& PackagePat
 	}
 
 	FEditorDomainReadArchive* Result = new FEditorDomainReadArchive(Locks, PackagePath, PackageSource);
-	UE::DerivedData::FRequest Request = RequestEditorDomainPackage(PackagePath, PackageSource->Digest,
+	const FPackageDigest PackageSourceDigest = PackageSource->Digest;
+	const bool bHasEditorSource = (PackageSource->Source == EPackageSource::Editor);
+
+	// Unlock before requesting the package because the completion callback takes the lock.
+	ScopeLock.Unlock();
+
+	UE::DerivedData::FRequest Request = RequestEditorDomainPackage(PackagePath, PackageSourceDigest,
 		UE::DerivedData::EPriority::Normal,
 		[Result](UE::DerivedData::FCacheGetCompleteParams&& Params)
 		{
@@ -285,21 +291,14 @@ FOpenPackageResult FEditorDomain::OpenReadPackage(const FPackagePath& PackagePat
 		});
 	Result->SetRequest(Request);
 
-	EPackageFormat Format;
-	if (PackageSource->Source == EPackageSource::Editor)
-	{
-		Format = EPackageFormat::Binary;
-	}
-	else
-	{
-		// EDITOR_DOMAIN_TODO: Reading GetPackageFormat forces us to wait for the cache response
-		// We should read just the metadata for the package so we don't have to block here on the transfer of the bytes.
-		Format = Result->GetPackageFormat();
-	}
 	if (OutUpdatedPath)
 	{
 		*OutUpdatedPath = PackagePath;
 	}
+
+	// EDITOR_DOMAIN_TODO: Reading GetPackageFormat forces us to wait for the cache response
+	// We should read just the metadata for the package so we don't have to block here on the transfer of the bytes.
+	const EPackageFormat Format = bHasEditorSource ? EPackageFormat::Binary : Result->GetPackageFormat();
 	return FOpenPackageResult{ TUniquePtr<FArchive>(Result), Format };
 }
 
