IterationProfile support for performance telemetry

This change adds:
* Extraction of performance measurements from utrace to csv care of a SummarizeTraceCommandlet
* Filtering of those measurements into another csv file that is legible by the Horde telemetry system
* Split of iteration profile into two passes with differing DDC settings: cold with empty cache and warm with seeded cache
* Synchronization of async processing of assets to wait for processing to complete (as a user will do usually)
* Culling of unused pathways in IterationProfile.xml for brevity's sake

diff --git a/Engine/Config/BaseEngine.ini b/Engine/Config/BaseEngine.ini
--- a/Engine/Config/BaseEngine.ini
+++ b/Engine/Config/BaseEngine.ini
@@ -1880,6 +1880,7 @@ Boot=(Type=Boot, Filename="%GAMEDIR%DerivedDataCache/Boot.ddc", MaxCacheSize=512
 Local=(Type=FileSystem, ReadOnly=false, Clean=false, Flush=false, PurgeTransient=true, DeleteUnused=true, UnusedFileAge=34, FoldersToClean=-1, Path=%ENGINEDIR%DerivedDataCache, EnvPathOverride=UE-LocalDataCachePath, EditorOverrideSetting=LocalDerivedDataCache)
 Pak=(Type=ReadPak, Filename="%GAMEDIR%DerivedDataCache/DDC.ddp")
 
+; Creates a completely new standalone cache so as to not destroy a valuable populated local cache
 [Cold]
 MinimumDaysToKeepFile=0
 Root=(Type=KeyLength, Length=120, Inner=AsyncPut)
@@ -1888,6 +1889,14 @@ Hierarchy=(Type=Hierarchical, Inner=Boot, Inner=Local)
 Boot=(Type=Boot, Filename="%GAMEDIR%DerivedDataCache/Boot.ddc", MaxCacheSize=512)
 Local=(Type=FileSystem, ReadOnly=false, Clean=false, Flush=true, PurgeTransient=true, DeleteUnused=true, UnusedFileAge=1, FoldersToClean=-1, Path="%GAMEDIR%ColdDerivedDataCache\")
 
+; Inherits the same standalone cache from [Cold], for incremental performance measurement purposes
+[Warm]
+MinimumDaysToKeepFile=0
+Root=(Type=KeyLength, Length=120, Inner=AsyncPut)
+AsyncPut=(Type=AsyncPut, Inner=Hierarchy)
+Hierarchy=(Type=Hierarchical, Inner=Boot, Inner=Local)
+Boot=(Type=Boot, Filename="%GAMEDIR%DerivedDataCache/Boot.ddc", MaxCacheSize=512)
+Local=(Type=FileSystem, ReadOnly=false, Clean=false, Flush=false, PurgeTransient=true, DeleteUnused=true, UnusedFileAge=1, FoldersToClean=-1, Path="%GAMEDIR%ColdDerivedDataCache\")
 
 [CreatePak]
 MinimumDaysToKeepFile=7
diff --git a/Engine/Plugins/Tests/EditorTests/Source/EditorTests/Private/UnrealEd/EditorIterationTests.cpp b/Engine/Plugins/Tests/EditorTests/Source/EditorTests/Private/UnrealEd/EditorIterationTests.cpp
--- a/Engine/Plugins/Tests/EditorTests/Source/EditorTests/Private/UnrealEd/EditorIterationTests.cpp
+++ b/Engine/Plugins/Tests/EditorTests/Source/EditorTests/Private/UnrealEd/EditorIterationTests.cpp
@@ -10,6 +10,7 @@
 #include "Tests/AutomationEditorCommon.h"
 #include "Tests/AutomationCommon.h"
 #include "AutomationControllerSettings.h"
+#include "AssetCompilingManager.h"
 
 #include "Subsystems/AssetEditorSubsystem.h"
 #include "Editor.h"
@@ -53,25 +54,13 @@ void FIterationOpenAssets::GetTests(TArray<FString>& OutBeautifiedNames, TArray<
 
 bool FIterationOpenAssets::RunTest(const FString& LongAssetPath)
 {
-	//start with all editors closed
-	GEditor->GetEditorSubsystem<UAssetEditorSubsystem>()->CloseAllAssetEditors();
-
-	// below is all latent action, so before sending there, verify the asset exists
-	UObject* Object = FSoftObjectPath(LongAssetPath).TryLoad();
-	if ( Object == nullptr )
-	{
-		UE_LOG(LogEditorAutomationTests, Error, TEXT("Failed to Open Asset '%s'."), *LongAssetPath);
-		return false;
-	}
-
+	AddCommand(new FCloseAllAssetEditorsCommand());
 	AddCommand(new FOpenEditorForAssetCommand(LongAssetPath));
 	AddCommand(new FWaitLatentCommand(0.5f));
-	AddCommand(new FDelayedFunctionLatentCommand([=] {
-		if ( Object->GetOutermost()->IsDirty() )
-		{
-			UE_LOG(LogEditorAutomationTests, Error, TEXT("Asset '%s' was dirty after opening it."), *LongAssetPath);
-		}
+	AddCommand(new FFunctionLatentCommand([] {
+		return FAssetCompilingManager::Get().GetNumRemainingAssets() == 0;
 	}));
+	AddCommand(new FWaitLatentCommand(0.5f));
 	AddCommand(new FCloseAllAssetEditorsCommand());
 	AddCommand(new FDelayedFunctionLatentCommand([] {
 		CollectGarbage(GARBAGE_COLLECTION_KEEPFLAGS);
diff --git a/Engine/Source/Editor/UnrealEd/Classes/Commandlets/SummarizeTraceCommandlet.h b/Engine/Source/Editor/UnrealEd/Classes/Commandlets/SummarizeTraceCommandlet.h
new file mode 100644
--- /dev/null
+++ b/Engine/Source/Editor/UnrealEd/Classes/Commandlets/SummarizeTraceCommandlet.h
@@ -0,0 +1,23 @@
+// Copyright Epic Games, Inc. All Rights Reserved.
+
+/*=============================================================================
+	SummarizeTraceCommandlet.cpp:
+	  Commandlet for summarizing utrace cpu events into to csv
+=============================================================================*/
+
+#pragma once
+#include "Commandlets/Commandlet.h"
+#include "SummarizeTraceCommandlet.generated.h"
+
+UCLASS(config=Editor)
+class USummarizeTraceCommandlet : public UCommandlet
+{
+	GENERATED_UCLASS_BODY()
+
+public:
+	//~ Begin UCommandlet Interface
+
+	virtual int32 Main(const FString& CmdLineParams) override;
+	
+	//~ End UCommandlet Interface
+};
diff --git a/Engine/Source/Editor/UnrealEd/Private/Commandlets/SummarizeTraceCommandlet.cpp b/Engine/Source/Editor/UnrealEd/Private/Commandlets/SummarizeTraceCommandlet.cpp
new file mode 100644
--- /dev/null
+++ b/Engine/Source/Editor/UnrealEd/Private/Commandlets/SummarizeTraceCommandlet.cpp
@@ -0,0 +1,676 @@
+// Copyright Epic Games, Inc. All Rights Reserved.
+
+/*=============================================================================
+	SummarizeTraceCommandlet.cpp: Commandlet for summarizing a utrace
+=============================================================================*/
+
+#include "Commandlets/SummarizeTraceCommandlet.h"
+
+#include "Containers/StringConv.h"
+#include "GenericPlatform/GenericPlatformFile.h"
+#include "HAL/PlatformFileManager.h"
+#include "Misc/Crc.h"
+#include "Misc/FileHelper.h"
+#include "Misc/Paths.h"
+#include "String/ParseTokens.h"
+#include "Trace/Analysis.h"
+#include "Trace/Analyzer.h"
+#include "Trace/DataStream.h"
+
+/*
+ * The following could be in TraceServices. This way if the format of CPU scope
+ * events change this interface acts as a compatibility contract for external
+ * tools. In the future it may be in a separate library so that Trace and
+ * Insights' instrumentation are more broadly available.
+ */
+
+static uint64 Decode7bit(const uint8*& Cursor)
+{
+	uint64 Value = 0;
+	uint64 ByteIndex = 0;
+	bool bHasMoreBytes;
+	do
+	{
+		uint8 ByteValue = *Cursor++;
+		bHasMoreBytes = ByteValue & 0x80;
+		Value |= uint64(ByteValue & 0x7f) << (ByteIndex * 7);
+		++ByteIndex;
+	} 	while (bHasMoreBytes);
+	return Value;
+}
+
+class FCpuAnalyzer
+	: public UE::Trace::IAnalyzer
+{
+public:
+	struct FScopeName
+	{
+		const TCHAR* Name;
+		uint32			Id;
+	};
+
+	struct FScopeEnter
+	{
+		double			TimeStamp;
+		uint32			ScopeId;
+		uint32			ThreadId;
+	};
+
+	struct FScopeExit
+	{
+		double			TimeStamp;
+		uint32			ThreadId;
+	};
+
+	virtual void		OnCpuScopeName(const FScopeName& ScopeName) = 0;
+	virtual void		OnCpuScopeEnter(const FScopeEnter& ScopeEnter) = 0;
+	virtual void		OnCpuScopeExit(const FScopeExit& ScopeExit) = 0;
+
+private:
+	virtual void		OnAnalysisBegin(const FOnAnalysisContext& Context) override;
+	virtual bool		OnEvent(uint16 RouteId, EStyle Style, const FOnEventContext& Context) override;
+	void				OnEventSpec(const FOnEventContext& Context);
+	void				OnBatch(const FOnEventContext& Context);
+};
+
+enum
+{
+	RouteId_EventSpec,
+	RouteId_EventBatch,
+	RouteId_EndCapture,
+};
+
+void FCpuAnalyzer::OnAnalysisBegin(const FOnAnalysisContext& Context)
+{
+	auto& Builder = Context.InterfaceBuilder;
+	Builder.RouteEvent(RouteId_EventSpec, "CpuProfiler", "EventSpec");
+	Builder.RouteEvent(RouteId_EventBatch, "CpuProfiler", "EventBatch");
+	Builder.RouteEvent(RouteId_EndCapture, "CpuProfiler", "EndCapture");
+}
+
+bool FCpuAnalyzer::OnEvent(uint16 RouteId, EStyle Style, const FOnEventContext& Context)
+{
+	switch (RouteId)
+	{
+	case RouteId_EventSpec:
+		OnEventSpec(Context);
+		break;
+
+	case RouteId_EventBatch:
+	case RouteId_EndCapture:
+		OnBatch(Context);
+		break;
+	};
+
+	return true;
+}
+
+void FCpuAnalyzer::OnEventSpec(const FOnEventContext& Context)
+{
+	const auto& EventData = Context.EventData;
+	FString Name;
+	uint32 Id = EventData.GetValue<uint32>("Id");
+	EventData.GetString("Name", Name);
+	OnCpuScopeName({ *Name, Id - 1 });
+}
+
+void FCpuAnalyzer::OnBatch(const FOnEventContext& Context)
+{
+	const auto& EventData = Context.EventData;
+	const auto& EventTime = Context.EventTime;
+
+	uint32 ThreadId = Context.ThreadInfo.GetId();
+
+	const uint8* Cursor = EventData.GetAttachment();
+	const uint8* End = Cursor + EventData.GetAttachmentSize();
+	uint64 LastCycle = 0;
+	while (Cursor < End)
+	{
+		uint64 Value = Decode7bit(Cursor);
+		uint64 Cycle = LastCycle + (Value >> 1);
+		LastCycle = Cycle;
+
+		double TimeStamp = EventTime.AsSeconds(Cycle);
+		if (Value & 1)
+		{
+			uint64 ScopeId = Decode7bit(Cursor);
+			OnCpuScopeEnter({ TimeStamp, uint32(ScopeId - 1), ThreadId });
+		}
+		else
+		{
+			OnCpuScopeExit({ TimeStamp, ThreadId });
+		}
+	}
+}
+
+/*
+ * This too could be housed elsewhere, along with an API to make it easier to
+ * run analysis on trace files. The current model is influenced a little too much
+ * on the store model that Insights' browser mode hosts.
+ */
+
+class FFileDataStream : public UE::Trace::IInDataStream
+{
+public:
+	FFileDataStream()
+		: Handle(nullptr)
+	{
+	}
+
+	~FFileDataStream()
+	{
+		delete Handle;
+	}
+
+	bool Open(const TCHAR* Path)
+	{
+		Handle = FPlatformFileManager::Get().GetPlatformFile().OpenRead(Path);
+		if (Handle == nullptr)
+		{
+			return false;
+		}
+		Remaining = Handle->Size();
+		return true;
+	}
+
+	virtual int32 Read(void* Data, uint32 Size) override
+	{
+		if (Remaining <= 0 || Handle == nullptr)
+		{
+			return 0;
+		}
+
+		Size = (Size < Remaining) ? Size : Remaining;
+		Remaining -= Size;
+		return Handle->Read((uint8*)Data, Size) ? Size : 0;
+	}
+
+	IFileHandle* Handle;
+	uint64 Remaining;
+};
+
+/*
+ * Helper class for the SummarizeTrace commandlet.Aggregates statistics about
+ * a trace's CPU scopes.
+ */
+
+class FSummarizeAnalyzer
+	: public FCpuAnalyzer
+{
+public:
+	virtual void OnCpuScopeName(const FScopeName& ScopeName) override;
+	virtual void OnCpuScopeEnter(const FScopeEnter& ScopeEnter) override;
+	virtual void OnCpuScopeExit(const FScopeExit& ScopeExit) override;
+
+	struct FThread
+	{
+		TArray<FScopeEnter> ScopeStack;
+	};
+
+	struct FScope
+	{
+		FString Name;
+		uint64 Count = 0;
+		double TotalSeconds = 0.0;
+		double FirstStartSeconds = 0.0;
+		double FirstEndSeconds = 0.0;
+		double FirstSeconds = 0.0;
+		double MinSeconds = 1e10;
+		double MaxSeconds = -1e10;
+		double MeanSeconds = 0.0;
+		double VarianceAcc = 0.0; // Accumulator for Welford's
+
+		void AddDuration(double StartSeconds, double EndSeconds)
+		{
+			Count += 1;
+
+			// only set first for the first sample, compare exact zero
+			if (FirstStartSeconds == 0.0)
+			{
+				FirstStartSeconds = StartSeconds;
+				FirstEndSeconds = EndSeconds;
+				FirstSeconds = FirstEndSeconds - FirstStartSeconds;
+			}
+
+			// compute the duration
+			double DurationSeconds = EndSeconds - StartSeconds;
+
+			// set duration statistics
+			TotalSeconds += DurationSeconds;
+			MinSeconds = FMath::Min(MinSeconds, DurationSeconds);
+			MaxSeconds = FMath::Max(MaxSeconds, DurationSeconds);
+			UpdateVariance(DurationSeconds);
+		}
+
+		void UpdateVariance(double DurationSeconds)
+		{
+			ensure(Count);
+
+			// Welford's increment
+			double OldMeanSeconds = MeanSeconds;
+			MeanSeconds = MeanSeconds + ((DurationSeconds - MeanSeconds) / double(Count));
+			VarianceAcc = VarianceAcc + ((DurationSeconds - MeanSeconds) * (DurationSeconds - OldMeanSeconds));
+		}
+
+		double GetDeviationSeconds() const
+		{
+			if (Count > 1)
+			{
+				// Welford's final step, dependent on sample count
+				double VarianceSecondsSquared = VarianceAcc / double(Count - 1);
+
+				// stddev is sqrt of variance, to restore to units of seconds (vs. seconds squared)
+				return sqrt(VarianceSecondsSquared);
+			}
+			else
+			{
+				return 0.0;
+			}
+		}
+
+		void Merge(const FScope& Scope)
+		{
+			check(Name == Scope.Name);
+			TotalSeconds += Scope.TotalSeconds;
+			MinSeconds = FMath::Min(MinSeconds, Scope.MinSeconds);
+			MaxSeconds = FMath::Max(MaxSeconds, Scope.MaxSeconds);
+			Count += Scope.Count;
+		}
+
+		FString GetValue(const FStringView& Statistic) const
+		{
+			if (Statistic == TEXT("Name"))
+			{
+				return Name;
+			}
+			else if (Statistic == TEXT("Count"))
+			{
+				return FString::Printf(TEXT("%llu"), Count);
+			}
+			else if (Statistic == TEXT("TotalSeconds"))
+			{
+				return FString::Printf(TEXT("%f"), TotalSeconds);
+			}
+			else if (Statistic == TEXT("FirstStartSeconds"))
+			{
+				return FString::Printf(TEXT("%f"), FirstStartSeconds);
+			}
+			else if (Statistic == TEXT("FirstEndSeconds"))
+			{
+				return FString::Printf(TEXT("%f"), FirstEndSeconds);
+			}
+			else if (Statistic == TEXT("FirstSeconds"))
+			{
+				return FString::Printf(TEXT("%f"), FirstSeconds);
+			}
+			else if (Statistic == TEXT("MinSeconds"))
+			{
+				return FString::Printf(TEXT("%f"), MinSeconds);
+			}
+			else if (Statistic == TEXT("MaxSeconds"))
+			{
+				return FString::Printf(TEXT("%f"), MaxSeconds);
+			}
+			else if (Statistic == TEXT("MeanSeconds"))
+			{
+				return FString::Printf(TEXT("%f"), MeanSeconds);
+			}
+			else if (Statistic == TEXT("DeviationSeconds"))
+			{
+				return FString::Printf(TEXT("%f"), GetDeviationSeconds());
+			}
+			return FString();
+		}
+
+		// for deduplication
+		bool operator==(const FScope& Scope) const
+		{
+			return Name == Scope.Name;
+		}
+
+		// for sorting descending
+		bool operator<(const FScope& Scope) const
+		{
+			return TotalSeconds > Scope.TotalSeconds;
+		}
+	};
+
+	TArray<FScope> Scopes;
+	TArray<FThread> Threads;
+};
+
+static uint32 GetTypeHash(const FSummarizeAnalyzer::FScope Key)
+{
+	return FCrc::StrCrc32(*Key.Name);
+}
+
+void FSummarizeAnalyzer::OnCpuScopeName(const FScopeName& ScopeName)
+{
+	if (ScopeName.Id >= uint32(Scopes.Num()))
+	{
+		uint32 Num = (ScopeName.Id + 128) & ~127;
+		Scopes.SetNum(Num);
+	}
+	Scopes[ScopeName.Id].Name = ScopeName.Name;
+}
+
+void FSummarizeAnalyzer::OnCpuScopeEnter(const FScopeEnter& ScopeEnter)
+{
+	uint32 ThreadId = ScopeEnter.ThreadId;
+	if (ThreadId >= uint32(Threads.Num()))
+	{
+		Threads.SetNum(ThreadId + 1);
+	}
+	Threads[ThreadId].ScopeStack.Add(ScopeEnter);
+}
+
+void FSummarizeAnalyzer::OnCpuScopeExit(const FScopeExit& ScopeExit)
+{
+	uint32 ThreadId = ScopeExit.ThreadId;
+	if (ThreadId >= uint32(Threads.Num()) || Threads[ThreadId].ScopeStack.Num() <= 0)
+	{
+		return;
+	}
+
+	FScopeEnter ScopeEnter = Threads[ThreadId].ScopeStack.Pop();
+	double ScopeDuration = ScopeExit.TimeStamp - ScopeEnter.TimeStamp;
+
+	// unclear why we are getting ids that are out-of-bounds, fewer specs than scopes shouldn't be possible
+	//  maybe we are losing spec data?
+	//  maybe scope to spec id enc/dec has edge cases?
+	if (ScopeEnter.ScopeId < uint32(Scopes.Num()))
+	{
+		Scopes[ScopeEnter.ScopeId].AddDuration(ScopeEnter.TimeStamp, ScopeExit.TimeStamp);
+	}
+}
+
+/*
+ * SummarizeTrace commandlet ingests a utrace file and summarizes the
+ * cpu scope events within it, and summarizes each event to a csv. It
+ * also can generate a telemetry file given statistics csv about what
+ * events and what statistics you would like to track.
+ */
+
+DEFINE_LOG_CATEGORY_STATIC(LogSummarizeTrace, Log, All);
+
+USummarizeTraceCommandlet::USummarizeTraceCommandlet(const FObjectInitializer& ObjectInitializer)
+	: Super(ObjectInitializer)
+{
+}
+
+// helper for USummarizeTraceCommandlet::Main
+static void FilePrint(IFileHandle* Handle, const FString& String)
+{
+	const auto& UTF8String = StringCast<ANSICHAR>(*String);
+	Handle->Write(reinterpret_cast<const uint8*>(UTF8String.Get()), UTF8String.Length());
+}
+
+int32 USummarizeTraceCommandlet::Main(const FString& CmdLineParams)
+{
+	TArray<FString> Tokens;
+	TArray<FString> Switches;
+	TMap<FString, FString> ParamVals;
+	UCommandlet::ParseCommandLine(*CmdLineParams, Tokens, Switches, ParamVals);
+
+	// Display help
+	if (Switches.Contains("help"))
+	{
+		UE_LOG(LogSummarizeTrace, Log, TEXT("SummarizeTrace"));
+		UE_LOG(LogSummarizeTrace, Log, TEXT("This commandlet will summarize a utrace into something more easily ingestable by a reporting tool (csv)."));
+		UE_LOG(LogSummarizeTrace, Log, TEXT("Options:"));
+		UE_LOG(LogSummarizeTrace, Log, TEXT(" Required: -inputfile=<utrace path>   (The utrace you wish to process)"));
+		UE_LOG(LogSummarizeTrace, Log, TEXT(" Optional: -statsfile=<csv path>      (The csv of statistics to generate)"));
+		return 0;
+	}
+
+	FString TraceFileName;
+	if (FParse::Value(*CmdLineParams, TEXT("inputfile="), TraceFileName, true))
+	{
+		UE_LOG(LogSummarizeTrace, Display, TEXT("Loading trace from %s"), *TraceFileName);
+	}
+	else
+	{
+		UE_LOG(LogSummarizeTrace, Error, TEXT("You must specify a utrace file using -inputfile=<path>"));
+		return 1;
+	}
+
+	FString StatisticsFileName;
+	TMultiMap<FString, FString> NameToStatisticMap;
+	if (FParse::Value(*CmdLineParams, TEXT("statsfile="), StatisticsFileName, true))
+	{
+		UE_LOG(LogSummarizeTrace, Display, TEXT("Generating statistics from %s"), *StatisticsFileName);
+
+		TArray<FString> ParsedCSVFile;
+		FFileHelper::LoadFileToStringArray(ParsedCSVFile, *StatisticsFileName);
+
+		int NameColumn = -1;
+		int StatisticColumn = -1;
+		struct Column
+		{
+			const TCHAR* Name = nullptr;
+			int* Index = nullptr;
+		}
+		Columns[] =
+		{
+			{ TEXT("Name"), &NameColumn },
+			{ TEXT("Statistic"), &StatisticColumn },
+		};
+
+		bool bValidColumns = true;
+		for (int CSVIndex = 0; CSVIndex < ParsedCSVFile.Num() && bValidColumns; ++CSVIndex)
+		{
+			const FString& CSVEntry = ParsedCSVFile[CSVIndex];
+			TArray<FString> Fields;
+			UE::String::ParseTokens(CSVEntry.TrimStartAndEnd(), TEXT(','),
+				[&Fields](FStringView Field)
+				{
+					if (!Field.IsEmpty())
+					{
+						Fields.Add(FString(Field));
+					}
+				});
+
+			for (auto& Column : Columns)
+			{
+				if (CSVIndex == 0)
+				{
+					for (int FieldIndex = 0; FieldIndex < Fields.Num(); ++FieldIndex)
+					{
+						if (Fields[FieldIndex] == Column.Name)
+						{
+							(*Column.Index) = FieldIndex;
+							break;
+						}
+					}
+
+					if (*Column.Index == -1)
+					{
+						bValidColumns = false;
+					}
+				}
+				else
+				{
+					NameToStatisticMap.AddUnique(FString(Fields[NameColumn]), FString(Fields[StatisticColumn]));
+				}
+			}
+		}
+	}
+
+	bool bFound;
+	if (FPaths::FileExists(TraceFileName))
+	{
+		bFound = true;
+	}
+	else
+	{
+		bFound = false;
+		TArray<FString> SearchPaths;
+		SearchPaths.Add(FPaths::Combine(FPaths::EngineDir(), TEXT("Programs"), TEXT("UnrealInsights"), TEXT("Saved"), TEXT("TraceSessions")));
+		SearchPaths.Add(FPaths::EngineDir());
+		SearchPaths.Add(FPaths::ProjectDir());
+		for (const auto& SearchPath : SearchPaths)
+		{
+			FString PossibleTraceFileName = FPaths::Combine(SearchPath, TraceFileName);
+			if (FPaths::FileExists(PossibleTraceFileName))
+			{
+				TraceFileName = PossibleTraceFileName;
+				bFound = true;
+				break;
+			}
+		}
+	}
+
+	if (!bFound)
+	{
+		UE_LOG(LogSummarizeTrace, Error, TEXT("Trace file '%s' was not found"), *TraceFileName);
+		return 1;
+	}
+
+	FFileDataStream DataStream;
+	if (!DataStream.Open(*TraceFileName))
+	{
+		UE_LOG(LogSummarizeTrace, Error, TEXT("Unable to open trace file '%s' for read"), *TraceFileName);
+		return 1;
+	}
+
+	FSummarizeAnalyzer Analyzer;
+	UE::Trace::FAnalysisContext Context;
+	Context.AddAnalyzer(Analyzer);
+	UE::Trace::FAnalysisProcessor Processor = Context.Process(DataStream);
+	Processor.Wait();
+
+	TSet<FSummarizeAnalyzer::FScope> DeduplicatedScopes;
+
+	struct DisallowedString
+	{
+		const TCHAR* String;
+		bool First;
+	}
+	DisallowedStrings[] =
+	{
+		// breaks simple csv files
+		{ TEXT("\n"), true },
+		{ TEXT("\r"), true },
+		{ TEXT(","), true },
+
+		// likely a path to an asset
+		//  this is broad but unfortunate with current instrumentation
+		//  at the current moment names can start with /Relpath/To/Asset
+		{ TEXT("/"), true }, 
+	};
+
+	bool bDisallowed = false;
+	for (const auto& Scope : Analyzer.Scopes)
+	{
+		if (Scope.Name.IsEmpty())
+		{
+			continue;
+		}
+
+		if (Scope.Count == 0)
+		{
+			continue;
+		}
+
+		bDisallowed = false;
+		for (auto& DisallowedString : DisallowedStrings)
+		{
+			if (Scope.Name.Contains(DisallowedString.String))
+			{
+				if (DisallowedString.First)
+				{
+					UE_LOG(LogSummarizeTrace, Display, TEXT("A scope contains disallowed string '%s'. See log for full list."), DisallowedString.String);
+					DisallowedString.First = false;
+				}
+
+				UE_LOG(LogSummarizeTrace, Verbose, TEXT("Scope '%s' contains disallowed string '%s', skipping..."), *Scope.Name, DisallowedString.String);
+				bDisallowed = true;
+			}
+
+			if (bDisallowed)
+			{
+				break;
+			}
+		}
+
+		if (bDisallowed)
+		{
+			continue;
+		}
+
+		FSummarizeAnalyzer::FScope* FoundScope = DeduplicatedScopes.Find(Scope);
+		if (FoundScope)
+		{
+			FoundScope->Merge(Scope);
+		}
+		else
+		{
+			DeduplicatedScopes.Add(Scope);
+		}
+	}
+
+	UE_LOG(LogSummarizeTrace, Display, TEXT("Sorting %d events by total time accumulated..."), DeduplicatedScopes.Num());
+	TArray<FSummarizeAnalyzer::FScope> SortedScopes;
+	for (const auto& Scope : DeduplicatedScopes)
+	{
+		SortedScopes.Add(Scope);
+	}
+	SortedScopes.Sort();
+
+	FString CsvFileName = FPaths::SetExtension(TraceFileName, "csv");
+	UE_LOG(LogSummarizeTrace, Display, TEXT("Writing summary to %s..."), *CsvFileName);
+	IFileHandle* CsvHandle = FPlatformFileManager::Get().GetPlatformFile().OpenWrite(*CsvFileName);
+	if (CsvHandle)
+	{
+		// no newline, see row printf
+		FilePrint(CsvHandle, FString::Printf(TEXT("Name,Count,TotalSeconds,FirstStartSeconds,FirstEndSeconds,FirstSeconds,MinSeconds,MaxSeconds,MeanSeconds,DeviationSeconds")));
+		for (const auto& Scope : SortedScopes)
+		{
+			// note newline is at the front of every data line to prevent final extraneous newline, per customary for csv
+			FilePrint(CsvHandle, FString::Printf(TEXT("\n%s,%llu,%f,%f,%f,%f,%f,%f,%f,%f,"), *Scope.Name, Scope.Count, Scope.TotalSeconds, Scope.FirstStartSeconds, Scope.FirstEndSeconds, Scope.FirstSeconds, Scope.MinSeconds, Scope.MaxSeconds, Scope.MeanSeconds, Scope.GetDeviationSeconds()));
+		}
+		CsvHandle->Flush();
+		delete CsvHandle;
+		CsvHandle = nullptr;
+	}
+	else
+	{
+		UE_LOG(LogSummarizeTrace, Error, TEXT("Unable to open csv '%s' for write"), *CsvFileName);
+		return 1;
+	}
+
+	if (!NameToStatisticMap.IsEmpty())
+	{
+		FString TracePath = FPaths::GetPath(TraceFileName);
+		FString TraceFileBasename = FPaths::GetBaseFilename(TraceFileName);
+		FString TelemetryCsvFileName = TraceFileBasename + TEXT("Telemetry");
+		TelemetryCsvFileName = FPaths::Combine(TracePath, FPaths::SetExtension(TelemetryCsvFileName, "csv"));
+		UE_LOG(LogSummarizeTrace, Display, TEXT("Writing telemetry to %s..."), *TelemetryCsvFileName);
+		IFileHandle* TelemetryCsvHandle = FPlatformFileManager::Get().GetPlatformFile().OpenWrite(*TelemetryCsvFileName);
+		if (TelemetryCsvHandle)
+		{
+			// no newline, see row printf
+			FilePrint(TelemetryCsvHandle, FString::Printf(TEXT("TestName,Context,DataPoint,Measurement")));
+			for (const auto& Scope : SortedScopes)
+			{
+				TArray<FString> Statistics;
+				NameToStatisticMap.MultiFind(Scope.Name, Statistics, true);
+				for (const auto& Statistic : Statistics)
+				{
+					// note newline is at the front of every data line to prevent final extraneous newline, per customary for csv
+					FilePrint(TelemetryCsvHandle, FString::Printf(TEXT("\n%s,%s,%s,%s,"), *TraceFileBasename, *Scope.Name, *Statistic, *Scope.GetValue(Statistic)));
+				}
+			}
+			TelemetryCsvHandle->Flush();
+			delete TelemetryCsvHandle;
+			TelemetryCsvHandle = nullptr;
+		}
+		else
+		{
+			UE_LOG(LogSummarizeTrace, Error, TEXT("Unable to open telemetry csv '%s' for write"), *TelemetryCsvFileName);
+			return 1;
+		}
+	}
+
+	return 0;
+}
diff --git a/Engine/Source/Editor/UnrealEd/UnrealEd.Build.cs b/Engine/Source/Editor/UnrealEd/UnrealEd.Build.cs
--- a/Engine/Source/Editor/UnrealEd/UnrealEd.Build.cs
+++ b/Engine/Source/Editor/UnrealEd/UnrealEd.Build.cs
@@ -194,6 +194,7 @@ public class UnrealEd : ModuleRules
 				"IoStoreUtilities",
 				"EditorInteractiveToolsFramework",
 				"TraceLog",
+				"TraceAnalysis",
 				"DeveloperSettings",
 				"AnimationModifiers",
 				"MaterialBaking",
diff --git a/Engine/Source/Programs/AutomationTool/Gauntlet/Unreal/Editor/EditorTests.cs b/Engine/Source/Programs/AutomationTool/Gauntlet/Unreal/Editor/EditorTests.cs
--- a/Engine/Source/Programs/AutomationTool/Gauntlet/Unreal/Editor/EditorTests.cs
+++ b/Engine/Source/Programs/AutomationTool/Gauntlet/Unreal/Editor/EditorTests.cs
@@ -117,19 +117,5 @@ namespace UnrealEditor
 			Config.RequireRole(UnrealTargetRole.Editor);	
 			return Config;
 		}
-
-		public override ITestReport CreateReport(TestResult Result)
-		{
-			ITestReport Report = base.CreateReport(Result);
-
-			if (Report != null && CachedConfig.TraceFile != string.Empty)
-			{
-				string TraceFilePath = System.IO.Path.GetFullPath( CachedConfig.TraceFile );
-				Log.Info("Attaching trace {0} to report", TraceFilePath);
-				Report.AttachArtifact(TraceFilePath);
-			}
-
-			return Report;
-		}
 	}
 }
diff --git a/Engine/Source/Runtime/Core/Public/Misc/AutomationTest.h b/Engine/Source/Runtime/Core/Public/Misc/AutomationTest.h
--- a/Engine/Source/Runtime/Core/Public/Misc/AutomationTest.h
+++ b/Engine/Source/Runtime/Core/Public/Misc/AutomationTest.h
@@ -3379,8 +3379,8 @@ public: \
 // Basic Latent Commands
 
 /**
- * Run some code latently with a predicate lambda.  If the predicate returns true, the latent action will be called 
- * again next frame.  If it returns false, the command will stop running.
+ * Run some code latently with a predicate lambda.  If the predicate returns false, the latent action will be called 
+ * again next frame.  If it returns true, the command will stop running.
  */
 class FFunctionLatentCommand : public IAutomationLatentCommand
 {
