Fix packaging of sample projects using external actor files.

diff --git a/Engine/Source/Editor/UnrealEd/Private/Commandlets/GenerateDistillFileSetsCommandlet.cpp b/Engine/Source/Editor/UnrealEd/Private/Commandlets/GenerateDistillFileSetsCommandlet.cpp
--- a/Engine/Source/Editor/UnrealEd/Private/Commandlets/GenerateDistillFileSetsCommandlet.cpp
+++ b/Engine/Source/Editor/UnrealEd/Private/Commandlets/GenerateDistillFileSetsCommandlet.cpp
@@ -14,6 +14,7 @@
 #include "Misc/RedirectCollector.h"
 #include "Editor.h"
 #include "Engine/AssetManager.h"
+#include "Engine/World.h"
 
 DEFINE_LOG_CATEGORY_STATIC(LogGenerateDistillFileSetsCommandlet, Log, All);
 
@@ -272,6 +273,18 @@ int32 UGenerateDistillFileSetsCommandlet::Main( const FString& InParams )
 						UE_LOG(LogGenerateDistillFileSetsCommandlet, Log, TEXT("Package: %s"), *OtherName);
 					}
 				}
+				if (UWorld* World = UWorld::FindWorldInPackage(Package))
+				{
+					TArray<FString> ExternalActorPackages = World->PersistentLevel->GetOnDiskExternalActorPackages();
+					for (const FString& ExternalActorPackage : ExternalActorPackages)
+					{
+						if (!AllPackageNames.Contains(ExternalActorPackage))
+						{
+							AllPackageNames.Add(ExternalActorPackage);
+							UE_LOG(LogGenerateDistillFileSetsCommandlet, Log, TEXT("Package: %s"), *ExternalActorPackage);
+						}
+					}
+				}
 				UE_LOG(LogGenerateDistillFileSetsCommandlet, Display, TEXT( "Collecting garbage..." ) );
 				CollectGarbage(RF_NoFlags);
 			}
