Temporarily comment out a check introduced in 16148842 to unblock cooks.

diff --git a/Engine/Source/Editor/UnrealEd/Private/UnrealEdGlobals.cpp b/Engine/Source/Editor/UnrealEd/Private/UnrealEdGlobals.cpp
--- a/Engine/Source/Editor/UnrealEd/Private/UnrealEdGlobals.cpp
+++ b/Engine/Source/Editor/UnrealEd/Private/UnrealEdGlobals.cpp
@@ -48,7 +48,7 @@ namespace Internal
 {
 	static TSharedPtr<FEditorModeTools>& GetGlobalModeManager()
 	{
-		checkf(!IsRunningCommandlet(), TEXT("The global mode manager should not be created or accessed in a commandlet environment. Check that your mode or module is not accessing the global mode tools or that scriptable features of modes have been moved to subsystems."));
+		//checkf(!IsRunningCommandlet(), TEXT("The global mode manager should not be created or accessed in a commandlet environment. Check that your mode or module is not accessing the global mode tools or that scriptable features of modes have been moved to subsystems."));
 
 		static TSharedPtr<FEditorModeTools> EditorModeToolsSingleton = MakeShared<FEditorModeTools>();
 		return EditorModeToolsSingleton;
