Another packaging fix of sample projects using external actor file (load all actorfiles to gather dependencies).

diff --git a/Engine/Source/Editor/UnrealEd/Private/Commandlets/GenerateDistillFileSetsCommandlet.cpp b/Engine/Source/Editor/UnrealEd/Private/Commandlets/GenerateDistillFileSetsCommandlet.cpp
--- a/Engine/Source/Editor/UnrealEd/Private/Commandlets/GenerateDistillFileSetsCommandlet.cpp
+++ b/Engine/Source/Editor/UnrealEd/Private/Commandlets/GenerateDistillFileSetsCommandlet.cpp
@@ -262,29 +262,54 @@ int32 UGenerateDistillFileSetsCommandlet::Main( const FString& InParams )
 				AllPackageNames.Add(Package->GetName());
 
 				UE_LOG(LogGenerateDistillFileSetsCommandlet, Display, TEXT( "Finding content referenced by %s..." ), *MapPackage );
-				TArray<UObject *> AllPackages;
-				GetObjectsOfClass(UPackage::StaticClass(), AllPackages);
-				for (int32 Index = 0; Index < AllPackages.Num(); Index++)
+
+				auto GatherLoadedPackages = [&]()
 				{
-					FString OtherName = AllPackages[Index]->GetOutermost()->GetName();
-					if (!AllPackageNames.Contains(OtherName))
+					TArray<UObject *> AllPackages;
+					GetObjectsOfClass(UPackage::StaticClass(), AllPackages);
+					for (int32 Index = 0; Index < AllPackages.Num(); Index++)
 					{
-						AllPackageNames.Add(OtherName);
-						UE_LOG(LogGenerateDistillFileSetsCommandlet, Log, TEXT("Package: %s"), *OtherName);
+						FString OtherName = AllPackages[Index]->GetOutermost()->GetName();
+						if (!AllPackageNames.Contains(OtherName))
+						{
+							AllPackageNames.Add(OtherName);
+							UE_LOG(LogGenerateDistillFileSetsCommandlet, Log, TEXT("Package: %s"), *OtherName);
+						}
 					}
-				}
+				};
+
+				// Load all external actor packages to gather their dependencies
 				if (UWorld* World = UWorld::FindWorldInPackage(Package))
 				{
+					World->AddToRoot();
+
+					uint32 ActorPackageIndex = 0;
 					TArray<FString> ExternalActorPackages = World->PersistentLevel->GetOnDiskExternalActorPackages();
+
 					for (const FString& ExternalActorPackage : ExternalActorPackages)
 					{
 						if (!AllPackageNames.Contains(ExternalActorPackage))
 						{
 							AllPackageNames.Add(ExternalActorPackage);
 							UE_LOG(LogGenerateDistillFileSetsCommandlet, Log, TEXT("Package: %s"), *ExternalActorPackage);
+
+							FString LongActorPackageName;
+							FPackageName::TryConvertFilenameToLongPackageName(ExternalActorPackage, LongActorPackageName);
+							UPackage* ActorPackage = LoadPackage(nullptr, *LongActorPackageName, LOAD_None);
+
+							if (!(++ActorPackageIndex % 50))
+							{
+								GatherLoadedPackages();
+								UE_LOG(LogGenerateDistillFileSetsCommandlet, Display, TEXT( "Collecting garbage..." ) );
+								CollectGarbage(RF_NoFlags);
+							}
 						}
 					}
+
+					World->RemoveFromRoot();
 				}
+
+				GatherLoadedPackages();
 				UE_LOG(LogGenerateDistillFileSetsCommandlet, Display, TEXT( "Collecting garbage..." ) );
 				CollectGarbage(RF_NoFlags);
 			}
