DerivedDataCacheCommandlet - Ensure all WP actors are processed
* Moved world setup from commandlet code to UWorldPartitionBuilder::RunBuilder()
* Moved actor iteration logic (reference loading + GC) from UWorldPartitionResaveActorsBuilder::RunInternal() to WorldPartition::ForEachActorWithLoading()
* Adjusted DerivedDataCacheCommandlet to (potentially) iterate on WP actors using UWorldPartitionActorsBuilder::ForEachActorWithLoading()
* Much faster processing of packages with no assets to cache (most ExternalActors packages)

diff --git a/Engine/Source/Editor/SceneOutliner/Private/ActorHierarchy.cpp b/Engine/Source/Editor/SceneOutliner/Private/ActorHierarchy.cpp
--- a/Engine/Source/Editor/SceneOutliner/Private/ActorHierarchy.cpp
+++ b/Engine/Source/Editor/SceneOutliner/Private/ActorHierarchy.cpp
@@ -17,7 +17,7 @@
 #include "EditorFolderUtils.h"
 #include "LevelInstance/LevelInstanceActor.h"
 #include "LevelInstance/LevelInstanceSubsystem.h"
-#include "WorldPartition/WorldPartitionSubsystem.h"
+#include "WorldPartition/WorldPartitionHelpers.h"
 #include "WorldPartition/WorldPartitionActorDesc.h"
 
 TUniquePtr<FActorHierarchy> FActorHierarchy::Create(ISceneOutlinerMode* Mode, const TWeakObjectPtr<UWorld>& World)
@@ -286,11 +286,9 @@ void FActorHierarchy::CreateWorldChildren(UWorld* World, TArray<FSceneOutlinerTr
 
 	if (bShowingUnloadedActors)
 	{
-		if (UWorldPartitionSubsystem* WorldPartitionSubsystem = World->GetSubsystem<UWorldPartitionSubsystem>())
+		if (UWorldPartition* WorldPartition = World->GetWorldPartition())
 		{
-			UWorldPartition* WorldPartition = World->GetWorldPartition();
-			check(WorldPartition);
-			WorldPartitionSubsystem->ForEachActorDesc(AActor::StaticClass(), [this, WorldPartition, &OutItems](const FWorldPartitionActorDesc* ActorDesc)
+			FWorldPartitionHelpers::ForEachActorDesc(WorldPartition, [this, WorldPartition, &OutItems](const FWorldPartitionActorDesc* ActorDesc)
                 {
                     if (ActorDesc != nullptr && !ActorDesc->IsLoaded())
                     {
diff --git a/Engine/Source/Editor/UnrealEd/Classes/Commandlets/DerivedDataCacheCommandlet.h b/Engine/Source/Editor/UnrealEd/Classes/Commandlets/DerivedDataCacheCommandlet.h
--- a/Engine/Source/Editor/UnrealEd/Classes/Commandlets/DerivedDataCacheCommandlet.h
+++ b/Engine/Source/Editor/UnrealEd/Classes/Commandlets/DerivedDataCacheCommandlet.h
@@ -21,8 +21,14 @@ class UDerivedDataCacheCommandlet : public UCommandlet
 
 	// We hook this up to a delegate to avoid reloading textures and whatnot
 	TSet<FName> ProcessedPackages;
+	TSet<FName> PackagesToProcess;
+	double FinishCacheTime = 0.0;
+	double BeginCacheTime = 0.0;
 
 	void MaybeMarkPackageAsAlreadyLoaded(UPackage *Package);
+
+	void CacheLoadedPackages(UPackage* CurrentPackage, uint8 PackageFilter, const TArray<ITargetPlatform*>& Platforms);
+	void CacheWorldPackages(UWorld* World, uint8 PackageFilter, const TArray<ITargetPlatform*>& Platforms);
 };
 
 
diff --git a/Engine/Source/Editor/UnrealEd/Private/Commandlets/ContentCommandlets.cpp b/Engine/Source/Editor/UnrealEd/Private/Commandlets/ContentCommandlets.cpp
--- a/Engine/Source/Editor/UnrealEd/Private/Commandlets/ContentCommandlets.cpp
+++ b/Engine/Source/Editor/UnrealEd/Private/Commandlets/ContentCommandlets.cpp
@@ -44,8 +44,8 @@
 #include "Editor.h"
 #include "FileHelpers.h"
 #include "CommandletSourceControlUtils.h"
-#include "WorldPartition/WorldPartitionSubsystem.h"
 #include "WorldPartition/WorldPartition.h"
+#include "WorldPartition/WorldPartitionHelpers.h"
 #include "LevelInstance/LevelInstanceActor.h"
 
 #include "PackageHelperFunctions.h"
@@ -1687,8 +1687,7 @@ void UResavePackagesCommandlet::PerformAdditionalOperations(class UWorld* World,
 	// Load and Save world partition actor packages
 	if (bResaveWorldPartitionExternalActors)
 	{
-		UWorldPartitionSubsystem* WorldPartitionSubsystem = World->GetSubsystem<UWorldPartitionSubsystem>();
-		WorldPartitionSubsystem->ForEachActorDesc(AActor::StaticClass(), [this, WorldPartition](const FWorldPartitionActorDesc* ActorDesc)
+		FWorldPartitionHelpers::ForEachActorDesc(WorldPartition, [this, WorldPartition](const FWorldPartitionActorDesc* ActorDesc)
 		{
 			++TotalPackagesForResave;
 			// Load & Register World Partition Actor
diff --git a/Engine/Source/Editor/UnrealEd/Private/Commandlets/DerivedDataCacheCommandlet.cpp b/Engine/Source/Editor/UnrealEd/Private/Commandlets/DerivedDataCacheCommandlet.cpp
--- a/Engine/Source/Editor/UnrealEd/Private/Commandlets/DerivedDataCacheCommandlet.cpp
+++ b/Engine/Source/Editor/UnrealEd/Private/Commandlets/DerivedDataCacheCommandlet.cpp
@@ -25,6 +25,10 @@ DerivedDataCacheCommandlet.cpp: Commandlet for DDC maintenence
 #include "Settings/ProjectPackagingSettings.h"
 #include "Editor.h"
 #include "AssetCompilingManager.h"
+#include "WorldPartition/WorldPartition.h"
+#include "WorldPartition/WorldPartitionHelpers.h"
+#include "WorldPartition/WorldPartitionSubsystem.h"
+#include "LevelInstance/LevelInstanceSubsystem.h"
 
 DEFINE_LOG_CATEGORY_STATIC(LogDerivedDataCacheCommandlet, Log, All);
 
@@ -100,6 +104,185 @@ static void PumpAsync(bool* bInOutHadActivity = nullptr)
 	}
 }
 
+void UDerivedDataCacheCommandlet::CacheLoadedPackages(UPackage* CurrentPackage, uint8 PackageFilter, const TArray<ITargetPlatform*>& Platforms)
+{
+	// Timing variables
+	double DDCCommandletMaxWaitSeconds = 60. * 10.;
+	GConfig->GetDouble(TEXT("CookSettings"), TEXT("DDCCommandletMaxWaitSeconds"), DDCCommandletMaxWaitSeconds, GEditorIni);
+
+	TArray<UObject*> CachingObjects;
+	TArray<UObject*> ObjectBuffer;
+	TArray<UPackage*> NewPackages;
+
+	bool bIsCaching = false;
+
+	{
+		const double BeginCacheTimeStart = FPlatformTime::Seconds();
+		CachingObjects.Reset();
+		NewPackages.Reset();
+
+		for (TObjectIterator<UPackage> PackageIter; PackageIter; ++PackageIter)
+		{
+			UPackage* ExistingPackage = *PackageIter;
+			if ((PackageFilter & NORMALIZE_ExcludeEnginePackages) == 0 || !ExistingPackage->GetName().StartsWith(TEXT("/Engine")))
+			{
+				FName ExistingPackageName = ExistingPackage->GetFName();
+				if (ExistingPackage == CurrentPackage || !PackagesToProcess.Contains(ExistingPackageName))
+				{
+					if (!ProcessedPackages.Contains(ExistingPackageName))
+					{
+						UE_LOG(LogDerivedDataCacheCommandlet, Display, TEXT("Processing %s"), *ExistingPackageName.ToString());
+
+						ProcessedPackages.Add(ExistingPackageName);
+						NewPackages.Add(ExistingPackage);
+						check((ExistingPackage->GetPackageFlags() & PKG_ReloadingForCooker) == 0);
+
+						ObjectBuffer.Reset();
+						GetObjectsWithOuter(ExistingPackage, ObjectBuffer, true /* bIncludeNestedObjects */, RF_ClassDefaultObject /* ExclusionFlags */);
+						for (UObject* Object : ObjectBuffer)
+						{
+							for (auto Platform : Platforms)
+							{
+								Object->BeginCacheForCookedPlatformData(Platform);
+								bIsCaching |= !Object->IsCachedCookedPlatformDataLoaded(Platform);
+							}
+							CachingObjects.Add(Object);
+						}
+					}
+				}
+			}
+		}
+
+		BeginCacheTime += FPlatformTime::Seconds() - BeginCacheTimeStart;
+	}
+
+	{
+		const double FinishCacheTimeStart = FPlatformTime::Seconds();
+
+		if (bIsCaching)
+		{
+			PumpAsync();
+
+			ObjectBuffer.Reset();
+			ObjectBuffer.Append(CachingObjects);
+
+			double LastActivityTime = FinishCacheTimeStart;
+			const double WaitingForCacheSleepTime = 0.050;
+			while (ObjectBuffer.Num() > 0)
+			{
+				bool bHadActivity = false;
+				for (int32 ObjectIndex = 0; ObjectIndex < ObjectBuffer.Num();)
+				{
+					UObject* Object = ObjectBuffer[ObjectIndex];
+					bool bIsFinished = true;
+					for (auto Platform : Platforms)
+					{
+						bIsFinished = Object->IsCachedCookedPlatformDataLoaded(Platform) && bIsFinished;
+					}
+					if (bIsFinished)
+					{
+						ObjectBuffer.RemoveAtSwap(ObjectIndex);
+						bHadActivity = true;
+					}
+					else
+					{
+						++ObjectIndex;
+					}
+				}
+
+				double CurrentTime = FPlatformTime::Seconds();
+				if (!bHadActivity)
+				{
+					PumpAsync(&bHadActivity);
+				}
+				if (!bHadActivity)
+				{
+					if (CurrentTime - LastActivityTime >= DDCCommandletMaxWaitSeconds)
+					{
+						UE_LOG(LogDerivedDataCacheCommandlet, Error, TEXT("Timed out for %.2lfs waiting for %d objects to finish caching. First object: %s."),
+							DDCCommandletMaxWaitSeconds, ObjectBuffer.Num(), *ObjectBuffer[0]->GetFullName());
+						ObjectBuffer.Reset();
+					}
+					else
+					{
+						FPlatformProcess::Sleep(WaitingForCacheSleepTime);
+					}
+				}
+				else
+				{
+					LastActivityTime = CurrentTime;
+				}
+			}
+
+			PumpAsync();
+		}
+
+		// Tear down all of the Cached data; we do this only after all objects have finished because we need to not
+		// tear down any object until all objects in its package have finished
+		for (UObject* Object : CachingObjects)
+		{
+			Object->WillNeverCacheCookedPlatformDataAgain();
+			Object->ClearAllCachedCookedPlatformData();
+		}
+
+		// Mark the packages as processed
+		for (UPackage* NewPackage : NewPackages)
+		{
+			NewPackage->SetPackageFlags(PKG_ReloadingForCooker);
+		}
+
+		FinishCacheTime += FPlatformTime::Seconds() - FinishCacheTimeStart;
+	}
+}
+
+void UDerivedDataCacheCommandlet::CacheWorldPackages(UWorld* World, uint8 PackageFilter, const TArray<ITargetPlatform*>& Platforms)
+{
+	check(World);
+
+	World->AddToRoot();
+	
+	// Setup the world
+	World->WorldType = EWorldType::Editor;
+	UWorld::InitializationValues IVS;
+	IVS.RequiresHitProxies(false);
+	IVS.ShouldSimulatePhysics(false);
+	IVS.EnableTraceCollision(false);
+	IVS.CreateNavigation(false);
+	IVS.CreateAISystem(false);
+	IVS.AllowAudioPlayback(false);
+	IVS.CreatePhysicsScene(true);
+
+	World->InitWorld(UWorld::InitializationValues(IVS));
+	World->PersistentLevel->UpdateModelComponents();
+	World->UpdateWorldComponents(true /*bRerunConstructionScripts*/, false /*bCurrentLevelOnly*/);
+
+	// If the world is partitioned
+	bool bResult = true;
+	if (World->HasSubsystem<UWorldPartitionSubsystem>())
+	{
+		// Commandlets aren't loading level instances by default, change that behavior
+		if (ULevelInstanceSubsystem* LevelInstanceSubsystem = World->GetSubsystem<ULevelInstanceSubsystem>())
+		{
+			LevelInstanceSubsystem->SetLoadInstancesOnRegistration(true);
+		}
+
+		// Ensure the world has a valid world partition.
+		UWorldPartition* WorldPartition = World->GetWorldPartition();
+		check(WorldPartition);
+
+		FWorldPartitionHelpers::ForEachActorWithLoading(WorldPartition, [this, PackageFilter, &Platforms](AActor* Actor)
+		{
+			UE_LOG(LogDerivedDataCacheCommandlet, Display, TEXT("Loaded actor %s"), *Actor->GetName());
+			CacheLoadedPackages(Actor->GetPackage(), PackageFilter, Platforms);
+			return true;
+		});
+	}
+
+	World->ClearWorldComponents();
+	World->CleanupWorld();
+	World->RemoveFromRoot();
+}
+
 int32 UDerivedDataCacheCommandlet::Main( const FString& Params )
 {
 	TArray<FString> Tokens, Switches;
@@ -115,14 +298,10 @@ int32 UDerivedDataCacheCommandlet::Main( const FString& Params )
 	FParse::Value(*Params, TEXT("SubsetTarget="), SubsetTarget);
 	bool bDoSubset = SubsetMod > 0 && SubsetTarget < SubsetMod;
 
-	// Timing variables
-	double DDCCommandletMaxWaitSeconds = 60. * 10.;
-	GConfig->GetDouble(TEXT("CookSettings"), TEXT("DDCCommandletMaxWaitSeconds"), DDCCommandletMaxWaitSeconds, GEditorIni);
-
 	double FindProcessedPackagesTime = 0.0;
 	double GCTime = 0.0;
-	double FinishCacheTime = 0.;
-	double BeginCacheTime = 0.;
+	FinishCacheTime = 0.;
+	BeginCacheTime = 0.;
 
 	if (!bStartupOnly && bFillCache)
 	{
@@ -202,7 +381,7 @@ int32 UDerivedDataCacheCommandlet::Main( const FString& Params )
 		}
 
 		// assume the first token is the map wildcard/pathname
-		TArray<FString> FilesInPath;
+		TSet<FString> FilesInPath;
 		TArray<FString> Unused;
 		TArray<FString> TokenFiles;
 		for ( int32 TokenIndex = 0; TokenIndex < Tokens.Num(); TokenIndex++ )
@@ -214,7 +393,7 @@ int32 UDerivedDataCacheCommandlet::Main( const FString& Params )
 				continue;
 			}
 
-			FilesInPath += TokenFiles;
+			FilesInPath.Append(TokenFiles);
 		}
 
 		TArray<TPair<FString, FName>> PackagePaths;
@@ -294,149 +473,61 @@ int32 UDerivedDataCacheCommandlet::Main( const FString& Params )
 			UE_LOG(LogDerivedDataCacheCommandlet, Display, TEXT("%d packages to load..."), PackagePaths.Num());
 		}
 
-		TArray<UObject*> CachingObjects;
-		TArray<UObject*> ObjectBuffer;
-		TArray<UPackage*> NewPackages;
+		// Gather the list of packages to process
+		PackagesToProcess.Empty(PackagePaths.Num());
+		for (int32 PackageIndex = PackagePaths.Num() - 1; PackageIndex >= 0; PackageIndex--)
+		{
+			PackagesToProcess.Add(PackagePaths[PackageIndex].Get<1>());
+		}
+
+		// Process each package
 		for (int32 PackageIndex = PackagePaths.Num() - 1; PackageIndex >= 0; PackageIndex-- )
 		{
+			TPair<FString, FName>& PackagePath = PackagePaths[PackageIndex];
+			const FString& Filename = PackagePath.Get<0>();
+			FName PackageFName = PackagePath.Get<1>();
+			check(!ProcessedPackages.Contains(PackageFName));
+
+			// If work is distributed, skip packages that are meant to be process by other machines
+			if (bDoSubset)
 			{
-				TPair<FString, FName>& PackagePath = PackagePaths[PackageIndex];
-				const FString& Filename = PackagePath.Get<0>();
-				FName PackageFName = PackagePath.Get<1>();
-				if (ProcessedPackages.Contains(PackageFName))
+				FString PackageName = PackageFName.ToString();
+				if (FCrc::StrCrc_DEPRECATED(*PackageName.ToUpper()) % SubsetMod != SubsetTarget)
 				{
 					continue;
 				}
-				if (bDoSubset)
-				{
-					FString PackageName = PackageFName.ToString();
-					if (FCrc::StrCrc_DEPRECATED(*PackageName.ToUpper()) % SubsetMod != SubsetTarget)
-					{
-						continue;
-					}
-				}
+			}
 
-				UE_LOG(LogDerivedDataCacheCommandlet, Display, TEXT("Loading (%d) %s"), FilesInPath.Num() - PackageIndex, *Filename);
+			UE_LOG(LogDerivedDataCacheCommandlet, Display, TEXT("Loading (%d) %s"), FilesInPath.Num() - PackageIndex, *Filename);
 
-				UPackage* Package = LoadPackage(NULL, *Filename, LOAD_None);
-				if (Package == NULL)
-				{
-					UE_LOG(LogDerivedDataCacheCommandlet, Error, TEXT("Error loading %s!"), *Filename);
-				}
-				else
-				{
-					bLastPackageWasMap = Package->ContainsMap();
-					NumProcessedSinceLastGC++;
-				}
+			UPackage* Package = LoadPackage(NULL, *Filename, LOAD_None);
+			if (Package == NULL)
+			{
+				UE_LOG(LogDerivedDataCacheCommandlet, Error, TEXT("Error loading %s!"), *Filename);
+				bLastPackageWasMap = false;
+			}
+			else
+			{
+				bLastPackageWasMap = Package->ContainsMap();
+				NumProcessedSinceLastGC++;
 			}
 
 			// even if the load failed this could be the first time through the loop so it might have all the startup packages to resolve
 			GRedirectCollector.ResolveAllSoftObjectPaths();
 
 			// Find any new packages and cache all the objects in each package
-			{
-				const double BeginCacheTimeStart = FPlatformTime::Seconds();
-				CachingObjects.Reset();
-				NewPackages.Reset();
-				for (TObjectIterator<UPackage> PackageIter; PackageIter; ++PackageIter)
-				{
-					UPackage* ExistingPackage = *PackageIter;
-					if ((PackageFilter & NORMALIZE_ExcludeEnginePackages) == 0 || !ExistingPackage->GetName().StartsWith(TEXT("/Engine")))
-					{
-						FName ExistingPackageName = ExistingPackage->GetFName();
-						if (!ProcessedPackages.Contains(ExistingPackageName))
-						{
-							ProcessedPackages.Add(ExistingPackageName);
-							NewPackages.Add(ExistingPackage);
-							check((ExistingPackage->GetPackageFlags() & PKG_ReloadingForCooker) == 0);
-
-							ObjectBuffer.Reset();
-							GetObjectsWithOuter(ExistingPackage, ObjectBuffer, true /* bIncludeNestedObjects */, RF_ClassDefaultObject /* ExclusionFlags */);
-							for (UObject* Object : ObjectBuffer)
-							{
-								for (auto Platform : Platforms)
-								{
-									Object->BeginCacheForCookedPlatformData(Platform);
-								}
-								CachingObjects.Add(Object);
-							}
-						}
-					}
-				}
-				PumpAsync();
-				BeginCacheTime += FPlatformTime::Seconds() - BeginCacheTimeStart;
-			}
+			CacheLoadedPackages(Package, PackageFilter, Platforms);
 
+			// Ensure we load maps to process all their referenced packages in case they are using world partition.
+			if (bLastPackageWasMap)
 			{
-				const double FinishCacheTimeStart = FPlatformTime::Seconds();
-				ObjectBuffer.Reset();
-				ObjectBuffer.Append(CachingObjects);
-
-				double LastActivityTime = FinishCacheTimeStart;
-				const double WaitingForCacheSleepTime = 0.050;
-				while (ObjectBuffer.Num() > 0)
-				{
-					bool bHadActivity = false;
-					for (int32 ObjectIndex = 0; ObjectIndex < ObjectBuffer.Num();)
-					{
-						UObject* Object = ObjectBuffer[ObjectIndex];
-						bool bIsFinished = true;
-						for (auto Platform : Platforms)
-						{
-							bIsFinished = Object->IsCachedCookedPlatformDataLoaded(Platform) && bIsFinished;
-						}
-						if (bIsFinished)
-						{
-							ObjectBuffer.RemoveAtSwap(ObjectIndex);
-							bHadActivity = true;
-						}
-						else
-						{
-							++ObjectIndex;
-						}
-					}
-
-					double CurrentTime = FPlatformTime::Seconds();
-					if (!bHadActivity)
-					{
-						PumpAsync(&bHadActivity);
-					}
-					if (!bHadActivity)
-					{
-						if (CurrentTime - LastActivityTime >= DDCCommandletMaxWaitSeconds)
-						{
-							UE_LOG(LogDerivedDataCacheCommandlet, Error, TEXT("Timed out for %.2lfs waiting for %d objects to finish caching. First object: %s."),
-								DDCCommandletMaxWaitSeconds, ObjectBuffer.Num(), *ObjectBuffer[0]->GetFullName());
-							ObjectBuffer.Reset();
-						}
-						else
-						{
-							FPlatformProcess::Sleep(WaitingForCacheSleepTime);
-						}
-					}
-					else
-					{
-						LastActivityTime = CurrentTime;
-					}
-				}
-
-				// Tear down all of the Cached data; we do this only after all objects have finished because we need to not
-				// tear down any object until all objects in its package have finished
-				for (UObject* Object : CachingObjects)
-				{
-					Object->WillNeverCacheCookedPlatformDataAgain();
-					Object->ClearAllCachedCookedPlatformData();
-				}
-				// Mark the packages as processed
-				for (UPackage* NewPackage : NewPackages)
+				if (UWorld* World = UWorld::FindWorldInPackage(Package))
 				{
-					NewPackage->SetPackageFlags(PKG_ReloadingForCooker);
+					CacheWorldPackages(World, PackageFilter, Platforms);
 				}
-
-				PumpAsync();
-				FinishCacheTime += FPlatformTime::Seconds() - FinishCacheTimeStart;
 			}
 
+			// Perform a GC if conditions are met
 			if (NumProcessedSinceLastGC >= GCInterval || PackageIndex < 0 || bLastPackageWasMap)
 			{
 				const double StartGCTime = FPlatformTime::Seconds();
diff --git a/Engine/Source/Editor/UnrealEd/Private/Commandlets/WorldPartitionBuilderCommandlet.cpp b/Engine/Source/Editor/UnrealEd/Private/Commandlets/WorldPartitionBuilderCommandlet.cpp
--- a/Engine/Source/Editor/UnrealEd/Private/Commandlets/WorldPartitionBuilderCommandlet.cpp
+++ b/Engine/Source/Editor/UnrealEd/Private/Commandlets/WorldPartitionBuilderCommandlet.cpp
@@ -5,15 +5,12 @@
 
 #include "CoreMinimal.h"
 #include "EngineUtils.h"
-#include "Editor.h"
 #include "Logging/LogMacros.h"
 #include "Misc/CommandLine.h"
-#include "Misc/ConfigCacheIni.h"
 #include "HAL/PlatformFileManager.h"
 #include "ProfilingDebugging/ScopedTimers.h"
 #include "WorldPartition/WorldPartition.h"
-#include "WorldPartition/WorldPartitionSubsystem.h"
-#include "LevelInstance/LevelInstanceSubsystem.h"
+
 #include "Trace/Trace.h"
 
 DEFINE_LOG_CATEGORY_STATIC(LogWorldPartitionBuilderCommandlet, All, All);
@@ -22,52 +19,6 @@ UWorldPartitionBuilderCommandlet::UWorldPartitionBuilderCommandlet(const FObject
 	: Super(ObjectInitializer)
 {}
 
-UWorldPartitionBuilder* UWorldPartitionBuilderCommandlet::CreateBuilder(const FString& WorldConfigFilename)
-{
-	// Parse builder class name
-	FString BuilderClassName;
-	if (!FParse::Value(FCommandLine::Get(), TEXT("Builder="), BuilderClassName, false))
-	{
-		UE_LOG(LogWorldPartitionBuilderCommandlet, Error, TEXT("Invalid builder name."));
-		return nullptr;
-	}
-
-	UClass* BuilderClass = FindObject<UClass>(ANY_PACKAGE, *BuilderClassName);
-
-	if (!BuilderClass)
-	{
-		UE_LOG(LogWorldPartitionBuilderCommandlet, Error, TEXT("Unknown builder %s."), *BuilderClassName);
-		return nullptr;
-	}
-
-	// Create builder instance
-	UWorldPartitionBuilder* Builder = NewObject<UWorldPartitionBuilder>(this, BuilderClass);
-	check(Builder);
-
-	Builder->AddToRoot();
-
-	// Validate builder settings
-	if (Builder->RequiresCommandletRendering() && !IsAllowCommandletRendering())
-	{
-		UE_LOG(LogWorldPartitionBuilderCommandlet, Error, TEXT("The option \"-AllowCommandletRendering\" must be provided for the %s process to work"), *BuilderClassName);
-		return nullptr;
-	}
-
-	// Load configuration file
-	if (FPlatformFileManager::Get().GetPlatformFile().FileExists(*WorldConfigFilename))
-	{
-		LoadConfig(GetClass(), *WorldConfigFilename);
-	}
-
-	// Load builder configuration
-	if (FPlatformFileManager::Get().GetPlatformFile().FileExists(*WorldConfigFilename))
-	{
-		Builder->LoadConfig(BuilderClass, *WorldConfigFilename);
-	}
-
-	return Builder;
-}
-
 int32 UWorldPartitionBuilderCommandlet::Main(const FString& Params)
 {
 	TRACE_CPUPROFILER_EVENT_SCOPE(UWorldPartitionBuilderCommandlet::Main);
@@ -96,18 +47,21 @@ int32 UWorldPartitionBuilderCommandlet::Main(const FString& Params)
 		return 1;
 	}
 
-	FString WorldConfigFilename = FPaths::ChangeExtension(WorldFilename, TEXT("ini"));
-
-	// Create builder instance
-	UWorldPartitionBuilder* Builder = CreateBuilder(WorldConfigFilename);
-	if (!Builder)
+	// Parse builder class name
+	FString BuilderClassName;
+	if (!FParse::Value(FCommandLine::Get(), TEXT("Builder="), BuilderClassName, false))
 	{
-		UE_LOG(LogWorldPartitionBuilderCommandlet, Error, TEXT("Failed to create builder, exiting."));
+		UE_LOG(LogWorldPartitionBuilderCommandlet, Error, TEXT("Invalid builder name."));
 		return 1;
 	}
 
-	// Perform builder pre world initialisation
-	Builder->PreWorldInitialization(*this);
+	// Find builder class
+	UClass* BuilderClass = FindObject<UClass>(ANY_PACKAGE, *BuilderClassName);
+	if (!BuilderClass)
+	{
+		UE_LOG(LogWorldPartitionBuilderCommandlet, Error, TEXT("Unknown builder %s."), *BuilderClassName);
+		return 1;
+	}
 
 	// Load the map package
 	UPackage* MapPackage = LoadPackage(NULL, *Tokens[0], LOAD_None);
@@ -125,70 +79,22 @@ int32 UWorldPartitionBuilderCommandlet::Main(const FString& Params)
 		return 1;
 	}
 
-	// Setup the world.
-	World->WorldType = EWorldType::Editor;
-	World->AddToRoot();
-	if (!World->bIsWorldInitialized)
-	{
-		UWorld::InitializationValues IVS;
-		IVS.RequiresHitProxies(false);
-		IVS.ShouldSimulatePhysics(false);
-		IVS.EnableTraceCollision(false);
-		IVS.CreateNavigation(false);
-		IVS.CreateAISystem(false);
-		IVS.AllowAudioPlayback(false);
-		IVS.CreatePhysicsScene(true);
-
-		World->InitWorld(UWorld::InitializationValues(IVS));
-		World->PersistentLevel->UpdateModelComponents();
-		World->UpdateWorldComponents(true /*bRerunConstructionScripts*/, false /*bCurrentLevelOnly*/);
-	}
-
-	// Make sure the world is partitioned.
-	if (!World->HasSubsystem<UWorldPartitionSubsystem>())
-	{
-		UE_LOG(LogWorldPartitionBuilderCommandlet, Error, TEXT("Commandlet only works on partitioned maps."));
-		return 1;
-	}
-
-	// Commandlets aren't loading level instances by default, change that behavior
-	if (ULevelInstanceSubsystem* LevelInstanceSubsystem = World->GetSubsystem<ULevelInstanceSubsystem>())
+	// Load configuration file
+	FString WorldConfigFilename = FPackageName::LongPackageNameToFilename(World->GetPackage()->GetName(), TEXT(".ini"));
+	if (FPlatformFileManager::Get().GetPlatformFile().FileExists(*WorldConfigFilename))
 	{
-		LevelInstanceSubsystem->SetLoadInstancesOnRegistration(true);
+		LoadConfig(GetClass(), *WorldConfigFilename);
 	}
 
-	// Retrieve the world partition.
-	UWorldPartition* WorldPartition = World->GetWorldPartition();
-	check(WorldPartition);
-
-	FWorldContext& WorldContext = GEditor->GetEditorWorldContext(true /*bEnsureIsGWorld*/);
-	WorldContext.SetCurrentWorld(World);
-	UWorld* PrevGWorld = GWorld;
-	GWorld = World;
-
 	// Run builder
-	bool bResult = Builder->Run(World, *this);
-	
-	// Save default configuration
-	if (bResult)
+	bool bResult = UWorldPartitionBuilder::RunBuilder(BuilderClass, World);
+
+	// Save configuration file
+	if (!FPlatformFileManager::Get().GetPlatformFile().FileExists(*WorldConfigFilename) ||
+		!FPlatformFileManager::Get().GetPlatformFile().IsReadOnly(*WorldConfigFilename))
 	{
-	    if (!FPlatformFileManager::Get().GetPlatformFile().FileExists(*WorldConfigFilename) ||
-		    !FPlatformFileManager::Get().GetPlatformFile().IsReadOnly(*WorldConfigFilename))
-	    {
-		    SaveConfig(CPF_Config, *WorldConfigFilename);
-    
-		    Builder->SaveConfig(CPF_Config, *WorldConfigFilename);
-	    }
+		SaveConfig(CPF_Config, *WorldConfigFilename);
 	}
 
-	// Cleanup
-	World->RemoveFromRoot();
-	World->ClearWorldComponents();
-	World->CleanupWorld();
-
-	// Restore previous world
-	WorldContext.SetCurrentWorld(PrevGWorld);
-	GWorld = PrevGWorld;
-
 	return bResult ? 0 : 1;
-}
\ No newline at end of file
+}
diff --git a/Engine/Source/Editor/UnrealEd/Private/WorldPartition/WorldPartitionBuilder.cpp b/Engine/Source/Editor/UnrealEd/Private/WorldPartition/WorldPartitionBuilder.cpp
--- a/Engine/Source/Editor/UnrealEd/Private/WorldPartition/WorldPartitionBuilder.cpp
+++ b/Engine/Source/Editor/UnrealEd/Private/WorldPartition/WorldPartitionBuilder.cpp
@@ -3,14 +3,19 @@
 #include "WorldPartition/WorldPartitionBuilder.h"
 
 #include "CoreMinimal.h"
-
+#include "Editor.h"
+#include "Misc/ConfigCacheIni.h"
+#include "HAL/PlatformFileManager.h"
 #include "DistanceFieldAtlas.h"
 #include "MeshCardRepresentation.h"
 #include "StaticMeshCompiler.h"
 #include "Engine/World.h"
 #include "WorldPartition/WorldPartition.h"
+#include "WorldPartition/WorldPartitionHelpers.h"
+#include "WorldPartition/WorldPartitionSubsystem.h"
 #include "WorldPartition/DataLayer/WorldDataLayers.h"
 #include "WorldPartition/DataLayer/DataLayer.h"
+#include "LevelInstance/LevelInstanceSubsystem.h"
 #include "Math/IntVector.h"
 
 DEFINE_LOG_CATEGORY_STATIC(LogWorldPartitionBuilder, Log, All);
@@ -20,6 +25,138 @@ UWorldPartitionBuilder::UWorldPartitionBuilder(const FObjectInitializer& ObjectI
 	: Super(ObjectInitializer)
 {}
 
+bool UWorldPartitionBuilder::RunBuilder(TSubclassOf<UWorldPartitionBuilder> BuilderClass, UWorld* World)
+{
+	// Create builder instance
+	UWorldPartitionBuilder* Builder = NewObject<UWorldPartitionBuilder>(GetTransientPackage(), BuilderClass);
+	if (!Builder)
+	{
+		UE_LOG(LogWorldPartitionBuilder, Error, TEXT("Failed to create builder."));
+		return false;
+	}
+
+	Builder->AddToRoot();
+	ON_SCOPE_EXIT
+	{
+		Builder->RemoveFromRoot();
+	};
+
+	return RunBuilder(Builder, World);
+}
+
+bool UWorldPartitionBuilder::RunBuilder(UWorldPartitionBuilder* Builder, UWorld* World)
+{
+	// Load configuration file & builder configuration
+	FString WorldConfigFilename = FPackageName::LongPackageNameToFilename(World->GetPackage()->GetName(), TEXT(".ini"));
+	if (FPlatformFileManager::Get().GetPlatformFile().FileExists(*WorldConfigFilename))
+	{
+		Builder->LoadConfig(Builder->GetClass(), *WorldConfigFilename);
+	}
+
+	// Validate builder settings
+	if (IsRunningCommandlet() && Builder->RequiresCommandletRendering() && !IsAllowCommandletRendering())
+	{
+		UE_LOG(LogWorldPartitionBuilder, Error, TEXT("The option \"-AllowCommandletRendering\" must be provided for the %s process to work"), *Builder->GetClass()->GetName());
+		return false;
+	}
+
+	FPackageSourceControlHelper SCCHelper;
+
+	// Perform builder pre world initialisation
+	Builder->PreWorldInitialization(SCCHelper);
+
+	bool bWorldWasRooted = World->IsRooted();
+	if (!bWorldWasRooted)
+	{
+		World->AddToRoot();
+	}
+
+	// Setup the world if needed
+	bool bWorldWasInitialized = World->bIsWorldInitialized;
+	if (!bWorldWasInitialized)
+	{
+		World->WorldType = EWorldType::Editor;
+
+		UWorld::InitializationValues IVS;
+		IVS.RequiresHitProxies(false);
+		IVS.ShouldSimulatePhysics(false);
+		IVS.EnableTraceCollision(false);
+		IVS.CreateNavigation(false);
+		IVS.CreateAISystem(false);
+		IVS.AllowAudioPlayback(false);
+		IVS.CreatePhysicsScene(true);
+
+		World->InitWorld(UWorld::InitializationValues(IVS));
+		World->PersistentLevel->UpdateModelComponents();
+		World->UpdateWorldComponents(true /*bRerunConstructionScripts*/, false /*bCurrentLevelOnly*/);
+	}
+	else
+	{
+		check(World->WorldType == EWorldType::Editor);
+	}
+
+	// Cleanup
+	ON_SCOPE_EXIT
+	{
+		// Restore world to previous state
+		if (!bWorldWasInitialized)
+		{
+			World->ClearWorldComponents();
+			World->CleanupWorld();
+		}
+
+		// Unroot world if required
+		if (!bWorldWasRooted)
+		{
+			World->RemoveFromRoot();
+		}
+	};
+
+	// Make sure the world is partitioned
+	bool bResult = true;
+	if (World->HasSubsystem<UWorldPartitionSubsystem>())
+	{
+		// Commandlets aren't loading level instances by default, change that behavior
+		if (ULevelInstanceSubsystem* LevelInstanceSubsystem = World->GetSubsystem<ULevelInstanceSubsystem>())
+		{
+			LevelInstanceSubsystem->SetLoadInstancesOnRegistration(true);
+		}
+
+		// Ensure the world has a valid world partition.
+		UWorldPartition* WorldPartition = World->GetWorldPartition();
+		check(WorldPartition);
+
+		FWorldContext& WorldContext = GEditor->GetEditorWorldContext(true /*bEnsureIsGWorld*/);
+		WorldContext.SetCurrentWorld(World);
+		UWorld* PrevGWorld = GWorld;
+		GWorld = World;
+
+		// Run builder
+		bResult = Builder->Run(World, SCCHelper);
+
+		// Restore previous world
+		WorldContext.SetCurrentWorld(PrevGWorld);
+		GWorld = PrevGWorld;
+
+		// Save default configuration
+		if (bResult)
+		{
+			if (!FPlatformFileManager::Get().GetPlatformFile().FileExists(*WorldConfigFilename) ||
+				!FPlatformFileManager::Get().GetPlatformFile().IsReadOnly(*WorldConfigFilename))
+			{
+				Builder->SaveConfig(CPF_Config, *WorldConfigFilename);
+			}
+		}
+	}
+	else
+	{
+		UE_LOG(LogWorldPartitionBuilder, Error, TEXT("WorldPartition builders only works on partitioned maps."));
+		bResult = false;
+	}
+
+	return bResult;
+}
+
 bool UWorldPartitionBuilder::Run(UWorld* World, FPackageSourceControlHelper& PackageHelper)
 {
 	UWorldPartition* WorldPartition = World->GetWorldPartition();
@@ -104,7 +241,7 @@ bool UWorldPartitionBuilder::Run(UWorld* World, FPackageSourceControlHelper& Pac
 						return false;
 					}
 
-					if (HasExceededMaxMemory())
+					if (FWorldPartitionHelpers::HasExceededMaxMemory())
 					{
 						WorldPartition->UnloadEditorCells(LoadedBounds);
 						// Reset Loaded Bounds
@@ -130,14 +267,6 @@ bool UWorldPartitionBuilder::Run(UWorld* World, FPackageSourceControlHelper& Pac
 	}
 }
 
-bool UWorldPartitionBuilder::HasExceededMaxMemory() const
-{
-	const uint64 MemoryMinFreePhysical = 1024ll * 1024 * 1024;
-	const uint64 MemoryMaxUsedPhysical = 32768ll * 1024 * 1024l;
-	const FPlatformMemoryStats MemStats = FPlatformMemory::GetStats();
-	return (MemStats.AvailablePhysical < MemoryMinFreePhysical) || (MemStats.UsedPhysical >= MemoryMaxUsedPhysical);
-};
-
 void UWorldPartitionBuilder::DoCollectGarbage() const
 {
 	if (GDistanceFieldAsyncQueue)
@@ -154,8 +283,8 @@ void UWorldPartitionBuilder::DoCollectGarbage() const
 	CollectGarbage(RF_NoFlags, true);
 	const FPlatformMemoryStats MemStatsAfter = FPlatformMemory::GetStats();
 
-	UE_LOG(LogWorldPartitionBuilder, Warning, TEXT("AvailablePhysical:%.2fGB AvailableVirtual %.2fGB"),
+	UE_LOG(LogWorldPartitionBuilder, Display, TEXT("GC Performed - Freed Physical: %.2fGB, Freed Virtual: %.2fGB"),
 		((int64)MemStatsAfter.AvailablePhysical - (int64)MemStatsBefore.AvailablePhysical) / (1024.0 * 1024.0 * 1024.0),
 		((int64)MemStatsAfter.AvailableVirtual - (int64)MemStatsBefore.AvailableVirtual) / (1024.0 * 1024.0 * 1024.0)
 	);
-};
+};
\ No newline at end of file
diff --git a/Engine/Source/Editor/UnrealEd/Private/WorldPartition/WorldPartitionHLODsBuilder.cpp b/Engine/Source/Editor/UnrealEd/Private/WorldPartition/WorldPartitionHLODsBuilder.cpp
--- a/Engine/Source/Editor/UnrealEd/Private/WorldPartition/WorldPartitionHLODsBuilder.cpp
+++ b/Engine/Source/Editor/UnrealEd/Private/WorldPartition/WorldPartitionHLODsBuilder.cpp
@@ -18,7 +18,7 @@
 #include "DerivedDataCacheInterface.h"
 
 #include "WorldPartition/WorldPartition.h"
-#include "WorldPartition/WorldPartitionSubsystem.h"
+#include "WorldPartition/WorldPartitionHelpers.h"
 #include "WorldPartition/HLOD/HLODActor.h"
 #include "WorldPartition/HLOD/HLODActorDesc.h"
 
@@ -420,7 +420,7 @@ bool UWorldPartitionHLODsBuilder::BuildHLODActors()
 			}
 		}
 
-		if (HasExceededMaxMemory())
+		if (FWorldPartitionHelpers::HasExceededMaxMemory())
 		{
 			DoCollectGarbage();
 		}
diff --git a/Engine/Source/Editor/UnrealEd/Private/WorldPartition/WorldPartitionResaveActorsBuilder.cpp b/Engine/Source/Editor/UnrealEd/Private/WorldPartition/WorldPartitionResaveActorsBuilder.cpp
--- a/Engine/Source/Editor/UnrealEd/Private/WorldPartition/WorldPartitionResaveActorsBuilder.cpp
+++ b/Engine/Source/Editor/UnrealEd/Private/WorldPartition/WorldPartitionResaveActorsBuilder.cpp
@@ -3,6 +3,7 @@
 #include "WorldPartition/WorldPartitionResaveActorsBuilder.h"
 
 #include "WorldPartition/WorldPartition.h"
+#include "WorldPartition/WorldPartitionHelpers.h"
 #include "PackageSourceControlHelper.h"
 #include "SourceControlHelpers.h"
 #include "Engine/World.h"
@@ -18,33 +19,13 @@ UWorldPartitionResaveActorsBuilder::UWorldPartitionResaveActorsBuilder(const FOb
 
 bool UWorldPartitionResaveActorsBuilder::RunInternal(UWorld* World, const FBox& Bounds, FPackageSourceControlHelper& PackageHelper)
 {
-	UWorldPartition* WorldPartition = World->GetWorldPartition();
-	check(WorldPartition);
-
-	TMap<FGuid, FWorldPartitionReference> ActorReferences;
-	TArray<FGuid> ActorsToSave;
-
-	// Recursive loading of references 
-	TFunction<void(const FGuid&, TMap<FGuid, FWorldPartitionReference>&)> LoadReferences = [WorldPartition, &LoadReferences](const FGuid& ActorGuid, TMap<FGuid, FWorldPartitionReference>& InOutActorReferences)
-	{
-		if (InOutActorReferences.Contains(ActorGuid))
-		{
-			return;
-		}
-
-		if (FWorldPartitionActorDesc* ActorDesc = WorldPartition->GetActorDesc(ActorGuid))
-		{
-			for (FGuid ReferenceGuid : ActorDesc->GetReferences())
-			{
-				LoadReferences(ReferenceGuid, InOutActorReferences);
-			}
+	FPackageSourceControlHelper SCCHelper;
 
-			InOutActorReferences.Add(ActorGuid, FWorldPartitionReference(WorldPartition, ActorGuid));
-		}
-	};
+	int32 SaveCount = 0;
+	int32 FailCount = 0;
 
 	// Actor Class Filter
-	UClass* ActorClass = nullptr;
+	UClass* ActorClass = AActor::StaticClass();
 	// @todo_ow: support BP Classes when the ActorDesc have that information
 	if (!ActorClassName.IsEmpty())
 	{
@@ -55,39 +36,27 @@ bool UWorldPartitionResaveActorsBuilder::RunInternal(UWorld* World, const FBox&
 			return false;
 		}
 	}
-		
-	UE_LOG(LogWorldPartitionResaveActorsBuilder, Display, TEXT("Gathering actors to save"));
-	// Get list of actors to save
-	for (UActorDescContainer::TIterator<AActor> ActorIterator(WorldPartition); ActorIterator; ++ActorIterator)
+
+	UWorldPartition* WorldPartition = World->GetWorldPartition();
+	if (!WorldPartition)
 	{
-		FWorldPartitionActorDesc* ActorDesc = *ActorIterator;	
-		if (!ActorClass || ActorDesc->GetActorClass()->IsChildOf(ActorClass))
-		{
-			ActorsToSave.Add(ActorDesc->GetGuid());
-		}
+		UE_LOG(LogWorldPartitionResaveActorsBuilder, Error, TEXT("Failed to retrieve WorldPartition."));
+		return false;
 	}
-	UE_LOG(LogWorldPartitionResaveActorsBuilder, Display, TEXT("Found %d actors to save"), ActorsToSave.Num());
-	FPackageSourceControlHelper SCCHelper;
 
-	int32 SaveCount = 0;
-	int32 FailCount = 0;
-
-	for (const FGuid& ActorGuid : ActorsToSave)
+	FWorldPartitionHelpers::ForEachActorWithLoading(WorldPartition, ActorClass, [&SaveCount, &FailCount, &SCCHelper](AActor* Actor)
 	{
-		LoadReferences(ActorGuid, ActorReferences);
-		FWorldPartitionReference ActorReference(WorldPartition, ActorGuid);
-		AActor* Actor = ActorReference.Get()->GetActor();
 		if (UPackage* Package = Actor->GetExternalPackage())
 		{
 			if (SCCHelper.Checkout(Package))
-			{				
+			{
 				// Save package
 				FString PackageFileName = SourceControlHelpers::PackageFilename(Package);
 				if (!UPackage::SavePackage(Package, nullptr, RF_Standalone, *PackageFileName))
 				{
 					UE_LOG(LogWorldPartitionResaveActorsBuilder, Error, TEXT("Error saving package %s."), *Package->GetName());
 					++FailCount;
-					continue;
+					return true;
 				}
 
 				// It is possible the resave can't checkout everything. Continue processing.
@@ -99,19 +68,14 @@ bool UWorldPartitionResaveActorsBuilder::RunInternal(UWorld* World, const FBox&
 				// It is possible the resave can't checkout everything. Continue processing.
 				UE_LOG(LogWorldPartitionResaveActorsBuilder, Warning, TEXT("Error checking out package %s."), *Package->GetName());
 				++FailCount;
-				continue;
+				return true;
 			}
 
-			UE_LOG(LogWorldPartitionResaveActorsBuilder, Display, TEXT("Processed %d / %d packages (%d Saved / %d Failed)"), SaveCount + FailCount, ActorsToSave.Num(), SaveCount, FailCount);
+			UE_LOG(LogWorldPartitionResaveActorsBuilder, Display, TEXT("Processed %d packages (%d Saved / %d Failed)"), SaveCount + FailCount, SaveCount, FailCount);
 		}
 
-		if (HasExceededMaxMemory())
-		{
-			UE_LOG(LogWorldPartitionResaveActorsBuilder, Log, TEXT("Freeing some memory"));
-			ActorReferences.Empty();
-			DoCollectGarbage();
-		}
-	}
+		return true;
+	});
 
 	UPackage::WaitForAsyncFileWrites();
 
diff --git a/Engine/Source/Editor/UnrealEd/Public/WorldPartition/WorldPartitionBuilder.h b/Engine/Source/Editor/UnrealEd/Public/WorldPartition/WorldPartitionBuilder.h
--- a/Engine/Source/Editor/UnrealEd/Public/WorldPartition/WorldPartitionBuilder.h
+++ b/Engine/Source/Editor/UnrealEd/Public/WorldPartition/WorldPartitionBuilder.h
@@ -18,6 +18,9 @@ public:
 		IterativeCells
 	};
 
+	static bool RunBuilder(UWorldPartitionBuilder* BuilderClass, UWorld* World);
+	static bool RunBuilder(TSubclassOf<UWorldPartitionBuilder> BuilderClass, UWorld* World);
+
 	virtual bool RequiresCommandletRendering() const PURE_VIRTUAL(UWorldPartitionBuilder::RequiresCommandletRendering, return false;);
 	virtual ELoadingMode GetLoadingMode() const PURE_VIRTUAL(UWorldPartitionBuilder::GetLoadingMode, return ELoadingMode::Custom;);
 	
@@ -28,7 +31,6 @@ public:
 protected:
 	virtual bool RunInternal(UWorld* World, const FBox& Bounds, FPackageSourceControlHelper& PackageHelper) PURE_VIRTUAL(UWorldPartition::RunInternal, return false;);
 
-	bool HasExceededMaxMemory() const;
 	void DoCollectGarbage() const;
 
 	int32 IterativeCellSize = 102400;
diff --git a/Engine/Source/Runtime/Engine/Private/ActorPartition/ActorPartitionSubsystem.cpp b/Engine/Source/Runtime/Engine/Private/ActorPartition/ActorPartitionSubsystem.cpp
--- a/Engine/Source/Runtime/Engine/Private/ActorPartition/ActorPartitionSubsystem.cpp
+++ b/Engine/Source/Runtime/Engine/Private/ActorPartition/ActorPartitionSubsystem.cpp
@@ -3,6 +3,7 @@
 #include "ActorPartition/ActorPartitionSubsystem.h"
 #include "ActorPartition/PartitionActor.h"
 #include "WorldPartition/ActorPartition/PartitionActorDesc.h"
+#include "WorldPartition/WorldPartitionHelpers.h"
 #include "WorldPartition/WorldPartitionSubsystem.h"
 #include "WorldPartition/DataLayer/DataLayerSubsystem.h"
 #include "Engine/World.h"
@@ -175,7 +176,7 @@ public:
 	FActorPartitionWorldPartition(UWorld* InWorld)
 		: FBaseActorPartition(InWorld)
 	{
-		WorldPartition = InWorld->GetSubsystem<UWorldPartitionSubsystem>();
+		WorldPartition = InWorld->GetWorldPartition();
 		check(WorldPartition);
 	}
 
@@ -224,11 +225,11 @@ public:
 		FBox CellBounds = UActorPartitionSubsystem::FCellCoord::GetCellBounds(InCellCoord, InGridSize);
 		if (bInBoundsSearch)
 		{
-			WorldPartition->ForEachIntersectingActorDesc(CellBounds, InActorPartitionId.GetClass(), FindActor);
+			FWorldPartitionHelpers::ForEachIntersectingActorDesc(WorldPartition, CellBounds, InActorPartitionId.GetClass(), FindActor);
 		}
 		else
 		{
-			WorldPartition->ForEachActorDesc(InActorPartitionId.GetClass(), FindActor);
+			FWorldPartitionHelpers::ForEachActorDesc(WorldPartition, InActorPartitionId.GetClass(), FindActor);
 		}
 				
 		if (bUnloadedActorExists)
@@ -307,7 +308,7 @@ public:
 	}
 
 private:
-	UWorldPartitionSubsystem* WorldPartition;
+	UWorldPartition* WorldPartition;
 };
 
 #endif // WITH_EDITOR
diff --git a/Engine/Source/Runtime/Engine/Private/WorldPartition/RuntimeSpatialHash/RuntimeSpatialHashHLOD.cpp b/Engine/Source/Runtime/Engine/Private/WorldPartition/RuntimeSpatialHash/RuntimeSpatialHashHLOD.cpp
--- a/Engine/Source/Runtime/Engine/Private/WorldPartition/RuntimeSpatialHash/RuntimeSpatialHashHLOD.cpp
+++ b/Engine/Source/Runtime/Engine/Private/WorldPartition/RuntimeSpatialHash/RuntimeSpatialHashHLOD.cpp
@@ -9,6 +9,7 @@
 #include "WorldPartition/HLOD/HLODLayer.h"
 #include "WorldPartition/HLOD/HLODActor.h"
 #include "WorldPartition/HLOD/HLODActorDesc.h"
+#include "WorldPartition/WorldPartitionHelpers.h"
 
 #include "HAL/PlatformFileManager.h"
 #include "Misc/PackageName.h"
@@ -26,9 +27,6 @@
 #include "AssetRegistryModule.h"
 #include "AssetData.h"
 
-#include "IDirectoryWatcher.h"
-#include "DirectoryWatcherModule.h"
-
 
 DEFINE_LOG_CATEGORY_STATIC(LogWorldPartitionRuntimeSpatialHashHLOD, Log, All);
 
@@ -96,14 +94,6 @@ static void DeletePackage(UWorldPartition* WorldPartition, FWorldPartitionActorD
 
 static TArray<FGuid> GenerateHLODsForGrid(UWorldPartition* WorldPartition, const FActorContainerInstance& MainContainerInstance, const FSpatialHashRuntimeGrid& RuntimeGrid, uint32 HLODLevel, FHLODCreationContext& Context, ISourceControlHelper* SourceControlHelper, bool bCreateActorsOnly, const TArray<const FActorClusterInstance*>& ClusterInstances)
 {
-	auto HasExceededMaxMemory = []()
-	{
-		const uint64 MemoryMinFreePhysical = 1024ll * 1024 * 1024;
-		const uint64 MemoryMaxUsedPhysical = 16384ll * 1024 * 1024l;
-		const FPlatformMemoryStats MemStats = FPlatformMemory::GetStats();
-		return (MemStats.AvailablePhysical < MemoryMinFreePhysical) || (MemStats.UsedPhysical >= MemoryMaxUsedPhysical);
-	};
-
 	auto DoCollectGarbage = []()
 	{
 		if (GDistanceFieldAsyncQueue)
@@ -120,7 +110,7 @@ static TArray<FGuid> GenerateHLODsForGrid(UWorldPartition* WorldPartition, const
 		CollectGarbage(RF_NoFlags, true);
 		const FPlatformMemoryStats MemStatsAfter = FPlatformMemory::GetStats();
 
-		UE_LOG(LogWorldPartitionRuntimeSpatialHashHLOD, Warning, TEXT("AvailablePhysical:%.2fGB AvailableVirtual %.2fGB"),
+		UE_LOG(LogWorldPartitionRuntimeSpatialHashHLOD, Display, TEXT("GC Performed - Freed Physical: %.2fGB, Freed Virtual: %.2fGB"),
 			((int64)MemStatsAfter.AvailablePhysical - (int64)MemStatsBefore.AvailablePhysical) / (1024.0 * 1024.0 * 1024.0),
 			((int64)MemStatsAfter.AvailableVirtual - (int64)MemStatsBefore.AvailableVirtual) / (1024.0 * 1024.0 * 1024.0)
 		);
@@ -160,9 +150,6 @@ static TArray<FGuid> GenerateHLODsForGrid(UWorldPartition* WorldPartition, const
 	FScopedSlowTask SlowTask(NbCellsToProcess, FText::FromString(FString::Printf(TEXT("Building HLODs for grid %s..."), *RuntimeGrid.GridName.ToString())));
 	SlowTask.MakeDialog();
 
-	FDirectoryWatcherModule& DirectoryWatcherModule = FModuleManager::LoadModuleChecked<FDirectoryWatcherModule>(TEXT("DirectoryWatcher"));
-	FAssetRegistryModule& AssetRegistryModule = FModuleManager::LoadModuleChecked<FAssetRegistryModule>(TEXT("AssetRegistry"));
-
 	TArray<FGuid> GridHLODActors;
 	PartitionedActors.ForEachCells([&](const FSquare2DGridHelper::FGridLevel::FGridCell& GridCell)
 	{
@@ -222,7 +209,7 @@ static TArray<FGuid> GenerateHLODsForGrid(UWorldPartition* WorldPartition, const
 						{
 							FStaticMeshCompilingManager::Get().FinishAllCompilation();
 
-							if (HasExceededMaxMemory())
+							if (FWorldPartitionHelpers::HasExceededMaxMemory())
 							{
 								DoCollectGarbage();
 							}
diff --git a/Engine/Source/Runtime/Engine/Private/WorldPartition/WorldPartition.cpp b/Engine/Source/Runtime/Engine/Private/WorldPartition/WorldPartition.cpp
--- a/Engine/Source/Runtime/Engine/Private/WorldPartition/WorldPartition.cpp
+++ b/Engine/Source/Runtime/Engine/Private/WorldPartition/WorldPartition.cpp
@@ -640,27 +640,6 @@ UWorldPartition* UWorldPartition::CreateWorldPartition(AWorldSettings* WorldSett
 	return WorldPartition;
 }
 
-void UWorldPartition::ForEachIntersectingActorDesc(const FBox& Box, TSubclassOf<AActor> ActorClass, TFunctionRef<bool(const FWorldPartitionActorDesc*)> Predicate) const
-{
-	EditorHash->ForEachIntersectingActor(Box, [&ActorClass, Predicate](const FWorldPartitionActorDesc* ActorDesc)
-	{
-		if (ActorDesc->GetActorClass()->IsChildOf(ActorClass))
-		{
-			Predicate(ActorDesc);
-		}
-	});
-}
-
-void UWorldPartition::ForEachActorDesc(TSubclassOf<AActor> ActorClass, TFunctionRef<bool(const FWorldPartitionActorDesc*)> Predicate) const
-{
-	for (UActorDescContainer::TConstIterator<> ActorDescIterator(this, ActorClass); ActorDescIterator; ++ActorDescIterator)
-	{
-		if (!Predicate(*ActorDescIterator))
-		{
-			return;
-		}
-	}
-}
 #endif
 
 const TArray<FWorldPartitionStreamingSource>& UWorldPartition::GetStreamingSources() const
diff --git a/Engine/Source/Runtime/Engine/Private/WorldPartition/WorldPartitionHelpers.cpp b/Engine/Source/Runtime/Engine/Private/WorldPartition/WorldPartitionHelpers.cpp
new file mode 100644
--- /dev/null
+++ b/Engine/Source/Runtime/Engine/Private/WorldPartition/WorldPartitionHelpers.cpp
@@ -0,0 +1,100 @@
+// Copyright Epic Games, Inc. All Rights Reserved.
+
+#include "WorldPartition/WorldPartitionHelpers.h"
+
+#include "WorldPartition/WorldPartition.h"
+#include "WorldPartition/WorldPartitionEditorHash.h"
+
+#if WITH_EDITOR
+
+void FWorldPartitionHelpers::ForEachIntersectingActorDesc(UWorldPartition* WorldPartition, const FBox& Box, TSubclassOf<AActor> ActorClass, TFunctionRef<bool(const FWorldPartitionActorDesc*)> Predicate)
+{
+	WorldPartition->EditorHash->ForEachIntersectingActor(Box, [&ActorClass, Predicate](const FWorldPartitionActorDesc* ActorDesc)
+	{
+		if (ActorDesc->GetActorClass()->IsChildOf(ActorClass))
+		{
+			Predicate(ActorDesc);
+		}
+	});
+}
+
+void FWorldPartitionHelpers::ForEachActorDesc(UWorldPartition* WorldPartition, TSubclassOf<AActor> ActorClass, TFunctionRef<bool(const FWorldPartitionActorDesc*)> Predicate)
+{
+	for (UActorDescContainer::TConstIterator<> ActorDescIterator(WorldPartition, ActorClass); ActorDescIterator; ++ActorDescIterator)
+	{
+		if (!Predicate(*ActorDescIterator))
+		{
+			return;
+		}
+	}
+}
+
+void FWorldPartitionHelpers::ForEachActorWithLoading(UWorldPartition* WorldPartition, TSubclassOf<AActor> ActorClass, TFunctionRef<bool(AActor*)> Predicate)
+{
+	// Recursive loading of references 
+	TFunction<void(const FGuid&, TMap<FGuid, FWorldPartitionReference>&)> LoadReferences = [WorldPartition, &LoadReferences](const FGuid& ActorGuid, TMap<FGuid, FWorldPartitionReference>& InOutActorReferences)
+	{
+		if (InOutActorReferences.Contains(ActorGuid))
+		{
+			return;
+		}
+
+		if (const FWorldPartitionActorDesc* ActorDesc = WorldPartition->GetActorDesc(ActorGuid))
+		{
+			for (FGuid ReferenceGuid : ActorDesc->GetReferences())
+			{
+				LoadReferences(ReferenceGuid, InOutActorReferences);
+			}
+
+			InOutActorReferences.Add(ActorGuid, FWorldPartitionReference(WorldPartition, ActorGuid));
+		}
+	};
+
+	TMap<FGuid, FWorldPartitionReference> ActorReferences;
+	TArray<FGuid> ActorsToProcess;
+
+	ForEachActorDesc(WorldPartition, ActorClass, [&](const FWorldPartitionActorDesc* ActorDesc)
+	{
+		const FGuid& ActorGuid = ActorDesc->GetGuid();
+		LoadReferences(ActorGuid, ActorReferences);
+
+		FWorldPartitionReference ActorReference(WorldPartition, ActorGuid);
+		AActor* Actor = ActorReference.Get()->GetActor();
+
+		if (!Predicate(Actor))
+		{
+			return false;
+		}
+
+		if (HasExceededMaxMemory())
+		{
+			ActorReferences.Empty();
+
+			const FPlatformMemoryStats MemStatsBefore = FPlatformMemory::GetStats();
+			CollectGarbage(RF_NoFlags, true);
+			const FPlatformMemoryStats MemStatsAfter = FPlatformMemory::GetStats();
+
+			UE_LOG(LogWorldPartition, Log, TEXT("GC Performed - Available Physical: %.2fGB, Available Virtual: %.2fGB"),
+				(int64)MemStatsAfter.AvailablePhysical / (1024.0 * 1024.0 * 1024.0),
+				(int64)MemStatsAfter.AvailableVirtual / (1024.0 * 1024.0 * 1024.0)
+			);
+		}
+
+		return true;
+	});
+}
+
+bool FWorldPartitionHelpers::HasExceededMaxMemory()
+{
+	const uint64 MemoryMinFreePhysical = 1 * 1024ll * 1024ll * 1024ll;
+	const uint64 MemoryMaxUsedPhysical = 32 * 1024ll * 1024ll * 1024ll;
+	const FPlatformMemoryStats MemStats = FPlatformMemory::GetStats();
+
+	const bool bHasExceededMinFreePhysical = MemStats.AvailablePhysical < MemoryMinFreePhysical;
+	const bool bHasExceededMaxUsedPhysical = MemStats.UsedPhysical >= MemoryMaxUsedPhysical;
+	const bool bHasExceededMaxMemory = bHasExceededMinFreePhysical || bHasExceededMaxUsedPhysical;
+
+	return bHasExceededMaxMemory;
+};
+
+#endif // #if WITH_EDITOR
diff --git a/Engine/Source/Runtime/Engine/Private/WorldPartition/WorldPartitionSubsystem.cpp b/Engine/Source/Runtime/Engine/Private/WorldPartition/WorldPartitionSubsystem.cpp
--- a/Engine/Source/Runtime/Engine/Private/WorldPartition/WorldPartitionSubsystem.cpp
+++ b/Engine/Source/Runtime/Engine/Private/WorldPartition/WorldPartitionSubsystem.cpp
@@ -266,28 +266,6 @@ void UWorldPartitionSubsystem::Draw(UCanvas* Canvas, class APlayerController* PC
 	}
 }
 
-//
-// Methods only used by Editor for non-game world
-//
-
-#if WITH_EDITOR
-void UWorldPartitionSubsystem::ForEachIntersectingActorDesc(const FBox& Box, TSubclassOf<AActor> ActorClass, TFunctionRef<bool(const FWorldPartitionActorDesc*)> Predicate) const
-{
-	if (const UWorldPartition* MainPartition = GetMainWorldPartition())
-	{
-		MainPartition->ForEachIntersectingActorDesc(Box, ActorClass, Predicate);
-	}
-}
-
-void UWorldPartitionSubsystem::ForEachActorDesc(TSubclassOf<AActor> ActorClass, TFunctionRef<bool(const FWorldPartitionActorDesc*)> Predicate) const
-{
-	if (const UWorldPartition* MainPartition = GetMainWorldPartition())
-	{
-		MainPartition->ForEachActorDesc(ActorClass, Predicate);
-	}
-}
-#endif // WITH_EDITOR
-
 //
 // UWorldPartitionSubsystem::FWorldPartitionCVars
 //
diff --git a/Engine/Source/Runtime/Engine/Public/WorldPartition/WorldPartition.h b/Engine/Source/Runtime/Engine/Public/WorldPartition/WorldPartition.h
--- a/Engine/Source/Runtime/Engine/Public/WorldPartition/WorldPartition.h
+++ b/Engine/Source/Runtime/Engine/Public/WorldPartition/WorldPartition.h
@@ -120,10 +120,7 @@ private:
 
 	// UWorldPartitionSubsystem interface+
 	friend class UWorldPartitionSubsystem;
-#if WITH_EDITOR
-	void ForEachIntersectingActorDesc(const FBox& Box, TSubclassOf<AActor> ActorClass, TFunctionRef<bool(const FWorldPartitionActorDesc*)> Predicate) const;
-	void ForEachActorDesc(TSubclassOf<AActor> ActorClass, TFunctionRef<bool(const FWorldPartitionActorDesc*)> Predicate) const;
-#endif
+
 	const TArray<FWorldPartitionStreamingSource>& GetStreamingSources() const;
 	// UActorPartitionSubsystem interface-
 
diff --git a/Engine/Source/Runtime/Engine/Public/WorldPartition/WorldPartitionHelpers.h b/Engine/Source/Runtime/Engine/Public/WorldPartition/WorldPartitionHelpers.h
new file mode 100644
--- /dev/null
+++ b/Engine/Source/Runtime/Engine/Public/WorldPartition/WorldPartitionHelpers.h
@@ -0,0 +1,40 @@
+// Copyright Epic Games, Inc. All Rights Reserved.
+#pragma once
+
+#include "CoreMinimal.h"
+#include "GameFramework/Actor.h"
+
+#if WITH_EDITOR
+
+class UWorldPartition;
+class FWorldPartitionActorDesc;
+
+class ENGINE_API FWorldPartitionHelpers
+{
+public:
+	template <class ActorClass = AActor>
+	static void ForEachIntersectingActorDesc(UWorldPartition* WorldPartition, const FBox& Box, bool bIncludeFromChildActors, TFunctionRef<bool(const FWorldPartitionActorDesc*)> Func)
+	{
+		ForEachIntersectingActorDesc(WorldPartition, Box, bIncludeFromChildActors, ActorClass::StaticClass(), Func);
+	}
+
+	template<class ActorClass = AActor>
+	static void ForEachActorDesc(UWorldPartition* WorldPartition, TFunctionRef<bool(const FWorldPartitionActorDesc*)> Func)
+	{
+		ForEachActorDesc(WorldPartition, ActorClass::StaticClass(), Func);
+	}
+
+	template<class ActorClass = AActor>
+	static void ForEachActorWithLoading(UWorldPartition* WorldPartition, TFunctionRef<bool(AActor*)> Func)
+	{
+		ForEachActorWithLoading(WorldPartition, ActorClass::StaticClass(), Func);
+	}
+
+	static void ForEachIntersectingActorDesc(UWorldPartition* WorldPartition, const FBox& Box, TSubclassOf<AActor> ActorClass, TFunctionRef<bool(const FWorldPartitionActorDesc*)> Func);
+	static void ForEachActorDesc(UWorldPartition* WorldPartition, TSubclassOf<AActor> ActorClass, TFunctionRef<bool(const FWorldPartitionActorDesc*)> Func);
+	static void ForEachActorWithLoading(UWorldPartition* WorldPartition, TSubclassOf<AActor> ActorClass, TFunctionRef<bool(AActor*)> Func);
+
+	static bool HasExceededMaxMemory();
+};
+
+#endif // WITH_EDITOR
diff --git a/Engine/Source/Runtime/Engine/Public/WorldPartition/WorldPartitionSubsystem.h b/Engine/Source/Runtime/Engine/Public/WorldPartition/WorldPartitionSubsystem.h
--- a/Engine/Source/Runtime/Engine/Public/WorldPartition/WorldPartitionSubsystem.h
+++ b/Engine/Source/Runtime/Engine/Public/WorldPartition/WorldPartitionSubsystem.h
@@ -50,10 +50,6 @@ public:
 	UFUNCTION(BlueprintCallable, Category = Streaming)
 	bool IsStreamingCompleted(EWorldPartitionRuntimeCellState QueryState, const TArray<FWorldPartitionStreamingQuerySource>& QuerySources, bool bExactState) const;
 
-#if WITH_EDITOR
-	void ForEachIntersectingActorDesc(const FBox& Box, TSubclassOf<AActor> ActorClass, TFunctionRef<bool(const FWorldPartitionActorDesc*)> Predicate) const;
-	void ForEachActorDesc(TSubclassOf<AActor> ActorClass, TFunctionRef<bool(const FWorldPartitionActorDesc*)> Predicate) const;
-#endif
 	void ToggleDrawRuntimeHash2D();
 
 private:
