Make sure that parallel proxy creation is disabled when a property coloration is enabled.  It is disabled by default, but could expose a race condition for any projects that enable it.

diff --git a/Engine/Source/Editor/UnrealEd/Classes/Editor/EditorEngine.h b/Engine/Source/Editor/UnrealEd/Classes/Editor/EditorEngine.h
--- a/Engine/Source/Editor/UnrealEd/Classes/Editor/EditorEngine.h
+++ b/Engine/Source/Editor/UnrealEd/Classes/Editor/EditorEngine.h
@@ -739,6 +739,7 @@ public:
 	virtual bool ShouldDrawBrushWireframe(AActor* InActor) override;
 	virtual void NotifyToolsOfObjectReplacement(const TMap<UObject*, UObject*>& OldToNewInstanceMap) override;
 	virtual bool ShouldThrottleCPUUsage() const override;
+	virtual bool IsPropertyColorationColorFeatureActivated() const override;
 	virtual bool GetPropertyColorationColor(class UObject* Object, FColor& OutColor) override;
 	virtual bool WorldIsPIEInNewViewport(UWorld* InWorld) override;
 	virtual void FocusNextPIEWorld(UWorld* CurrentPieWorld, bool previous = false) override;
diff --git a/Engine/Source/Editor/UnrealEd/Private/EditorServer.cpp b/Engine/Source/Editor/UnrealEd/Private/EditorServer.cpp
--- a/Engine/Source/Editor/UnrealEd/Private/EditorServer.cpp
+++ b/Engine/Source/Editor/UnrealEd/Private/EditorServer.cpp
@@ -3794,11 +3794,15 @@ void UEditorEngine::GetPropertyColorationTarget(FString& OutPropertyValue, FProp
 	OutPropertyChain	= GPropertyColorationChain;
 }
 
+bool UEditorEngine::IsPropertyColorationColorFeatureActivated() const
+{
+	return GPropertyColorationClass && GPropertyColorationChain;
+}
 
 bool UEditorEngine::GetPropertyColorationColor(UObject* Object, FColor& OutColor)
 {
 	bool bResult = false;
-	if ( GPropertyColorationClass && GPropertyColorationChain && GPropertyColorationChain->Num() > 0 )
+	if (IsPropertyColorationColorFeatureActivated() && GPropertyColorationChain->Num() > 0)
 	{
 		UObject* MatchingBase = NULL;
 		AActor* Owner = NULL;
diff --git a/Engine/Source/Runtime/Engine/Classes/Engine/Engine.h b/Engine/Source/Runtime/Engine/Classes/Engine/Engine.h
--- a/Engine/Source/Runtime/Engine/Classes/Engine/Engine.h
+++ b/Engine/Source/Runtime/Engine/Classes/Engine/Engine.h
@@ -2257,6 +2257,11 @@ public:
 		// Intentionally empty.
 	}
 
+	/**
+	 * @return true if property coloration color has a valid implementation
+	 */
+	virtual bool IsPropertyColorationColorFeatureActivated() const { return false; }
+
 	/**
 	 * Computes a color to use for property coloration for the given object.
 	 *
diff --git a/Engine/Source/Runtime/Engine/Private/Components/SkinnedMeshComponent.cpp b/Engine/Source/Runtime/Engine/Private/Components/SkinnedMeshComponent.cpp
--- a/Engine/Source/Runtime/Engine/Private/Components/SkinnedMeshComponent.cpp
+++ b/Engine/Source/Runtime/Engine/Private/Components/SkinnedMeshComponent.cpp
@@ -649,6 +649,13 @@ void USkinnedMeshComponent::DestroyRenderState_Concurrent()
 
 bool USkinnedMeshComponent::RequiresGameThreadEndOfFrameRecreate() const
 {
+#if STATICMESH_ENABLE_DEBUG_RENDERING
+	if (GIsEditor && GEngine->IsPropertyColorationColorFeatureActivated())
+	{
+		return true;
+	}
+#endif
+
 	// When we are a master/slave, we cannot recreate render state in parallel as this could 
 	// happen concurrently with our dependent component(s)
 	return MasterPoseComponent.Get() != nullptr || SlavePoseComponents.Num() > 0;
diff --git a/Engine/Source/Runtime/Engine/Private/Components/StaticMeshComponent.cpp b/Engine/Source/Runtime/Engine/Private/Components/StaticMeshComponent.cpp
--- a/Engine/Source/Runtime/Engine/Private/Components/StaticMeshComponent.cpp
+++ b/Engine/Source/Runtime/Engine/Private/Components/StaticMeshComponent.cpp
@@ -691,6 +691,12 @@ void UStaticMeshComponent::OnUnregister()
 
 bool UStaticMeshComponent::RequiresGameThreadEndOfFrameRecreate() const
 {
+#if STATICMESH_ENABLE_DEBUG_RENDERING
+	if (GIsEditor && GEngine->IsPropertyColorationColorFeatureActivated())
+	{
+		return true;
+	}
+#endif
 	return false;
 }
 
diff --git a/Engine/Source/Runtime/Engine/Private/SkeletalMesh.cpp b/Engine/Source/Runtime/Engine/Private/SkeletalMesh.cpp
--- a/Engine/Source/Runtime/Engine/Private/SkeletalMesh.cpp
+++ b/Engine/Source/Runtime/Engine/Private/SkeletalMesh.cpp
@@ -5602,8 +5602,10 @@ FSkeletalMeshSceneProxy::FSkeletalMeshSceneProxy(const USkinnedMeshComponent* Co
 
 	// Get a color for property coloration
 	FColor NewPropertyColor;
-	GEngine->GetPropertyColorationColor( (UObject*)Component, NewPropertyColor );
-	SetPropertyColor(NewPropertyColor);
+	if (GEngine->GetPropertyColorationColor((UObject*)Component, NewPropertyColor))
+	{
+		SetPropertyColor(NewPropertyColor);
+	}
 
 	// Copy out shadow physics asset data
 	if(SkinnedMeshComponent)
