Sequencer: Export baked transforms by evaluating each frame and also ticking the hierarchy. This fixes exporting an object (ie. camera) that is attached to a skeletal mesh or a camera rig which requires those objects to be ticked up the chain.

Also, make sure we export all of the skeletal mesh components and not just the ones that have skeletal animation tracks. This ensures that the entire skeletal actor is exported properly.

#b mike.zyracki

diff --git a/Engine/Source/Editor/UnrealEd/Private/Fbx/FbxAnimationExport.cpp b/Engine/Source/Editor/UnrealEd/Private/Fbx/FbxAnimationExport.cpp
--- a/Engine/Source/Editor/UnrealEd/Private/Fbx/FbxAnimationExport.cpp
+++ b/Engine/Source/Editor/UnrealEd/Private/Fbx/FbxAnimationExport.cpp
@@ -678,6 +678,8 @@ void FFbxExporter::ExportAnimTrack(IAnimTrackAdapter& AnimTrackAdapter, AActor*
 	TArray<USkeletalMeshComponent*> SkeletalMeshComponents;
 	Actor->GetComponents(SkeletalMeshComponents);
 
+	const float TickRate = 1.0f/FrameRate;
+
 	for (int32 FrameCount = 0; FrameCount <= AnimationLength; ++FrameCount)
 	{
 		if (FrameCount == 0)
@@ -699,14 +701,28 @@ void FFbxExporter::ExportAnimTrack(IAnimTrackAdapter& AnimTrackAdapter, AActor*
 		// @todo - hack - this will be removed at some point
 		for (USkeletalMeshComponent* SkeletalMeshComponent : SkeletalMeshComponents)
 		{
-			SkeletalMeshComponent->TickAnimation(0.03f, false);
-
-			SkeletalMeshComponent->RefreshBoneTransforms();
-			SkeletalMeshComponent->RefreshSlaveComponents();
-			SkeletalMeshComponent->UpdateComponentToWorld();
-			SkeletalMeshComponent->FinalizeBoneTransform();
-			SkeletalMeshComponent->MarkRenderTransformDirty();
-			SkeletalMeshComponent->MarkRenderDynamicDataDirty();
+			USceneComponent* Child = SkeletalMeshComponent;
+			while (Child)
+			{
+				if (USkeletalMeshComponent* ChildSkeletalMeshComponent = Cast<USkeletalMeshComponent>(Child))
+				{
+					SkeletalMeshComponent->TickAnimation(TickRate, false);
+
+					SkeletalMeshComponent->RefreshBoneTransforms();
+					SkeletalMeshComponent->RefreshSlaveComponents();
+					SkeletalMeshComponent->UpdateComponentToWorld();
+					SkeletalMeshComponent->FinalizeBoneTransform();
+					SkeletalMeshComponent->MarkRenderTransformDirty();
+					SkeletalMeshComponent->MarkRenderDynamicDataDirty();
+				}
+
+				if (Child->GetOwner())
+				{
+					Child->GetOwner()->Tick(TickRate);
+				}
+
+				Child = Child->GetAttachParent();
+			}
 		}
 
 		FbxTime ExportTime; 
diff --git a/Engine/Source/Editor/UnrealEd/Private/Fbx/FbxMainExport.cpp b/Engine/Source/Editor/UnrealEd/Private/Fbx/FbxMainExport.cpp
--- a/Engine/Source/Editor/UnrealEd/Private/Fbx/FbxMainExport.cpp
+++ b/Engine/Source/Editor/UnrealEd/Private/Fbx/FbxMainExport.cpp
@@ -91,8 +91,6 @@
 #include "Evaluation/MovieSceneEvaluationTemplateInstance.h"
 #include "MovieSceneSequence.h"
 #include "MovieSceneTimeHelpers.h"
-#include "EntitySystem/Interrogation/MovieSceneInterrogationLinker.h"
-#include "Systems/MovieSceneComponentTransformSystem.h"
 
 #if WITH_PHYSX
 #include "DynamicMeshBuilder.h"
@@ -1988,19 +1986,20 @@ bool FFbxExporter::ExportLevelSequenceTracks(UMovieScene* MovieScene, IMovieScen
 
 	const bool bIsCameraActor = BoundActor ? BoundActor->IsA(ACameraActor::StaticClass()) : BoundComponent ? BoundComponent->IsA(UCameraComponent::StaticClass()) : false;
 
-	// If there's more than one transform track for this actor (ie. on the actor and on the root component) or if there's more than one section, evaluate through interrogation
+	// If there's more than one transform track for this actor (ie. on the actor and on the root component) or if there's more than one section, evaluate baked
 	if (bIsCameraActor || TransformTracks.Num() > 1 || (TransformTracks.Num() != 0 && TransformTracks[0].Get()->GetAllSections().Num() > 1))
 	{
 		if (!bSkip3DTransformTrack)
 		{
-			ExportLevelSequenceInterrogated3DTransformTrack(FbxActor, MovieScenePlayer, InSequenceID, TransformTracks, BoundObject, MovieScene->GetPlaybackRange(), RootToLocalTransform);
+			FLevelSequenceAnimTrackAdapter AnimTrackAdapter(MovieScenePlayer, MovieScene, RootToLocalTransform, nullptr);
+			ExportLevelSequenceBaked3DTransformTrack(AnimTrackAdapter, FbxActor, MovieScenePlayer, InSequenceID, TransformTracks, BoundObject, MovieScene->GetPlaybackRange(), RootToLocalTransform);
 		}
 
 		bSkip3DTransformTrack = true;
 	}
-	
+
 	// Look for the tracks that we currently support
-	bool bExportedAnimTrack = false; // Only export the anim track once since the evaluation is already blended
+	UMovieSceneSkeletalAnimationTrack* SkeletalAnimationTrack = nullptr;
 	for (UMovieSceneTrack* Track : Tracks)
 	{
 		if (Track->IsA(UMovieScene3DTransformTrack::StaticClass()))
@@ -2011,20 +2010,28 @@ bool FFbxExporter::ExportLevelSequenceTracks(UMovieScene* MovieScene, IMovieScen
 				ExportLevelSequence3DTransformTrack(FbxActor, MovieScenePlayer, InSequenceID, *TransformTrack, BoundObject, MovieScene->GetPlaybackRange(), RootToLocalTransform);
 			}
 		}
-		else if (Track->IsA(UMovieSceneSkeletalAnimationTrack::StaticClass()) && !bExportedAnimTrack)
+		else if (Track->IsA(UMovieSceneSkeletalAnimationTrack::StaticClass()))
 		{
-			if (SkeletalMeshComp)
-			{
-				bExportedAnimTrack = true;
-				FLevelSequenceAnimTrackAdapter AnimTrackAdapter(MovieScenePlayer, MovieScene, RootToLocalTransform, Cast<UMovieSceneSkeletalAnimationTrack>(Track));
-				ExportAnimTrack(AnimTrackAdapter, Actor, SkeletalMeshComp, 1.0 / DisplayRate.AsDecimal());
-			}
+			SkeletalAnimationTrack = Cast<UMovieSceneSkeletalAnimationTrack>(Track);
 		}
 		else
 		{
 			ExportLevelSequenceTrackChannels(FbxActor, *Track, MovieScene->GetPlaybackRange(), RootToLocalTransform);
 		}
 	}
+
+	// Export all of the skeletal animation components for this actor
+	if (SkeletalMeshComp && SkeletalAnimationTrack)
+	{
+		TArray<USkeletalMeshComponent*> SkeletalMeshComponents;
+		SkeletalMeshComp->GetOwner()->GetComponents(SkeletalMeshComponents);
+		for (USkeletalMeshComponent* SkeletalMeshComponent : SkeletalMeshComponents)
+		{
+			FLevelSequenceAnimTrackAdapter AnimTrackAdapter(MovieScenePlayer, MovieScene, RootToLocalTransform, SkeletalAnimationTrack);
+			ExportAnimTrack(AnimTrackAdapter, Actor, SkeletalMeshComponent, 1.0 / DisplayRate.AsDecimal());
+		}
+	}
+
 	return true;
 }
 
@@ -3112,7 +3119,7 @@ void FFbxExporter::ExportLevelSequence3DTransformTrack(FbxNode* FbxNode, IMovieS
 	}
 }
 
-void FFbxExporter::ExportLevelSequenceInterrogated3DTransformTrack(FbxNode* FbxNode, IMovieScenePlayer* MovieScenePlayer, FMovieSceneSequenceIDRef InSequenceID, TArray<TWeakObjectPtr<UMovieScene3DTransformTrack> > TransformTracks, UObject* BoundObject, const TRange<FFrameNumber>& InPlaybackRange, const FMovieSceneSequenceTransform& RootToLocalTransform)
+void FFbxExporter::ExportLevelSequenceBaked3DTransformTrack(IAnimTrackAdapter& AnimTrackAdapter, FbxNode* FbxNode, IMovieScenePlayer* MovieScenePlayer, FMovieSceneSequenceIDRef InSequenceID, TArray<TWeakObjectPtr<UMovieScene3DTransformTrack> > TransformTracks, UObject* BoundObject, const TRange<FFrameNumber>& InPlaybackRange, const FMovieSceneSequenceTransform& RootToLocalTransform)
 {
 	using namespace UE::MovieScene;
 
@@ -3192,7 +3199,7 @@ void FFbxExporter::ExportLevelSequenceInterrogated3DTransformTrack(FbxNode* FbxN
 		FbxCurveScaleZ->KeyModifyBegin();
 	}
 
-	FMovieSceneTimeTransform LocatToRootTransform = RootToLocalTransform.InverseLinearOnly();
+	FMovieSceneTimeTransform LocalToRootTransform = RootToLocalTransform.InverseLinearOnly();
 
 	USceneComponent* InterrogatedComponent = nullptr;
 	if (BoundComponent)
@@ -3217,24 +3224,45 @@ void FFbxExporter::ExportLevelSequenceInterrogated3DTransformTrack(FbxNode* FbxN
 		}
 	}
 
-	FSystemInterrogator Interrogator;
-	Interrogator.ImportTransformHierarchy(InterrogatedComponent, MovieScenePlayer, InSequenceID);
-	
+	TArray<FTransform> WorldTransforms;
 	int32 LocalStartFrame = FFrameRate::TransformTime(FFrameTime(DiscreteInclusiveLower(InPlaybackRange)), TickResolution, DisplayRate).RoundToFrame().Value;
 	int32 AnimationLength = FFrameRate::TransformTime(FFrameTime(FFrameNumber(DiscreteSize(InPlaybackRange))), TickResolution, DisplayRate).RoundToFrame().Value + 1; // Add one so that we export a key for the end frame
 
+	const float SampleRate = 1.0f/DisplayRate.AsDecimal();
+
 	for (int32 FrameNumber = LocalStartFrame; FrameNumber < LocalStartFrame + AnimationLength; ++FrameNumber)
 	{
 		const FFrameTime FrameTime = FFrameRate::TransformTime(FrameNumber, DisplayRate, TickResolution);
-		Interrogator.AddInterrogation(FrameTime);
-	}
 
-	Interrogator.Update();
+		// This will call UpdateSkelPose on the skeletal mesh component to move bones based on animations in the matinee group
+		AnimTrackAdapter.UpdateAnimation(FrameNumber);
 
-	TArray<FTransform> WorldTransforms;
-	Interrogator.QueryWorldSpaceTransforms(InterrogatedComponent, WorldTransforms);
+		USceneComponent* Child = InterrogatedComponent;
+		while (Child)
+		{
+			if (USkeletalMeshComponent* ChildSkeletalMeshComponent = Cast<USkeletalMeshComponent>(Child))
+			{
+				ChildSkeletalMeshComponent->TickAnimation(SampleRate, false);
+
+				ChildSkeletalMeshComponent->RefreshBoneTransforms();
+				ChildSkeletalMeshComponent->RefreshSlaveComponents();
+				ChildSkeletalMeshComponent->UpdateComponentToWorld();
+				ChildSkeletalMeshComponent->FinalizeBoneTransform();
+				ChildSkeletalMeshComponent->MarkRenderTransformDirty();
+				ChildSkeletalMeshComponent->MarkRenderDynamicDataDirty();
+			}
+
+			if (Child->GetOwner())
+			{
+				Child->GetOwner()->Tick(SampleRate);
+			}
 
-	ensure(WorldTransforms.Num() == AnimationLength);
+			Child = Child->GetAttachParent();
+		}
+		
+		FTransform WorldTransform = InterrogatedComponent->GetComponentToWorld();
+		WorldTransforms.Add(WorldTransform);
+	}
 
 	for (int32 TransformIndex = 0; TransformIndex < WorldTransforms.Num(); ++TransformIndex)
 	{
@@ -3252,7 +3280,7 @@ void FFbxExporter::ExportLevelSequenceInterrogated3DTransformTrack(FbxNode* FbxN
 		}
 		else
 		{
-			FbxTime.SetSecondDouble(DisplayRate.AsSeconds(CurrentFrame * LocatToRootTransform));
+			FbxTime.SetSecondDouble(DisplayRate.AsSeconds(CurrentFrame * LocalToRootTransform));
 		}
 
 		FbxCurveTransX->KeySet(FbxCurveTransX->KeyAdd(FbxTime), FbxTime, KeyTrans[0]);
@@ -3355,6 +3383,11 @@ void FFbxExporter::ExportLevelSequenceTrackChannels( FbxNode* FbxNode, UMovieSce
 			{
 				Property = FbxNode->FindProperty("UE_Radius", false);
 			}
+			else if (PropertyName == "FieldOfView" && FbxCamera)
+			{
+				Property = FbxCamera->FocalLength;
+				IsFoV = true;
+			}
 			else if (PropertyName == "FOVAngle" && FbxCamera)
 			{
 				Property = FbxCamera->FocalLength;
diff --git a/Engine/Source/Editor/UnrealEd/Private/FbxExporter.h b/Engine/Source/Editor/UnrealEd/Private/FbxExporter.h
--- a/Engine/Source/Editor/UnrealEd/Private/FbxExporter.h
+++ b/Engine/Source/Editor/UnrealEd/Private/FbxExporter.h
@@ -468,9 +468,9 @@ private:
 	void ExportLevelSequence3DTransformTrack(FbxNode* FbxActor, IMovieScenePlayer* MovieScenePlayer, FMovieSceneSequenceIDRef InSequenceID, UMovieScene3DTransformTrack& TransformTrack, UObject* BoundObject, const TRange<FFrameNumber>& InPlaybackRange, const FMovieSceneSequenceTransform& RootToLocalTransform);
 
 	/**
-	 * Exports a level sequence 3D transform track that's getting interrogated(sample all sections) nto the FBX animation stack.
+	 * Exports a level sequence 3D transform track that's getting baked (sample all sections) onto the FBX animation stack.
 	 */
-	void ExportLevelSequenceInterrogated3DTransformTrack(FbxNode* FbxActor, IMovieScenePlayer* MovieScenePlayer, FMovieSceneSequenceIDRef InSequenceID, TArray<TWeakObjectPtr<UMovieScene3DTransformTrack> > TransformTracks, UObject* BoundObject, const TRange<FFrameNumber>& InPlaybackRange, const FMovieSceneSequenceTransform& RootToLocalTransform);
+	void ExportLevelSequenceBaked3DTransformTrack(IAnimTrackAdapter& AnimTrackAdapter, FbxNode* FbxActor, IMovieScenePlayer* MovieScenePlayer, FMovieSceneSequenceIDRef InSequenceID, TArray<TWeakObjectPtr<UMovieScene3DTransformTrack> > TransformTracks, UObject* BoundObject, const TRange<FFrameNumber>& InPlaybackRange, const FMovieSceneSequenceTransform& RootToLocalTransform);
 
 	/** 
 	 * Exports a level sequence property track into the FBX animation stack. 
