Fix missing UE_LOG argument.

diff --git a/Engine/Source/Editor/UnrealEd/Private/Commandlets/WorldPartitionBuilderCommandlet.cpp b/Engine/Source/Editor/UnrealEd/Private/Commandlets/WorldPartitionBuilderCommandlet.cpp
--- a/Engine/Source/Editor/UnrealEd/Private/Commandlets/WorldPartitionBuilderCommandlet.cpp
+++ b/Engine/Source/Editor/UnrealEd/Private/Commandlets/WorldPartitionBuilderCommandlet.cpp
@@ -75,7 +75,7 @@ int32 UWorldPartitionBuilderCommandlet::Main(const FString& Params)
 	UWorld* World = UWorld::FindWorldInPackage(MapPackage);
 	if (!World)
 	{
-		UE_LOG(LogWorldPartitionBuilderCommandlet, Error, TEXT("No world in specified package %s."));
+		UE_LOG(LogWorldPartitionBuilderCommandlet, Error, TEXT("No world in specified package %s."), *Tokens[0]);
 		return 1;
 	}
 
