Tentative fix for assert when DistillFileSetsCommandlet runs until a proper fix is found

diff --git a/Engine/Source/Editor/UnrealEd/Private/Commandlets/GenerateDistillFileSetsCommandlet.cpp b/Engine/Source/Editor/UnrealEd/Private/Commandlets/GenerateDistillFileSetsCommandlet.cpp
--- a/Engine/Source/Editor/UnrealEd/Private/Commandlets/GenerateDistillFileSetsCommandlet.cpp
+++ b/Engine/Source/Editor/UnrealEd/Private/Commandlets/GenerateDistillFileSetsCommandlet.cpp
@@ -297,12 +297,15 @@ int32 UGenerateDistillFileSetsCommandlet::Main( const FString& InParams )
 							FPackageName::TryConvertFilenameToLongPackageName(ExternalActorPackage, LongActorPackageName);
 							UPackage* ActorPackage = LoadPackage(nullptr, *LongActorPackageName, LOAD_None);
 
+							/* commenting this section for now has this seems to free BodySetup object too early while another thread is using it 
+							 * will certainly come back to it after figure out the right solution to the problem 
 							if (!(++ActorPackageIndex % 50))
 							{
 								GatherLoadedPackages();
 								UE_LOG(LogGenerateDistillFileSetsCommandlet, Display, TEXT( "Collecting garbage..." ) );
 								CollectGarbage(RF_NoFlags);
 							}
+							*/
 						}
 					}
 
