Add a profile tag in UEditorEngine::EndPlayMap.

diff --git a/Engine/Source/Editor/UnrealEd/Private/PlayLevel.cpp b/Engine/Source/Editor/UnrealEd/Private/PlayLevel.cpp
--- a/Engine/Source/Editor/UnrealEd/Private/PlayLevel.cpp
+++ b/Engine/Source/Editor/UnrealEd/Private/PlayLevel.cpp
@@ -173,6 +173,8 @@ void UEditorEngine::EndPlayMap()
 		return;
 	}
 
+	TRACE_CPUPROFILER_EVENT_SCOPE(UEditorEngine::EndPlayMap);
+
 	TGuardValue<bool> GuardIsEndingPlay(bIsEndingPlay, true);
 
 	FEditorDelegates::PrePIEEnded.Broadcast( bIsSimulatingInEditor );
