revert changes changelist 16184267
This quick fix is no longer necessary in this branch as 16030882 fix the problem properly

diff --git a/Engine/Source/Editor/UnrealEd/Private/Commandlets/GenerateDistillFileSetsCommandlet.cpp b/Engine/Source/Editor/UnrealEd/Private/Commandlets/GenerateDistillFileSetsCommandlet.cpp
--- a/Engine/Source/Editor/UnrealEd/Private/Commandlets/GenerateDistillFileSetsCommandlet.cpp
+++ b/Engine/Source/Editor/UnrealEd/Private/Commandlets/GenerateDistillFileSetsCommandlet.cpp
@@ -297,15 +297,12 @@ int32 UGenerateDistillFileSetsCommandlet::Main( const FString& InParams )
 							FPackageName::TryConvertFilenameToLongPackageName(ExternalActorPackage, LongActorPackageName);
 							UPackage* ActorPackage = LoadPackage(nullptr, *LongActorPackageName, LOAD_None);
 
-							/* commenting this section for now has this seems to free BodySetup object too early while another thread is using it 
-							 * will certainly come back to it after figure out the right solution to the problem 
 							if (!(++ActorPackageIndex % 50))
 							{
 								GatherLoadedPackages();
 								UE_LOG(LogGenerateDistillFileSetsCommandlet, Display, TEXT( "Collecting garbage..." ) );
 								CollectGarbage(RF_NoFlags);
 							}
-							*/
 						}
 					}
 
