Add dependency to the subobject data interface for the subobject editor. Add a new delegate to the Component Class combo picker to use the new subobject data interface instead.

#rnx

diff --git a/Engine/Source/Editor/UnrealEd/Private/SComponentClassCombo.cpp b/Engine/Source/Editor/UnrealEd/Private/SComponentClassCombo.cpp
--- a/Engine/Source/Editor/UnrealEd/Private/SComponentClassCombo.cpp
+++ b/Engine/Source/Editor/UnrealEd/Private/SComponentClassCombo.cpp
@@ -48,6 +48,7 @@ void SComponentClassCombo::Construct(const FArguments& InArgs)
 {
 	PrevSelectedIndex = INDEX_NONE;
 	OnComponentClassSelected = InArgs._OnComponentClassSelected;
+	OnSubobjectClassSelected = InArgs._OnSubobjectClassSelected;
 	TextFilter = MakeShared<FTextFilterExpressionEvaluator>(ETextFilterExpressionEvaluatorMode::BasicString);
 
 	FComponentTypeRegistry::Get().SubscribeToComponentList(ComponentClassList).AddRaw(this, &SComponentClassCombo::UpdateComponentClassList);
@@ -242,22 +243,36 @@ void SComponentClassCombo::OnAddComponentSelectionChanged( FComponentClassComboE
 			// Neither do we want the combo dropdown staying open once the user has clicked on a valid option
 			AddNewButton->SetIsMenuOpen(false, false);
 
-			if( OnComponentClassSelected.IsBound() )
+			if( OnComponentClassSelected.IsBound() || OnSubobjectClassSelected.IsBound() )
 			{
 				UClass* ComponentClass = InItem->GetComponentClass();
 				if (ComponentClass == nullptr)
 				{
 					// The class is not loaded yet, so load it:
 					const ELoadFlags LoadFlags = LOAD_None;
-					UBlueprint* LoadedObject = LoadObject<UBlueprint>(NULL, *InItem->GetComponentPath(), NULL, LoadFlags, NULL);
+					UBlueprint* LoadedObject = LoadObject<UBlueprint>(nullptr, *InItem->GetComponentPath(), nullptr, LoadFlags, nullptr);
 					ComponentClass = GetAuthoritativeBlueprintClass(LoadedObject);
 				}
 
-				UActorComponent* NewActorComponent = OnComponentClassSelected.Execute(ComponentClass, InItem->GetComponentCreateAction(), InItem->GetAssetOverride());
+				UActorComponent* NewActorComponent =
+					OnComponentClassSelected.IsBound() ?
+					OnComponentClassSelected.Execute(ComponentClass, InItem->GetComponentCreateAction(), InItem->GetAssetOverride())
+					: nullptr;
+				
+				FSubobjectDataHandle NewActorCompHandle =
+					OnSubobjectClassSelected.IsBound() ?
+					OnSubobjectClassSelected.Execute(ComponentClass, InItem->GetComponentCreateAction(), InItem->GetAssetOverride())
+					: FSubobjectDataHandle::InvalidHandle;
+				
 				if(NewActorComponent)
 				{
 					InItem->GetOnComponentCreated().ExecuteIfBound(NewActorComponent);
 				}
+				
+				if(NewActorCompHandle.IsValid())
+				{
+					InItem->GetOnSubobjectCreated().ExecuteIfBound(NewActorCompHandle);
+				}
 			}
 		}
 	}
diff --git a/Engine/Source/Editor/UnrealEd/Public/SComponentClassCombo.h b/Engine/Source/Editor/UnrealEd/Public/SComponentClassCombo.h
--- a/Engine/Source/Editor/UnrealEd/Public/SComponentClassCombo.h
+++ b/Engine/Source/Editor/UnrealEd/Public/SComponentClassCombo.h
@@ -6,6 +6,7 @@
 #include "SlateFwd.h"
 #include "Templates/SubclassOf.h"
 #include "Components/ActorComponent.h"
+#include "SubobjectData.h"
 #include "Widgets/DeclarativeSyntaxSupport.h"
 #include "Widgets/Input/SComboButton.h"
 #include "Widgets/Views/STableViewBase.h"
@@ -34,8 +35,10 @@ namespace EComponentCreateAction
 
 
 DECLARE_DELEGATE_OneParam(FOnComponentCreated, UActorComponent*);
+DECLARE_DELEGATE_OneParam(FOnSubobjectCreated, FSubobjectDataHandle);
 
 DECLARE_DELEGATE_RetVal_ThreeParams( UActorComponent*, FComponentClassSelected, TSubclassOf<UActorComponent>, EComponentCreateAction::Type, UObject*);
+DECLARE_DELEGATE_RetVal_ThreeParams( FSubobjectDataHandle, FSubobjectClassSelected, TSubclassOf<UActorComponent>, EComponentCreateAction::Type, UObject*);
 
 struct FComponentEntryCustomizationArgs
 {
@@ -45,6 +48,8 @@ struct FComponentEntryCustomizationArgs
 	FString ComponentNameOverride;
 	/** Callback when a new component is created */
 	FOnComponentCreated OnComponentCreated;
+	/** Callback when a new subobject is created */
+	FOnSubobjectCreated OnOnSubobjectCreated;
 	/** Brush icon to use instead of the class icon */
 	FName IconOverrideBrushName;
 	/** Custom sort priority to use (smaller means sorted first) */
@@ -138,6 +143,12 @@ public:
 	{
 		return CustomizationArgs.OnComponentCreated;
 	}
+	
+	FOnSubobjectCreated& GetOnSubobjectCreated()
+    {
+    	return CustomizationArgs.OnOnSubobjectCreated;
+    }
+	
 	FString GetClassName() const;
 	FString GetComponentPath() const { return ComponentPath.ToString(); }
 
@@ -176,6 +187,7 @@ public:
 
 		SLATE_ATTRIBUTE(bool, IncludeText)
 		SLATE_EVENT( FComponentClassSelected, OnComponentClassSelected )
+		SLATE_EVENT( FSubobjectClassSelected, OnSubobjectClassSelected )
 
 	SLATE_END_ARGS()
 
@@ -217,9 +229,11 @@ private:
 	FText GetFriendlyComponentName(FComponentClassComboEntryPtr Entry) const;
 
 	TSharedRef<SToolTip> GetComponentToolTip(FComponentClassComboEntryPtr Entry) const;
-
+	
 	FComponentClassSelected OnComponentClassSelected;
 
+	FSubobjectClassSelected OnSubobjectClassSelected;
+	
 	/** List of component class names used by combo box */
 	TArray<FComponentClassComboEntryPtr>* ComponentClassList;
 
diff --git a/Engine/Source/Editor/UnrealEd/UnrealEd.Build.cs b/Engine/Source/Editor/UnrealEd/UnrealEd.Build.cs
--- a/Engine/Source/Editor/UnrealEd/UnrealEd.Build.cs
+++ b/Engine/Source/Editor/UnrealEd/UnrealEd.Build.cs
@@ -113,6 +113,7 @@ public class UnrealEd : ModuleRules
 				"InterchangeCore",
 				"InterchangeEngine",
 				"DeveloperToolSettings",
+				"SubobjectDataInterface",
 			}
 		);
 
