Fixed issue running the WP HLOD Builder with -DeleteHLODs when always loaded HLOD actors are used
Missing ActorDesc when not running setup & not running a distributed build shouldn't be an error

diff --git a/Engine/Source/Editor/UnrealEd/Private/WorldPartition/WorldPartitionHLODsBuilder.cpp b/Engine/Source/Editor/UnrealEd/Private/WorldPartition/WorldPartitionHLODsBuilder.cpp
--- a/Engine/Source/Editor/UnrealEd/Private/WorldPartition/WorldPartitionHLODsBuilder.cpp
+++ b/Engine/Source/Editor/UnrealEd/Private/WorldPartition/WorldPartitionHLODsBuilder.cpp
@@ -465,12 +465,23 @@ bool UWorldPartitionHLODsBuilder::DeleteHLODActors()
 {
 	UE_LOG(LogWorldPartitionHLODsBuilder, Display, TEXT("#### Deleting HLOD actors ####"));
 
-	int32 NumDeleted = 0;
+	TArray<FString> PackagesToDelete;
 	for (UActorDescContainer::TIterator<AWorldPartitionHLOD> HLODIterator(WorldPartition); HLODIterator; ++HLODIterator)
 	{
 		FWorldPartitionActorDesc* HLODActorDesc = *HLODIterator;
 		FString PackageName = HLODActorDesc->GetActorPackage().ToString();
+		PackagesToDelete.Add(PackageName);
+	}
+
+	// Ensure we don't hold on to packages of always loaded actors
+	// When running distributed builds, we wanna leave the machine clean, so added files are deleted, checked out files are reverted
+	// and deleted files are restored.
+	WorldPartition->Uninitialize();
+	DoCollectGarbage();
 
+	int32 NumDeleted = 0;
+	for (const FString& PackageName : PackagesToDelete)
+	{
 		bool bDeleted = SourceControlHelper->Delete(PackageName);
 		if (bDeleted)
 		{
@@ -580,9 +591,7 @@ TArray<TArray<FGuid>> UWorldPartitionHLODsBuilder::GetHLODWorldloads(int32 NumWo
 		for (const FGuid& SubActorGuid : HLODIterator->GetSubActors())
 		{
 			FWorldPartitionActorDesc* SubActorDesc = WorldPartition->GetActorDesc(SubActorGuid);
-			check(SubActorDesc);
-
-			if (SubActorDesc->GetActorClass()->IsChildOf<AWorldPartitionHLOD>())
+			if (SubActorDesc && SubActorDesc->GetActorClass()->IsChildOf<AWorldPartitionHLOD>())
 			{
 				ChildHLODs.Add(SubActorGuid);
 			}
@@ -674,13 +683,11 @@ bool UWorldPartitionHLODsBuilder::ValidateWorkload(const TArray<FGuid>&Workload)
 		for (const FGuid& SubActorGuid : HLODActorDesc->GetSubActors())
 		{
 			FWorldPartitionActorDesc* SubActorDesc = WorldPartition->GetActorDesc(SubActorGuid);
-			if (!SubActorDesc)
-			{
-				UE_LOG(LogWorldPartitionHLODsBuilder, Error, TEXT("Unknown sub actor guid found, your HLOD actors are probably out of date. Run with -SetupHLODs to fix this. Exiting..."));
-				return false;
-			}
 
-			if (SubActorDesc->GetActorClass()->IsChildOf<AWorldPartitionHLOD>())
+			// Invalid sub actor guid found, this is unexpected when running distributed builds as the build step is always preceeded by the setup step.
+			check(SubActorDesc || !bDistributedBuild);
+
+			if (SubActorDesc && SubActorDesc->GetActorClass()->IsChildOf<AWorldPartitionHLOD>())
 			{
 				if(!ProcessedHLOD.Contains(SubActorGuid))
 				{
