Avoid access to the global mode manager in UEditorEngine::ShouldDrawBrushWireframe() if running from a cmdlet

diff --git a/Engine/Source/Editor/UnrealEd/Private/EditorEngine.cpp b/Engine/Source/Editor/UnrealEd/Private/EditorEngine.cpp
--- a/Engine/Source/Editor/UnrealEd/Private/EditorEngine.cpp
+++ b/Engine/Source/Editor/UnrealEd/Private/EditorEngine.cpp
@@ -895,11 +895,7 @@ void UEditorEngine::InitializeObjectReferences()
 
 bool UEditorEngine::ShouldDrawBrushWireframe( AActor* InActor )
 {
-	bool bResult = true;
-
-	bResult = GLevelEditorModeTools().ShouldDrawBrushWireframe( InActor );
-	
-	return bResult;
+	return !IsRunningCommandlet() ? GLevelEditorModeTools().ShouldDrawBrushWireframe(InActor) : false;
 }
 
 //
