Override the component selection in FPersonaEditorModeManager when we set a preview scene in Persona. This allows the tools to pick up the currently editing component from the skelmesh editor.
Note that this is a temporary band-aid until the modeling tools start using the new UContextObjectStore mechanism as means to pass components from editor to the tools context.

#rnx

diff --git a/Engine/Source/Editor/Persona/Private/PersonaEditorModeManager.cpp b/Engine/Source/Editor/Persona/Private/PersonaEditorModeManager.cpp
--- a/Engine/Source/Editor/Persona/Private/PersonaEditorModeManager.cpp
+++ b/Engine/Source/Editor/Persona/Private/PersonaEditorModeManager.cpp
@@ -2,6 +2,10 @@
 
 #include "PersonaEditorModeManager.h"
 #include "IPersonaEditMode.h"
+#include "IPersonaPreviewScene.h"
+#include "Selection.h"
+#include "Animation/DebugSkelMeshComponent.h"
+
 
 bool FPersonaEditorModeManager::GetCameraTarget(FSphere& OutTarget) const
 {
@@ -33,3 +37,19 @@ void FPersonaEditorModeManager::GetOnScreenDebugInfo(TArray<FText>& OutDebugText
 		}
 	}
 }
+
+
+void FPersonaEditorModeManager::SetPreviewScene(FPreviewScene* NewPreviewScene)
+{
+	const IPersonaPreviewScene *PersonaPreviewScene = static_cast<const IPersonaPreviewScene *>(NewPreviewScene);
+
+	if (PersonaPreviewScene && PersonaPreviewScene->GetPreviewMeshComponent())
+	{
+		ComponentSet->BeginBatchSelectOperation();
+		ComponentSet->DeselectAll();
+		ComponentSet->Select(PersonaPreviewScene->GetPreviewMeshComponent(), true);
+		ComponentSet->EndBatchSelectOperation();
+	}
+
+	FAssetEditorModeManager::SetPreviewScene(NewPreviewScene);
+}
diff --git a/Engine/Source/Editor/Persona/Private/PersonaEditorModeManager.h b/Engine/Source/Editor/Persona/Private/PersonaEditorModeManager.h
--- a/Engine/Source/Editor/Persona/Private/PersonaEditorModeManager.h
+++ b/Engine/Source/Editor/Persona/Private/PersonaEditorModeManager.h
@@ -12,4 +12,7 @@ public:
 	/** IPersonaEditorModeManager interface  */
 	virtual bool GetCameraTarget(FSphere& OutTarget) const override;
 	virtual void GetOnScreenDebugInfo(TArray<FText>& OutDebugText) const override;
+
+	// FAssetEditorModeManager overrides
+	virtual void SetPreviewScene(FPreviewScene* NewPreviewScene) override;
 };
diff --git a/Engine/Source/Editor/UnrealEd/Private/AssetEditorModeManager.cpp b/Engine/Source/Editor/UnrealEd/Private/AssetEditorModeManager.cpp
--- a/Engine/Source/Editor/UnrealEd/Private/AssetEditorModeManager.cpp
+++ b/Engine/Source/Editor/UnrealEd/Private/AssetEditorModeManager.cpp
@@ -29,8 +29,6 @@ FAssetEditorModeManager::FAssetEditorModeManager()
 
 FAssetEditorModeManager::~FAssetEditorModeManager()
 {
-	SetPreviewScene(nullptr);
-
 	// We may be destroyed after the UObject system has already shutdown, 
 	// which would mean that these instances will be garbage
 	if (UObjectInitialized())
@@ -45,10 +43,7 @@ FAssetEditorModeManager::~FAssetEditorModeManager()
 		}
 
 		ActorSet->RemoveFromRoot();
-		ActorSet = nullptr;
-
 		ComponentSet->RemoveFromRoot();
-		ComponentSet = nullptr;
 
 		if (!ObjectSet->HasAnyFlags(RF_BeginDestroyed | RF_FinishDestroyed))
 		{
@@ -59,7 +54,6 @@ FAssetEditorModeManager::~FAssetEditorModeManager()
 		}
 
 		ObjectSet->RemoveFromRoot();
-		ObjectSet = nullptr;
 	}
 }
 
@@ -83,7 +77,7 @@ UWorld* FAssetEditorModeManager::GetWorld() const
 	return (PreviewScene != nullptr) ? PreviewScene->GetWorld() : GEditor->GetEditorWorldContext().World();
 }
 
-void FAssetEditorModeManager::SetPreviewScene(class FPreviewScene* NewPreviewScene)
+void FAssetEditorModeManager::SetPreviewScene(FPreviewScene* NewPreviewScene)
 {
 	PreviewScene = NewPreviewScene;
 }
diff --git a/Engine/Source/Editor/UnrealEd/Public/AssetEditorModeManager.h b/Engine/Source/Editor/UnrealEd/Public/AssetEditorModeManager.h
--- a/Engine/Source/Editor/UnrealEd/Public/AssetEditorModeManager.h
+++ b/Engine/Source/Editor/UnrealEd/Public/AssetEditorModeManager.h
@@ -6,6 +6,7 @@
 #include "EditorModeManager.h"
 
 class FPreviewScene;
+class UTypedElementList;
 
 //////////////////////////////////////////////////////////////////////////
 // FAssetEditorModeManager
@@ -14,21 +15,21 @@ class UNREALED_API FAssetEditorModeManager : public FEditorModeTools
 {
 public:
 	FAssetEditorModeManager();
-	virtual ~FAssetEditorModeManager();
+	virtual ~FAssetEditorModeManager() override;
 
 	// FEditorModeTools interface
-	virtual class USelection* GetSelectedActors() const override;
-	virtual class USelection* GetSelectedObjects() const override;
-	virtual class USelection* GetSelectedComponents() const override;
+	virtual USelection* GetSelectedActors() const override;
+	virtual USelection* GetSelectedObjects() const override;
+	virtual USelection* GetSelectedComponents() const override;
 	virtual UWorld* GetWorld() const override;
 	// End of FEditorModeTools interface
 
-	void SetPreviewScene(class FPreviewScene* NewPreviewScene);
+	virtual void SetPreviewScene(FPreviewScene* NewPreviewScene);
 	FPreviewScene* GetPreviewScene() const;
 
 protected:
-	class USelection* ActorSet = nullptr;
-	class USelection* ObjectSet = nullptr;
-	class USelection* ComponentSet = nullptr;
-	class FPreviewScene* PreviewScene = nullptr;
+	USelection* ActorSet = nullptr;
+	USelection* ObjectSet = nullptr;
+	USelection* ComponentSet = nullptr;
+	FPreviewScene* PreviewScene = nullptr;
 };
