Fixed crash when dropping a BP actor with a nested water body actor in a map with a landscape but no water brush and/or water mesh actor: in that case, the addition of the actor will trigger a UWorld::SpawnActor to create the missing actors, which is forbidden during a construction script :
* Changed the API of the actor factory's CreateActor/SpawnActor in order to take a FActorSpawnParameters in parameter, so as to expose UWorld::SpawnActor's own option to the factory user.
* The level pointer is still part of the function prototypes in order to make it explicit to the user that he needs to pass in a valid ULevel (rather than relying on FActorSpawnParameters's OverrideLevel, which is easy to miss). In the case where OverrideLevel is indeed passed within FActorSpawnParameters, the functions validate that the pointer is the same as the level param
* Also, removed automatic creation of water mesh actor on map load : this is not playing well with world partition because it will keep on creating new actors if the level containing the water mesh actor is not loaded

#tests editor + full build preflight

diff --git a/Engine/Plugins/Enterprise/DatasmithImporter/Source/DatasmithImporter/Private/ActorFactoryDatasmithScene.cpp b/Engine/Plugins/Enterprise/DatasmithImporter/Source/DatasmithImporter/Private/ActorFactoryDatasmithScene.cpp
--- a/Engine/Plugins/Enterprise/DatasmithImporter/Source/DatasmithImporter/Private/ActorFactoryDatasmithScene.cpp
+++ b/Engine/Plugins/Enterprise/DatasmithImporter/Source/DatasmithImporter/Private/ActorFactoryDatasmithScene.cpp
@@ -145,7 +145,7 @@ bool UActorFactoryDatasmithScene::CanCreateActorFrom(const FAssetData& AssetData
 	return true;
 }
 
-AActor* UActorFactoryDatasmithScene::SpawnActor( UObject* Asset, ULevel* InLevel, const FTransform& Transform, EObjectFlags InObjectFlags, const FName Name )
+AActor* UActorFactoryDatasmithScene::SpawnActor( UObject* Asset, ULevel* InLevel, const FTransform& Transform, const FActorSpawnParameters& InSpawnParams)
 {
 	UDatasmithScene* DatasmithScene = Cast< UDatasmithScene >( Asset );
 
@@ -156,13 +156,13 @@ AActor* UActorFactoryDatasmithScene::SpawnActor( UObject* Asset, ULevel* InLevel
 
 	AActor* ResultingActor = nullptr;
 
-	if ( InObjectFlags & RF_Transient )
+	if (InSpawnParams.ObjectFlags & RF_Transient)
 	{
 		// This is a hack for drag and drop so that we don't spawn all the actors for the preview actor since it gets deleted right after.
-		FActorSpawnParameters SpawnParameters;
-		SpawnParameters.ObjectFlags = InObjectFlags;
 
-		AStaticMeshActor* DragActor = Cast< AStaticMeshActor >( InLevel->GetWorld()->SpawnActor( AStaticMeshActor::StaticClass(), &Transform, SpawnParameters ) );
+		FActorSpawnParameters SpawnInfo(InSpawnParams);
+		SpawnInfo.OverrideLevel = InLevel;
+		AStaticMeshActor* DragActor = Cast< AStaticMeshActor >( InLevel->GetWorld()->SpawnActor( AStaticMeshActor::StaticClass(), &Transform, SpawnInfo) );
 		DragActor->GetStaticMeshComponent()->SetStaticMesh( Cast< UStaticMesh >( FSoftObjectPath( TEXT("StaticMesh'/Engine/EditorMeshes/EditorSphere.EditorSphere'") ).TryLoad() ) );
 		DragActor->SetActorScale3D( FVector( 0.1f ) );
 
@@ -170,7 +170,7 @@ AActor* UActorFactoryDatasmithScene::SpawnActor( UObject* Asset, ULevel* InLevel
 	}
 	else
 	{
-		ResultingActor = UActorFactoryDatasmithSceneImpl::ImportActors( DatasmithScene, nullptr, Transform, InObjectFlags, false );
+		ResultingActor = UActorFactoryDatasmithSceneImpl::ImportActors( DatasmithScene, nullptr, Transform, InSpawnParams.ObjectFlags, false );
 	}
 
 	return ResultingActor;
diff --git a/Engine/Plugins/Enterprise/DatasmithImporter/Source/DatasmithImporter/Public/ActorFactoryDatasmithScene.h b/Engine/Plugins/Enterprise/DatasmithImporter/Source/DatasmithImporter/Public/ActorFactoryDatasmithScene.h
--- a/Engine/Plugins/Enterprise/DatasmithImporter/Source/DatasmithImporter/Public/ActorFactoryDatasmithScene.h
+++ b/Engine/Plugins/Enterprise/DatasmithImporter/Source/DatasmithImporter/Public/ActorFactoryDatasmithScene.h
@@ -6,6 +6,8 @@
 
 #include "ActorFactoryDatasmithScene.generated.h"
 
+struct FActorSpawnParameters;
+
 UCLASS(MinimalAPI)
 class UActorFactoryDatasmithScene : public UActorFactory
 {
@@ -18,6 +20,6 @@ public:
 	virtual bool CanCreateActorFrom( const FAssetData& AssetData, FText& OutErrorMsg ) override;
 
 protected:
-	virtual AActor* SpawnActor( UObject* Asset, ULevel* InLevel, const FTransform& Transform, EObjectFlags InObjectFlags, const FName Name ) override;
+	virtual AActor* SpawnActor( UObject* Asset, ULevel* InLevel, const FTransform& Transform, const FActorSpawnParameters& InSpawnParams) override;
 	//~ End UActorFactory Interface
 };
\ No newline at end of file
diff --git a/Engine/Plugins/Enterprise/VariantManagerContent/Source/VariantManagerContentEditor/Private/LevelVariantSetsActorFactory.cpp b/Engine/Plugins/Enterprise/VariantManagerContent/Source/VariantManagerContentEditor/Private/LevelVariantSetsActorFactory.cpp
--- a/Engine/Plugins/Enterprise/VariantManagerContent/Source/VariantManagerContentEditor/Private/LevelVariantSetsActorFactory.cpp
+++ b/Engine/Plugins/Enterprise/VariantManagerContent/Source/VariantManagerContentEditor/Private/LevelVariantSetsActorFactory.cpp
@@ -33,13 +33,13 @@ bool ULevelVariantSetsActorFactory::CanCreateActorFrom( const FAssetData& AssetD
 	return true;
 }
 
-AActor* ULevelVariantSetsActorFactory::SpawnActor( UObject* Asset, ULevel* InLevel, const FTransform& Transform, EObjectFlags InObjectFlags, const FName Name )
+AActor* ULevelVariantSetsActorFactory::SpawnActor(UObject* InAsset, ULevel* InLevel, const FTransform& InTransform, const FActorSpawnParameters& InSpawnParams)
 {
-	ALevelVariantSetsActor* NewActor = Cast<ALevelVariantSetsActor>(Super::SpawnActor(Asset, InLevel, Transform, InObjectFlags, Name));
+	ALevelVariantSetsActor* NewActor = Cast<ALevelVariantSetsActor>(Super::SpawnActor(InAsset, InLevel, InTransform, InSpawnParams));
 
 	if (NewActor)
 	{
-		if (ULevelVariantSets* LevelVariantSets = Cast<ULevelVariantSets>(Asset))
+		if (ULevelVariantSets* LevelVariantSets = Cast<ULevelVariantSets>(InAsset))
 		{
 			NewActor->SetLevelVariantSets(LevelVariantSets);
 		}
diff --git a/Engine/Plugins/Enterprise/VariantManagerContent/Source/VariantManagerContentEditor/Public/LevelVariantSetsActorFactory.h b/Engine/Plugins/Enterprise/VariantManagerContent/Source/VariantManagerContentEditor/Public/LevelVariantSetsActorFactory.h
--- a/Engine/Plugins/Enterprise/VariantManagerContent/Source/VariantManagerContentEditor/Public/LevelVariantSetsActorFactory.h
+++ b/Engine/Plugins/Enterprise/VariantManagerContent/Source/VariantManagerContentEditor/Public/LevelVariantSetsActorFactory.h
@@ -9,6 +9,7 @@
 #include "LevelVariantSetsActorFactory.generated.h"
 
 class AActor;
+struct FActorSpawnParameters;
 struct FAssetData;
 class ULevel;
 
@@ -18,7 +19,7 @@ class ULevelVariantSetsActorFactory : public UActorFactory
 	GENERATED_UCLASS_BODY()
 
 	// Begin UActorFactory Interface
-	virtual AActor* SpawnActor( UObject* Asset, ULevel* InLevel, const FTransform& Transform, EObjectFlags ObjectFlags, const FName Name ) override;
+	virtual AActor* SpawnActor(UObject* InAsset, ULevel* InLevel, const FTransform& InTransform, const FActorSpawnParameters& InSpawnParams) override;
 	virtual bool CanCreateActorFrom( const FAssetData& AssetData, FText& OutErrorMsg ) override;
 	virtual UObject* GetAssetFromActorInstance(AActor* ActorInstance) override;
 	// End UActorFactory Interface
diff --git a/Engine/Plugins/Experimental/Water/Source/Editor/Private/WaterEditorModule.cpp b/Engine/Plugins/Experimental/Water/Source/Editor/Private/WaterEditorModule.cpp
--- a/Engine/Plugins/Experimental/Water/Source/Editor/Private/WaterEditorModule.cpp
+++ b/Engine/Plugins/Experimental/Water/Source/Editor/Private/WaterEditorModule.cpp
@@ -63,7 +63,6 @@ void FWaterEditorModule::StartupModule()
 
 	GEngine->OnLevelActorAdded().AddRaw(this, &FWaterEditorModule::OnLevelActorAddedToWorld);
 
-	FEditorDelegates::OnMapOpened.AddRaw(this, &FWaterEditorModule::OnMapLoaded);
 	RegisterComponentVisualizer(UWaterSplineComponent::StaticClass()->GetFName(), MakeShareable(new FWaterSplineComponentVisualizer));
 
 	if (GEditor)
@@ -165,9 +164,9 @@ void FWaterEditorModule::OnLevelActorAddedToWorld(AActor* Actor)
 						FName BrushActorName = MakeUniqueObjectName(Landscape->GetOuter(), WaterBrushClassPtr, FName(BrushActorString));
 						FActorSpawnParameters SpawnParams;
 						SpawnParams.Name = BrushActorName;
-						SpawnParams.OverrideLevel = Landscape->GetLevel();
+						SpawnParams.bAllowDuringConstructionScript = true; // This can be called by construction script if the actor being added to the world is part of a blueprint, for example : 
 						AWaterLandscapeBrush* NewBrush = (WaterBrushActorFactory != nullptr) 
-							? CastChecked<AWaterLandscapeBrush>(WaterBrushActorFactory->CreateActor(ActorWorld, SpawnParams.OverrideLevel, FTransform::Identity, RF_Transactional, BrushActorName))
+							? Cast<AWaterLandscapeBrush>(WaterBrushActorFactory->CreateActor(ActorWorld, Landscape->GetLevel(), FTransform::Identity, SpawnParams))
 							: ActorWorld->SpawnActor<AWaterLandscapeBrush>(WaterBrushClassPtr, SpawnParams);
 						if (NewBrush)
 						{
@@ -190,39 +189,13 @@ void FWaterEditorModule::OnLevelActorAddedToWorld(AActor* Actor)
 			{
 				FActorSpawnParameters SpawnParams;
 				SpawnParams.OverrideLevel = ActorWorld->PersistentLevel;
+				SpawnParams.bAllowDuringConstructionScript = true; // This can be called by construction script if the actor being added to the world is part of a blueprint, for example : 
 				ActorWorld->SpawnActor<AWaterMeshActor>(AWaterMeshActor::StaticClass(), SpawnParams);
 			}
 		}
 	}
 }
 
-void FWaterEditorModule::OnMapLoaded(const FString& /* Filename */, bool /*bAsTemplate*/)
-{
-	UWorld* EditorWorld = GWorld;
-	TActorIterator<AWaterLandscapeBrush> BrushIterator(EditorWorld);
-
-	AWaterLandscapeBrush* WaterManagerActor = nullptr;
-	if (BrushIterator)
-	{
-		WaterManagerActor = *BrushIterator;
-	}
-
-	if (WaterManagerActor)
-	{
-		const bool bHasMeshActor = !!TActorIterator<AWaterMeshActor>(EditorWorld);
-
-		if (!bHasMeshActor)
-		{
-			FActorSpawnParameters SpawnParams;
-			SpawnParams.OverrideLevel = WaterManagerActor->GetLevel();
-			EditorWorld->SpawnActor<AWaterMeshActor>(AWaterMeshActor::StaticClass(), SpawnParams);
-
-			// If the texture is newly created force an update to get textures into the water mesh actor
-			WaterManagerActor->ForceWaterTextureUpdate();
-		}
-	}
-}
-
 IMPLEMENT_MODULE(FWaterEditorModule, WaterEditor);
 
 #undef LOCTEXT_NAMESPACE
\ No newline at end of file
diff --git a/Engine/Plugins/Experimental/Water/Source/Editor/Public/WaterEditorModule.h b/Engine/Plugins/Experimental/Water/Source/Editor/Public/WaterEditorModule.h
--- a/Engine/Plugins/Experimental/Water/Source/Editor/Public/WaterEditorModule.h
+++ b/Engine/Plugins/Experimental/Water/Source/Editor/Public/WaterEditorModule.h
@@ -36,8 +36,6 @@ private:
 
 	void OnLevelActorAddedToWorld(AActor* Actor);
 
-	void OnMapLoaded(const FString& /* Filename */, bool /*bAsTemplate*/);
-
 private:
 	/** Array of component class names we have registered, so we know what to unregister afterwards */
 	TArray<FName> RegisteredComponentClassNames;
diff --git a/Engine/Plugins/Importers/USDImporter/Source/USDImporter/Private/USDPrimResolver.cpp b/Engine/Plugins/Importers/USDImporter/Source/USDImporter/Private/USDPrimResolver.cpp
--- a/Engine/Plugins/Importers/USDImporter/Source/USDImporter/Private/USDPrimResolver.cpp
+++ b/Engine/Plugins/Importers/USDImporter/Source/USDImporter/Private/USDPrimResolver.cpp
@@ -296,7 +296,7 @@ AActor* UDEPRECATED_UUSDPrimResolver::SpawnActor(FUSDSceneImportContext& ImportC
 			else
 			{
 				// We'll set the name later by using the ActorLabels cache
-				SpawnedActor = ActorFactory->CreateActor(ActorAsset, ImportContext.World->GetCurrentLevel(), FTransform::Identity, RF_Transactional);
+				SpawnedActor = ActorFactory->CreateActor(ActorAsset, ImportContext.World->GetCurrentLevel(), FTransform::Identity);
 
 				// For empty group actors set their initial mobility to static
 				if ( ActorFactory == ImportContext.EmptyActorFactory )
diff --git a/Engine/Plugins/MovieScene/LevelSequenceEditor/Source/LevelSequenceEditor/Private/Misc/LevelSequenceEditorActorSpawner.cpp b/Engine/Plugins/MovieScene/LevelSequenceEditor/Source/LevelSequenceEditor/Private/Misc/LevelSequenceEditorActorSpawner.cpp
--- a/Engine/Plugins/MovieScene/LevelSequenceEditor/Source/LevelSequenceEditor/Private/Misc/LevelSequenceEditorActorSpawner.cpp
+++ b/Engine/Plugins/MovieScene/LevelSequenceEditor/Source/LevelSequenceEditor/Private/Misc/LevelSequenceEditorActorSpawner.cpp
@@ -111,7 +111,11 @@ TValueOrError<FNewSpawnable, FText> FLevelSequenceEditorActorSpawner::CreateNewS
 			{
 				const FName ActorName = MakeUniqueObjectName(World->PersistentLevel, FactoryToUse->NewActorClass->StaticClass(), TemplateName);
 
-				AActor* Instance = FactoryToUse->CreateActor(&SourceObject, World->PersistentLevel, FTransform(), RF_Transient | RF_Transactional, ActorName );
+				FActorSpawnParameters SpawnParams;
+				SpawnParams.ObjectFlags = RF_Transient | RF_Transactional;
+				SpawnParams.Name = ActorName;
+
+				AActor* Instance = FactoryToUse->CreateActor(&SourceObject, World->PersistentLevel, FTransform(), SpawnParams);
 				Instance->bIsEditorPreviewActor = false;
 				NewSpawnable.ObjectTemplate = StaticDuplicateObject(Instance, &OwnerMovieScene, TemplateName, RF_AllFlags & ~RF_Transient);
 
diff --git a/Engine/Plugins/VirtualProduction/RemoteControl/Source/RemoteControlUI/Private/Factories/RemoteControlPresetActorFactory.cpp b/Engine/Plugins/VirtualProduction/RemoteControl/Source/RemoteControlUI/Private/Factories/RemoteControlPresetActorFactory.cpp
--- a/Engine/Plugins/VirtualProduction/RemoteControl/Source/RemoteControlUI/Private/Factories/RemoteControlPresetActorFactory.cpp
+++ b/Engine/Plugins/VirtualProduction/RemoteControl/Source/RemoteControlUI/Private/Factories/RemoteControlPresetActorFactory.cpp
@@ -27,13 +27,13 @@ bool URemoteControlPresetActorFactory::CanCreateActorFrom( const FAssetData& Ass
 	return true;
 }
 
-AActor* URemoteControlPresetActorFactory::SpawnActor( UObject* Asset, ULevel* InLevel, const FTransform& Transform, EObjectFlags InObjectFlags, const FName Name )
+AActor* URemoteControlPresetActorFactory::SpawnActor(UObject* InAsset, ULevel* InLevel, const FTransform& InTransform, const FActorSpawnParameters& InSpawnParams)
 {
-	ARemoteControlPresetActor* NewActor = Cast<ARemoteControlPresetActor>(Super::SpawnActor(Asset, InLevel, Transform, InObjectFlags, Name));
+	ARemoteControlPresetActor* NewActor = Cast<ARemoteControlPresetActor>(Super::SpawnActor(InAsset, InLevel, InTransform, InSpawnParams));
 
 	if (NewActor)
 	{
-		if (URemoteControlPreset* Preset = Cast<URemoteControlPreset>(Asset))
+		if (URemoteControlPreset* Preset = Cast<URemoteControlPreset>(InAsset))
 		{
 			NewActor->Preset = Preset;
 		}
diff --git a/Engine/Plugins/VirtualProduction/RemoteControl/Source/RemoteControlUI/Private/Factories/RemoteControlPresetActorFactory.h b/Engine/Plugins/VirtualProduction/RemoteControl/Source/RemoteControlUI/Private/Factories/RemoteControlPresetActorFactory.h
--- a/Engine/Plugins/VirtualProduction/RemoteControl/Source/RemoteControlUI/Private/Factories/RemoteControlPresetActorFactory.h
+++ b/Engine/Plugins/VirtualProduction/RemoteControl/Source/RemoteControlUI/Private/Factories/RemoteControlPresetActorFactory.h
@@ -6,6 +6,7 @@
 #include "RemoteControlPresetActorFactory.generated.h"
 
 class AActor;
+struct FActorSpawnParameters;
 struct FAssetData;
 class ULevel;
 class SNotificationItem;
@@ -19,7 +20,7 @@ public:
 	URemoteControlPresetActorFactory();
 
 	//~ Begin UActorFactory Interface
-	virtual AActor* SpawnActor( UObject* Asset, ULevel* InLevel, const FTransform& Transform, EObjectFlags ObjectFlags, const FName Name ) override;
+	virtual AActor* SpawnActor(UObject* InAsset, ULevel* InLevel, const FTransform& InTransform, const FActorSpawnParameters& InSpawnParams) override;
 	virtual bool CanCreateActorFrom( const FAssetData& AssetData, FText& OutErrorMsg ) override;
 	virtual UObject* GetAssetFromActorInstance(AActor* ActorInstance) override;
 	//~ End UActorFactory Interface
diff --git a/Engine/Source/Editor/LandscapeEditor/Private/ActorFactoryLandscape.cpp b/Engine/Source/Editor/LandscapeEditor/Private/ActorFactoryLandscape.cpp
--- a/Engine/Source/Editor/LandscapeEditor/Private/ActorFactoryLandscape.cpp
+++ b/Engine/Source/Editor/LandscapeEditor/Private/ActorFactoryLandscape.cpp
@@ -21,7 +21,7 @@ UActorFactoryLandscape::UActorFactoryLandscape(const FObjectInitializer& ObjectI
 	NewActorClass = ALandscapeProxy::StaticClass();
 }
 
-AActor* UActorFactoryLandscape::SpawnActor(UObject* Asset, ULevel* InLevel, const FTransform& Transform, EObjectFlags InObjectFlags, const FName Name)
+AActor* UActorFactoryLandscape::SpawnActor(UObject* Asset, ULevel* InLevel, const FTransform& Transform, const FActorSpawnParameters& InSpawnParams)
 {
 	GLevelEditorModeTools().ActivateMode(FBuiltinEditorModes::EM_Landscape);
 
@@ -32,11 +32,15 @@ AActor* UActorFactoryLandscape::SpawnActor(UObject* Asset, ULevel* InLevel, cons
 
 	EdMode->SetCurrentTool("NewLandscape");
 
-	FActorSpawnParameters SpawnInfo;
-	SpawnInfo.OverrideLevel = InLevel;
-	SpawnInfo.ObjectFlags = InObjectFlags;
-	SpawnInfo.Name = Name;
-	return InLevel->OwningWorld->SpawnActor(ALandscapePlaceholder::StaticClass(), &Transform, SpawnInfo);
+	ULevel* LocalLevel = ValidateSpawnActorLevel(InLevel, InSpawnParams);
+	if (LocalLevel != nullptr)
+	{
+		FActorSpawnParameters SpawnInfo(InSpawnParams);
+		SpawnInfo.OverrideLevel = LocalLevel;
+		return LocalLevel->OwningWorld->SpawnActor(ALandscapePlaceholder::StaticClass(), &Transform, SpawnInfo);
+	}
+
+	return nullptr;
 }
 
 ALandscapePlaceholder::ALandscapePlaceholder(const FObjectInitializer& ObjectInitializer)
diff --git a/Engine/Source/Editor/LandscapeEditor/Private/Classes/ActorFactoryLandscape.h b/Engine/Source/Editor/LandscapeEditor/Private/Classes/ActorFactoryLandscape.h
--- a/Engine/Source/Editor/LandscapeEditor/Private/Classes/ActorFactoryLandscape.h
+++ b/Engine/Source/Editor/LandscapeEditor/Private/Classes/ActorFactoryLandscape.h
@@ -9,6 +9,7 @@
 
 class AActor;
 class ULevel;
+struct FActorSpawnParameters;
 
 UCLASS(MinimalAPI, config=Editor)
 class UActorFactoryLandscape : public UActorFactory
@@ -16,6 +17,6 @@ class UActorFactoryLandscape : public UActorFactory
 	GENERATED_UCLASS_BODY()
 
 	//~ Begin UActorFactory Interface
-	virtual AActor* SpawnActor( UObject* Asset, ULevel* InLevel, const FTransform& Transform, EObjectFlags ObjectFlags, const FName Name ) override;
+	virtual AActor* SpawnActor( UObject* Asset, ULevel* InLevel, const FTransform& Transform, const FActorSpawnParameters& InSpawnParams) override;
 	//~ End UActorFactory Interface
 };
diff --git a/Engine/Source/Editor/LandscapeEditor/Private/LandscapeEdModeBPCustomTools.cpp b/Engine/Source/Editor/LandscapeEditor/Private/LandscapeEdModeBPCustomTools.cpp
--- a/Engine/Source/Editor/LandscapeEditor/Private/LandscapeEdModeBPCustomTools.cpp
+++ b/Engine/Source/Editor/LandscapeEditor/Private/LandscapeEdModeBPCustomTools.cpp
@@ -111,7 +111,7 @@ public:
 
 			UWorld* ActorWorld = ViewportClient->GetWorld();
 			ALandscapeBlueprintBrushBase* Brush = (BrushActorFactory != nullptr)
-				? CastChecked<ALandscapeBlueprintBrushBase>(BrushActorFactory->CreateActor(ActorWorld, SpawnInfo.OverrideLevel, FTransform(SpawnLocation), RF_Transactional, BrushActorName))
+				? CastChecked<ALandscapeBlueprintBrushBase>(BrushActorFactory->CreateActor(ActorWorld, SpawnInfo.OverrideLevel, FTransform(SpawnLocation), SpawnInfo))
 				: ActorWorld->SpawnActor<ALandscapeBlueprintBrushBase>(BrushClassPtr, SpawnLocation, FRotator(0.0f), SpawnInfo);
 			EdMode->UISettings->BlueprintBrush = nullptr;
 
diff --git a/Engine/Source/Editor/UnrealEd/Classes/ActorFactories/ActorFactory.h b/Engine/Source/Editor/UnrealEd/Classes/ActorFactories/ActorFactory.h
--- a/Engine/Source/Editor/UnrealEd/Classes/ActorFactories/ActorFactory.h
+++ b/Engine/Source/Editor/UnrealEd/Classes/ActorFactories/ActorFactory.h
@@ -8,12 +8,14 @@
 #include "UObject/Object.h"
 #include "Templates/SubclassOf.h"
 #include "Factories/AssetFactoryInterface.h"
+#include "Engine/World.h"
 
 #include "ActorFactory.generated.h"
 
 UNREALED_API DECLARE_LOG_CATEGORY_EXTERN(LogActorFactory, Log, All);
 
 class AActor;
+struct FActorSpawnParameters;
 struct FAssetData;
 class UBlueprint;
 class ULevel;
@@ -51,7 +53,9 @@ class UNREALED_API UActorFactory : public UObject, public IAssetFactoryInterface
 	FVector SpawnPositionOffset;
 
 	/** Called to actual create an actor with the supplied transform (scale is ignored), using the properties in the ActorFactory */
-	AActor* CreateActor( UObject* Asset, ULevel* InLevel, FTransform Transform, EObjectFlags InObjectFlags = RF_Transactional, const FName InName = NAME_None );
+	UE_DEPRECATED(5.0, "This function has been deprecated in favor of the other version that takes a FActorSpawnParameters in parameter")
+	AActor* CreateActor(UObject* Asset, ULevel* InLevel, const FTransform& Transform, EObjectFlags InObjectFlags, const FName InName = NAME_None);
+	AActor* CreateActor(UObject* InAsset, ULevel* InLevel, const FTransform& InTransform, const FActorSpawnParameters& InSpawnParams = FActorSpawnParameters());
 
 	/** Called to create a blueprint class that can be used to spawn an actor from this factory */
 	UE_DEPRECATED(5.0, "This function is no longer used. See FKismetEditorUtilities::CreateBlueprint.")
@@ -90,8 +94,14 @@ class UNREALED_API UActorFactory : public UObject, public IAssetFactoryInterface
 
 protected:
 
+	/** Validates the input params for SpawnActor and returns the appropriate level to use depending on whether InLevel and/or InSpawnParams.OverrideLevel is passed : */
+	ULevel* ValidateSpawnActorLevel(ULevel* InLevel, const FActorSpawnParameters& InSpawnParams) const;
+
 	virtual bool PreSpawnActor( UObject* Asset, FTransform& InOutLocation);
-	virtual AActor* SpawnActor( UObject* Asset, ULevel* InLevel, const FTransform& Transform, EObjectFlags ObjectFlags, const FName Name );
+
+	UE_DEPRECATED(5.0, "This function has been deprecated in favor of the other version that takes a FActorSpawnParameters in parameter")
+	virtual AActor* SpawnActor(UObject* Asset, ULevel* InLevel, const FTransform& Transform, EObjectFlags ObjectFlags, const FName Name ) final;
+	virtual AActor* SpawnActor(UObject* InAsset, ULevel* InLevel, const FTransform& InTransform, const FActorSpawnParameters& InSpawnParams);
 
 	/** Subclasses may implement this to modify the actor after it has been spawned 
 	    IMPORTANT: If you override this, you should usually also override PostCreateBlueprint()! */
diff --git a/Engine/Source/Editor/UnrealEd/Classes/ActorFactories/ActorFactoryClass.h b/Engine/Source/Editor/UnrealEd/Classes/ActorFactories/ActorFactoryClass.h
--- a/Engine/Source/Editor/UnrealEd/Classes/ActorFactories/ActorFactoryClass.h
+++ b/Engine/Source/Editor/UnrealEd/Classes/ActorFactories/ActorFactoryClass.h
@@ -11,6 +11,7 @@
 class AActor;
 struct FAssetData;
 class ULevel;
+struct FActorSpawnParameters;
 
 UCLASS(MinimalAPI, config=Editor, collapsecategories, hidecategories=Object)
 class UActorFactoryClass : public UActorFactory
@@ -24,5 +25,5 @@ class UActorFactoryClass : public UActorFactory
 	//~ End UActorFactory Interface	
 
 protected:
-	virtual AActor* SpawnActor( UObject* Asset, ULevel* InLevel, const FTransform& Transform, EObjectFlags ObjectFlags, const FName Name ) override;
+	virtual AActor* SpawnActor( UObject* Asset, ULevel* InLevel, const FTransform& Transform, const FActorSpawnParameters& InSpawnParams) override;
 };
diff --git a/Engine/Source/Editor/UnrealEd/Classes/ActorFactories/ActorFactoryEmptyActor.h b/Engine/Source/Editor/UnrealEd/Classes/ActorFactories/ActorFactoryEmptyActor.h
--- a/Engine/Source/Editor/UnrealEd/Classes/ActorFactories/ActorFactoryEmptyActor.h
+++ b/Engine/Source/Editor/UnrealEd/Classes/ActorFactories/ActorFactoryEmptyActor.h
@@ -9,6 +9,7 @@
 #include "ActorFactoryEmptyActor.generated.h"
 
 class AActor;
+struct FActorSpawnParameters;
 struct FAssetData;
 class ULevel;
 
@@ -25,6 +26,6 @@ public:
 	virtual bool CanCreateActorFrom( const FAssetData& AssetData, FText& OutErrorMsg ) override;
 
 protected:
-	virtual AActor* SpawnActor( UObject* Asset, ULevel* InLevel, const FTransform& Transform, EObjectFlags ObjectFlags, const FName Name ) override;
+	virtual AActor* SpawnActor(UObject* InAsset, ULevel* InLevel, const FTransform& InTransform, const FActorSpawnParameters& InSpawnParams) override;
 
 };
diff --git a/Engine/Source/Editor/UnrealEd/Private/AssetSelection.cpp b/Engine/Source/Editor/UnrealEd/Private/AssetSelection.cpp
--- a/Engine/Source/Editor/UnrealEd/Private/AssetSelection.cpp
+++ b/Engine/Source/Editor/UnrealEd/Private/AssetSelection.cpp
@@ -638,7 +638,10 @@ static AActor* PrivateAddActor( UObject* Asset, UActorFactory* Factory, bool Sel
 		
 		if (!Actor)
 		{
-			Actor = Factory->CreateActor(Asset, DesiredLevel, ActorTransform, ObjectFlags, Name);
+			FActorSpawnParameters SpawnParams;
+			SpawnParams.ObjectFlags = ObjectFlags;
+			SpawnParams.Name = Name;
+			Actor = Factory->CreateActor(Asset, DesiredLevel, ActorTransform, SpawnParams);
 		}
 
 		if (Actor)
diff --git a/Engine/Source/Editor/UnrealEd/Private/Factories/ActorFactory.cpp b/Engine/Source/Editor/UnrealEd/Private/Factories/ActorFactory.cpp
--- a/Engine/Source/Editor/UnrealEd/Private/Factories/ActorFactory.cpp
+++ b/Engine/Source/Editor/UnrealEd/Private/Factories/ActorFactory.cpp
@@ -303,7 +303,9 @@ bool UActorFactory::PrePlaceAsset(FAssetPlacementInfo& InPlacementInfo, const FP
 TArray<FTypedElementHandle> UActorFactory::PlaceAsset(const FAssetPlacementInfo& InPlacementInfo, const FPlacementOptions& InPlacementOptions)
 {
 	TArray<FTypedElementHandle> PlacedActorHandles;
-	AActor* NewActor = CreateActor(InPlacementInfo.AssetToPlace.GetAsset(), InPlacementInfo.PreferredLevel.Get(), InPlacementInfo.FinalizedTransform, RF_Transactional, InPlacementInfo.NameOverride);
+	FActorSpawnParameters SpawnParams;
+	SpawnParams.Name = InPlacementInfo.NameOverride;
+	AActor* NewActor = CreateActor(InPlacementInfo.AssetToPlace.GetAsset(), InPlacementInfo.PreferredLevel.Get(), InPlacementInfo.FinalizedTransform, SpawnParams);
 	if (NewActor)
 	{
 		PlacedActorHandles.Add(UEngineElementsLibrary::AcquireEditorActorElementHandle(NewActor));
@@ -384,22 +386,31 @@ void UActorFactory::EndPlacement(TArrayView<const FTypedElementHandle> InPlacedE
 {
 }
 
-AActor* UActorFactory::CreateActor( UObject* Asset, ULevel* InLevel, FTransform SpawnTransform, EObjectFlags InObjectFlags, const FName Name )
+AActor* UActorFactory::CreateActor( UObject* Asset, ULevel* InLevel, const FTransform& SpawnTransform, EObjectFlags InObjectFlags, const FName Name)
+{
+	FActorSpawnParameters SpawnParams;
+	SpawnParams.ObjectFlags = InObjectFlags;
+	SpawnParams.Name = Name;
+	return CreateActor(Asset, InLevel, SpawnTransform, SpawnParams);
+}
+
+AActor* UActorFactory::CreateActor(UObject* InAsset, ULevel* InLevel, const FTransform& InTransform, const FActorSpawnParameters& InSpawnParams)
 {
 	AActor* NewActor = nullptr;
 
-	if ( PreSpawnActor(Asset, SpawnTransform) )
+	FTransform Transform(InTransform);
+	if (PreSpawnActor(InAsset, Transform))
 	{
-		NewActor = SpawnActor(Asset, InLevel, SpawnTransform, InObjectFlags, Name);
+		NewActor = SpawnActor(InAsset, InLevel, Transform, InSpawnParams);
 
-		if ( NewActor )
+		if (NewActor)
 		{
-			PostSpawnActor(Asset, NewActor);
+			PostSpawnActor(InAsset, NewActor);
 
 			// Only do this if the actor wasn't already given a name
-			if (Name == NAME_None && Asset)
+			if ((InAsset != nullptr) && (InSpawnParams.Name == NAME_None))
 			{
-				FActorLabelUtilities::SetActorLabelUnique(NewActor, Asset->GetName());
+				FActorLabelUtilities::SetActorLabelUnique(NewActor, InAsset->GetName());
 			}
 
 			NewActor->PostEditChange();
@@ -421,21 +432,47 @@ UBlueprint* UActorFactory::CreateBlueprint( UObject* Asset, UObject* Outer, cons
 	return NewBlueprint;
 }
 
+ULevel* UActorFactory::ValidateSpawnActorLevel(ULevel* InLevel, const FActorSpawnParameters& InSpawnParams) const
+{
+	ULevel* LocalLevel = InLevel;
+	if (InLevel != nullptr)
+	{
+		// If InLevel is passed, then InSpawnParams.OverrideLevel shouldn't be used or be identical to that level, because in the end, InLevel will be the one we'll use :
+		ensure((InSpawnParams.OverrideLevel == nullptr) || (InSpawnParams.OverrideLevel == InLevel));
+	}
+	else
+	{
+		// If InLevel is not passed then the level should at least be specified in InSpawnParams.OverrideLevel :
+		ensure(InSpawnParams.OverrideLevel != nullptr);
+		LocalLevel = InSpawnParams.OverrideLevel;
+	}
+
+	return LocalLevel;
+}
+
 bool UActorFactory::PreSpawnActor( UObject* Asset, FTransform& InOutLocation)
 {
 	// Subclasses may implement this to set up a spawn or to adjust the spawn location or rotation.
 	return true;
 }
 
-AActor* UActorFactory::SpawnActor( UObject* Asset, ULevel* InLevel, const FTransform& Transform, EObjectFlags InObjectFlags, const FName Name )
+AActor* UActorFactory::SpawnActor( UObject* InAsset, ULevel* InLevel, const FTransform& Transform, EObjectFlags InObjectFlags, const FName Name )
+{
+	FActorSpawnParameters SpawnParams;
+	SpawnParams.ObjectFlags = InObjectFlags;
+	SpawnParams.Name = Name;
+	return SpawnActor(InAsset, InLevel, Transform, SpawnParams);
+}
+
+AActor* UActorFactory::SpawnActor(UObject* InAsset, ULevel* InLevel, const FTransform& InTransform, const FActorSpawnParameters& InSpawnParams)
 {
-	AActor* DefaultActor = GetDefaultActor( FAssetData( Asset ) );
-	if ( DefaultActor )
+	ULevel* LocalLevel = ValidateSpawnActorLevel(InLevel, InSpawnParams);
+
+	AActor* DefaultActor = GetDefaultActor(FAssetData(InAsset));
+	if ((DefaultActor != nullptr) && (LocalLevel != nullptr))
 	{
-		FActorSpawnParameters SpawnInfo;
-		SpawnInfo.OverrideLevel = InLevel;
-		SpawnInfo.ObjectFlags = InObjectFlags;
-		SpawnInfo.Name = Name;
+		FActorSpawnParameters SpawnParamsCopy(InSpawnParams);
+		SpawnParamsCopy.OverrideLevel = LocalLevel;
 
 		const bool bIsCreatingPreviewElements = FLevelEditorViewportClient::IsDroppingPreviewActor();
 		bool bIsPlacementSystemCreatingPreviewElements = false;
@@ -443,10 +480,10 @@ AActor* UActorFactory::SpawnActor( UObject* Asset, ULevel* InLevel, const FTrans
 		{
 			bIsPlacementSystemCreatingPreviewElements = PlacementSubsystem->IsCreatingPreviewElements();
 		}
-		SpawnInfo.bTemporaryEditorActor = bIsCreatingPreviewElements || bIsPlacementSystemCreatingPreviewElements;
-		SpawnInfo.bHideFromSceneOutliner = bIsPlacementSystemCreatingPreviewElements;
+		SpawnParamsCopy.bTemporaryEditorActor = bIsCreatingPreviewElements || bIsPlacementSystemCreatingPreviewElements;
+		SpawnParamsCopy.bHideFromSceneOutliner = bIsPlacementSystemCreatingPreviewElements;
 
-		return InLevel->OwningWorld->SpawnActor( DefaultActor->GetClass(), &Transform, SpawnInfo );
+		return LocalLevel->OwningWorld->SpawnActor(DefaultActor->GetClass(), &InTransform, SpawnParamsCopy);
 	}
 
 	return nullptr;
@@ -1336,17 +1373,17 @@ bool UActorFactoryEmptyActor::CanCreateActorFrom( const FAssetData& AssetData, F
 	return AssetData.ObjectPath == FName(*AActor::StaticClass()->GetPathName());
 }
 
-AActor* UActorFactoryEmptyActor::SpawnActor( UObject* Asset, ULevel* InLevel, const FTransform& Transform, EObjectFlags InObjectFlags, const FName Name )
+AActor* UActorFactoryEmptyActor::SpawnActor(UObject* InAsset, ULevel* InLevel, const FTransform& InTransform, const FActorSpawnParameters& InSpawnParams)
 {
 	AActor* NewActor = nullptr;
 	{
 		// Spawn a temporary actor for dragging around
-		NewActor = Super::SpawnActor(Asset, InLevel, Transform, InObjectFlags, Name);
+		NewActor = Super::SpawnActor(InAsset, InLevel, InTransform, InSpawnParams);
 
 		USceneComponent* RootComponent = NewObject<USceneComponent>(NewActor, USceneComponent::GetDefaultSceneRootVariableName(), RF_Transactional);
 		RootComponent->Mobility = EComponentMobility::Movable;
 		RootComponent->bVisualizeComponent = bVisualizeActor;
-		RootComponent->SetWorldTransform(Transform);
+		RootComponent->SetWorldTransform(InTransform);
 
 		NewActor->SetRootComponent(RootComponent);
 		NewActor->AddInstanceComponent(RootComponent);
@@ -1504,17 +1541,15 @@ bool UActorFactoryClass::PreSpawnActor( UObject* Asset, FTransform& InOutLocatio
 	return false;
 }
 
-AActor* UActorFactoryClass::SpawnActor( UObject* Asset, ULevel* InLevel, const FTransform& Transform, EObjectFlags InObjectFlags, const FName Name )
+AActor* UActorFactoryClass::SpawnActor( UObject* Asset, ULevel* InLevel, const FTransform& Transform, const FActorSpawnParameters& InSpawnParams)
 {
 	UClass* ActualClass = Cast<UClass>(Asset);
 
 	if ( (nullptr != ActualClass) && ActualClass->IsChildOf(AActor::StaticClass()) )
 	{
-		FActorSpawnParameters SpawnInfo;
+		FActorSpawnParameters SpawnInfo(InSpawnParams);
 		SpawnInfo.OverrideLevel = InLevel;
-		SpawnInfo.ObjectFlags = InObjectFlags;
-		SpawnInfo.Name = Name;
-		return InLevel->OwningWorld->SpawnActor( ActualClass, &Transform, SpawnInfo );
+		return InLevel->OwningWorld->SpawnActor(ActualClass, &Transform, SpawnInfo);
 	}
 
 	return nullptr;
@@ -2073,13 +2108,13 @@ bool UActorFactoryLevelSequence::CanCreateActorFrom( const FAssetData& AssetData
 	return true;
 }
 
-AActor* UActorFactoryLevelSequence::SpawnActor( UObject* Asset, ULevel* InLevel, const FTransform& Transform, EObjectFlags InObjectFlags, const FName Name )
+AActor* UActorFactoryLevelSequence::SpawnActor(UObject* InAsset, ULevel* InLevel, const FTransform& InTransform, const FActorSpawnParameters& InSpawnParams)
 {
-	ALevelSequenceActor* NewActor = Cast<ALevelSequenceActor>(Super::SpawnActor(Asset, InLevel, Transform, InObjectFlags, Name));
+	ALevelSequenceActor* NewActor = Cast<ALevelSequenceActor>(Super::SpawnActor(InAsset, InLevel, InTransform, InSpawnParams));
 
 	if (NewActor)
 	{
-		if (ULevelSequence* LevelSequence = Cast<ULevelSequence>(Asset))
+		if (ULevelSequence* LevelSequence = Cast<ULevelSequence>(InAsset))
 		{
 			NewActor->SetSequence(LevelSequence);
 		}
diff --git a/Engine/Source/Editor/UnrealEd/Private/Factories/ActorFactoryLevelSequence.h b/Engine/Source/Editor/UnrealEd/Private/Factories/ActorFactoryLevelSequence.h
--- a/Engine/Source/Editor/UnrealEd/Private/Factories/ActorFactoryLevelSequence.h
+++ b/Engine/Source/Editor/UnrealEd/Private/Factories/ActorFactoryLevelSequence.h
@@ -8,6 +8,7 @@
 #include "ActorFactoryLevelSequence.generated.h"
 
 class AActor;
+struct FActorSpawnParameters;
 struct FAssetData;
 class ULevel;
 
@@ -17,7 +18,7 @@ class UActorFactoryLevelSequence : public UActorFactory
 	GENERATED_UCLASS_BODY()
 
 	// Begin UActorFactory Interface
-	virtual AActor* SpawnActor( UObject* Asset, ULevel* InLevel, const FTransform& Transform, EObjectFlags ObjectFlags, const FName Name ) override;
+	virtual AActor* SpawnActor(UObject* InAsset, ULevel* InLevel, const FTransform& InTransform, const FActorSpawnParameters& InSpawnParams) override;
 	virtual bool CanCreateActorFrom( const FAssetData& AssetData, FText& OutErrorMsg ) override;
 	virtual UObject* GetAssetFromActorInstance(AActor* ActorInstance) override;
 	// End UActorFactory Interface
