Bugfix related to skeleton remapping, where some if statement had to be inversed.

#rnx

diff --git a/Engine/Source/Editor/UnrealEd/Private/Animation/DebugSkelMeshComponent.cpp b/Engine/Source/Editor/UnrealEd/Private/Animation/DebugSkelMeshComponent.cpp
--- a/Engine/Source/Editor/UnrealEd/Private/Animation/DebugSkelMeshComponent.cpp
+++ b/Engine/Source/Editor/UnrealEd/Private/Animation/DebugSkelMeshComponent.cpp
@@ -266,7 +266,7 @@ void UDebugSkelMeshComponent::InitAnim(bool bForceReinit)
 	// then we need to clear it up to avoid an issue
 	if ( PreviewInstance && PreviewInstance->GetCurrentAsset() && SkeletalMesh )
 	{
-		if (SkeletalMesh->GetSkeleton()->IsCompatible(PreviewInstance->GetCurrentAsset()->GetSkeleton()))
+		if (!SkeletalMesh->GetSkeleton()->IsCompatible(PreviewInstance->GetCurrentAsset()->GetSkeleton()))
 		{
 			// if it doesn't match, just clear it
 			PreviewInstance->SetAnimationAsset(NULL);
