Fixed WorldPartition distributed HLOD build
- ActorDesc was out of date after copying package from prior build step

diff --git a/Engine/Source/Editor/UnrealEd/Private/WorldPartition/WorldPartitionHLODsBuilder.cpp b/Engine/Source/Editor/UnrealEd/Private/WorldPartition/WorldPartitionHLODsBuilder.cpp
--- a/Engine/Source/Editor/UnrealEd/Private/WorldPartition/WorldPartitionHLODsBuilder.cpp
+++ b/Engine/Source/Editor/UnrealEd/Private/WorldPartition/WorldPartitionHLODsBuilder.cpp
@@ -15,6 +15,7 @@
 #include "ISourceControlProvider.h"
 #include "SourceControlOperations.h"
 
+#include "AssetRegistryModule.h"
 #include "DerivedDataCacheInterface.h"
 
 #include "WorldPartition/WorldPartition.h"
@@ -967,5 +968,9 @@ bool UWorldPartitionHLODsBuilder::CopyFilesFromWorkingDir(const FString& SourceD
 		}
 	}
 
+	// Force a rescan of the updated files
+	IAssetRegistry& AssetRegistry = FModuleManager::LoadModuleChecked<FAssetRegistryModule>(TEXT("AssetRegistry")).Get();
+	AssetRegistry.ScanPathsSynchronous(ModifiedFiles.Array(), /*bForceRescan*/true, /*bIgnoreBlackListScanFilters*/true);
+
 	return true;
-}
\ No newline at end of file
+}
