Add PreObjectSaveContext to OnObjectSaved delegate. Use the information to skip the recalculate of Actor descriptions when making an EditorDomainSave; this was causing dangling pointers since EditorDomainSaves can occur during a load.
#rnx

diff --git a/Engine/Plugins/MovieScene/ActorSequence/Source/ActorSequenceEditor/Private/ActorSequenceEditorTabSummoner.cpp b/Engine/Plugins/MovieScene/ActorSequence/Source/ActorSequenceEditor/Private/ActorSequenceEditorTabSummoner.cpp
--- a/Engine/Plugins/MovieScene/ActorSequence/Source/ActorSequenceEditor/Private/ActorSequenceEditorTabSummoner.cpp
+++ b/Engine/Plugins/MovieScene/ActorSequence/Source/ActorSequenceEditor/Private/ActorSequenceEditorTabSummoner.cpp
@@ -18,6 +18,7 @@
 #include "Framework/MultiBox/MultiBoxBuilder.h"
 #include "Framework/Application/SlateApplication.h"
 #include "ActorSequenceEditorStyle.h"
+#include "UObject/ObjectSaveContext.h"
 
 #define LOCTEXT_NAMESPACE "ActorSequenceEditorSummoner"
 
@@ -214,7 +215,7 @@ public:
 
 		GEditor->UnregisterForUndo(this);
 		GEditor->OnBlueprintPreCompile().Remove(OnBlueprintPreCompileHandle);
-		FCoreUObjectDelegates::OnObjectSaved.Remove(OnObjectSavedHandle);
+		FCoreUObjectDelegates::OnObjectPreSave.Remove(OnObjectSavedHandle);
 	}
 
 	~SActorSequenceEditorWidgetImpl()
@@ -236,7 +237,7 @@ public:
 	void Construct(const FArguments&, TWeakPtr<FBlueprintEditor> InBlueprintEditor)
 	{
 		OnBlueprintPreCompileHandle = GEditor->OnBlueprintPreCompile().AddSP(this, &SActorSequenceEditorWidgetImpl::OnBlueprintPreCompile);
-		OnObjectSavedHandle = FCoreUObjectDelegates::OnObjectSaved.AddSP(this, &SActorSequenceEditorWidgetImpl::OnObjectPreSave);
+		OnObjectSavedHandle = FCoreUObjectDelegates::OnObjectPreSave.AddSP(this, &SActorSequenceEditorWidgetImpl::OnObjectPreSave);
 
 		WeakBlueprintEditor = InBlueprintEditor;
 
@@ -410,7 +411,7 @@ public:
 		}
 	}
 
-	void OnObjectPreSave(UObject* InObject)
+	void OnObjectPreSave(UObject* InObject, FObjectPreSaveContext SaveContext)
 	{
 		TSharedPtr<FBlueprintEditor> BlueprintEditor = WeakBlueprintEditor.Pin();
 		if (Sequencer.IsValid() && BlueprintEditor.IsValid() && InObject && InObject == BlueprintEditor->GetBlueprintObj())
diff --git a/Engine/Plugins/MovieScene/LevelSequenceEditor/Source/LevelSequenceEditor/Private/Misc/LevelSequenceEditorSpawnRegister.cpp b/Engine/Plugins/MovieScene/LevelSequenceEditor/Source/LevelSequenceEditor/Private/Misc/LevelSequenceEditorSpawnRegister.cpp
--- a/Engine/Plugins/MovieScene/LevelSequenceEditor/Source/LevelSequenceEditor/Private/Misc/LevelSequenceEditorSpawnRegister.cpp
+++ b/Engine/Plugins/MovieScene/LevelSequenceEditor/Source/LevelSequenceEditor/Private/Misc/LevelSequenceEditorSpawnRegister.cpp
@@ -17,6 +17,7 @@
 #include "EntitySystem/MovieSceneSpawnablesSystem.h"
 #include "Evaluation/MovieSceneEvaluationTemplateInstance.h"
 #include "Sections/MovieScene3DTransformSection.h"
+#include "UObject/ObjectSaveContext.h"
 
 #define LOCTEXT_NAMESPACE "LevelSequenceEditorSpawnRegister"
 
@@ -34,7 +35,7 @@ FLevelSequenceEditorSpawnRegister::FLevelSequenceEditorSpawnRegister()
 	FCoreUObjectDelegates::OnObjectsReplaced.AddRaw(this, &FLevelSequenceEditorSpawnRegister::OnObjectsReplaced);
 
 	OnObjectModifiedHandle = FCoreUObjectDelegates::OnObjectModified.AddRaw(this, &FLevelSequenceEditorSpawnRegister::OnObjectModified);
-	OnObjectSavedHandle    = FCoreUObjectDelegates::OnObjectSaved.AddRaw(this, &FLevelSequenceEditorSpawnRegister::OnPreObjectSaved);
+	OnObjectSavedHandle    = FCoreUObjectDelegates::OnObjectPreSave.AddRaw(this, &FLevelSequenceEditorSpawnRegister::OnPreObjectSaved);
 #endif
 }
 
@@ -56,7 +57,7 @@ FLevelSequenceEditorSpawnRegister::~FLevelSequenceEditorSpawnRegister()
 #if WITH_EDITOR
 	FCoreUObjectDelegates::OnObjectsReplaced.RemoveAll(this);
 	FCoreUObjectDelegates::OnObjectModified.Remove(OnObjectModifiedHandle);
-	FCoreUObjectDelegates::OnObjectSaved.Remove(OnObjectSavedHandle);
+	FCoreUObjectDelegates::OnObjectPreSave.Remove(OnObjectSavedHandle);
 #endif
 }
 
@@ -275,7 +276,7 @@ void FLevelSequenceEditorSpawnRegister::OnObjectModified(UObject* ModifiedObject
 	}
 }
 
-void FLevelSequenceEditorSpawnRegister::OnPreObjectSaved(UObject* Object)
+void FLevelSequenceEditorSpawnRegister::OnPreObjectSaved(UObject* Object, FObjectPreSaveContext SaveContext)
 {
 	UMovieSceneSequence* SequenceBeingSaved = Cast<UMovieSceneSequence>(Object);
 	if (SequenceBeingSaved && SequencesWithModifiedObjects.Contains(SequenceBeingSaved))
diff --git a/Engine/Plugins/MovieScene/LevelSequenceEditor/Source/LevelSequenceEditor/Public/Misc/LevelSequenceEditorSpawnRegister.h b/Engine/Plugins/MovieScene/LevelSequenceEditor/Source/LevelSequenceEditor/Public/Misc/LevelSequenceEditorSpawnRegister.h
--- a/Engine/Plugins/MovieScene/LevelSequenceEditor/Source/LevelSequenceEditor/Public/Misc/LevelSequenceEditorSpawnRegister.h
+++ b/Engine/Plugins/MovieScene/LevelSequenceEditor/Source/LevelSequenceEditor/Public/Misc/LevelSequenceEditorSpawnRegister.h
@@ -12,6 +12,7 @@
 class IMovieScenePlayer;
 class ISequencer;
 class UMovieScene;
+class FObjectPreSaveContext;
 
 /**
  * Spawn register used in the editor to add some usability features like maintaining selection states, and projecting spawned state onto spawnable defaults
@@ -61,7 +62,7 @@ private:
 	void OnObjectModified(UObject* ModifiedObject);
 
 	/** Called before an object is saved in the editor */
-	void OnPreObjectSaved(UObject* Object);
+	void OnPreObjectSaved(UObject* Object, FObjectPreSaveContext SaveContext);
 
 private:
 
@@ -102,6 +103,6 @@ private:
 	/** Handle to a delegate that is bound to FCoreUObjectDelegates::OnObjectModified to harvest changes to spawned objects. */
 	FDelegateHandle OnObjectModifiedHandle;
 
-	/** Handle to a delegate that is bound to FCoreUObjectDelegates::OnObjectSaved to harvest changes to spawned objects. */
+	/** Handle to a delegate that is bound to FCoreUObjectDelegates::OnObjectPreSave to harvest changes to spawned objects. */
 	FDelegateHandle OnObjectSavedHandle;
 };
diff --git a/Engine/Source/Editor/UnrealEd/Classes/CookOnTheSide/CookOnTheFlyServer.h b/Engine/Source/Editor/UnrealEd/Classes/CookOnTheSide/CookOnTheFlyServer.h
--- a/Engine/Source/Editor/UnrealEd/Classes/CookOnTheSide/CookOnTheFlyServer.h
+++ b/Engine/Source/Editor/UnrealEd/Classes/CookOnTheSide/CookOnTheFlyServer.h
@@ -613,7 +613,7 @@ public:
 	void OnObjectModified( UObject *ObjectMoving );
 	void OnObjectPropertyChanged(UObject* ObjectBeingModified, FPropertyChangedEvent& PropertyChangedEvent);
 	void OnObjectUpdated( UObject *Object );
-	void OnObjectSaved( UObject *ObjectSaved );
+	void OnObjectSaved( UObject *ObjectSaved, FObjectPreSaveContext SaveContext );
 
 	/**
 	* Marks a package as dirty for cook
diff --git a/Engine/Source/Editor/UnrealEd/Private/CookOnTheFlyServer.cpp b/Engine/Source/Editor/UnrealEd/Private/CookOnTheFlyServer.cpp
--- a/Engine/Source/Editor/UnrealEd/Private/CookOnTheFlyServer.cpp
+++ b/Engine/Source/Editor/UnrealEd/Private/CookOnTheFlyServer.cpp
@@ -3466,11 +3466,11 @@ void UCookOnTheFlyServer::OnObjectPropertyChanged(UObject* ObjectBeingModified,
 	OnObjectUpdated( ObjectBeingModified );
 }
 
-void UCookOnTheFlyServer::OnObjectSaved( UObject* ObjectSaved )
+void UCookOnTheFlyServer::OnObjectSaved( UObject* ObjectSaved, FObjectPreSaveContext SaveContext)
 {
-	if (GIsCookerLoadingPackage)
+	if (SaveContext.IsProceduralSave() )
 	{
-		// This is the cooker saving a cooked package, ignore
+		// This is a procedural save (e.g. our own saving of the cooked package) rather than a user save, ignore
 		return;
 	}
 
@@ -4370,7 +4370,7 @@ void UCookOnTheFlyServer::Initialize( ECookMode::Type DesiredCookMode, ECookInit
 	{
 		FCoreUObjectDelegates::OnObjectPropertyChanged.AddUObject(this, &UCookOnTheFlyServer::OnObjectPropertyChanged);
 		FCoreUObjectDelegates::OnObjectModified.AddUObject(this, &UCookOnTheFlyServer::OnObjectModified);
-		FCoreUObjectDelegates::OnObjectSaved.AddUObject(this, &UCookOnTheFlyServer::OnObjectSaved);
+		FCoreUObjectDelegates::OnObjectPreSave.AddUObject(this, &UCookOnTheFlyServer::OnObjectSaved);
 
 		FCoreDelegates::OnTargetPlatformChangedSupportedFormats.AddUObject(this, &UCookOnTheFlyServer::OnTargetPlatformChangedSupportedFormats);
 	}
diff --git a/Engine/Source/Editor/UnrealEd/Private/UnrealEdMisc.cpp b/Engine/Source/Editor/UnrealEd/Private/UnrealEdMisc.cpp
--- a/Engine/Source/Editor/UnrealEd/Private/UnrealEdMisc.cpp
+++ b/Engine/Source/Editor/UnrealEd/Private/UnrealEdMisc.cpp
@@ -251,7 +251,7 @@ void FUnrealEdMisc::OnInit()
 	FEditorSupportDelegates::RedrawAllViewports.AddRaw(this, &FUnrealEdMisc::CB_RedrawAllViewports);
 	GEngine->OnLevelActorAdded().AddRaw( this, &FUnrealEdMisc::CB_LevelActorsAdded );
 
-	FCoreUObjectDelegates::OnObjectSaved.AddRaw(this, &FUnrealEdMisc::OnObjectSaved);
+	FCoreUObjectDelegates::OnObjectPreSave.AddRaw(this, &FUnrealEdMisc::OnObjectSaved);
 
 #if USE_UNIT_TESTS
 	FAutomationTestFramework::Get().PreTestingEvent.AddRaw(this, &FUnrealEdMisc::CB_PreAutomationTesting);
@@ -1163,7 +1163,7 @@ void FUnrealEdMisc::CB_RefreshEditor()
 
 void FUnrealEdMisc::PreSaveWorld(UWorld* World, FObjectPreSaveContext ObjectSaveContext)
 {
-	LogAssetUpdate(World);
+	LogAssetUpdate(World, ObjectSaveContext);
 	const bool bAutosaveOrPIE = (ObjectSaveContext.GetSaveFlags() & SAVE_FromAutosave) != 0;
 	if (bAutosaveOrPIE || !World || World != GEditor->GetEditorWorldContext().World() || !FEngineAnalytics::IsAvailable())
 	{
@@ -1711,22 +1711,22 @@ void FUnrealEdMisc::OnGotoAsset(const FString& InAssetPath) const
 	}	
 }
 
-void FUnrealEdMisc::OnObjectSaved(UObject* SavedObject)
+void FUnrealEdMisc::OnObjectSaved(UObject* SavedObject, FObjectPreSaveContext SaveContext)
 {
 	// Ensure the saved object is a non-UWorld asset (UWorlds are handled separately)
 	if (!SavedObject->IsA<UWorld>() && SavedObject->IsAsset())
 	{
-		LogAssetUpdate(SavedObject);
+		LogAssetUpdate(SavedObject, SaveContext);
 	}
 }
 
-void FUnrealEdMisc::LogAssetUpdate(UObject* UpdatedAsset)
+void FUnrealEdMisc::LogAssetUpdate(UObject* UpdatedAsset, FObjectPreSaveContext SaveContext)
 {
 	UPackage* AssetPackage = UpdatedAsset->GetOutermost();
 	const bool bIsPIESave = AssetPackage->RootPackageHasAnyFlags(PKG_PlayInEditor);
 	const bool bIsAutosave = GUnrealEd->GetPackageAutoSaver().IsAutoSaving();
 
-	if (!bIsPIESave && !bIsAutosave && !GIsAutomationTesting)
+	if (!bIsPIESave && !bIsAutosave && !GIsAutomationTesting && !SaveContext.IsProceduralSave())
 	{
 		uint32& NumUpdates = NumUpdatesByAssetName.FindOrAdd(UpdatedAsset->GetClass()->GetFName());
 		NumUpdates++;
diff --git a/Engine/Source/Editor/UnrealEd/Public/UnrealEdMisc.h b/Engine/Source/Editor/UnrealEd/Public/UnrealEdMisc.h
--- a/Engine/Source/Editor/UnrealEd/Public/UnrealEdMisc.h
+++ b/Engine/Source/Editor/UnrealEd/Public/UnrealEdMisc.h
@@ -288,10 +288,10 @@ public:
 	void OnGotoAsset(const FString& InAssetPath) const;
 
 	/** Delegate used to update the map of asset update counts */
-	void OnObjectSaved(UObject* SavedObject);
+	void OnObjectSaved(UObject* SavedObject, FObjectPreSaveContext SaveContext);
 
 	/** Logs an update to an asset */
-	void LogAssetUpdate(UObject* UpdatedAsset);
+	void LogAssetUpdate(UObject* UpdatedAsset, FObjectPreSaveContext SaveContext);
 
 	/** Initialize engine analytics */
 	void InitEngineAnalytics();
diff --git a/Engine/Source/Runtime/CoreUObject/Private/UObject/Obj.cpp b/Engine/Source/Runtime/CoreUObject/Private/UObject/Obj.cpp
--- a/Engine/Source/Runtime/CoreUObject/Private/UObject/Obj.cpp
+++ b/Engine/Source/Runtime/CoreUObject/Private/UObject/Obj.cpp
@@ -1218,7 +1218,10 @@ void UObject::PreSave(FObjectPreSaveContext SaveContext)
 {
 	SaveContext.Data.bBaseClassCalled = true;
 #if WITH_EDITOR
+	PRAGMA_DISABLE_DEPRECATION_WARNINGS
 	FCoreUObjectDelegates::OnObjectSaved.Broadcast(this);
+	PRAGMA_ENABLE_DEPRECATION_WARNINGS
+	FCoreUObjectDelegates::OnObjectPreSave.Broadcast(this, SaveContext);
 #endif
 }
 
diff --git a/Engine/Source/Runtime/CoreUObject/Private/UObject/UObjectGlobals.cpp b/Engine/Source/Runtime/CoreUObject/Private/UObject/UObjectGlobals.cpp
--- a/Engine/Source/Runtime/CoreUObject/Private/UObject/UObjectGlobals.cpp
+++ b/Engine/Source/Runtime/CoreUObject/Private/UObject/UObjectGlobals.cpp
@@ -148,7 +148,10 @@ FCoreUObjectDelegates::FOnObjectModified FCoreUObjectDelegates::OnObjectModified
 FCoreUObjectDelegates::FOnObjectTransacted FCoreUObjectDelegates::OnObjectTransacted;
 FCoreUObjectDelegates::FOnObjectsReplaced FCoreUObjectDelegates::OnObjectsReplaced;
 FCoreUObjectDelegates::FOnAssetLoaded FCoreUObjectDelegates::OnAssetLoaded;
+PRAGMA_DISABLE_DEPRECATION_WARNINGS
 FCoreUObjectDelegates::FOnObjectSaved FCoreUObjectDelegates::OnObjectSaved;
+PRAGMA_ENABLE_DEPRECATION_WARNINGS
+FCoreUObjectDelegates::FOnObjectPreSave FCoreUObjectDelegates::OnObjectPreSave;
 #endif // WITH_EDITOR
 
 FSimpleMulticastDelegate& FCoreUObjectDelegates::GetPreGarbageCollectDelegate()
diff --git a/Engine/Source/Runtime/CoreUObject/Public/UObject/UObjectGlobals.h b/Engine/Source/Runtime/CoreUObject/Public/UObject/UObjectGlobals.h
--- a/Engine/Source/Runtime/CoreUObject/Public/UObject/UObjectGlobals.h
+++ b/Engine/Source/Runtime/CoreUObject/Public/UObject/UObjectGlobals.h
@@ -21,6 +21,7 @@
 
 struct FCustomPropertyListNode;
 struct FObjectInstancingGraph;
+class FObjectPreSaveContext;
 struct FObjectPtr;
 struct FStaticConstructObjectParameters;
 template <typename T>
@@ -2361,10 +2362,14 @@ struct COREUOBJECT_API FCoreUObjectDelegates
 	DECLARE_MULTICAST_DELEGATE_OneParam(FOnObjectsReplaced, const FReplacementObjectMap&);
 	static FOnObjectsReplaced OnObjectsReplaced;
 
-	/** Callback for when an asset is saved. This is called from UObject::PreSave before it is actually written to disk, for every object saved */
 	DECLARE_MULTICAST_DELEGATE_OneParam(FOnObjectSaved, UObject*);
+	UE_DEPRECATED(5.0, "Use OnObjectPreSave instead.")
 	static FOnObjectSaved OnObjectSaved;
 
+	/** Callback for when an asset is saved. This is called from UObject::PreSave before it is actually written to disk, for every object saved */
+	DECLARE_MULTICAST_DELEGATE_TwoParams(FOnObjectPreSave, UObject*, FObjectPreSaveContext);
+	static FOnObjectPreSave OnObjectPreSave;
+
 	/** Callback for when an asset is loaded. This gets called in both the editor and standalone game editor builds, but only for objects that return true for IsAsset() */
 	DECLARE_MULTICAST_DELEGATE_OneParam(FOnAssetLoaded, UObject*);
 	static FOnAssetLoaded OnAssetLoaded;
diff --git a/Engine/Source/Runtime/Engine/Classes/Engine/AssetManager.h b/Engine/Source/Runtime/Engine/Classes/Engine/AssetManager.h
--- a/Engine/Source/Runtime/Engine/Classes/Engine/AssetManager.h
+++ b/Engine/Source/Runtime/Engine/Classes/Engine/AssetManager.h
@@ -648,7 +648,7 @@ protected:
 	virtual void OnInMemoryAssetDeleted(UObject *Object);
 
 	/** Called when object is saved */
-	virtual void OnObjectPreSave(UObject* Object);
+	virtual void OnObjectPreSave(UObject* Object, FObjectPreSaveContext SaveContext);
 
 	/** When asset is renamed */
 	virtual void OnAssetRenamed(const FAssetData& NewData, const FString& OldPath);
diff --git a/Engine/Source/Runtime/Engine/Private/AssetManager.cpp b/Engine/Source/Runtime/Engine/Private/AssetManager.cpp
--- a/Engine/Source/Runtime/Engine/Private/AssetManager.cpp
+++ b/Engine/Source/Runtime/Engine/Private/AssetManager.cpp
@@ -10,6 +10,7 @@
 #include "Engine/BlueprintGeneratedClass.h"
 #include "Interfaces/IPluginManager.h"
 #include "UObject/ConstructorHelpers.h"
+#include "UObject/ObjectSaveContext.h"
 #include "UObject/UObjectHash.h"
 #include "Misc/FileHelper.h"
 #include "Misc/ScopedSlowTask.h"
@@ -293,7 +294,7 @@ void UAssetManager::PostInitProperties()
 
 		FEditorDelegates::PreBeginPIE.AddUObject(this, &UAssetManager::PreBeginPIE);
 		FEditorDelegates::EndPIE.AddUObject(this, &UAssetManager::EndPIE);
-		FCoreUObjectDelegates::OnObjectSaved.AddUObject(this, &UAssetManager::OnObjectPreSave);
+		FCoreUObjectDelegates::OnObjectPreSave.AddUObject(this, &UAssetManager::OnObjectPreSave);
 
 		// In editor builds guess the type/name if allowed
 		bShouldGuessTypeAndName = Settings.bShouldGuessTypeAndNameInEditor;
@@ -4155,7 +4156,7 @@ void UAssetManager::OnInMemoryAssetDeleted(UObject *Object)
 	RemovePrimaryAssetId(PrimaryAssetId);
 }
 
-void UAssetManager::OnObjectPreSave(UObject* Object)
+void UAssetManager::OnObjectPreSave(UObject* Object, FObjectPreSaveContext SaveContext)
 {
 	// If this is in the asset manager dictionary, make sure it actually has a primary asset id that matches
 	const bool bIsAssetOrClass = Object->IsAsset() || Object->IsA(UClass::StaticClass()); 
diff --git a/Engine/Source/Runtime/Engine/Private/WorldPartition/ActorDescContainer.cpp b/Engine/Source/Runtime/Engine/Private/WorldPartition/ActorDescContainer.cpp
--- a/Engine/Source/Runtime/Engine/Private/WorldPartition/ActorDescContainer.cpp
+++ b/Engine/Source/Runtime/Engine/Private/WorldPartition/ActorDescContainer.cpp
@@ -7,6 +7,7 @@
 #include "AssetRegistryModule.h"
 #include "WorldPartition/WorldPartitionHandle.h"
 #include "Misc/Base64.h"
+#include "UObject/ObjectSaveContext.h"
 #endif
 
 UActorDescContainer::UActorDescContainer(const FObjectInitializer& ObjectInitializer)
@@ -162,8 +163,15 @@ bool UActorDescContainer::ShouldHandleActorEvent(const AActor* Actor)
 	return false;
 }
 
-void UActorDescContainer::OnObjectPreSave(UObject* Object)
+void UActorDescContainer::OnObjectPreSave(UObject* Object, FObjectPreSaveContext SaveContext)
 {
+	if (SaveContext.IsProceduralSave())
+	{
+		// Do not delete and recreate FWorldPartitionActorDesc when making a procedural save. Procedural saves such as
+		// EditorDomain saves can occur during a load package, and FWorldPartitionReferenceImpl::IncRefCount 
+		// can load (and therefore save) an actor package while holding a pointer to the FWorldPartitionActorDesc. 
+		return;
+	}
 	if (const AActor* Actor = Cast<AActor>(Object))
 	{
 		if (ShouldHandleActorEvent(Actor))
@@ -290,7 +298,7 @@ void UActorDescContainer::RegisterEditorDelegates()
 {
 	if (GEditor && !IsTemplate() && !World->IsGameWorld())
 	{
-		FCoreUObjectDelegates::OnObjectSaved.AddUObject(this, &UActorDescContainer::OnObjectPreSave);
+		FCoreUObjectDelegates::OnObjectPreSave.AddUObject(this, &UActorDescContainer::OnObjectPreSave);
 		FEditorDelegates::OnPackageDeleted.AddUObject(this, &UActorDescContainer::OnPackageDeleted);
 		FCoreUObjectDelegates::OnObjectsReplaced.AddUObject(this, &UActorDescContainer::OnObjectsReplaced);
 	}
@@ -300,7 +308,7 @@ void UActorDescContainer::UnregisterEditorDelegates()
 {
 	if (GEditor && !IsTemplate() && !World->IsGameWorld())
 	{
-		FCoreUObjectDelegates::OnObjectSaved.RemoveAll(this);
+		FCoreUObjectDelegates::OnObjectPreSave.RemoveAll(this);
 		FEditorDelegates::OnPackageDeleted.RemoveAll(this);
 		FCoreUObjectDelegates::OnObjectsReplaced.RemoveAll(this);
 	}
diff --git a/Engine/Source/Runtime/Engine/Public/WorldPartition/ActorDescContainer.h b/Engine/Source/Runtime/Engine/Public/WorldPartition/ActorDescContainer.h
--- a/Engine/Source/Runtime/Engine/Public/WorldPartition/ActorDescContainer.h
+++ b/Engine/Source/Runtime/Engine/Public/WorldPartition/ActorDescContainer.h
@@ -25,7 +25,7 @@ public:
 
 #if WITH_EDITOR
 	// Events
-	virtual void OnObjectPreSave(UObject* Object);
+	virtual void OnObjectPreSave(UObject* Object, FObjectPreSaveContext SaveContext);
 	virtual void OnPackageDeleted(UPackage* Package);
 	virtual void OnObjectsReplaced(const TMap<UObject*, UObject*>& OldToNewObjectMap);
 
