Put a workaround in place for a limitation when using memory tracing with other systems, who may allocate memory while writing fields of events they trace. Havoc ensues.

#rnx

diff --git a/Engine/Source/Editor/UnrealEd/Private/CookOnTheFlyServer.cpp b/Engine/Source/Editor/UnrealEd/Private/CookOnTheFlyServer.cpp
--- a/Engine/Source/Editor/UnrealEd/Private/CookOnTheFlyServer.cpp
+++ b/Engine/Source/Editor/UnrealEd/Private/CookOnTheFlyServer.cpp
@@ -3021,8 +3021,13 @@ void UCookOnTheFlyServer::PumpSaves(UE::Cook::FTickStackData& StackData, uint32
 		TArray<bool> SucceededSavePackage;
 		TArray<FSavePackageResultStruct> SavePackageResults;
 		{
+			TStringBuilder<512> PackageDataFileName;
+			if (UE_TRACE_CHANNELEXPR_IS_ENABLED(CookChannel))
+			{
+				PackageData.GetFileName().ToString(PackageDataFileName);
+			}
 			UE_SCOPED_HIERARCHICAL_CUSTOM_COOKTIMER_AND_DURATION(SaveCookedPackage, DetailedCookStats::TickCookOnTheSideSaveCookedPackageTimeSec)
-				UE_ADD_CUSTOM_COOKTIMER_META(SaveCookedPackage, PackageName, *PackageData.GetFileName().ToString());
+				UE_ADD_CUSTOM_COOKTIMER_META(SaveCookedPackage, PackageName, *PackageDataFileName);
 
 			uint32 SaveFlags = SAVE_KeepGUID | (bSaveAsyncAllowed ? SAVE_Async : SAVE_None) | (IsCookFlagSet(ECookInitializationFlags::Unversioned) ? SAVE_Unversioned : 0);
 
diff --git a/Engine/Source/Runtime/CoreUObject/Private/UObject/UObjectGlobals.cpp b/Engine/Source/Runtime/CoreUObject/Private/UObject/UObjectGlobals.cpp
--- a/Engine/Source/Runtime/CoreUObject/Private/UObject/UObjectGlobals.cpp
+++ b/Engine/Source/Runtime/CoreUObject/Private/UObject/UObjectGlobals.cpp
@@ -1105,8 +1105,14 @@ UPackage* LoadPackageInternal(UPackage* InOuter, const FPackagePath& PackagePath
 	const FLinkerInstancingContext* InstancingContext, const FPackagePath* DiffPackagePath)
 {
 	DECLARE_SCOPE_CYCLE_COUNTER(TEXT("LoadPackageInternal"), STAT_LoadPackageInternal, STATGROUP_ObjectVerbose);
+
+	FString TracePackageName;
+	if (UE_TRACE_CHANNELEXPR_IS_ENABLED(LoadTimeChannel))
+	{
+		TracePackageName = PackagePath.GetPackageNameOrFallback();
+	}
 	SCOPED_CUSTOM_LOADTIMER(LoadPackageInternal)
-		ADD_CUSTOM_LOADTIMER_META(LoadPackageInternal, PackageName, *PackagePath.GetPackageNameOrFallback());
+		ADD_CUSTOM_LOADTIMER_META(LoadPackageInternal, PackageName, *TracePackageName);
 
 	if (PackagePath.IsEmpty())
 	{
diff --git a/Engine/Source/Runtime/Engine/Private/Materials/Material.cpp b/Engine/Source/Runtime/Engine/Private/Materials/Material.cpp
--- a/Engine/Source/Runtime/Engine/Private/Materials/Material.cpp
+++ b/Engine/Source/Runtime/Engine/Private/Materials/Material.cpp
@@ -2985,8 +2985,14 @@ UE_TRACE_EVENT_END()
 void UMaterial::CacheResourceShadersForRendering(bool bRegenerateId, EMaterialShaderPrecompileMode PrecompileMode)
 {
 	TRACE_CPUPROFILER_EVENT_SCOPE(UMaterial::CacheResourceShadersForRendering);
+
+	FString TraceMaterialName;
+	if (UE_TRACE_CHANNELEXPR_IS_ENABLED(CpuChannel))
+	{
+		TraceMaterialName = GetFullName();
+	}
 	UE_TRACE_LOG_SCOPED_T(Cpu, CacheResourceShadersForRendering, CpuChannel)
-		<< CacheResourceShadersForRendering.MaterialName(*GetFullName());
+		<< CacheResourceShadersForRendering.MaterialName(*TraceMaterialName);
 
 #if WITH_EDITOR
 	// Always rebuild the shading model field on recompile
