WorldPartitionHLOD builder - Use ScanModifiedAssetFiles() in place of ScanPathsSynchronous(), to properly handle the deleted external actors files

diff --git a/Engine/Source/Editor/UnrealEd/Private/WorldPartition/WorldPartitionHLODsBuilder.cpp b/Engine/Source/Editor/UnrealEd/Private/WorldPartition/WorldPartitionHLODsBuilder.cpp
--- a/Engine/Source/Editor/UnrealEd/Private/WorldPartition/WorldPartitionHLODsBuilder.cpp
+++ b/Engine/Source/Editor/UnrealEd/Private/WorldPartition/WorldPartitionHLODsBuilder.cpp
@@ -970,7 +970,7 @@ bool UWorldPartitionHLODsBuilder::CopyFilesFromWorkingDir(const FString& SourceD
 
 	// Force a rescan of the updated files
 	IAssetRegistry& AssetRegistry = FModuleManager::LoadModuleChecked<FAssetRegistryModule>(TEXT("AssetRegistry")).Get();
-	AssetRegistry.ScanPathsSynchronous(ModifiedFiles.Array(), /*bForceRescan*/true, /*bIgnoreBlackListScanFilters*/true);
+	AssetRegistry.ScanModifiedAssetFiles(ModifiedFiles.Array());
 
 	return true;
 }
