ICookPackageSplitter: Remove spurious error if Package is GC'd before PostSaveGeneratorPackage.
#rnx

diff --git a/Engine/Source/Editor/UnrealEd/Private/CookOnTheFlyServer.cpp b/Engine/Source/Editor/UnrealEd/Private/CookOnTheFlyServer.cpp
--- a/Engine/Source/Editor/UnrealEd/Private/CookOnTheFlyServer.cpp
+++ b/Engine/Source/Editor/UnrealEd/Private/CookOnTheFlyServer.cpp
@@ -2696,12 +2696,7 @@ void UCookOnTheFlyServer::ReleaseCookedPlatformData(UE::Cook::FPackageData& Pack
 	if (GeneratorPackage && PackageData.HasCompletedGeneration())
 	{
 		UObject* SplitObject = GeneratorPackage->FindSplitDataObject();
-		if (!SplitObject || !PackageData.GetPackage())
-		{
-			UE_LOG(LogCook, Error, TEXT("PackageSplitter: Could not find %s. This prevents us from restoring the package to its precooked state. Splitter=%s."),
-				(!PackageData.GetPackage() ? TEXT("UPackage") : TEXT("SplitDataObject")), *GeneratorPackage->GetSplitDataObjectName().ToString());
-		}
-		else
+		if (SplitObject && PackageData.GetPackage()) // If GC deleted the Package or SplitObject, then there is nothing we need to clean up
 		{
 			GeneratorPackage->GetCookPackageSplitterInstance()->PostSaveGeneratorPackage(PackageData.GetPackage(), SplitObject);
 		}
