CookPackageSplitter: Change created package to be at the proper PackageName for the target domain. Use the intermediate generation directory as an override filename when saving/loading to the WorkspaceDomain.
Pass placeholder versions of the packages to populate into PreSaveGeneratorPackage so that hard links to them can be saved when using deferred population.
#rnx

diff --git a/Engine/Source/Editor/UnrealEd/Classes/CookOnTheSide/CookOnTheFlyServer.h b/Engine/Source/Editor/UnrealEd/Classes/CookOnTheSide/CookOnTheFlyServer.h
--- a/Engine/Source/Editor/UnrealEd/Classes/CookOnTheSide/CookOnTheFlyServer.h
+++ b/Engine/Source/Editor/UnrealEd/Classes/CookOnTheSide/CookOnTheFlyServer.h
@@ -796,9 +796,12 @@ private:
 	 *
 	 * @param BuildFilename long package name of the package to load 
 	 * @param OutPackage UPackage of the package loaded, non-null on success, may be non-null on failure if the UPackage existed but had another failure
+	 * @param OverrideFileName Filename to use instead of PackageData.GetFileName(), for e.g. generated packages
+	 * @param ReportingPackageData PackageData which caused the load, for e.g. generated packages
 	 * @return Whether LoadPackage was completely successful and the package can be cooked
 	 */
-	bool LoadPackageForCooking(UE::Cook::FPackageData& PackageData, UPackage*& OutPackage);
+	bool LoadPackageForCooking(UE::Cook::FPackageData& PackageData, UPackage*& OutPackage, FString* LoadFromFileName = nullptr,
+		UE::Cook::FPackageData* ReportingPackageData = nullptr);
 
 	/**
 	* Initialize the sandbox for @param TargetPlatforms
diff --git a/Engine/Source/Editor/UnrealEd/Private/CookOnTheFlyServer.cpp b/Engine/Source/Editor/UnrealEd/Private/CookOnTheFlyServer.cpp
--- a/Engine/Source/Editor/UnrealEd/Private/CookOnTheFlyServer.cpp
+++ b/Engine/Source/Editor/UnrealEd/Private/CookOnTheFlyServer.cpp
@@ -44,6 +44,7 @@
 #include "UObject/SavePackage.h"
 #include "UObject/MetaData.h"
 #include "UObject/ConstructorHelpers.h"
+#include "UObject/ReferenceChainSearch.h"
 #include "UObject/ObjectSaveContext.h"
 #include "UObject/UObjectArray.h"
 #include "Misc/PackageName.h"
@@ -228,7 +229,7 @@ namespace UE
 {
 namespace Cook
 {
-static const TCHAR* GeneratedPackageSubPath = TEXT("_Generated_");
+const TCHAR* GeneratedPackageSubPath = TEXT("_Generated_");
 }
 }
 
@@ -1925,7 +1926,7 @@ void UCookOnTheFlyServer::LoadPackageInQueue(UE::Cook::FPackageData& PackageData
 	OutNumPushed = 0;
 
 	FName PackageFileName(PackageData.GetFileName());
-	if (PackageData.IsLoadFromDisk())
+	if (!PackageData.IsGenerated())
 	{
 		bool bLoadFullySuccessful = LoadPackageForCooking(PackageData, LoadedPackage);
 		if (!bLoadFullySuccessful)
@@ -1971,7 +1972,6 @@ void UCookOnTheFlyServer::LoadPackageInQueue(UE::Cook::FPackageData& PackageData
 	}
 	else
 	{
-		check(PackageData.IsGenerated());
 		FGeneratorPackage* GeneratorStruct = PackageData.GetGeneratedOwner();
 		if (!GeneratorStruct)
 		{
@@ -1979,7 +1979,6 @@ void UCookOnTheFlyServer::LoadPackageInQueue(UE::Cook::FPackageData& PackageData
 			RejectPackageToLoad(PackageData, TEXT("is an orphaned generated package"));
 			return;
 		}
-		check(GeneratorStruct->IsDeferredPopulate());
 		FGeneratorPackage::FGeneratedStruct* GeneratedStruct = GeneratorStruct->FindGeneratedStruct(&PackageData);
 		if (!GeneratedStruct)
 		{
@@ -1987,24 +1986,49 @@ void UCookOnTheFlyServer::LoadPackageInQueue(UE::Cook::FPackageData& PackageData
 			RejectPackageToLoad(PackageData, TEXT("is an orphaned generated package"));
 			return;
 		}
-		const FPackageData& OwnerPackageData = GeneratorStruct->GetOwner();
 
-		UPackage* OwnerPackage = LoadPackage(nullptr, *OwnerPackageData.GetFileName().ToString(), LOAD_None);
-		if (!OwnerPackage)
+		if (GeneratorStruct->IsDeferredPopulate())
 		{
-			ResultFlags |= COSR_ErrorLoadingPackage;
-			UE_LOG(LogCook, Error, TEXT("Package %s is a generated package and we could not load it's generator package %s. It can not be loaded."),
-				*PackageFileName.ToString(), *OwnerPackageData.GetFileName().ToString());
-			RejectPackageToLoad(PackageData, TEXT("is a generated package which could not load its generator"));
-			return;
+			FPackageData& OwnerPackageData = const_cast<FPackageData&>(GeneratorStruct->GetOwner());
+			UPackage* OwnerPackage;
+			bool bLoadFullySuccessful = LoadPackageForCooking(OwnerPackageData, OwnerPackage, nullptr, &PackageData);
+			if (!bLoadFullySuccessful)
+			{
+				ResultFlags |= COSR_ErrorLoadingPackage;
+				UE_LOG(LogCook, Error, TEXT("Package %s is a generated package and we could not load it's generator package %s. It can not be loaded."),
+					*PackageFileName.ToString(), *OwnerPackageData.GetFileName().ToString());
+				RejectPackageToLoad(PackageData, TEXT("is a generated package which could not load its generator"));
+				return;
+			}
+			FPopulatePackageContext Context(GeneratorStruct, OwnerPackage);
+			Context.GeneratedStruct = GeneratedStruct;
+			LoadedPackage = TryPopulateGeneratedPackage(Context);
+			if (!LoadedPackage)
+			{
+				RejectPackageToLoad(PackageData, TEXT("is a generated package which could not be populated"));
+				return;
+			}
 		}
-		FPopulatePackageContext Context(GeneratorStruct, OwnerPackage);
-		Context.GeneratedStruct = GeneratedStruct;
-		LoadedPackage = TryPopulateGeneratedPackage(Context);
-		if (!LoadedPackage)
+		else
 		{
-			RejectPackageToLoad(PackageData, TEXT("is a generated package which could not be populated"));
-			return;
+			FString LoadFromFileName = GeneratorStruct->GetIntermediateLocalPath(*GeneratedStruct);
+			bool bLoadFullySuccessful = LoadPackageForCooking(PackageData, LoadedPackage, &LoadFromFileName);
+			if (!bLoadFullySuccessful)
+			{
+				ResultFlags |= COSR_ErrorLoadingPackage;
+				UE_LOG(LogCook, Verbose, TEXT("Not cooking package %s"), *LoadFromFileName);
+				RejectPackageToLoad(PackageData, TEXT("failed to load"));
+				return;
+			}
+			check(LoadedPackage != nullptr && LoadedPackage->IsFullyLoaded());
+			if (LoadedPackage->GetFName() != PackageData.GetPackageName())
+			{
+				ResultFlags |= COSR_ErrorLoadingPackage;
+				UE_LOG(LogCook, Error, TEXT("Generated package loaded for package name %s, but received package name %s. Remapping of names for generated packages is not supported so it can not be loaded."),
+					*PackageData.GetPackageName().ToString(), *LoadedPackage->GetName());
+				RejectPackageToLoad(PackageData, TEXT("had a remapped name"));
+				return;
+			}
 		}
 	}
 
@@ -2208,23 +2232,6 @@ UE::Cook::FGeneratorPackage* UCookOnTheFlyServer::CreateGeneratorPackage(UE::Coo
 
 	UE_LOG(LogCook, Display, TEXT("Splitting Package %s with class %s acting on object %s."), *PackageData.GetPackageName().ToString(), *Splitter->GetSplitDataClass()->GetName(), *SplitDataObject->GetFullName());
 
-	// Prepare necessary paths
-	const UPackage* GeneratorPackage = PackageData.GetPackage();
-	const FString GeneratorPackagePath = FPackageName::GetLongPackagePath(GeneratorPackage->GetPathName());
-	const FString GeneratorPackageShortName = FPackageName::GetShortName(GeneratorPackage->GetName());
-	const FString GeneratedUncookedRootPath = FPaths::RemoveDuplicateSlashes(FString::Printf(TEXT("/%s%s/"), *GeneratorPackageShortName, UE::Cook::GeneratedPackageSubPath));
-	const FString GeneratedUncookedContentPath = FPaths::RemoveDuplicateSlashes(FString::Printf(TEXT("%s/Cooked/%s/%s/"), *FPaths::ProjectIntermediateDir(), *GeneratorPackageShortName, UE::Cook::GeneratedPackageSubPath));
-
-	// Empty the uncooked intermediate directory
-	const FString GeneratedDirectoryToDelete = FPaths::ConvertRelativePathToFull(FPaths::CreateStandardFilename(GeneratedUncookedContentPath));
-	GetAsyncIODelete().DeleteDirectory(*GeneratedDirectoryToDelete);
-
-	// Register Mount Point for Generator Package : /<GeneratorPackageShortName>_Generated_/ -> <ProjectIntermediateDir>/Cooked/<GeneratorPackageShortName>/_Generated_/
-	if (!FPackageName::MountPointExists(GeneratedUncookedRootPath))
-	{
-		FPackageName::RegisterMountPoint(GeneratedUncookedRootPath, GeneratedUncookedContentPath);
-	}
-
 	// Create instance of CookPackageSplitter class
 	ICookPackageSplitter* SplitterInstance = Splitter->CreateInstance(SplitDataObject);
 	if (!SplitterInstance)
@@ -2234,9 +2241,21 @@ UE::Cook::FGeneratorPackage* UCookOnTheFlyServer::CreateGeneratorPackage(UE::Coo
 	}
 
 	// Create a FGeneratorPackage helper object using this CookPackageSplitter instance and return it
-	PackageData.CreateGeneratorPackage(SplitDataObject, SplitterInstance, GeneratedUncookedRootPath);
-
-	return PackageData.GetGeneratorPackage();
+	PackageData.CreateGeneratorPackage(SplitDataObject, SplitterInstance);
+	UE::Cook::FGeneratorPackage* GeneratorPackage = PackageData.GetGeneratorPackage();
+	if (!SplitterInstance->UseDeferredPopulate())
+	{
+		// Prepare the uncooked intermediate mountpoint directory
+		FString MountPackagePath;
+		FString MountLocalFilePath;
+		GeneratorPackage->GetIntermediateMountPoint(MountPackagePath, MountLocalFilePath);
+		GetAsyncIODelete().DeleteDirectory(*MountLocalFilePath);
+		if (!FPackageName::MountPointExists(MountPackagePath))
+		{
+			FPackageName::RegisterMountPoint(MountPackagePath, MountLocalFilePath);
+		}
+	}
+	return GeneratorPackage;
 }
 
 void UCookOnTheFlyServer::SplitPackage(UE::Cook::FGeneratorPackage* GeneratorStruct, bool& bOutCompleted, bool& bOutError)
@@ -2246,75 +2265,129 @@ void UCookOnTheFlyServer::SplitPackage(UE::Cook::FGeneratorPackage* GeneratorStr
 	bOutCompleted = false;
 	bOutError = false;
 	check(GeneratorStruct);
-
 	ICookPackageSplitter* Splitter = GeneratorStruct->GetCookPackageSplitterInstance();
-	if (Splitter->UseDeferredPopulate())
+	UObject* SplitObject = GeneratorStruct->FindSplitDataObject();
+	if (!SplitObject)
 	{
-		while (const FGeneratorPackage::FGeneratedStruct* GeneratedStruct = GeneratorStruct->GetNextPackageToGenerate())
+		UE_LOG(LogCook, Error, TEXT("Could not find SplitDataObject %s"), *GeneratorStruct->GetSplitDataObjectName().ToString());
+		bOutError = true;
+		return;
+	}
+
+	if (!GeneratorStruct->HasGeneratedList())
+	{
+		// Call the splitter to generate the list
+		if (!GeneratorStruct->TryGenerateList(SplitObject, GetPackageNameCache(), *PackageDatas))
 		{
-			FPackageData* GeneratedPackageData = GeneratedStruct->PackageData;
-			if (!GeneratedPackageData)
-			{
-				UE_LOG(LogCook, Error, TEXT("PackageSplitter was not given valid PackageData to queue for generated package: no package data. Splitter=%s, Generated=%s."),
-					*GeneratorStruct->GetSplitDataObjectName().ToString(), *GeneratedStruct->GeneratorData.RelativePath);
-				bOutError = true;
-				return;
-			}
-			GeneratedPackageData->ClearCookedPlatformData();
-			QueueDiscoveredPackageData(*GeneratedPackageData);
+			bOutError = true;
+			return;
 		}
+		GeneratorStruct->SetGeneratedList();
 	}
-	else
-	{
-		UPackage* Owner = GeneratorStruct->GetOwner().GetPackage();
-		FPopulatePackageContext Context(GeneratorStruct, Owner);
 
-		// Start Splitting
-		while (const FGeneratorPackage::FGeneratedStruct* GeneratedStruct = GeneratorStruct->GetNextPackageToGenerate())
+	if (!GeneratorStruct->HasClearedOldPackages())
+	{
+		for (const FGeneratorPackage::FGeneratedStruct& GeneratedStruct : GeneratorStruct->GetPackagesToGenerate())
 		{
-			FPackageData* GeneratedPackageData = GeneratedStruct->PackageData;
-			if (!GeneratedPackageData)
+			const FString GeneratedPackageName = GeneratedStruct.PackageData->GetPackageName().ToString();
+			if (FindObject<UPackage>(nullptr, *GeneratedPackageName))
 			{
-				UE_LOG(LogCook, Error, TEXT("PackageSplitter was not given valid PackageData to queue for generated package: no package data. Splitter=%s, Generated=%s."),
-					*GeneratorStruct->GetSplitDataObjectName().ToString(), *GeneratedStruct->GeneratorData.RelativePath);
-				bOutError = true;
-				return;
+				if (GeneratorStruct->HasClearedOldPackagesWithGC())
+				{
+					UE_LOG(LogCook, Error, TEXT("PackageSplitter was unable to construct new generated packages because an old version of the package is already in memory and GC did not remove it. Splitter=%s, Generated=%s."),
+						*GeneratorStruct->GetSplitDataObjectName().ToString(), *GeneratedStruct.RelativePath);
+					bOutError = true;
+					return;
+				}
+				else
+				{
+					GeneratorStruct->SetClearedOldPackagesWithGC();
+					return;
+				}
 			}
-			GeneratedPackageData->ClearCookedPlatformData();
+		}
 
-			Context.GeneratedStruct = GeneratedStruct;
-			UPackage* GeneratedPackage = TryPopulateGeneratedPackage(Context);
-			if (!GeneratedPackage)
-			{
-				bOutError = true;
-				return;
-			}
+		GeneratorStruct->SetClearedOldPackages();
+	}
 
-			// Save package into the uncooked intermediate directory
-			FSavePackageResultStruct SaveResult = GEditor->Save(GeneratedPackage, nullptr, RF_Standalone,
-				*GeneratedPackageData->GetFileName().ToString(), GError, nullptr, /*bForceByteSwapping*/ false,
-				/*bWarnOfLongFilename*/ true, SAVE_None);
-			if (SaveResult.Result != ESavePackageResult::Success)
+	if (!GeneratorStruct->HasQueuedGeneratedPackages())
+	{
+		if (Splitter->UseDeferredPopulate())
+		{
+			for (const FGeneratorPackage::FGeneratedStruct& GeneratedStruct : GeneratorStruct->GetPackagesToGenerate())
 			{
-				UE_LOG(LogCook, Error, TEXT("PackageSplitter could not save uncooked generated package. Splitter=%s, Generated=%s."),
-					*GeneratorStruct->GetSplitDataObjectName().ToString(), *GeneratedPackageData->GetFileName().ToString());
-				bOutError = true;
-				return;
+				FPackageData* GeneratedPackageData = GeneratedStruct.PackageData;
+				GeneratedPackageData->ClearCookedPlatformData();
+				QueueDiscoveredPackageData(*GeneratedPackageData);
 			}
+		}
+		else
+		{
+			UPackage* Owner = GeneratorStruct->GetOwner().GetPackage();
+			FPopulatePackageContext Context(GeneratorStruct, Owner);
+			TArrayView<FGeneratorPackage::FGeneratedStruct> PackagesToGenerate = GeneratorStruct->GetPackagesToGenerate();
+			// Start Splitting
+			int32& NextIndex = GeneratorStruct->GetNextPopulateIndex();
+			while (NextIndex < PackagesToGenerate.Num())
+			{
+				const FGeneratorPackage::FGeneratedStruct& GeneratedStruct = PackagesToGenerate[NextIndex++];
+				FPackageData* GeneratedPackageData = GeneratedStruct.PackageData;
+				GeneratedPackageData->ClearCookedPlatformData();
 
-			// Update Asset Registry with the generated package
-			AssetRegistry->ScanFilesSynchronous({ GeneratedPackageData->GetFileName().ToString() }, true);
+				Context.GeneratedStruct = &GeneratedStruct;
+				UPackage* GeneratedPackage = TryPopulateGeneratedPackage(Context);
+				if (!GeneratedPackage)
+				{
+					bOutError = true;
+					return;
+				}
+
+				// Save package into the uncooked intermediate directory
+				FString IntermediateFile = GeneratorStruct->GetIntermediateLocalPath(GeneratedStruct);
+				FSavePackageResultStruct SaveResult = GEditor->Save(GeneratedPackage, nullptr, RF_Standalone,
+					*IntermediateFile, GError, nullptr, /*bForceByteSwapping*/ false,
+					/*bWarnOfLongFilename*/ true, SAVE_None);
+				if (SaveResult.Result != ESavePackageResult::Success)
+				{
+					UE_LOG(LogCook, Error, TEXT("PackageSplitter could not save uncooked generated package. Splitter=%s, Generated=%s."),
+						*GeneratorStruct->GetSplitDataObjectName().ToString(), *IntermediateFile);
+					bOutError = true;
+					return;
+				}
 
-			// Queue the generated package
-			QueueDiscoveredPackageData(*GeneratedPackageData, true /* bLoadReady */);
+				// Queue the generated package
+				QueueDiscoveredPackageData(*GeneratedPackageData, true /* bLoadReady */);
 
-			if (HasExceededMaxMemory())
-			{
-				return;
+				if (HasExceededMaxMemory())
+				{
+					return;
+				}
 			}
 		}
+		GeneratorStruct->SetQueuedGeneratedPackages();
 	}
 
+	TArrayView<FGeneratorPackage::FGeneratedStruct> PackagesToGenerate = GeneratorStruct->GetPackagesToGenerate();
+	TArray<ICookPackageSplitter::FGeneratedPackageForPreSave> SplitterDatas;
+	SplitterDatas.Reserve(PackagesToGenerate.Num());
+	for (const FGeneratorPackage::FGeneratedStruct& GeneratedStruct : PackagesToGenerate)
+	{
+		ICookPackageSplitter::FGeneratedPackageForPreSave& SplitterData = SplitterDatas.Emplace_GetRef();
+		SplitterData.RelativePath = GeneratedStruct.RelativePath;
+		SplitterData.bCreatedAsMap = GeneratedStruct.bCreateAsMap;
+
+		const FString GeneratedPackageName = GeneratedStruct.PackageData->GetPackageName().ToString();
+		SplitterData.Package = FindObject<UPackage>(nullptr, *GeneratedPackageName);
+		if (!SplitterData.Package)
+		{
+			SplitterData.Package = CreatePackage(*GeneratedPackageName);
+		}
+	}
+
+	UPackage* OwnerPackage = GeneratorStruct->GetOwner().GetPackage();
+	check(OwnerPackage);
+	GeneratorStruct->GetCookPackageSplitterInstance()->PreSaveGeneratorPackage(OwnerPackage, SplitObject, SplitterDatas);
+
 	bOutCompleted = true;
 }
 
@@ -2323,8 +2396,8 @@ UPackage* UCookOnTheFlyServer::TryPopulateGeneratedPackage(UE::Cook::FPopulatePa
 	using namespace UE::Cook;
 
 	FGeneratorPackage* GeneratorStruct = Context.GeneratorStruct;
+	const FGeneratorPackage::FGeneratedStruct& GeneratedStruct = *Context.GeneratedStruct;
 	const UPackage* OwnerPackage = Context.OwnerPackage;
-	const ICookPackageSplitter::FGeneratedPackage& GeneratorData = Context.GeneratedStruct->GeneratorData;
 	UE::Cook::FPackageData* GeneratedPackageData = Context.GeneratedStruct->PackageData;
 	const FString GeneratedPackageName = GeneratedPackageData->GetPackageName().ToString();
 
@@ -2347,45 +2420,62 @@ UPackage* UCookOnTheFlyServer::TryPopulateGeneratedPackage(UE::Cook::FPopulatePa
 			*Context.GeneratorPackagePath, *Context.GeneratorPackageShortName, UE::Cook::GeneratedPackageSubPath));
 	}
 
-	check(GeneratorData.GetCreateAsMap().IsSet()); // Should have been confirmed in TryGenerateList and not allowed to get here if not set
 	check(GeneratedPackageData); // Caller already checked this
 
-	// Make package name to fit Mounted GeneratedUncookedRootPath
-	const FString GeneratedPackageCookedName = FPaths::RemoveDuplicateSlashes(FString::Printf(TEXT("/%s/%s"), *Context.GeneratedCookedRootPath, *GeneratorData.RelativePath));
-	const FString GeneratedPackageExtension = FString(FPathViews::GetExtension(GeneratedPackageData->GetFileName().ToString(), true /* bIncludeDot */));
+	ICookPackageSplitter* Splitter = GeneratorStruct->GetCookPackageSplitterInstance();
 
 	// Create package
-	if (FindObject<UPackage>(nullptr, *GeneratedPackageName))
+	UPackage* GeneratedPackage = FindObject<UPackage>(nullptr, *GeneratedPackageName);
+	bool bPopulatedByPreSave = false;
+	if (GeneratedPackage)
 	{
-		UE_LOG(LogCook, Error, TEXT("PackageSplitter found an existing copy of a package it was trying to generate; this is not yet supported. Splitter=%s, Generated=%s."),
-			*GeneratorStruct->GetSplitDataObjectName().ToString(), *GeneratedPackageName);
-		return nullptr;
+		if (!Splitter->UseDeferredPopulate())
+		{
+			UE_LOG(LogCook, Error, TEXT("PackageSplitter found an existing copy of a package it was trying to populate;")
+				TEXT("this is unexpected since we previously checked and did not find it. Splitter=%s, Generated=%s."),
+				*GeneratorStruct->GetSplitDataObjectName().ToString(), *GeneratedPackageName);
+			return nullptr;
+		}
+		else if (GeneratorStruct->HasGCOccurredAfterGenerate())
+		{
+			UE_LOG(LogCook, Error, TEXT("PackageSplitter found an existing copy of a package it was trying to populate;")
+				TEXT("this is unexpected since garbage has been collected and the package should have been unreferenced so it should have been collected.")
+				TEXT("Splitter=%s, Generated=%s."),
+				*GeneratorStruct->GetSplitDataObjectName().ToString(), *GeneratedPackageName);
+			FReferenceChainSearch RefChainSearch(GeneratedPackage, EReferenceChainSearchMode::Shortest | EReferenceChainSearchMode::PrintAllResults);
+			return nullptr;
+		}
+		// Otherwise this is the package that was created and passed to presave, and it is still valid because there has not been a GC since
+		// we created it. Mark its state and use it
+		bPopulatedByPreSave = true;
+	}
+	else
+	{
+		GeneratedPackage = CreatePackage(*GeneratedPackageName);
 	}
-	UPackage* GeneratedPackage = CreatePackage(*GeneratedPackageName);
 
 	// Populate package using CookPackageSplitterInstance and pass GeneratedPackage's cooked name for it to
 	// properly setup any internal reference to this package (SoftObjectPaths or others)
-	if (!GeneratorStruct->GetCookPackageSplitterInstance()->TryPopulatePackage(OwnerPackage, Context.OwnerObject, GeneratedPackage,
-		GeneratorData.RelativePath, GeneratedPackageCookedName))
+	ICookPackageSplitter::FGeneratedPackageForPopulate PopulateData;
+	PopulateData.RelativePath = GeneratedStruct.RelativePath;
+	PopulateData.Package = GeneratedPackage;
+	PopulateData.bPopulatedByPreSave = bPopulatedByPreSave;
+	PopulateData.bCreatedAsMap = GeneratedStruct.bCreateAsMap;
+	if (!Splitter->TryPopulatePackage(OwnerPackage, Context.OwnerObject, PopulateData))
 	{
 		UE_LOG(LogCook, Error, TEXT("PackageSplitter returned false from TryPopulatePackage. Splitter=%s, Generated=%s."),
 			*GeneratorStruct->GetSplitDataObjectName().ToString(), *GeneratedPackageName);
 		return nullptr;
 	}
 	bool bPackageIsMap = GeneratedPackage->ContainsMap();
-	if (bPackageIsMap != *GeneratorData.GetCreateAsMap())
+	if (bPackageIsMap != GeneratedStruct.bCreateAsMap)
 	{
 		UE_LOG(LogCook, Error, TEXT("PackageSplitter specified generated package is %s in GetGenerateList results, but then in TryPopulatePackage created it as %s. Splitter=%s, Generated=%s."),
-			(*GeneratorData.GetCreateAsMap() ? TEXT("map") : TEXT("uasset")), (bPackageIsMap ? TEXT("map") : TEXT("uasset")),
+			(GeneratedStruct.bCreateAsMap ? TEXT("map") : TEXT("uasset")), (bPackageIsMap ? TEXT("map") : TEXT("uasset")),
 			*GeneratorStruct->GetSplitDataObjectName().ToString(), *GeneratedPackageName);
 		return nullptr;
 	}
-
-	// Set override cooked FileName to be used when saving the cooked generated package
-	const FString GeneratePackageCookedFileName = FPaths::RemoveDuplicateSlashes(FString::Printf(TEXT("%s/%s/%s/%s%s"),
-		*FPaths::GetPath(GeneratorStruct->GetOwner().GetFileName().ToString()), *Context.GeneratorPackageShortName,
-		UE::Cook::GeneratedPackageSubPath, *GeneratorData.RelativePath, *GeneratedPackageExtension));
-	GeneratedPackageData->SetCookedFileName(FPackageNameCache::GetStandardFileName(GeneratePackageCookedFileName));
+	GeneratedPackage->MarkAsFullyLoaded();
 
 	return GeneratedPackage;
 }
@@ -2451,13 +2541,6 @@ bool UCookOnTheFlyServer::BeginPrepareSave(UE::Cook::FPackageData& PackageData,
 					PackageData.SetHasBeginPrepareSaveFailed(true);
 					return false;
 				}
-
-				// Call the splitter to generate the list
-				if (!Generator->TryGenerateList(Package, SplitDataObject, GetPackageNameCache(), *PackageDatas))
-				{
-					PackageData.SetHasBeginPrepareSaveFailed(true);
-					return false;
-				}
 			}
 		}
 		else
@@ -2483,14 +2566,6 @@ bool UCookOnTheFlyServer::BeginPrepareSave(UE::Cook::FPackageData& PackageData,
 			{
 				return false;
 			}
-			UObject* SplitObject = Generator->FindSplitDataObject();
-			if (!SplitObject)
-			{
-				UE_LOG(LogCook, Error, TEXT("Could not find SplitDataObject for %s"), *Package->GetName());
-				PackageData.SetHasBeginPrepareSaveFailed(true);
-				return false;
-			}
-			Generator->GetCookPackageSplitterInstance()->PreSaveGeneratorPackage(Package, SplitObject);
 		}
 		PackageData.SetCompletedGeneration(true);
 	}
@@ -2721,7 +2796,8 @@ void UCookOnTheFlyServer::TickCancels()
 	PackageDatas->PollPendingCookedPlatformDatas();
 }
 
-bool UCookOnTheFlyServer::LoadPackageForCooking(UE::Cook::FPackageData& PackageData, UPackage*& OutPackage)
+bool UCookOnTheFlyServer::LoadPackageForCooking(UE::Cook::FPackageData& PackageData, UPackage*& OutPackage, FString* LoadFromFilename,
+	UE::Cook::FPackageData* ReportingPackageData)
 {
 	UE_SCOPED_HIERARCHICAL_COOKTIMER_AND_DURATION(LoadPackageForCooking, DetailedCookStats::TickCookOnTheSideLoadPackagesTimeSec);
 
@@ -2732,13 +2808,13 @@ bool UCookOnTheFlyServer::LoadPackageForCooking(UE::Cook::FPackageData& PackageD
 		PackageTracker->LoadingPackageData = nullptr;
 	};
 
-	OutPackage = NULL;
-	FString PackageNameString;
-	OutPackage = FindObject<UPackage>(ANY_PACKAGE, *PackageData.GetPackageName().ToString());
+	FString PackageName = PackageData.GetPackageName().ToString();
+	OutPackage = FindObject<UPackage>(ANY_PACKAGE, *PackageName);
 
-	FString FileName(PackageData.GetFileName().ToString());
+	FString FileName(LoadFromFilename ? *LoadFromFilename : PackageData.GetFileName().ToString());
+	FString ReportingFileName(ReportingPackageData ? ReportingPackageData->GetFileName().ToString() : FileName);
 #if DEBUG_COOKONTHEFLY
-	UE_LOG(LogCook, Display, TEXT("Processing request %s"), *FileName);
+	UE_LOG(LogCook, Display, TEXT("Processing request %s"), *ReportingFileName);
 #endif
 	static TSet<FString> CookWarningsList;
 	if (CookWarningsList.Contains(FileName) == false)
@@ -2753,7 +2829,21 @@ bool UCookOnTheFlyServer::LoadPackageForCooking(UE::Cook::FPackageData& PackageD
 	{
 		bool bWasPartiallyLoaded = OutPackage != nullptr;
 		GIsCookerLoadingPackage = true;
-		UPackage* LoadedPackage = LoadPackage(NULL, *FileName, LOAD_None);
+		UPackage* LoadIntoPackage = nullptr;
+		if (LoadFromFilename)
+		{
+			// When loading from a specified filename, we provide the packagename to load into via LoadPackage's first argument - the InOuter package
+			if (OutPackage)
+			{
+				LoadIntoPackage = OutPackage;
+			}
+			else
+			{
+				LoadIntoPackage = CreatePackage(*PackageName);
+			}
+		}
+
+		UPackage* LoadedPackage = LoadPackage(LoadIntoPackage, *FileName, LOAD_None);
 		if (LoadedPackage)
 		{
 			OutPackage = LoadedPackage;
@@ -2777,7 +2867,7 @@ bool UCookOnTheFlyServer::LoadPackageForCooking(UE::Cook::FPackageData& PackageD
 #if DEBUG_COOKONTHEFLY
 	else
 	{
-		UE_LOG(LogCook, Display, TEXT("Package already loaded %s avoiding reload"), *FileName);
+		UE_LOG(LogCook, Display, TEXT("Package already loaded %s avoiding reload"), *ReportingFileName);
 	}
 #endif
 
@@ -2785,7 +2875,7 @@ bool UCookOnTheFlyServer::LoadPackageForCooking(UE::Cook::FPackageData& PackageD
 	{
 		if ((!IsCookOnTheFlyMode()) || (!IsCookingInEditor()))
 		{
-			LogCookerMessage(FString::Printf(TEXT("Error loading %s!"), *FileName), EMessageSeverity::Error);
+			LogCookerMessage(FString::Printf(TEXT("Error loading %s!"), *ReportingFileName), EMessageSeverity::Error);
 		}
 	}
 	GOutputCookingWarnings = false;
@@ -3857,6 +3947,14 @@ void UCookOnTheFlyServer::PostGarbageCollect()
 	check(!SavingPackageData || SavingPackageData->GetPackage() != nullptr);
 
 	GCKeepObjects.Empty();
+
+	for (FPackageData* PackageData : *PackageDatas.Get())
+	{
+		if (FGeneratorPackage* GeneratorPackage = PackageData->GetGeneratorPackage())
+		{
+			GeneratorPackage->PostGarbageCollect();
+		}
+	}
 }
 
 void UCookOnTheFlyServer::BeginDestroy()
@@ -4163,12 +4261,6 @@ void UCookOnTheFlyServer::SaveCookedPackage(UE::Cook::FPackageData& PackageData,
 			UE_LOG(LogCook, Fatal, TEXT("Package %s marked as reloading for cook by was requested to save"), *Package->GetName());
 		}
 
-		// If valid, override cooked FileName. Currently used for Generated Packages (see SplitPackage)
-		if (PackageData.GetCookedFileName() != NAME_None)
-		{
-			Filename = PackageData.GetCookedFileName().ToString();
-		}
-
 		// Use SandboxFile to do path conversion to properly handle sandbox paths (outside of standard paths in particular).
 		Filename = ConvertToFullSandboxPath(*Filename, true);
 
@@ -8862,4 +8954,4 @@ uint32 UCookOnTheFlyServer::FullLoadAndSave(uint32& CookedPackageCount)
 	return Result;
 }
 
-#undef LOCTEXT_NAMESPACE
+#undef LOCTEXT_NAMESPACE
\ No newline at end of file
diff --git a/Engine/Source/Editor/UnrealEd/Private/Cooker/CookPackageData.cpp b/Engine/Source/Editor/UnrealEd/Private/Cooker/CookPackageData.cpp
--- a/Engine/Source/Editor/UnrealEd/Private/Cooker/CookPackageData.cpp
+++ b/Engine/Source/Editor/UnrealEd/Private/Cooker/CookPackageData.cpp
@@ -7,6 +7,7 @@
 #include "CookOnTheSide/CookOnTheFlyServer.h"
 #include "CookPlatformManager.h"
 #include "Containers/StringView.h"
+#include "HAL/FileManager.h"
 #include "Misc/CoreMiscDefines.h"
 #include "Misc/Paths.h"
 #include "Misc/PreloadableFile.h"
@@ -687,9 +688,8 @@ namespace Cook
 			SetIsPreloadAttempted(true);
 			return true;
 		}
-		if (IsGenerated() && (!GeneratorPackage || GeneratorPackage->IsDeferredPopulate()))
+		if (IsGenerated())
 		{
-			// Orphaned generator packages will fail to load, so no need to preload them
 			// Deferred populate generated packages are loaded from their generator, not from disk
 			ClearPreload();
 			SetIsPreloadAttempted(true);
@@ -933,26 +933,6 @@ namespace Cook
 		GeneratedOwner = InGeneratedOwner;
 	}
 
-	bool FPackageData::IsLoadFromDisk() const
-	{
-		// Non-generated packages always load from disk
-		if (!IsGenerated())
-		{
-			return true;
-		}
-
-		// Packages generated from a splitter that does not support deferred populate load from their generated
-		// package in the intermediate content directory
-		if (GeneratedOwner && !GeneratedOwner->IsDeferredPopulate())
-		{
-			return true;
-		}
-
-		// Orphaned generated packages should be treated as failed to load
-		// Deferred populate generated packages are loaded from their Generator package
-		return false;
-	}
-
 	bool FPackageData::GeneratorPackageRequiresGC() const
 	{
 		// We consider that if a FPackageData has valid GeneratorPackage helper object,
@@ -961,22 +941,20 @@ namespace Cook
 		return IsGenerating() && !GetHasBeginPrepareSaveFailed();
 	}
 
-	UE::Cook::FGeneratorPackage* FPackageData::CreateGeneratorPackage(const UObject* InSplitDataObject, ICookPackageSplitter* InCookPackageSplitterInstance, const FString& InGeneratedUncookedRootPath)
+	UE::Cook::FGeneratorPackage* FPackageData::CreateGeneratorPackage(const UObject* InSplitDataObject, ICookPackageSplitter* InCookPackageSplitterInstance)
 	{
 		check(!GetGeneratorPackage());
-		GeneratorPackage.Reset(new UE::Cook::FGeneratorPackage(*this, InSplitDataObject, InCookPackageSplitterInstance, InGeneratedUncookedRootPath));
+		GeneratorPackage.Reset(new UE::Cook::FGeneratorPackage(*this, InSplitDataObject, InCookPackageSplitterInstance));
 		return GetGeneratorPackage();
 	}
 	
 	//////////////////////////////////////////////////////////////////////////
 	// FGeneratorPackage
 
-	FGeneratorPackage::FGeneratorPackage(UE::Cook::FPackageData& InOwner, const UObject* InSplitDataObject, ICookPackageSplitter* InCookPackageSplitterInstance, const FString& InGeneratedUncookedRootPath)
+	FGeneratorPackage::FGeneratorPackage(UE::Cook::FPackageData& InOwner, const UObject* InSplitDataObject, ICookPackageSplitter* InCookPackageSplitterInstance)
 	: Owner(InOwner)
 	, SplitDataObjectName(*InSplitDataObject->GetFullName())
-	, GeneratedUncookedRootPath(InGeneratedUncookedRootPath)
 	{
-		check(PackageToGenerateNextIndex == 0);
 		check(InCookPackageSplitterInstance);
 		CookPackageSplitterInstance.Reset(InCookPackageSplitterInstance);
 	}
@@ -999,27 +977,31 @@ namespace Cook
 		}
 	}
 
-	bool FGeneratorPackage::TryGenerateList(UPackage* OwnerPackage, UObject* OwnerObject, const FPackageNameCache& PackageNameCache, FPackageDatas& PackageDatas)
+	bool FGeneratorPackage::TryGenerateList(UObject* OwnerObject, const FPackageNameCache& PackageNameCache, FPackageDatas& PackageDatas)
 	{
+		UPackage* OwnerPackage = Owner.GetPackage();
+		check(OwnerPackage);
 		TArray<ICookPackageSplitter::FGeneratedPackage> GeneratorDatas = CookPackageSplitterInstance->GetGenerateList(OwnerPackage, OwnerObject);
 		PackagesToGenerate.Reset(GeneratorDatas.Num());
-		for (ICookPackageSplitter::FGeneratedPackage& ReturnedGeneratorData : GeneratorDatas)
+		for (ICookPackageSplitter::FGeneratedPackage& SplitterData : GeneratorDatas)
 		{
 			FGeneratedStruct& GeneratedStruct = PackagesToGenerate.Emplace_GetRef();
-			GeneratedStruct.GeneratorData = MoveTemp(ReturnedGeneratorData);
-			ICookPackageSplitter::FGeneratedPackage& GeneratorData = GeneratedStruct.GeneratorData;
-			FString PackageName = FPaths::RemoveDuplicateSlashes(FString::Printf(TEXT("/%s/%s"),
-				*GetGeneratedUncookedRootPath(), *GeneratorData.RelativePath));
-			FName PackageFName(*PackageName);
+			GeneratedStruct.RelativePath = MoveTemp(SplitterData.RelativePath);
+			GeneratedStruct.Dependencies = MoveTemp(SplitterData.Dependencies);
+			FString PackageName = FPaths::RemoveDuplicateSlashes(FString::Printf(TEXT("%s/%s/%s"),
+				*this->Owner.GetPackageName().ToString(), GeneratedPackageSubPath, *GeneratedStruct.RelativePath));
 
-			if (!GeneratorData.GetCreateAsMap().IsSet())
+			if (!SplitterData.GetCreateAsMap().IsSet())
 			{
 				UE_LOG(LogCook, Error, TEXT("PackageSplitter did not specify whether CreateAsMap is true for generated package. Splitter=%s, Generated=%s."),
 					*this->GetSplitDataObjectName().ToString(), *PackageName);
 				return false;
 			}
+			GeneratedStruct.bCreateAsMap = *SplitterData.GetCreateAsMap();
+
+			const FName PackageFName(*PackageName);
 			const FName FileName = PackageNameCache.GetCachedStandardFileName(PackageFName, false /* bRequireExists */,
-				*GeneratorData.GetCreateAsMap());
+				GeneratedStruct.bCreateAsMap);
 			if (FileName.IsNone())
 			{
 				UE_LOG(LogCook, Error, TEXT("PackageSplitter could not find mounted filename for generated packagepath. Splitter=%s, Generated=%s."),
@@ -1027,7 +1009,7 @@ namespace Cook
 				return false;
 			}
 			UE::Cook::FPackageData& PackageData = PackageDatas.FindOrAddPackageData(PackageFName, FileName);
-			if (IAssetRegistry::Get()->GetAssetPackageDataCopy(PackageFName).IsSet() && !PackageData.IsGenerated())
+			if (IFileManager::Get().FileExists(*FileName.ToString()))
 			{
 				UE_LOG(LogCook, Warning, TEXT("PackageSplitter specified a generated package that already exists in the workspace domain. Splitter=%s, Generated=%s."),
 					*this->GetSplitDataObjectName().ToString(), *PackageName);
@@ -1043,15 +1025,6 @@ namespace Cook
 		return true;
 	}
 
-	const FGeneratorPackage::FGeneratedStruct* FGeneratorPackage::GetNextPackageToGenerate()
-	{
-		if (PackagesToGenerate.IsValidIndex(PackageToGenerateNextIndex))
-		{
-			return &PackagesToGenerate[PackageToGenerateNextIndex++];
-		}
-		return nullptr;
-	}
-
 	FGeneratorPackage::FGeneratedStruct* FGeneratorPackage::FindGeneratedStruct(FPackageData* PackageData)
 	{
 		for (FGeneratedStruct& GeneratedStruct : PackagesToGenerate)
@@ -1077,6 +1050,33 @@ namespace Cook
 		return FindObject<UObject>(nullptr, *ObjectPath);
 	}
 
+	void FGeneratorPackage::PostGarbageCollect()
+	{
+		if (Owner.HasCompletedGeneration())
+		{
+			bGCOccurredAfterGenerate = true;
+		}
+	}
+
+	void FGeneratorPackage::GetIntermediateMountPoint(FString& OutPackagePath, FString& OutLocalFilePath) const
+	{
+		const FString OwnerShortName = FPackageName::GetShortName(Owner.GetPackageName().ToString());
+		OutPackagePath = FPaths::RemoveDuplicateSlashes(FString::Printf(TEXT("/%s%s/"),
+			*OwnerShortName, UE::Cook::GeneratedPackageSubPath));
+		OutLocalFilePath = FPaths::RemoveDuplicateSlashes(FString::Printf(TEXT("%s/Cooked/%s/%s/"),
+			*FPaths::ProjectIntermediateDir(), *OwnerShortName, UE::Cook::GeneratedPackageSubPath));
+
+	}
+	FString FGeneratorPackage::GetIntermediateLocalPath(const FGeneratorPackage::FGeneratedStruct& GeneratedStruct) const
+	{
+		FString UnusedPackagePath;
+		FString MountLocalFilePath;
+		GetIntermediateMountPoint(UnusedPackagePath, MountLocalFilePath);
+		FString Extension = FPaths::GetExtension(GeneratedStruct.PackageData->GetFileName().ToString(), true /* bIncludeDot */);
+		return FPaths::RemoveDuplicateSlashes(FString::Printf(TEXT("%s/%s%s"),
+			*MountLocalFilePath, *GeneratedStruct.RelativePath, *Extension));
+	}
+
 	//////////////////////////////////////////////////////////////////////////
 	// FPendingCookedPlatformData
 
diff --git a/Engine/Source/Editor/UnrealEd/Private/Cooker/CookPackageData.h b/Engine/Source/Editor/UnrealEd/Private/Cooker/CookPackageData.h
--- a/Engine/Source/Editor/UnrealEd/Private/Cooker/CookPackageData.h
+++ b/Engine/Source/Editor/UnrealEd/Private/Cooker/CookPackageData.h
@@ -37,6 +37,8 @@ namespace Cook
 	class FPackageDataQueue;
 	struct FPendingCookedPlatformDataCancelManager;
 
+	extern const TCHAR* GeneratedPackageSubPath;
+
 	/** Flags specifying the behavior of FPackageData::SendToState */
 	enum class ESendFlags
 	{
@@ -258,10 +260,6 @@ namespace Cook
 		bool GetHasBeginPrepareSaveFailed() const { return static_cast<bool>(bHasBeginPrepareSaveFailed); }
 		void SetHasBeginPrepareSaveFailed(bool bValue) { bHasBeginPrepareSaveFailed = bValue != 0; }
 
-		/** Get/Set override filename to be used when saving cooked package. */
-		FName GetCookedFileName() const { return CookedFileName; }
-		void SetCookedFileName(FName InCookedFileName) { CookedFileName = InCookedFileName; }
-
 		/** Validate that the fields related to the BeginCacheForCookedPlatformData calls are empty, as required when not in the save state. */
 		void CheckCookedPlatformDataEmpty() const;
 		/** Clear the fields related to the BeginCacheForCookedPlatformData calls, when e.g. leaving the save state. Caller is responsible for having already executed any required cancellation steps to avoid dangling pending operations. */
@@ -281,7 +279,7 @@ namespace Cook
 		void RemapTargetPlatforms(const TMap<ITargetPlatform*, ITargetPlatform*>& Remap);
 
 		/** Create and return a GeneratorPackage helper object for a given CookPackageSplitter. */
-		FGeneratorPackage* CreateGeneratorPackage(const UObject* InSplitDataObject, ICookPackageSplitter* InCookPackageSplitterInstance, const FString& InGeneratedUncookedRootPath);
+		FGeneratorPackage* CreateGeneratorPackage(const UObject* InSplitDataObject, ICookPackageSplitter* InCookPackageSplitterInstance);
 		/** Destroy GeneratorPackage helper object. */
 		void DestroyGeneratorPackage() { GeneratorPackage.Reset(); }
 		/** Get GeneratorPackage helper object. */
@@ -304,8 +302,6 @@ namespace Cook
 		void SetGeneratedOwner(FGeneratorPackage* InGeneratedOwner);
 		/** Return the owning generator PackageData. Will be null if not a generated package or if orphaned */
 		FGeneratorPackage* GetGeneratedOwner() const { return GeneratedOwner; }
-		/** Return whether PumpLoads should load the package from disk or from the special generated path */
-		bool IsLoadFromDisk() const;
 
 		/** Return whether a GC is required by a generator package as soon as possible. */
 		bool GeneratorPackageRequiresGC() const;
@@ -374,7 +370,6 @@ namespace Cook
 		FCompletionCallback CompletionCallback;
 		FName PackageName;
 		FName FileName;
-		FName CookedFileName; // Override filename to be used when saving cooked package (used by generated packages)
 		TWeakObjectPtr<UPackage> Package;
 		FPackageDatas& PackageDatas; // The one-per-CookOnTheFlyServer owner of this PackageData
 		/**
@@ -419,20 +414,23 @@ namespace Cook
 	public:
 		struct FGeneratedStruct
 		{
-			ICookPackageSplitter::FGeneratedPackage GeneratorData;
+			FString RelativePath;
+			TArray<FName> Dependencies;
 			FPackageData* PackageData = nullptr;
+			bool bCreateAsMap = false;
 		};
 
 		/** Store the provided CookPackageSplitter and prepare the packages to generate. */
-		FGeneratorPackage(UE::Cook::FPackageData& InOwner, const UObject* InSplitDataObject, ICookPackageSplitter* InCookPackageSplitterInstance, const FString& InGeneratedUncookedRootPath);
+		FGeneratorPackage(UE::Cook::FPackageData& InOwner, const UObject* InSplitDataObject, ICookPackageSplitter* InCookPackageSplitterInstance);
 		~FGeneratorPackage();
 		/** Clear references to owned generated packages, and mark those packages as orphaned */
 		void ClearGeneratedPackages();
 
 		/** Call the Splitter's GetGenerateList and create the PackageDatas */
-		bool TryGenerateList(UPackage* OwnerPackage, UObject* OwnerObject, const FPackageNameCache& PackageNameCache, FPackageDatas& PackageDatas);
-		/** Return next package to generate. */
-		const FGeneratedStruct* GetNextPackageToGenerate();
+		bool TryGenerateList(UObject* OwnerObject, const FPackageNameCache& PackageNameCache, FPackageDatas& PackageDatas);
+
+		/** Accessor for the packages to generate */
+		TArrayView<FGeneratedStruct> GetPackagesToGenerate() { return PackagesToGenerate; }
 
 		/** Return CookPackageSplitter. */
 		ICookPackageSplitter* GetCookPackageSplitterInstance() const { return CookPackageSplitterInstance.Get(); }
@@ -440,8 +438,17 @@ namespace Cook
 		const UE::Cook::FPackageData& GetOwner() const { return Owner; }
 		/** Return the SplitDataObject's FullObjectPath. */
 		const FName GetSplitDataObjectName() const { return SplitDataObjectName; }
-		/** Returns root path of generated packages. */
-		const FString& GetGeneratedUncookedRootPath() const { return GeneratedUncookedRootPath; }
+		/**
+		 * Return the packagepath and localpath to the owner's mount point used for loading generated packages,
+		 * if we need to save them to disk.
+		 */
+		void GetIntermediateMountPoint(FString& OutPackagePath, FString& OutLocalPath) const;
+		/**
+		 * Return the standard path to the temporary file (including extension) where we store the given
+		 * generated package, if we need to save it to disk for later retrieval in the cook session.
+		 */
+		FString GetIntermediateLocalPath(const FGeneratorPackage::FGeneratedStruct& GeneratedStruct) const;
+
 		/** Splitter capability for whether the generated packages should populate from owner rather than load from disk */
 		bool IsDeferredPopulate() const { return CookPackageSplitterInstance->UseDeferredPopulate(); }
 		/** Return the FGeneratedStruct with the given PackageData, or null if not found. */
@@ -452,13 +459,47 @@ namespace Cook
 		 */
 		UObject* FindSplitDataObject() const;
 
+		/** State variables for reentrant SplitPackage calls */
+		bool HasGeneratedList() const { return bGeneratedList; }
+		void SetGeneratedList() { bGeneratedList = true; }
+		bool HasClearedOldPackages() { return bClearedOldPackages; }
+		void SetClearedOldPackages() { bClearedOldPackages = true; }
+		bool HasClearedOldPackagesWithGC() { return bClearedOldPackagesWithGC; }
+		void SetClearedOldPackagesWithGC() { bClearedOldPackagesWithGC = true; }
+		bool HasQueuedGeneratedPackages() const { return bQueuedGeneratedPackages; }
+		void SetQueuedGeneratedPackages() { bQueuedGeneratedPackages = true; }
+		int32& GetNextPopulateIndex() { return NextPopulateIndex; }
+
+		/**
+		 * Records whether GC occurred after creating the placeholder packages.
+		 * If GC has occurred, we expect to no longer find the placeholder packages in memory,
+		 * and if we do find them its an error because we no longer know which objects are in the packages.
+		 */
+		bool HasGCOccurredAfterGenerate() const { return bGCOccurredAfterGenerate; }
+
+		/** Callback during garbage collection */
+		void PostGarbageCollect();
+
 	private:
+		/** PackageData for the package that is being split */
 		const UE::Cook::FPackageData& Owner;
+		/** Name of the object that prompted the splitter creation */
 		FName SplitDataObjectName;
-		const FString GeneratedUncookedRootPath;
-		TUniquePtr<ICookPackageSplitter> CookPackageSplitterInstance; // Cached CookPackageSplitter
-		TArray<FGeneratedStruct> PackagesToGenerate; // Generator's packages to generate
-		int32 PackageToGenerateNextIndex = 0; // Index of next package to generate
+		/** Cached CookPackageSplitter */
+		TUniquePtr<ICookPackageSplitter> CookPackageSplitterInstance;
+		/** Recorded list of packages to generate from the splitter, and data we need about them */
+		TArray<FGeneratedStruct> PackagesToGenerate;
+
+		/** State variables for reentrant SplitPackage calls */
+		int32 NextPopulateIndex = 0;
+		bool bGeneratedList = false;
+		bool bClearedOldPackages = false;
+		bool bClearedOldPackagesWithGC = false;
+		bool bQueuedGeneratedPackages = false;
+
+		/** Records whether GC occurred after creating the placeholder packages. */
+		bool bGCOccurredAfterGenerate = false;
+
 	};
 
 
diff --git a/Engine/Source/Editor/UnrealEd/Public/CookPackageSplitter.h b/Engine/Source/Editor/UnrealEd/Public/CookPackageSplitter.h
--- a/Engine/Source/Editor/UnrealEd/Public/CookPackageSplitter.h
+++ b/Engine/Source/Editor/UnrealEd/Public/CookPackageSplitter.h
@@ -49,6 +49,20 @@ public:
 	/** Return the list of packages to generate. */
 	virtual TArray<FGeneratedPackage> GetGenerateList(const UPackage* OwnerPackage, const UObject* OwnerObject) = 0;
 
+	struct FGeneratedPackageForPopulate
+	{
+		/** RelativePath returned from GetGenerateList */
+		FString RelativePath;
+		/**
+		 * Non-null, constructed package. May have been previously passed into PreSaveGeneratedPackage,
+		 * or may be completely empty if that previous package was garbage collected before this call.
+		 */
+		UPackage* Package = nullptr;
+		/** *GetCreateAsMap returned from GetGenerateList. The package filename extension has already been set based on this. */
+		bool bCreatedAsMap = false;
+		/** Whether the Package was populated by PreSave or constructed empty */
+		bool bPopulatedByPreSave = false;
+	};
 	/**
 	 * Try to populate a generated package.
 	 *
@@ -59,19 +73,29 @@ public:
 	 *
 	 * @param OwnerPackage				The parent package being split
 	 * @param OwnerObject				The SplitDataClass instance that this CookPackageSplitter instance was created for
-	 * @param GeneratedPackage			The package the cooker generated for the given FGeneratedPackage
-	 * @param RelativePath				The relative path from the given FGeneratedPackage
-	 * @param GeneratedPackageCookName	The name that will be given for the cooked package (differs from the uncooked package)
+	 * @param GeneratedPackage			Pointer and information about the package to populate
 	 * @return							True if successfully populates,  false on error (this will cause a cook error).
 	 */
-	virtual bool TryPopulatePackage(const UPackage* OwnerPackage, const UObject* OwnerObject, UPackage* GeneratedPackage, const FStringView& RelativePath, const FStringView& GeneratedPackageCookName) = 0;
+	virtual bool TryPopulatePackage(const UPackage* OwnerPackage, const UObject* OwnerObject, const FGeneratedPackageForPopulate& GeneratedPackage) = 0;
 
+	struct FGeneratedPackageForPreSave
+	{
+		/** RelativePath returned from GetGenerateList */
+		FString RelativePath;
+		/**
+		 * Non-null, constructed package. May have been previously passed into PreSaveGeneratedPackage,
+		 * or may be completely empty if that previous package was garbage collected before this call.
+		 */
+		UPackage* Package = nullptr;
+		/** *GetCreateAsMap returned from GetGenerateList. The package filename extension has already been set based on this. */
+		bool bCreatedAsMap = false;
+	};
 	/**
 	 * Called before saving the parent generator package. Depending on the value of UseDeferredPopulate, this may be
 	 * before or after the TryPopulatePackage calls.
 	 * Make any required adjustments to the parent package before it is saved into the target domain.
 	 */
-	virtual void PreSaveGeneratorPackage(UPackage* OwnerPackage, UObject* OwnerObject) {}
+	virtual void PreSaveGeneratorPackage(UPackage* OwnerPackage, UObject* OwnerObject, const TArray<FGeneratedPackageForPreSave>& PlaceholderPackages) {}
 
 	/**
 	 * Called after saving the parent generator package. Undo any required adjustments to the parent package that
diff --git a/Engine/Source/Runtime/Engine/Private/WorldPartition/WorldPartition.cpp b/Engine/Source/Runtime/Engine/Private/WorldPartition/WorldPartition.cpp
--- a/Engine/Source/Runtime/Engine/Private/WorldPartition/WorldPartition.cpp
+++ b/Engine/Source/Runtime/Engine/Private/WorldPartition/WorldPartition.cpp
@@ -1197,10 +1197,10 @@ bool UWorldPartition::GenerateStreaming(EWorldPartitionStreamingMode Mode, TArra
 	return Result;
 }
 
-bool UWorldPartition::PopulateGeneratedPackageForCook(UPackage* InPackage, const FString& InPackageRelativePath, const FString& InPackageCookName)
+bool UWorldPartition::PopulateGeneratedPackageForCook(UPackage* InPackage, const FString& InPackageRelativePath)
 {
 	check(RuntimeHash);
-	return RuntimeHash->PopulateGeneratedPackageForCook(InPackage, InPackageRelativePath, InPackageCookName);
+	return RuntimeHash->PopulateGeneratedPackageForCook(InPackage, InPackageRelativePath);
 }
 
 void UWorldPartition::FinalizeGeneratedPackageForCook()
diff --git a/Engine/Source/Runtime/Engine/Private/WorldPartition/WorldPartitionCookPackageSplitter.cpp b/Engine/Source/Runtime/Engine/Private/WorldPartition/WorldPartitionCookPackageSplitter.cpp
--- a/Engine/Source/Runtime/Engine/Private/WorldPartition/WorldPartitionCookPackageSplitter.cpp
+++ b/Engine/Source/Runtime/Engine/Private/WorldPartition/WorldPartitionCookPackageSplitter.cpp
@@ -63,17 +63,19 @@ TArray<ICookPackageSplitter::FGeneratedPackage> FWorldPartitionCookPackageSplitt
 	return PackagesToGenerate;
 }
 
-bool FWorldPartitionCookPackageSplitter::TryPopulatePackage(const UPackage* OwnerPackage, const UObject* OwnerObject, UPackage* GeneratedPackage, const FStringView& RelativePath, const FStringView& GeneratedPackageCookName)
+bool FWorldPartitionCookPackageSplitter::TryPopulatePackage(const UPackage* OwnerPackage, const UObject* OwnerObject,
+	const FGeneratedPackageForPopulate& GeneratedPackage)
 {
 	// TODO: Make PopulateGeneratedPackageForCook const so we can honor the constness of the OwnerObject in this API function
 	const UWorld* ConstPartitionedWorld = ValidateDataObject(OwnerObject);
 	UWorld* PartitionedWorld = const_cast<UWorld*>(ConstPartitionedWorld);
 	UWorldPartition* WorldPartition = PartitionedWorld->PersistentLevel->GetWorldPartition();
-	bool bResult = WorldPartition->PopulateGeneratedPackageForCook(GeneratedPackage, FString(RelativePath), FString(GeneratedPackageCookName));
+	bool bResult = WorldPartition->PopulateGeneratedPackageForCook(GeneratedPackage.Package, GeneratedPackage.RelativePath);
 	return bResult;
 }
 
-void FWorldPartitionCookPackageSplitter::PreSaveGeneratorPackage(UPackage* OwnerPackage, UObject* OwnerObject)
+void FWorldPartitionCookPackageSplitter::PreSaveGeneratorPackage(UPackage* OwnerPackage, UObject* OwnerObject,
+	const TArray<FGeneratedPackageForPreSave>& PlaceholderPackages)
 {
 	UWorld* PartitionedWorld = ValidateDataObject(OwnerObject);
 	UWorldPartition* WorldPartition = PartitionedWorld->PersistentLevel->GetWorldPartition();
diff --git a/Engine/Source/Runtime/Engine/Private/WorldPartition/WorldPartitionCookPackageSplitter.h b/Engine/Source/Runtime/Engine/Private/WorldPartition/WorldPartitionCookPackageSplitter.h
--- a/Engine/Source/Runtime/Engine/Private/WorldPartition/WorldPartitionCookPackageSplitter.h
+++ b/Engine/Source/Runtime/Engine/Private/WorldPartition/WorldPartitionCookPackageSplitter.h
@@ -15,8 +15,10 @@ public:
 	virtual bool UseDeferredPopulate() { return false; }
 	virtual ~FWorldPartitionCookPackageSplitter() {}
 	virtual TArray<ICookPackageSplitter::FGeneratedPackage> GetGenerateList(const UPackage* OwnerPackage, const UObject* OwnerObject) override;
-	virtual bool TryPopulatePackage(const UPackage* OwnerPackage, const UObject* OwnerObject, UPackage* GeneratedPackage, const FStringView& RelativePath, const FStringView& GeneratedPackageCookName) override;
-	virtual void PreSaveGeneratorPackage(UPackage* OwnerPackage, UObject* OwnerObject) override;
+	virtual bool TryPopulatePackage(const UPackage* OwnerPackage, const UObject* OwnerObject,
+		const ICookPackageSplitter::FGeneratedPackageForPopulate& GeneratedPackage) override;
+	virtual void PreSaveGeneratorPackage(UPackage* OwnerPackage, UObject* OwnerObject,
+		const TArray<ICookPackageSplitter::FGeneratedPackageForPreSave>& PlaceholderPackages) override;
 	//~ End of ICookPackageSplitter
 
 private:
diff --git a/Engine/Source/Runtime/Engine/Private/WorldPartition/WorldPartitionRuntimeLevelStreamingCell.cpp b/Engine/Source/Runtime/Engine/Private/WorldPartition/WorldPartitionRuntimeLevelStreamingCell.cpp
--- a/Engine/Source/Runtime/Engine/Private/WorldPartition/WorldPartitionRuntimeLevelStreamingCell.cpp
+++ b/Engine/Source/Runtime/Engine/Private/WorldPartition/WorldPartitionRuntimeLevelStreamingCell.cpp
@@ -147,10 +147,10 @@ void UWorldPartitionRuntimeLevelStreamingCell::MoveAlwaysLoadedContentToPersiste
 	}
 }
 
-bool UWorldPartitionRuntimeLevelStreamingCell::PopulateGeneratedPackageForCook(UPackage* InPackage, const FString& InPackageCookName)
+bool UWorldPartitionRuntimeLevelStreamingCell::PopulateGeneratedPackageForCook(UPackage* InPackage)
 {
 	check(!IsAlwaysLoaded());
-	if (!InPackage || InPackageCookName.IsEmpty())
+	if (!InPackage)
 	{
 		return false;
 	}
@@ -159,7 +159,7 @@ bool UWorldPartitionRuntimeLevelStreamingCell::PopulateGeneratedPackageForCook(U
 	{
 		UWorldPartition* WorldPartition = GetOuterUWorldPartition();
 		UWorld* OuterWorld = WorldPartition->GetTypedOuter<UWorld>();
-		ULevelStreaming* NewLevelStreaming = CreateLevelStreaming(InPackageCookName);
+		ULevelStreaming* NewLevelStreaming = CreateLevelStreaming(InPackage->GetName());
 		check(NewLevelStreaming)
 
 		// Load cell Actors
diff --git a/Engine/Source/Runtime/Engine/Private/WorldPartition/WorldPartitionRuntimeSpatialHash.cpp b/Engine/Source/Runtime/Engine/Private/WorldPartition/WorldPartitionRuntimeSpatialHash.cpp
--- a/Engine/Source/Runtime/Engine/Private/WorldPartition/WorldPartitionRuntimeSpatialHash.cpp
+++ b/Engine/Source/Runtime/Engine/Private/WorldPartition/WorldPartitionRuntimeSpatialHash.cpp
@@ -947,14 +947,14 @@ bool UWorldPartitionRuntimeSpatialHash::CreateStreamingGrid(const FSpatialHashRu
 	return true;
 }
 
-bool UWorldPartitionRuntimeSpatialHash::PopulateGeneratedPackageForCook(UPackage* InPackage, const FString& InPackageRelativePath, const FString& InPackageCookName)
+bool UWorldPartitionRuntimeSpatialHash::PopulateGeneratedPackageForCook(UPackage* InPackage, const FString& InPackageRelativePath)
 {
 	if (UWorldPartitionRuntimeCell** MatchingCell = PackagesToGenerateForCook.Find(InPackageRelativePath))
 	{
 		UWorldPartitionRuntimeCell* Cell = *MatchingCell;
 		if (ensure(Cell))
 		{
-			return Cell->PopulateGeneratedPackageForCook(InPackage, InPackageCookName);
+			return Cell->PopulateGeneratedPackageForCook(InPackage);
 		}
 	}
 	return false;
diff --git a/Engine/Source/Runtime/Engine/Public/WorldPartition/WorldPartition.h b/Engine/Source/Runtime/Engine/Public/WorldPartition/WorldPartition.h
--- a/Engine/Source/Runtime/Engine/Public/WorldPartition/WorldPartition.h
+++ b/Engine/Source/Runtime/Engine/Public/WorldPartition/WorldPartition.h
@@ -143,7 +143,7 @@ public:
 	void RemapSoftObjectPath(FSoftObjectPath& ObjectPath);
 
 	// Cook Methods
-	bool PopulateGeneratedPackageForCook(UPackage* InPackage, const FString& InPackageRelativePath, const FString& InPackageCookName);
+	bool PopulateGeneratedPackageForCook(UPackage* InPackage, const FString& InPackageRelativePath);
 	void FinalizeGeneratedPackageForCook();
 
 	FBox GetWorldBounds() const;
diff --git a/Engine/Source/Runtime/Engine/Public/WorldPartition/WorldPartitionRuntimeCell.h b/Engine/Source/Runtime/Engine/Public/WorldPartition/WorldPartitionRuntimeCell.h
--- a/Engine/Source/Runtime/Engine/Public/WorldPartition/WorldPartitionRuntimeCell.h
+++ b/Engine/Source/Runtime/Engine/Public/WorldPartition/WorldPartitionRuntimeCell.h
@@ -136,7 +136,7 @@ class UWorldPartitionRuntimeCell : public UObject
 	virtual int32 GetActorCount() const PURE_VIRTUAL(UWorldPartitionRuntimeCell::GetActorCount, return 0;);
 
 	// Cook methods
-	virtual bool PopulateGeneratedPackageForCook(UPackage* InPackage, const FString& InPackageCookName) PURE_VIRTUAL(UWorldPartitionRuntimeCell::PopulateGeneratedPackageForCook, return false;);
+	virtual bool PopulateGeneratedPackageForCook(UPackage* InPackage) PURE_VIRTUAL(UWorldPartitionRuntimeCell::PopulateGeneratedPackageForCook, return false;);
 	virtual void MoveAlwaysLoadedContentToPersistentLevel() PURE_VIRTUAL(UWorldPartitionRuntimeCell::MoveAlwaysLoadedContentToPersistentLevel);
 	virtual void FinalizeGeneratedPackageForCook() {}
 	virtual FString GetPackageNameToCreate() const PURE_VIRTUAL(UWorldPartitionRuntimeCell::GetPackageNameToCreate, return FString(""););
diff --git a/Engine/Source/Runtime/Engine/Public/WorldPartition/WorldPartitionRuntimeHash.h b/Engine/Source/Runtime/Engine/Public/WorldPartition/WorldPartitionRuntimeHash.h
--- a/Engine/Source/Runtime/Engine/Public/WorldPartition/WorldPartitionRuntimeHash.h
+++ b/Engine/Source/Runtime/Engine/Public/WorldPartition/WorldPartitionRuntimeHash.h
@@ -26,7 +26,7 @@ class ENGINE_API UWorldPartitionRuntimeHash : public UObject
 
 	virtual void SetDefaultValues() {}
 	virtual void ImportFromWorldComposition(class UWorldComposition* WorldComposition) {}
-	virtual bool PopulateGeneratedPackageForCook(UPackage* InPackage, const FString& InPackageRelativePath, const FString& InPackageCookName) { return false; }
+	virtual bool PopulateGeneratedPackageForCook(UPackage* InPackage, const FString& InPackageRelativePath) { return false; }
 	virtual void FinalizeGeneratedPackageForCook() {}
 	virtual void FlushStreaming() {}
 	virtual bool GenerateHLOD(ISourceControlHelper* SourceControlHelper, bool bCreateActorsOnly) { return false; }
diff --git a/Engine/Source/Runtime/Engine/Public/WorldPartition/WorldPartitionRuntimeLevelStreamingCell.h b/Engine/Source/Runtime/Engine/Public/WorldPartition/WorldPartitionRuntimeLevelStreamingCell.h
--- a/Engine/Source/Runtime/Engine/Public/WorldPartition/WorldPartitionRuntimeLevelStreamingCell.h
+++ b/Engine/Source/Runtime/Engine/Public/WorldPartition/WorldPartitionRuntimeLevelStreamingCell.h
@@ -27,7 +27,7 @@ class UWorldPartitionRuntimeLevelStreamingCell : public UWorldPartitionRuntimeSp
 	const TArray<FWorldPartitionRuntimeCellObjectMapping>& GetPackages() const { return Packages; }
 
 	// Cook methods
-	virtual bool PopulateGeneratedPackageForCook(UPackage* InPackage, const FString& InPackageCookName) override;
+	virtual bool PopulateGeneratedPackageForCook(UPackage* InPackage) override;
 	virtual FString GetPackageNameToCreate() const override;
 #endif
 
diff --git a/Engine/Source/Runtime/Engine/Public/WorldPartition/WorldPartitionRuntimeSpatialHash.h b/Engine/Source/Runtime/Engine/Public/WorldPartition/WorldPartitionRuntimeSpatialHash.h
--- a/Engine/Source/Runtime/Engine/Public/WorldPartition/WorldPartitionRuntimeSpatialHash.h
+++ b/Engine/Source/Runtime/Engine/Public/WorldPartition/WorldPartitionRuntimeSpatialHash.h
@@ -202,7 +202,7 @@ public:
 #if WITH_EDITOR
 	virtual void SetDefaultValues() override;
 	virtual void ImportFromWorldComposition(class UWorldComposition* WorldComposition) override;
-	virtual bool PopulateGeneratedPackageForCook(UPackage* InPackage, const FString& InPackageRelativePath, const FString& InPackageCookName) override;
+	virtual bool PopulateGeneratedPackageForCook(UPackage* InPackage, const FString& InPackageRelativePath) override;
 	virtual void FinalizeGeneratedPackageForCook() override;
 	virtual void FlushStreaming() override;
 	virtual bool GenerateHLOD(ISourceControlHelper* SourceControlHelper, bool bCreateActorsOnly) override;
