Fix issue where the message about reload being disabled would be incorrectly displayed when reload is enabled.

#rnx

diff --git a/Engine/Source/Editor/UnrealEd/Private/Kismet2/ReloadUtilities.cpp b/Engine/Source/Editor/UnrealEd/Private/Kismet2/ReloadUtilities.cpp
--- a/Engine/Source/Editor/UnrealEd/Private/Kismet2/ReloadUtilities.cpp
+++ b/Engine/Source/Editor/UnrealEd/Private/Kismet2/ReloadUtilities.cpp
@@ -735,12 +735,12 @@ FReload::~FReload()
 
 bool FReload::GetEnableReinstancing(bool bHasChanged) const  
 { 
-	if (bHasChanged && !bEnabledMessage)
+	if (bHasChanged && !bEnableReinstancing  && !bEnabledMessage)
 	{
 		bEnabledMessage = true;
 		Ar.Logf(ELogVerbosity::Display, TEXT("Re-instancing has been disabled.  Some changes will be ignored."));
 	}
-	return bEnableReinstancing; 
+	return bEnableReinstancing;
 }
 
 
