Access Violation attempting to convert a map to World Partition

Properly shutdown world in WorldPartition conversion commandlet

diff --git a/Engine/Source/Editor/UnrealEd/Private/Commandlets/WorldPartitionConvertCommandlet.cpp b/Engine/Source/Editor/UnrealEd/Private/Commandlets/WorldPartitionConvertCommandlet.cpp
--- a/Engine/Source/Editor/UnrealEd/Private/Commandlets/WorldPartitionConvertCommandlet.cpp
+++ b/Engine/Source/Editor/UnrealEd/Private/Commandlets/WorldPartitionConvertCommandlet.cpp
@@ -852,6 +852,13 @@ int32 UWorldPartitionConvertCommandlet::Main(const FString& Params)
 		return 1;
 	}
 
+	ON_SCOPE_EXIT
+	{
+		MainWorld->ClearWorldComponents();
+		MainWorld->CleanupWorld();
+		MainWorld->RemoveFromRoot();
+	};
+
 	UPackage* MainPackage = MainLevel->GetPackage();
 	AWorldDataLayers* MainWorldDataLayers = MainWorld->GetWorldDataLayers();
 	if (!MainWorldDataLayers)
