ICookPackageSplitter: Prevent generator packages from generating twice due to a GC after their first generation before they finish saving.
#rnx

diff --git a/Engine/Source/Editor/UnrealEd/Classes/CookOnTheSide/CookOnTheFlyServer.h b/Engine/Source/Editor/UnrealEd/Classes/CookOnTheSide/CookOnTheFlyServer.h
--- a/Engine/Source/Editor/UnrealEd/Classes/CookOnTheSide/CookOnTheFlyServer.h
+++ b/Engine/Source/Editor/UnrealEd/Classes/CookOnTheSide/CookOnTheFlyServer.h
@@ -901,8 +901,10 @@ private:
 	/**
 	 * Frees all the memory used to call BeginCacheForCookedPlatformData on all the objects in PackageData.
 	 * If the calls were incomplete because the PackageData's save was cancelled, handles canceling them and leaving any required CancelManagers in GetPendingCookedPlatformDatas
+	 * 
+	 * @param bCompletedSave If false, data that can not be efficiently recomputed will be preserved to try the save again. If true, all data will be wiped.
 	 */
-	void ReleaseCookedPlatformData(UE::Cook::FPackageData& PackageData);
+	void ReleaseCookedPlatformData(UE::Cook::FPackageData& PackageData, bool bCompletedSave);
 
 	/**
 	 * Poll the GetPendingCookedPlatformDatas and release their resources when they are complete.
diff --git a/Engine/Source/Editor/UnrealEd/Private/CookOnTheFlyServer.cpp b/Engine/Source/Editor/UnrealEd/Private/CookOnTheFlyServer.cpp
--- a/Engine/Source/Editor/UnrealEd/Private/CookOnTheFlyServer.cpp
+++ b/Engine/Source/Editor/UnrealEd/Private/CookOnTheFlyServer.cpp
@@ -2671,7 +2671,7 @@ bool UCookOnTheFlyServer::FinishPrepareSave(UE::Cook::FPackageData& PackageData,
 	return true;
 }
 
-void UCookOnTheFlyServer::ReleaseCookedPlatformData(UE::Cook::FPackageData& PackageData)
+void UCookOnTheFlyServer::ReleaseCookedPlatformData(UE::Cook::FPackageData& PackageData, bool bCompletedSave)
 {
 	using namespace UE::Cook;
 
@@ -2682,7 +2682,7 @@ void UCookOnTheFlyServer::ReleaseCookedPlatformData(UE::Cook::FPackageData& Pack
 	}
 
 	// For every Object on which we called BeginCacheForCookedPlatformData, we need to call ClearAllCachedCookedPlatformData
-	if (PackageData.GetCookedPlatformDataComplete())
+	if (PackageData.GetCookedPlatformDataComplete() && bCompletedSave)
 	{
 		// Since we have completed CookedPlatformData, we know we called BeginCacheForCookedPlatformData on all objects in the package, and none are pending
 		if (!IsCookingInEditor()) // ClearAllCachedCookedPlatformData and WillNeverCacheCookedPlatformDataAgain calls are only used when not in editor
@@ -2767,26 +2767,37 @@ void UCookOnTheFlyServer::ReleaseCookedPlatformData(UE::Cook::FPackageData& Pack
 		}
 	}
 
-	FGeneratorPackage* GeneratorPackage = PackageData.GetGeneratorPackage();
-	if (GeneratorPackage && PackageData.HasCompletedGeneration())
+	if (bCompletedSave)
 	{
-		UObject* SplitObject = GeneratorPackage->FindSplitDataObject();
-		if (SplitObject && PackageData.GetPackage()) // If GC deleted the Package or SplitObject, then there is nothing we need to clean up
+		FGeneratorPackage* GeneratorPackage = PackageData.GetGeneratorPackage();
+		if (GeneratorPackage && PackageData.HasCompletedGeneration())
 		{
-			GeneratorPackage->GetCookPackageSplitterInstance()->PostSaveGeneratorPackage(PackageData.GetPackage(), SplitObject);
+			UObject* SplitObject = GeneratorPackage->FindSplitDataObject();
+			if (!SplitObject || !PackageData.GetPackage())
+			{
+				UE_LOG(LogCook, Error, TEXT("PackageSplitter: %s was GarbageCollected before we finished saving it. This will cause a failure later when we attempt to save it again and generate a second time. Splitter=%s."),
+					(!PackageData.GetPackage() ? TEXT("UPackage") : TEXT("SplitDataObject")), *GeneratorPackage->GetSplitDataObjectName().ToString());
+			}
+			else
+			{
+				GeneratorPackage->GetCookPackageSplitterInstance()->PostSaveGeneratorPackage(PackageData.GetPackage(), SplitObject);
+			}
 		}
+		PackageData.ResetGenerationProgress();
 	}
-
 	PackageData.ClearCookedPlatformData();
 
-	if (CurrentCookMode == ECookMode::CookByTheBook)
+	if (bCompletedSave)
 	{
-		UPackage* Package = PackageData.GetPackage();
-		if (Package && Package->LinkerLoad)
+		if (CurrentCookMode == ECookMode::CookByTheBook)
 		{
-			// Loaders and their handles can have large buffers held in process memory and in the system file cache from the
-			// data that was loaded.  Keeping this for the lifetime of the cook is costly, so we try and unload it here.
-			Package->LinkerLoad->FlushCache();
+			UPackage* Package = PackageData.GetPackage();
+			if (Package && Package->LinkerLoad)
+			{
+				// Loaders and their handles can have large buffers held in process memory and in the system file cache from the
+				// data that was loaded.  Keeping this for the lifetime of the cook is costly, so we try and unload it here.
+				Package->LinkerLoad->FlushCache();
+			}
 		}
 	}
 }
@@ -3026,7 +3037,7 @@ void UCookOnTheFlyServer::PumpSaves(UE::Cook::FTickStackData& StackData, uint32
 		{
 			if (PackageData.GetHasBeginPrepareSaveFailed())
 			{
-				ReleaseCookedPlatformData(PackageData);
+				ReleaseCookedPlatformData(PackageData, true /* bCompletedSave */);
 				PackageData.AddCookedPlatforms(PackageData.GetRequestedPlatforms(), false);
 				PackageData.SendToState(EPackageState::Idle, ESendFlags::QueueAdd);
 				++OutNumPushed;
@@ -3187,10 +3198,7 @@ void UCookOnTheFlyServer::PumpSaves(UE::Cook::FTickStackData& StackData, uint32
 			StackData.Timer.SavedPackage();
 		}
 
-		if (!IsCookingInEditor())
-		{
-			ReleaseCookedPlatformData(PackageData);
-		}
+		ReleaseCookedPlatformData(PackageData, true /* bCompletedSave */);
 
 		FName FileName = PackageData.GetFileName();
 
@@ -3823,9 +3831,10 @@ void UCookOnTheFlyServer::PreGarbageCollect()
 	TArray<UPackage*> GCKeepPackages;
 	for (FPackageData* PackageData : PackageDatas->GetSaveQueue())
 	{
-		// When a package is in the process of generating other packages,
-		// we need to keep it loaded (and all objects outered to it)
-		if (PackageData->IsGenerating())
+		// Generator packages that have started generation and have not yet finished saving should not be garbage
+		// collected because we do not want to kick them out of save and then have to reexecute their generation 
+		// on the next save, so keep them loaded.
+		if (PackageData->GetGeneratorPackage())
 		{
 			GCKeepPackages.Add(PackageData->GetPackage());
 		}
@@ -3920,25 +3929,13 @@ void UCookOnTheFlyServer::PostGarbageCollect()
 {
 	using namespace UE::Cook;
 
-	// If any PackageDatas with ObjectPointers lost had any of their object pointers deleted out from under them, demote them back to request
+	// If any PackageDatas with ObjectPointers had any of their object pointers deleted out from under them, demote them back to request
 	TArray<FPackageData*> Demotes;
 	for (FPackageData* PackageData : PackageDatas->GetSaveQueue())
 	{
 		if (PackageData->IsSaveInvalidated())
 		{
-			if (PackageData->IsGenerating())
-			{
-				// Generator packages cannot be demoted as that would break their generation data which must be preserved
-				// We're relying on BeginCacheForCookedPlatformData not having started; this is currently guaranteed when PackageData has a valid GeneratorPackage.
-				check(PackageData->GetCookedPlatformDataNextIndex() == 0);
-				// Recalculate object cache
-				PackageData->ClearObjectCache();
-				PackageData->CreateObjectCache();
-			}
-			else
-			{
-				Demotes.Add(PackageData);
-			}
+			Demotes.Add(PackageData);
 		}
 	}
 	for (FPackageData* PackageData : Demotes)
diff --git a/Engine/Source/Editor/UnrealEd/Private/Cooker/CookPackageData.cpp b/Engine/Source/Editor/UnrealEd/Private/Cooker/CookPackageData.cpp
--- a/Engine/Source/Editor/UnrealEd/Private/Cooker/CookPackageData.cpp
+++ b/Engine/Source/Editor/UnrealEd/Private/Cooker/CookPackageData.cpp
@@ -606,15 +606,13 @@ namespace Cook
 		check(GetPackage() != nullptr && GetPackage()->IsFullyLoaded());
 
 		check(!GetHasBeginPrepareSaveFailed());
-		check(!HasInitializedGeneratorSave());
-		check(!HasCompletedGeneration());
 		CheckObjectCacheEmpty();
 		CheckCookedPlatformDataEmpty();
 	}
 
 	void FPackageData::OnExitSave()
 	{
-		PackageDatas.GetCookOnTheFlyServer().ReleaseCookedPlatformData(*this);
+		PackageDatas.GetCookOnTheFlyServer().ReleaseCookedPlatformData(*this, false /* bCompletedSave */);
 		ClearObjectCache();
 		SetHasBeginPrepareSaveFailed(false);
 	}
@@ -879,8 +877,6 @@ namespace Cook
 		check(!GetCookedPlatformDataStarted());
 		check(!GetCookedPlatformDataCalled());
 		check(!GetCookedPlatformDataComplete());
-		check(!HasInitializedGeneratorSave());
-		check(!HasCompletedGeneration());
 	}
 
 	void FPackageData::ClearCookedPlatformData()
@@ -890,6 +886,10 @@ namespace Cook
 		SetCookedPlatformDataStarted(false);
 		SetCookedPlatformDataCalled(false);
 		SetCookedPlatformDataComplete(false);
+	}
+
+	void FPackageData::ResetGenerationProgress()
+	{
 		SetInitializedGeneratorSave(false);
 		SetCompletedGeneration(false);
 	}
diff --git a/Engine/Source/Editor/UnrealEd/Private/Cooker/CookPackageData.h b/Engine/Source/Editor/UnrealEd/Private/Cooker/CookPackageData.h
--- a/Engine/Source/Editor/UnrealEd/Private/Cooker/CookPackageData.h
+++ b/Engine/Source/Editor/UnrealEd/Private/Cooker/CookPackageData.h
@@ -302,6 +302,8 @@ namespace Cook
 		void SetGeneratedOwner(FGeneratorPackage* InGeneratedOwner);
 		/** Return the owning generator PackageData. Will be null if not a generated package or if orphaned */
 		FGeneratorPackage* GetGeneratedOwner() const { return GeneratedOwner; }
+		/** Mark that the package has finished generation and its save attempt, and should generate again if it needs to be resaved */
+		void ResetGenerationProgress();
 
 		/** Return whether a GC is required by a generator package as soon as possible. */
 		bool GeneratorPackageRequiresGC() const;
