Add support for trace counters to stall detector and trace summarize commandlet for telemetry in Iteration Profile

StallDetector:
* Replace static and per-callsite counter variables with trace classes
* Compose a trace api appropriate counter name (Stall Detector category)

SummarizeTraceCommandlet:
* Add support for extract and decode of counter messages from utrace files with a new analyzer class
* Report counters with name and count columns in the existing csv
* Support generation of counter events and values through to telemetry csv

diff --git a/Engine/Source/Editor/MainFrame/Private/MainFrameModule.cpp b/Engine/Source/Editor/MainFrame/Private/MainFrameModule.cpp
--- a/Engine/Source/Editor/MainFrame/Private/MainFrameModule.cpp
+++ b/Engine/Source/Editor/MainFrame/Private/MainFrameModule.cpp
@@ -393,7 +393,7 @@ TSharedRef<SWidget> FMainFrameModule::MakeDeveloperTools( const TArray<FMainFram
 #if STALL_DETECTOR
 		static FText GetStallCountAsString()
 		{
-			return FText::AsNumber(UE::FStallDetectorStats::TotalTriggeredCount);
+			return FText::AsNumber(UE::FStallDetectorStats::TotalTriggeredCount.Get());
 		}
 #endif // STALL_DETECTOR
 
diff --git a/Engine/Source/Editor/UnrealEd/Private/Commandlets/SummarizeTraceCommandlet.cpp b/Engine/Source/Editor/UnrealEd/Private/Commandlets/SummarizeTraceCommandlet.cpp
--- a/Engine/Source/Editor/UnrealEd/Private/Commandlets/SummarizeTraceCommandlet.cpp
+++ b/Engine/Source/Editor/UnrealEd/Private/Commandlets/SummarizeTraceCommandlet.cpp
@@ -16,6 +16,7 @@
 #include "Trace/Analysis.h"
 #include "Trace/Analyzer.h"
 #include "Trace/DataStream.h"
+#include "ProfilingDebugging/CountersTrace.h"
 
 /*
  * The following could be in TraceServices. This way if the format of CPU scope
@@ -45,7 +46,7 @@ class FCpuAnalyzer
 public:
 	struct FScopeName
 	{
-		const TCHAR* Name;
+		const TCHAR*	Name;
 		uint32			Id;
 	};
 
@@ -75,29 +76,29 @@ private:
 
 enum
 {
-	RouteId_EventSpec,
-	RouteId_EventBatch,
-	RouteId_EndCapture,
+	// CpuProfilerTrace.cpp
+	RouteId_CpuProfiler_EventSpec,
+	RouteId_CpuProfiler_EventBatch,
+	RouteId_CpuProfiler_EndCapture,
 };
 
 void FCpuAnalyzer::OnAnalysisBegin(const FOnAnalysisContext& Context)
 {
-	auto& Builder = Context.InterfaceBuilder;
-	Builder.RouteEvent(RouteId_EventSpec, "CpuProfiler", "EventSpec");
-	Builder.RouteEvent(RouteId_EventBatch, "CpuProfiler", "EventBatch");
-	Builder.RouteEvent(RouteId_EndCapture, "CpuProfiler", "EndCapture");
+	Context.InterfaceBuilder.RouteEvent(RouteId_CpuProfiler_EventSpec, "CpuProfiler", "EventSpec");
+	Context.InterfaceBuilder.RouteEvent(RouteId_CpuProfiler_EventBatch, "CpuProfiler", "EventBatch");
+	Context.InterfaceBuilder.RouteEvent(RouteId_CpuProfiler_EndCapture, "CpuProfiler", "EndCapture");
 }
 
 bool FCpuAnalyzer::OnEvent(uint16 RouteId, EStyle Style, const FOnEventContext& Context)
 {
 	switch (RouteId)
 	{
-	case RouteId_EventSpec:
+	case RouteId_CpuProfiler_EventSpec:
 		OnEventSpec(Context);
 		break;
 
-	case RouteId_EventBatch:
-	case RouteId_EndCapture:
+	case RouteId_CpuProfiler_EventBatch:
+	case RouteId_CpuProfiler_EndCapture:
 		OnBatch(Context);
 		break;
 	};
@@ -107,7 +108,7 @@ bool FCpuAnalyzer::OnEvent(uint16 RouteId, EStyle Style, const FOnEventContext&
 
 void FCpuAnalyzer::OnEventSpec(const FOnEventContext& Context)
 {
-	const auto& EventData = Context.EventData;
+	const FEventData& EventData = Context.EventData;
 	FString Name;
 	uint32 Id = EventData.GetValue<uint32>("Id");
 	EventData.GetString("Name", Name);
@@ -116,8 +117,8 @@ void FCpuAnalyzer::OnEventSpec(const FOnEventContext& Context)
 
 void FCpuAnalyzer::OnBatch(const FOnEventContext& Context)
 {
-	const auto& EventData = Context.EventData;
-	const auto& EventTime = Context.EventTime;
+	const FEventData& EventData = Context.EventData;
+	const FEventTime& EventTime = Context.EventTime;
 
 	uint32 ThreadId = Context.ThreadInfo.GetId();
 
@@ -143,6 +144,102 @@ void FCpuAnalyzer::OnBatch(const FOnEventContext& Context)
 	}
 }
 
+class FCountersAnalyzer
+	: public UE::Trace::IAnalyzer
+{
+public:
+	struct FCounterName
+	{
+		const TCHAR*		Name;
+		ETraceCounterType	Type;
+		uint16				Id;
+	};
+
+	struct FCounterIntValue
+	{
+		uint16				Id;
+		int64				Value;
+	};
+
+	struct FCounterFloatValue
+	{
+		uint16				Id;
+		double				Value;
+	};
+
+	virtual void		OnCounterName(const FCounterName& CounterName) = 0;
+	virtual void		OnCounterIntValue(const FCounterIntValue& NewValue) = 0;
+	virtual void		OnCounterFloatValue(const FCounterFloatValue& NewValue) = 0;
+
+private:
+	virtual void		OnAnalysisBegin(const FOnAnalysisContext& Context) override;
+	virtual bool		OnEvent(uint16 RouteId, EStyle Style, const FOnEventContext& Context) override;
+	void				OnCountersSpec(const FOnEventContext& Context);
+	void				OnCountersSetValueInt(const FOnEventContext& Context);
+	void				OnCountersSetValueFloat(const FOnEventContext& Context);
+};
+
+enum
+{
+	// CountersTrace.cpp
+	RouteId_Counters_Spec,
+	RouteId_Counters_SetValueInt,
+	RouteId_Counters_SetValueFloat,
+};
+
+void FCountersAnalyzer::OnAnalysisBegin(const FOnAnalysisContext& Context)
+{
+	Context.InterfaceBuilder.RouteEvent(RouteId_Counters_Spec, "Counters", "Spec");
+	Context.InterfaceBuilder.RouteEvent(RouteId_Counters_SetValueInt, "Counters", "SetValueInt");
+	Context.InterfaceBuilder.RouteEvent(RouteId_Counters_SetValueFloat, "Counters", "SetValueFloat");
+}
+
+bool FCountersAnalyzer::OnEvent(uint16 RouteId, EStyle Style, const FOnEventContext& Context)
+{
+	switch (RouteId)
+	{
+	case RouteId_Counters_Spec:
+		OnCountersSpec(Context);
+		break;
+
+	case RouteId_Counters_SetValueInt:
+		OnCountersSetValueInt(Context);
+		break;
+
+	case RouteId_Counters_SetValueFloat:
+		OnCountersSetValueFloat(Context);
+		break;
+	}
+
+	return true;
+}
+
+void FCountersAnalyzer::OnCountersSpec(const FOnEventContext& Context)
+{
+	const FEventData& EventData = Context.EventData;
+	FString Name;
+	uint16 Id = EventData.GetValue<uint16>("Id");
+	ETraceCounterType Type = static_cast<ETraceCounterType>(EventData.GetValue<uint8>("Type"));
+	EventData.GetString("Name", Name);
+	OnCounterName({ *Name, Type, uint16(Id - 1) });
+}
+
+void FCountersAnalyzer::OnCountersSetValueInt(const FOnEventContext& Context)
+{
+	const FEventData& EventData = Context.EventData;
+	uint16 CounterId = EventData.GetValue<uint16>("CounterId");
+	int64 Value = EventData.GetValue<int64>("Value");
+	OnCounterIntValue({ uint16(CounterId - 1), Value });
+}
+
+void FCountersAnalyzer::OnCountersSetValueFloat(const FOnEventContext& Context)
+{
+	const FEventData& EventData = Context.EventData;
+	uint16 CounterId = EventData.GetValue<uint16>("CounterId");
+	double Value = EventData.GetValue<double>("Value");
+	OnCounterFloatValue({ uint16(CounterId - 1), Value });
+}
+
 /*
  * This too could be housed elsewhere, along with an API to make it easier to
  * run analysis on trace files. The current model is influenced a little too much
@@ -190,11 +287,10 @@ public:
 };
 
 /*
- * Helper class for the SummarizeTrace commandlet.Aggregates statistics about
- * a trace's CPU scopes.
+ * Helper classes for the SummarizeTrace commandlet. Aggregates statistics about a trace.
  */
 
-class FSummarizeAnalyzer
+class FSummarizeCpuAnalyzer
 	: public FCpuAnalyzer
 {
 public:
@@ -339,12 +435,12 @@ public:
 	TArray<FThread> Threads;
 };
 
-static uint32 GetTypeHash(const FSummarizeAnalyzer::FScope Key)
+static uint32 GetTypeHash(const FSummarizeCpuAnalyzer::FScope Key)
 {
 	return FCrc::StrCrc32(*Key.Name);
 }
 
-void FSummarizeAnalyzer::OnCpuScopeName(const FScopeName& ScopeName)
+void FSummarizeCpuAnalyzer::OnCpuScopeName(const FScopeName& ScopeName)
 {
 	if (ScopeName.Id >= uint32(Scopes.Num()))
 	{
@@ -354,7 +450,7 @@ void FSummarizeAnalyzer::OnCpuScopeName(const FScopeName& ScopeName)
 	Scopes[ScopeName.Id].Name = ScopeName.Name;
 }
 
-void FSummarizeAnalyzer::OnCpuScopeEnter(const FScopeEnter& ScopeEnter)
+void FSummarizeCpuAnalyzer::OnCpuScopeEnter(const FScopeEnter& ScopeEnter)
 {
 	uint32 ThreadId = ScopeEnter.ThreadId;
 	if (ThreadId >= uint32(Threads.Num()))
@@ -364,7 +460,7 @@ void FSummarizeAnalyzer::OnCpuScopeEnter(const FScopeEnter& ScopeEnter)
 	Threads[ThreadId].ScopeStack.Add(ScopeEnter);
 }
 
-void FSummarizeAnalyzer::OnCpuScopeExit(const FScopeExit& ScopeExit)
+void FSummarizeCpuAnalyzer::OnCpuScopeExit(const FScopeExit& ScopeExit)
 {
 	uint32 ThreadId = ScopeExit.ThreadId;
 	if (ThreadId >= uint32(Threads.Num()) || Threads[ThreadId].ScopeStack.Num() <= 0)
@@ -384,6 +480,109 @@ void FSummarizeAnalyzer::OnCpuScopeExit(const FScopeExit& ScopeExit)
 	}
 }
 
+class FSummarizeCountersAnalyzer
+	: public FCountersAnalyzer
+{
+public:
+	virtual void OnCounterName(const FCounterName& CounterName) override;
+	virtual void OnCounterIntValue(const FCounterIntValue& NewValue) override;
+	virtual void OnCounterFloatValue(const FCounterFloatValue& NewValue) override;
+
+	struct FCounter
+	{
+		FString Name;
+		ETraceCounterType Type;
+
+		union
+		{
+			int64 IntValue;
+			double FloatValue;
+		};
+
+		FCounter(FString InName, ETraceCounterType InType)
+		{
+			Name = InName;
+			Type = InType;
+			switch (Type)
+			{
+			case TraceCounterType_Int:
+				IntValue = 0;
+				break;
+
+			case TraceCounterType_Float:
+				FloatValue = 0.0;
+				break;
+			}
+		}
+
+		template<typename ValueType>
+		void SetValue(ValueType InValue) = delete;
+
+		template<>
+		void SetValue(int64 InValue)
+		{
+			ensure(Type == TraceCounterType_Int);
+			if (Type == TraceCounterType_Int)
+			{
+				IntValue = InValue;
+			}
+		}
+
+		template<>
+		void SetValue(double InValue)
+		{
+			ensure(Type == TraceCounterType_Float);
+			if (Type == TraceCounterType_Float)
+			{
+				FloatValue = InValue;
+			}
+		}
+
+		FString GetValue() const
+		{
+			switch (Type)
+			{
+			case TraceCounterType_Int:
+				return FString::Printf(TEXT("%lld"), IntValue);
+
+			case TraceCounterType_Float:
+				return FString::Printf(TEXT("%f"), FloatValue);
+			}
+
+			ensure(false);
+			return TEXT("");
+		}
+	};
+
+	TMap<uint16, FCounter> Counters;
+};
+
+void FSummarizeCountersAnalyzer::OnCounterName(const FCounterName& CounterName)
+{
+	Counters.Add(CounterName.Id, FCounter(CounterName.Name, CounterName.Type));
+}
+
+void FSummarizeCountersAnalyzer::OnCounterIntValue(const FCounterIntValue& NewValue)
+{
+	FCounter* FoundCounter = Counters.Find(NewValue.Id);
+	ensure(FoundCounter);
+	if (FoundCounter)
+	{
+		FoundCounter->SetValue(NewValue.Value);
+	}
+}
+
+void FSummarizeCountersAnalyzer::OnCounterFloatValue(const FCounterFloatValue& NewValue)
+{
+	FCounter* FoundCounter = Counters.Find(NewValue.Id);
+	ensure(FoundCounter);
+	if (FoundCounter)
+	{
+		FoundCounter->SetValue(NewValue.Value);
+	}
+}
+
+
 /*
  * SummarizeTrace commandlet ingests a utrace file and summarizes the
  * cpu scope events within it, and summarizes each event to a csv. It
@@ -434,6 +633,8 @@ int32 USummarizeTraceCommandlet::Main(const FString& CmdLineParams)
 		return 1;
 	}
 
+	// load the stats file to know which event name and statistic name to generate in the telementry csv
+	// the telemetry csv is ingested completely, so this just whitelists specific data elements we want to track
 	FString StatisticsFileName;
 	TMultiMap<FString, FString> NameToStatisticMap;
 	if (FParse::Value(*CmdLineParams, TEXT("statsfile="), StatisticsFileName, true))
@@ -470,9 +671,9 @@ int32 USummarizeTraceCommandlet::Main(const FString& CmdLineParams)
 					}
 				});
 
-			for (auto& Column : Columns)
+			for (struct Column& Column : Columns)
 			{
-				if (CSVIndex == 0)
+				if (CSVIndex == 0) // is this the header row?
 				{
 					for (int FieldIndex = 0; FieldIndex < Fields.Num(); ++FieldIndex)
 					{
@@ -488,7 +689,7 @@ int32 USummarizeTraceCommandlet::Main(const FString& CmdLineParams)
 						bValidColumns = false;
 					}
 				}
-				else
+				else // else it is a data row, pull each element from appropriate column
 				{
 					NameToStatisticMap.AddUnique(FString(Fields[NameColumn]), FString(Fields[StatisticColumn]));
 				}
@@ -508,7 +709,7 @@ int32 USummarizeTraceCommandlet::Main(const FString& CmdLineParams)
 		SearchPaths.Add(FPaths::Combine(FPaths::EngineDir(), TEXT("Programs"), TEXT("UnrealInsights"), TEXT("Saved"), TEXT("TraceSessions")));
 		SearchPaths.Add(FPaths::EngineDir());
 		SearchPaths.Add(FPaths::ProjectDir());
-		for (const auto& SearchPath : SearchPaths)
+		for (const FString& SearchPath : SearchPaths)
 		{
 			FString PossibleTraceFileName = FPaths::Combine(SearchPath, TraceFileName);
 			if (FPaths::FileExists(PossibleTraceFileName))
@@ -533,13 +734,20 @@ int32 USummarizeTraceCommandlet::Main(const FString& CmdLineParams)
 		return 1;
 	}
 
-	FSummarizeAnalyzer Analyzer;
+	// setup analysis context with analyzers
 	UE::Trace::FAnalysisContext Context;
-	Context.AddAnalyzer(Analyzer);
+	FSummarizeCpuAnalyzer CpuAnalyzer;
+	Context.AddAnalyzer(CpuAnalyzer);
+	FSummarizeCountersAnalyzer CountersAnalyzer;
+	Context.AddAnalyzer(CountersAnalyzer);
+
+	// kick processing on a thread
 	UE::Trace::FAnalysisProcessor Processor = Context.Process(DataStream);
+
+	// sync on completion
 	Processor.Wait();
 
-	TSet<FSummarizeAnalyzer::FScope> DeduplicatedScopes;
+	TSet<FSummarizeCpuAnalyzer::FScope> DeduplicatedScopes;
 
 	struct DisallowedString
 	{
@@ -560,7 +768,7 @@ int32 USummarizeTraceCommandlet::Main(const FString& CmdLineParams)
 	};
 
 	bool bDisallowed = false;
-	for (const auto& Scope : Analyzer.Scopes)
+	for (const FSummarizeCpuAnalyzer::FScope& Scope : CpuAnalyzer.Scopes)
 	{
 		if (Scope.Name.IsEmpty())
 		{
@@ -572,8 +780,9 @@ int32 USummarizeTraceCommandlet::Main(const FString& CmdLineParams)
 			continue;
 		}
 
+		// sanitize strings for a bog-simple csv file
 		bDisallowed = false;
-		for (auto& DisallowedString : DisallowedStrings)
+		for (struct DisallowedString& DisallowedString : DisallowedStrings)
 		{
 			if (Scope.Name.Contains(DisallowedString.String))
 			{
@@ -598,7 +807,7 @@ int32 USummarizeTraceCommandlet::Main(const FString& CmdLineParams)
 			continue;
 		}
 
-		FSummarizeAnalyzer::FScope* FoundScope = DeduplicatedScopes.Find(Scope);
+		FSummarizeCpuAnalyzer::FScope* FoundScope = DeduplicatedScopes.Find(Scope);
 		if (FoundScope)
 		{
 			FoundScope->Merge(Scope);
@@ -610,25 +819,31 @@ int32 USummarizeTraceCommandlet::Main(const FString& CmdLineParams)
 	}
 
 	UE_LOG(LogSummarizeTrace, Display, TEXT("Sorting %d events by total time accumulated..."), DeduplicatedScopes.Num());
-	TArray<FSummarizeAnalyzer::FScope> SortedScopes;
-	for (const auto& Scope : DeduplicatedScopes)
+	TArray<FSummarizeCpuAnalyzer::FScope> SortedScopes;
+	for (const FSummarizeCpuAnalyzer::FScope& Scope : DeduplicatedScopes)
 	{
 		SortedScopes.Add(Scope);
 	}
 	SortedScopes.Sort();
 
+	// generate a summary csv, always
 	FString CsvFileName = FPaths::SetExtension(TraceFileName, "csv");
 	UE_LOG(LogSummarizeTrace, Display, TEXT("Writing summary to %s..."), *CsvFileName);
 	IFileHandle* CsvHandle = FPlatformFileManager::Get().GetPlatformFile().OpenWrite(*CsvFileName);
 	if (CsvHandle)
 	{
-		// no newline, see row printf
+		// no newline, see row printfs
 		FilePrint(CsvHandle, FString::Printf(TEXT("Name,Count,TotalSeconds,FirstStartSeconds,FirstEndSeconds,FirstSeconds,MinSeconds,MaxSeconds,MeanSeconds,DeviationSeconds")));
-		for (const auto& Scope : SortedScopes)
+		for (const FSummarizeCpuAnalyzer::FScope& Scope : SortedScopes)
 		{
 			// note newline is at the front of every data line to prevent final extraneous newline, per customary for csv
 			FilePrint(CsvHandle, FString::Printf(TEXT("\n%s,%llu,%f,%f,%f,%f,%f,%f,%f,%f,"), *Scope.Name, Scope.Count, Scope.TotalSeconds, Scope.FirstStartSeconds, Scope.FirstEndSeconds, Scope.FirstSeconds, Scope.MinSeconds, Scope.MaxSeconds, Scope.MeanSeconds, Scope.GetDeviationSeconds()));
 		}
+		for (const TMap<uint16, FSummarizeCountersAnalyzer::FCounter>::ElementType& Counter : CountersAnalyzer.Counters)
+		{
+			// note newline is at the front of every data line to prevent final extraneous newline, per customary for csv
+			FilePrint(CsvHandle, FString::Printf(TEXT("\n%s,%s,,,,,,,,,"), *Counter.Value.Name, *Counter.Value.GetValue()));
+		}
 		CsvHandle->Flush();
 		delete CsvHandle;
 		CsvHandle = nullptr;
@@ -639,6 +854,7 @@ int32 USummarizeTraceCommandlet::Main(const FString& CmdLineParams)
 		return 1;
 	}
 
+	// if we were asked to generate a telememtry file, generate it
 	if (!NameToStatisticMap.IsEmpty())
 	{
 		FString TracePath = FPaths::GetPath(TraceFileName);
@@ -649,18 +865,28 @@ int32 USummarizeTraceCommandlet::Main(const FString& CmdLineParams)
 		IFileHandle* TelemetryCsvHandle = FPlatformFileManager::Get().GetPlatformFile().OpenWrite(*TelemetryCsvFileName);
 		if (TelemetryCsvHandle)
 		{
-			// no newline, see row printf
+			// no newline, see row printfs
 			FilePrint(TelemetryCsvHandle, FString::Printf(TEXT("TestName,Context,DataPoint,Measurement")));
-			for (const auto& Scope : SortedScopes)
+			for (const FSummarizeCpuAnalyzer::FScope& Scope : SortedScopes)
 			{
 				TArray<FString> Statistics;
 				NameToStatisticMap.MultiFind(Scope.Name, Statistics, true);
-				for (const auto& Statistic : Statistics)
+				for (const FString& Statistic : Statistics)
 				{
 					// note newline is at the front of every data line to prevent final extraneous newline, per customary for csv
 					FilePrint(TelemetryCsvHandle, FString::Printf(TEXT("\n%s,%s,%s,%s,"), *TraceFileBasename, *Scope.Name, *Statistic, *Scope.GetValue(Statistic)));
 				}
 			}
+			for (const TMap<uint16, FSummarizeCountersAnalyzer::FCounter>::ElementType& Counter : CountersAnalyzer.Counters)
+			{
+				TArray<FString> Statistics;
+				NameToStatisticMap.MultiFind(Counter.Value.Name, Statistics, true);
+				for (const FString& Statistic : Statistics)
+				{
+					// note newline is at the front of every data line to prevent final extraneous newline, per customary for csv
+					FilePrint(TelemetryCsvHandle, FString::Printf(TEXT("\n%s,%s,%s,%s,"), *TraceFileBasename, *Counter.Value.Name, *Statistic, *Counter.Value.GetValue()));
+				}
+			}
 			TelemetryCsvHandle->Flush();
 			delete TelemetryCsvHandle;
 			TelemetryCsvHandle = nullptr;
diff --git a/Engine/Source/Programs/AutomationTool/Gauntlet/Unreal/Editor/EditorTests.cs b/Engine/Source/Programs/AutomationTool/Gauntlet/Unreal/Editor/EditorTests.cs
--- a/Engine/Source/Programs/AutomationTool/Gauntlet/Unreal/Editor/EditorTests.cs
+++ b/Engine/Source/Programs/AutomationTool/Gauntlet/Unreal/Editor/EditorTests.cs
@@ -78,7 +78,7 @@ namespace UnrealEditor
 			{
 				AppConfig.CommandLineParams.Add(string.Format("tracefile={0}", TraceFile));
 				AppConfig.CommandLineParams.Add("tracefiletrunc"); // replace existing
-				AppConfig.CommandLineParams.Add("trace=cpu");
+				AppConfig.CommandLineParams.Add("trace=cpu,counters");
 				AppConfig.CommandLineParams.Add("statnamedevents");
 			}
 
diff --git a/Engine/Source/Runtime/Core/Private/ProfilingDebugging/StallDetector.cpp b/Engine/Source/Runtime/Core/Private/ProfilingDebugging/StallDetector.cpp
--- a/Engine/Source/Runtime/Core/Private/ProfilingDebugging/StallDetector.cpp
+++ b/Engine/Source/Runtime/Core/Private/ProfilingDebugging/StallDetector.cpp
@@ -11,6 +11,10 @@
 #include "HAL/Runnable.h"
 #include "HAL/RunnableThread.h"
 
+// counters for sending information into trace system
+TRACE_DECLARE_INT_COUNTER(StallCount, TEXT("StallDetector/Count"));
+TRACE_DECLARE_FLOAT_COUNTER(StallTimeSeconds, TEXT("StallDetector/TimeSeconds"));
+
 // force normal behavior in the face of debug configuration and debugger attached
 #define STALL_DETECTOR_DEBUG 0
 
@@ -127,16 +131,26 @@ uint32 UE::FStallDetectorRunnable::Run()
 
 FCriticalSection UE::FStallDetectorStats::InstancesSection;
 TSet<UE::FStallDetectorStats*> UE::FStallDetectorStats::Instances;
-std::atomic<uint32> UE::FStallDetectorStats::TotalTriggeredCount;
-std::atomic<uint32> UE::FStallDetectorStats::TotalReportedCount;
+FCountersTrace::TCounter<std::atomic<int64>, TraceCounterType_Int> UE::FStallDetectorStats::TotalTriggeredCount (TEXT("StallDetector/TotalTriggeredCount"), TraceCounterDisplayHint_None);
+FCountersTrace::TCounter<std::atomic<int64>, TraceCounterType_Int> UE::FStallDetectorStats::TotalReportedCount (TEXT("StallDetector/TotalReportedCount"), TraceCounterDisplayHint_None);
 
 UE::FStallDetectorStats::FStallDetectorStats(const TCHAR* InName, const double InBudgetSeconds, const EStallDetectorReportingMode InReportingMode)
 	: Name(InName)
 	, BudgetSeconds(InBudgetSeconds)
 	, ReportingMode(InReportingMode)
 	, bReported(false)
-	, TriggerCount(0)
-	, OverageSeconds(0.0)
+	, TriggerCount(
+		(
+			FCString::Strcat(TriggerCountCounterName, TEXT("StallDetector/")),
+			FCString::Strcat(TriggerCountCounterName, InName),
+			FCString::Strcat(TriggerCountCounterName, TEXT(" TriggerCount"))
+		), TraceCounterDisplayHint_None)
+	, OverageSeconds(
+		(
+			FCString::Strcat(OverageSecondsCounterName, TEXT("StallDetector/")),
+			FCString::Strcat(OverageSecondsCounterName, InName),
+			FCString::Strcat(OverageSecondsCounterName, TEXT(" OverageSeconds"))
+		), TraceCounterDisplayHint_None)
 {
 	// Add at the end of construction
 	FScopeLock ScopeLock(&InstancesSection);
@@ -154,8 +168,8 @@ void UE::FStallDetectorStats::OnStallCompleted(double InOverageSeconds)
 {
 	// we sync access around these for coherency reasons, can be polled from another thread (tabulation)
 	FScopeLock StatsLock(&StatsSection);
-	TriggerCount++;
-	OverageSeconds += InOverageSeconds;
+	TriggerCount.Increment();
+	OverageSeconds.Add(InOverageSeconds);
 }
 
 void UE::FStallDetectorStats::TabulateStats(TArray<TabulatedResult>& TabulatedResults)
@@ -169,9 +183,9 @@ void UE::FStallDetectorStats::TabulateStats(TArray<TabulatedResult>& TabulatedRe
 			, OverageRatio(0.0)
 		{
 			FScopeLock Lock(&InStats->StatsSection);
-			if (InStats->TriggerCount && InStats->BudgetSeconds > 0.0)
+			if (InStats->TriggerCount.Get() && InStats->BudgetSeconds > 0.0)
 			{
-				OverageRatio = (InStats->OverageSeconds / InStats->TriggerCount) / InStats->BudgetSeconds;
+				OverageRatio = (InStats->OverageSeconds.Get() / InStats->TriggerCount.Get()) / InStats->BudgetSeconds;
 			}
 		}
 
@@ -189,7 +203,7 @@ void UE::FStallDetectorStats::TabulateStats(TArray<TabulatedResult>& TabulatedRe
 	FScopeLock InstancesLock(&UE::FStallDetectorStats::GetInstancesSection());
 	for (const UE::FStallDetectorStats* StallStats : UE::FStallDetectorStats::GetInstances())
 	{
-		if (StallStats->TriggerCount && StallStats->ReportingMode != UE::EStallDetectorReportingMode::Disabled)
+		if (StallStats->TriggerCount.Get() && StallStats->ReportingMode != UE::EStallDetectorReportingMode::Disabled)
 		{
 			StatsArray.Emplace(StallStats);
 		}
@@ -207,8 +221,8 @@ void UE::FStallDetectorStats::TabulateStats(TArray<TabulatedResult>& TabulatedRe
 
 			// we sync access around these for coherency reasons, can be polled from another thread (detector or scope)
 			FScopeLock Lock(&Result.Stats->StatsSection);
-			Result.TriggerCount = Result.Stats->TriggerCount;
-			Result.OverageSeconds = Result.Stats->OverageSeconds;
+			Result.TriggerCount = Result.Stats->TriggerCount.Get();
+			Result.OverageSeconds = Result.Stats->OverageSeconds.Get();
 		}
 	}
 }
@@ -345,7 +359,7 @@ void UE::FStallDetector::CheckAndReset()
 
 void UE::FStallDetector::OnStallDetected(uint32 InThreadId, const double InElapsedSeconds)
 {
-	Stats.TotalTriggeredCount++;
+	Stats.TotalTriggeredCount.Increment();
 
 	//
 	// Determine if we want to undermine the specified reporting mode
@@ -396,7 +410,7 @@ void UE::FStallDetector::OnStallDetected(uint32 InThreadId, const double InElaps
 	if (bSendReport)
 	{
 		Stats.bReported = true;
-		Stats.TotalReportedCount++;
+		Stats.TotalReportedCount.Increment();
 		const int NumStackFramesToIgnore = FPlatformTLS::GetCurrentThreadId() == InThreadId ? 2 : 0;
 		UE_LOG(LogStall, Log, TEXT("Stall detector '%s' exceeded budget of %fs, reporting..."), Stats.Name, Stats.BudgetSeconds);
 		double ReportSeconds = FStallDetector::Seconds();
diff --git a/Engine/Source/Runtime/Core/Public/ProfilingDebugging/CountersTrace.h b/Engine/Source/Runtime/Core/Public/ProfilingDebugging/CountersTrace.h
--- a/Engine/Source/Runtime/Core/Public/ProfilingDebugging/CountersTrace.h
+++ b/Engine/Source/Runtime/Core/Public/ProfilingDebugging/CountersTrace.h
@@ -47,6 +47,11 @@ struct FCountersTrace
 			CounterId = OutputInitCounter(InCounterName, CounterType, CounterDisplayHint);
 		}
 
+		const ValueType& Get() const
+		{
+			return Value;
+		}
+
 		void Set(ValueType InValue)
 		{
 			if (Value != InValue)
diff --git a/Engine/Source/Runtime/Core/Public/ProfilingDebugging/StallDetector.h b/Engine/Source/Runtime/Core/Public/ProfilingDebugging/StallDetector.h
--- a/Engine/Source/Runtime/Core/Public/ProfilingDebugging/StallDetector.h
+++ b/Engine/Source/Runtime/Core/Public/ProfilingDebugging/StallDetector.h
@@ -10,6 +10,7 @@
 #include "HAL/PreprocessorHelpers.h"
 #include "Logging/LogMacros.h"
 #include "Misc/ScopeLock.h"
+#include "ProfilingDebugging/CountersTrace.h"
 
 #include <atomic>
 
@@ -108,17 +109,19 @@ namespace UE
 		bool bReported;
 
 		// The total number of times all callsites have been triggered
-		static std::atomic<uint32> TotalTriggeredCount;
+		static FCountersTrace::TCounter<std::atomic<int64>, TraceCounterType_Int> TotalTriggeredCount;
 
 		// The total number of reports that have been sent
-		static std::atomic<uint32> TotalReportedCount;
+		static FCountersTrace::TCounter<std::atomic<int64>, TraceCounterType_Int> TotalReportedCount;
 
 	private:
 		// The number of times this callsite has been triggered
-		uint32 TriggerCount;
+		FCountersTrace::TCounter<int64, TraceCounterType_Int> TriggerCount;
+		TCHAR TriggerCountCounterName[256];
 
 		// The cumulative overage time for this callsite
-		double OverageSeconds;
+		FCountersTrace::TCounter<double, TraceCounterType_Float> OverageSeconds;
+		TCHAR OverageSecondsCounterName[256];
 
 		// Guards access to the stats from multiple threads, for coherency
 		mutable FCriticalSection StatsSection;
diff --git a/Engine/Source/Runtime/Engine/Private/Analytics/EditorSessionSummaryWriter.cpp b/Engine/Source/Runtime/Engine/Private/Analytics/EditorSessionSummaryWriter.cpp
--- a/Engine/Source/Runtime/Engine/Private/Analytics/EditorSessionSummaryWriter.cpp
+++ b/Engine/Source/Runtime/Engine/Private/Analytics/EditorSessionSummaryWriter.cpp
@@ -356,8 +356,8 @@ void FEditorSessionSummaryWriter::Tick(float DeltaTime)
 		CurrentSession->TopStallTriggerCount = TopResult.TriggerCount;
 	}
 
-	CurrentSession->TotalStallCount = UE::FStallDetectorStats::TotalTriggeredCount;
-	CurrentSession->TotalStallReported = UE::FStallDetectorStats::TotalReportedCount;
+	CurrentSession->TotalStallCount = UE::FStallDetectorStats::TotalTriggeredCount.Get();
+	CurrentSession->TotalStallReported = UE::FStallDetectorStats::TotalReportedCount.Get();
 #endif // STALL_DETECTOR
 
 	if (bSaveSession)
