Replace deprecated Subobject Editor stuff in the CAC Editor utils

#rnx

diff --git a/Engine/Source/Editor/UnrealEd/Private/Kismet2/ChildActorComponentEditorUtils.cpp b/Engine/Source/Editor/UnrealEd/Private/Kismet2/ChildActorComponentEditorUtils.cpp
--- a/Engine/Source/Editor/UnrealEd/Private/Kismet2/ChildActorComponentEditorUtils.cpp
+++ b/Engine/Source/Editor/UnrealEd/Private/Kismet2/ChildActorComponentEditorUtils.cpp
@@ -5,6 +5,7 @@
 #include "ToolMenus.h"
 #include "SSCSEditor.h"
 #include "SSCSEditorMenuContext.h"
+#include "SSubobjectEditor.h"
 
 #define LOCTEXT_NAMESPACE "ChildActorComponentEditorUtils"
 
@@ -26,7 +27,7 @@ struct FLocalChildActorComponentEditorUtils
 		return InMode == CurrentMode;
 	}
 
-	static void OnSetChildActorTreeViewVisualizationMode(UChildActorComponent* InChildActorComponent, EChildActorComponentTreeViewVisualizationMode InMode, TWeakPtr<SSCSEditor> InWeakSCSEditorPtr)
+	static void OnSetChildActorTreeViewVisualizationMode(UChildActorComponent* InChildActorComponent, EChildActorComponentTreeViewVisualizationMode InMode, TWeakPtr<SSubobjectEditor> InWeakSubobjectEditorPtr)
 	{
 		if (!InChildActorComponent)
 		{
@@ -35,14 +36,14 @@ struct FLocalChildActorComponentEditorUtils
 
 		InChildActorComponent->SetEditorTreeViewVisualizationMode(InMode);
 
-		TSharedPtr<SSCSEditor> SCSEditorPtr = InWeakSCSEditorPtr.Pin();
-		if (SCSEditorPtr.IsValid())
+		TSharedPtr<SSubobjectEditor> SubobjectEditorPtr = InWeakSubobjectEditorPtr.Pin();
+		if (SubobjectEditorPtr.IsValid())
 		{
-			SCSEditorPtr->UpdateTree();
+			SubobjectEditorPtr->UpdateTree();
 		}
 	}
 
-	static void CreateChildActorVisualizationModesSubMenu(UToolMenu* InSubMenu, UChildActorComponent* InChildActorComponent, TWeakPtr<SSCSEditor> InWeakSCSEditorPtr)
+	static void CreateChildActorVisualizationModesSubMenu(UToolMenu* InSubMenu, UChildActorComponent* InChildActorComponent, TWeakPtr<SSubobjectEditor> InWeakSCSEditorPtr)
 	{
 		FToolMenuSection& SubMenuSection = InSubMenu->AddSection("ExpansionModes");
 		SubMenuSection.AddMenuEntry(
@@ -185,10 +186,10 @@ void FChildActorComponentEditorUtils::FillComponentContextMenuOptions(UToolMenu*
 
 	FToolMenuSection& Section = Menu->AddSection("ChildActorComponent", LOCTEXT("ChildActorComponentHeading", "Child Actor Component"));
 	{
-		TWeakPtr<SSCSEditor> WeakEditorPtr;
+		TWeakPtr<SSubobjectEditor> WeakEditorPtr;
 		if (USSCSEditorMenuContext* MenuContext = Menu->FindContext<USSCSEditorMenuContext>())
 		{
-			WeakEditorPtr = MenuContext->SCSEditor;
+			WeakEditorPtr = MenuContext->SubobjectEditor;
 		}
 
 		Section.AddSubMenu(
@@ -222,10 +223,10 @@ void FChildActorComponentEditorUtils::FillChildActorContextMenuOptions(UToolMenu
 
 	FToolMenuSection& Section = Menu->AddSection("ChildActor", LOCTEXT("ChildActorHeading", "Child Actor"));
 	{
-		TWeakPtr<SSCSEditor> WeakEditorPtr;
+		TWeakPtr<SSubobjectEditor> WeakEditorPtr;
 		if (USSCSEditorMenuContext* MenuContext = Menu->FindContext<USSCSEditorMenuContext>())
 		{
-			WeakEditorPtr = MenuContext->SCSEditor;
+			WeakEditorPtr = MenuContext->SubobjectEditor;
 		}
 
 		Section.AddMenuEntry(
diff --git a/Engine/Source/Editor/UnrealEd/UnrealEd.Build.cs b/Engine/Source/Editor/UnrealEd/UnrealEd.Build.cs
--- a/Engine/Source/Editor/UnrealEd/UnrealEd.Build.cs
+++ b/Engine/Source/Editor/UnrealEd/UnrealEd.Build.cs
@@ -115,6 +115,7 @@ public class UnrealEd : ModuleRules
 				"InterchangeEngine",
 				"DeveloperToolSettings",
 				"SubobjectDataInterface",
+				"SubobjectEditor",
 			}
 		);
 
