EditorDomain: add Schema Hash for the package's classes to the PackageDigest.
#rnx

diff --git a/Engine/Source/Editor/UnrealEd/Private/Cooker/CookPackageData.cpp b/Engine/Source/Editor/UnrealEd/Private/Cooker/CookPackageData.cpp
--- a/Engine/Source/Editor/UnrealEd/Private/Cooker/CookPackageData.cpp
+++ b/Engine/Source/Editor/UnrealEd/Private/Cooker/CookPackageData.cpp
@@ -7,6 +7,7 @@
 #include "CookOnTheSide/CookOnTheFlyServer.h"
 #include "CookPlatformManager.h"
 #include "Containers/StringView.h"
+#include "EditorDomain/EditorDomain.h"
 #include "HAL/FileManager.h"
 #include "Misc/CoreMiscDefines.h"
 #include "Misc/Paths.h"
@@ -695,6 +696,10 @@ namespace Cook
 		}
 		if (!PreloadableFile.Get())
 		{
+			if (FEditorDomain* EditorDomain = FEditorDomain::Get())
+			{
+				EditorDomain->PrecachePackageDigest(GetPackageName());
+			}
 			TStringBuilder<NAME_SIZE> FileNameString;
 			GetFileName().ToString(FileNameString);
 			PreloadableFile.Set(MakeShared<FPreloadableArchive>(FileNameString.ToString()), *this);
diff --git a/Engine/Source/Editor/UnrealEd/Private/EditorDomain/EditorDomain.cpp b/Engine/Source/Editor/UnrealEd/Private/EditorDomain/EditorDomain.cpp
--- a/Engine/Source/Editor/UnrealEd/Private/EditorDomain/EditorDomain.cpp
+++ b/Engine/Source/Editor/UnrealEd/Private/EditorDomain/EditorDomain.cpp
@@ -141,8 +141,9 @@ bool FEditorDomain::TryFindOrAddPackageSource(const FPackagePath& PackagePath, T
 		return true;
 	}
 
+	FString ErrorMessage;
 	FPackageDigest PackageDigest;
-	EPackageDigestResult Result = GetPackageDigest(*AssetRegistry, PackageName, PackageDigest);
+	EPackageDigestResult Result = GetPackageDigest(*AssetRegistry, PackageName, PackageDigest, ErrorMessage, ClassDigests);
 	switch (Result)
 	{
 	case EPackageDigestResult::Success:
@@ -160,14 +161,24 @@ bool FEditorDomain::TryFindOrAddPackageSource(const FPackagePath& PackagePath, T
 		// but for non-existent packages we want it not to be there to avoid wasting memory on it
 		PackageSources.Remove(PackageName);
 		return false;
-	case EPackageDigestResult::WrongThread:
+	default:
+		UE_LOG(LogEditorDomain, Warning,
+			TEXT("Could not load package from EditorDomain; it will be loaded from the WorkspaceDomain: %s."),
+			*ErrorMessage)
 		PackageSource = new FPackageSource();
 		PackageSource->Source = EPackageSource::Workspace;
 		OutSource = PackageSource;
 		return true;
-	default:
-		check(false);
-		return false;
+	}
+}
+
+void FEditorDomain::PrecachePackageDigest(FName PackageName)
+{
+	AssetRegistry->WaitForPackage(PackageName.ToString());
+	TOptional<FAssetPackageData> PackageData = AssetRegistry->GetAssetPackageDataCopy(PackageName);
+	if (PackageData)
+	{
+		UE::EditorDomain::PrecacheClassDigests(PackageData->ImportedClasses, ClassDigests);
 	}
 }
 
@@ -419,10 +430,7 @@ void FEditorDomain::OnEndLoad(TConstArrayView<UPackage*> LoadedPackages)
 
 	for (UPackage* Package : PackagesToSave)
 	{
-		if (!UE::EditorDomain::TrySavePackage(Package))
-		{
-			UE_LOG(LogEditorDomain, Warning, TEXT("Could not save package %s into EditorDomain."), *Package->GetName());
-		}
+		UE::EditorDomain::TrySavePackage(Package, ClassDigests);
 	}
 }
 
@@ -452,10 +460,7 @@ void FEditorDomain::OnPostEngineInit()
 
 	for (UPackage* Package : PackagesToSave)
 	{
-		if (!UE::EditorDomain::TrySavePackage(Package))
-		{
-			UE_LOG(LogEditorDomain, Warning, TEXT("Could not save package %s into EditorDomain."), *Package->GetName());
-		}
+		UE::EditorDomain::TrySavePackage(Package, ClassDigests);
 	}
 }
 
diff --git a/Engine/Source/Editor/UnrealEd/Private/EditorDomain/EditorDomain.h b/Engine/Source/Editor/UnrealEd/Private/EditorDomain/EditorDomain.h
--- a/Engine/Source/Editor/UnrealEd/Private/EditorDomain/EditorDomain.h
+++ b/Engine/Source/Editor/UnrealEd/Private/EditorDomain/EditorDomain.h
@@ -3,23 +3,37 @@
 #pragma once
 
 #include "AssetRegistry/AssetData.h"
+#include "Containers/Map.h"
 #include "HAL/CriticalSection.h"
 #include "HAL/Platform.h"
+#include "Hash/Blake3.h"
 #include "IO/IoHash.h"
 #include "Logging/LogMacros.h"
 #include "Templates/RefCounting.h"
 #include "TickableEditorObject.h"
+#include "UObject/NameTypes.h"
 #include "UObject/PackageResourceManager.h"
 
 class FAssetPackageData;
 class FEditorDomainSaveClient;
 class IAssetRegistry;
-namespace UE
-{
-namespace EditorDomain
+class UPackage;
+namespace UE::EditorDomain
 {
 	typedef FIoHash FPackageDigest;
-}
+
+	/** A UClass's data that is used in the EditorDomain Digest */
+	struct FClassDigestData
+	{
+		FBlake3Hash SchemaHash;
+		bool bNative;
+	};
+	/** Threadsafe cache of ClassName -> Digest data for calculating EditorDomain Digests */
+	struct FClassDigestMap
+	{
+		TMap<FName, FClassDigestData> Map;
+		FCriticalSection Lock;
+	};
 }
 
 DECLARE_LOG_CATEGORY_EXTERN(LogEditorDomain, Log, All);
@@ -81,6 +95,9 @@ public:
 	virtual ETickableTickType GetTickableTickType() const override { return ETickableTickType::Always; }
 	virtual TStatId GetStatId() const override { return TStatId(); }
 
+	// EditorDomain interface
+	void PrecachePackageDigest(FName PackageName);
+
 private:
 	/**
 	 * Reference-counted struct holding the locks used for multithreaded synchronization.
@@ -140,6 +157,8 @@ private:
 	TRefCountPtr<FLocks> Locks;
 	/** Digests previously found for a package. Used for optimization, but also to record loaded-from-domain. */
 	TMap<FName, TRefCountPtr<FPackageSource>> PackageSources;
+	/** Cache of GetSchemaHash by class name */
+	UE::EditorDomain::FClassDigestMap ClassDigests;
 	/** True by default, set to false when reading is disabled for testing. */
 	bool bEditorDomainReadEnabled = true;
 	/** If true, use an out-of-process EditorDomainSaveServer for saves, else save in process in EndLoad */
diff --git a/Engine/Source/Editor/UnrealEd/Private/EditorDomain/EditorDomainSave.cpp b/Engine/Source/Editor/UnrealEd/Private/EditorDomain/EditorDomainSave.cpp
--- a/Engine/Source/Editor/UnrealEd/Private/EditorDomain/EditorDomainSave.cpp
+++ b/Engine/Source/Editor/UnrealEd/Private/EditorDomain/EditorDomainSave.cpp
@@ -712,7 +712,7 @@ bool FEditorDomainSaveServer::TrySavePackage(const FPackagePath& PackagePath, FS
 		return false;
 	}
 
-	if (!UE::EditorDomain::TrySavePackage(Package))
+	if (!UE::EditorDomain::TrySavePackage(Package, this->ClassDigests))
 	{
 		OutErrorMessage = FString::Printf(TEXT("Could not save package %s."), *PackagePath.GetDebugName());
 		return false;
diff --git a/Engine/Source/Editor/UnrealEd/Private/EditorDomain/EditorDomainSave.h b/Engine/Source/Editor/UnrealEd/Private/EditorDomain/EditorDomainSave.h
--- a/Engine/Source/Editor/UnrealEd/Private/EditorDomain/EditorDomainSave.h
+++ b/Engine/Source/Editor/UnrealEd/Private/EditorDomain/EditorDomainSave.h
@@ -7,6 +7,7 @@
 #include "Containers/RingBuffer.h"
 #include "Containers/StringView.h"
 #include "Containers/UnrealString.h"
+#include "EditorDomain/EditorDomain.h"
 #include "HAL/CriticalSection.h"
 #include "HAL/Platform.h"
 #include "HAL/PlatformProcess.h"
@@ -138,6 +139,8 @@ private:
 	IAssetRegistry* AssetRegistry;
 	/** In-process handle to the interprocess synchronization object. */
 	FPlatformProcess::FSemaphore* ProcessLock = nullptr;
+	/** Cache of GetSchemaHash by class name */
+	UE::EditorDomain::FClassDigestMap ClassDigests;
 	/** The time at which the current idleness started. Used to decide abdication. */
 	double IdleStartTime = 0;
 	/** The time at which garbage was last collected. Used for scheduling the next collection. */
diff --git a/Engine/Source/Editor/UnrealEd/Private/EditorDomain/EditorDomainUtils.cpp b/Engine/Source/Editor/UnrealEd/Private/EditorDomain/EditorDomainUtils.cpp
--- a/Engine/Source/Editor/UnrealEd/Private/EditorDomain/EditorDomainUtils.cpp
+++ b/Engine/Source/Editor/UnrealEd/Private/EditorDomain/EditorDomainUtils.cpp
@@ -16,6 +16,7 @@
 #include "Misc/PackagePath.h"
 #include "Misc/StringBuilder.h"
 #include "Serialization/CompactBinaryWriter.h"
+#include "UObject/CoreRedirects.h"
 #include "UObject/ObjectVersion.h"
 #include "UObject/Package.h"
 #include "UObject/UObjectHash.h"
@@ -30,7 +31,8 @@ const TCHAR* EditorDomainVersion = TEXT("C217EB656E9B4C04816D3DC0E21901F6");
 // Identifier of the CacheBucket for EditorDomainPackages
 const TCHAR* EditorDomainPackageBucketName = TEXT("EditorDomainPackage");
 
-FPackageDigest GetPackageDigest(const FAssetPackageData& PackageData, FName PackageName)
+EPackageDigestResult GetPackageDigest(FPackageDigest& OutDigest, FString& OutErrorMessage,
+	const FAssetPackageData& PackageData, FName PackageName, FClassDigestMap& ClassDigests)
 {
 	FCbWriter Writer;
 	int32 CurrentFileVersionUE = GPackageFileUEVersion;
@@ -52,25 +54,119 @@ FPackageDigest GetPackageDigest(const FAssetPackageData& PackageData, FName Pack
 		}
 		else
 		{
-			UE_LOG(LogEditorDomain, Error, TEXT("CustomVersion guid %s is used by package %s but is not ")
-				TEXT("available in FCurrentCustomVersions. PackageDigest will set its version to 0."),
-				*PackageVersion.Key.ToString(), *PackageName.ToString());
+			OutErrorMessage = FString::Printf(TEXT("Package %s uses CustomVersion guid %s but that guid is not available in FCurrentCustomVersions"),
+				*PackageName.ToString(), *PackageVersion.Key.ToString());
+			return EPackageDigestResult::MissingCustomVersion;
+		}
+	}
+	FNameBuilder NameBuilder;
+	int32 NextClass = 0;
+	for (int32 Attempt = 0; NextClass < PackageData.ImportedClasses.Num(); ++Attempt)
+	{
+		if (Attempt > 0)
+		{
+			// EDITORDOMAIN_TODO: Remove this !IsInGameThread check once FindObject no longer asserts if GIsSavingPackage
+			if (Attempt > 1 || !IsInGameThread())
+			{
+				OutErrorMessage = FString::Printf(TEXT("Package %s uses Class %s but that class is not loaded"),
+					*PackageName.ToString(), *PackageData.ImportedClasses[NextClass].ToString());
+				return EPackageDigestResult::MissingClass;
+			}
+			TConstArrayView<FName> RemainingClasses(PackageData.ImportedClasses);
+			RemainingClasses = RemainingClasses.Slice(NextClass, PackageData.ImportedClasses.Num() - NextClass);
+			PrecacheClassDigests(RemainingClasses, ClassDigests);
+		}
+		FScopeLock ClassDigestsScopeLock(&ClassDigests.Lock);
+		for (; NextClass < PackageData.ImportedClasses.Num(); ++NextClass)
+		{
+			FName ClassName = PackageData.ImportedClasses[NextClass];
+			FClassDigestData* ExistingData = ClassDigests.Map.Find(ClassName);
+			if (!ExistingData)
+			{
+				break;
+			}
+			if (ExistingData->bNative)
+			{
+				Writer << ExistingData->SchemaHash;
+			}
+		}
+	}
+	OutDigest = Writer.Save().GetRangeHash();
+	return EPackageDigestResult::Success;
+}
+
+void PrecacheClassDigests(TConstArrayView<FName> ClassNames, FClassDigestMap& ClassDigests)
+{
+	TArray<TPair<FName,FClassDigestData>> ClassesToAdd;
+	FString ClassNameStr;
+	{
+		FScopeLock ClassDigestsScopeLock(&ClassDigests.Lock);
+		for (FName ClassName : ClassNames)
+		{
+			if (!ClassDigests.Map.Find(ClassName))
+			{
+				ClassesToAdd.Emplace_GetRef().Get<0>() = ClassName;
+			}
+		}
+	}
+	if (ClassesToAdd.Num())
+	{
+		FString TargetClassNameString;
+		for (int32 Index = 0; Index < ClassesToAdd.Num();)
+		{
+			FName OldClassFName = ClassesToAdd[Index].Get<0>();
+			FClassDigestData& Data = ClassesToAdd[Index].Get<1>();
+			OldClassFName.ToString(TargetClassNameString);
+			FCoreRedirectObjectName OldClassName(TargetClassNameString);
+			FCoreRedirectObjectName NewClassName = FCoreRedirects::GetRedirectedName(ECoreRedirectFlags::Type_Class, OldClassName);
+			if (OldClassName != NewClassName)
+			{
+				TargetClassNameString = NewClassName.ToString();
+			}
+
+			if (FPackageName::IsScriptPackage(TargetClassNameString))
+			{
+				UStruct* Struct = FindObject<UStruct>(nullptr, *TargetClassNameString);
+				if (Struct)
+				{
+					Data.SchemaHash = Struct->GetSchemaHash(false /* bSkipEditorOnly */);
+					Data.bNative = true;
+					++Index;
+				}
+				else
+				{
+					ClassesToAdd.RemoveAtSwap(Index);
+				}
+			}
+			else
+			{
+				Data.SchemaHash.Reset();
+				Data.bNative = false;
+				++Index;
+			}
+		}
+		{
+			FScopeLock ClassDigestsScopeLock(&ClassDigests.Lock);
+			for (TPair<FName,FClassDigestData>& Pair: ClassesToAdd)
+			{
+				ClassDigests.Map.Add(Pair.Get<0>(), MoveTemp(Pair.Get<1>()));
+			}
 		}
 	}
-	return Writer.Save().GetRangeHash();
 }
 
 EPackageDigestResult GetPackageDigest(IAssetRegistry& AssetRegistry, FName PackageName,
-	FPackageDigest& OutPackageDigest)
+	FPackageDigest& OutPackageDigest, FString& OutErrorMessage, FClassDigestMap& ClassDigests)
 {
 	AssetRegistry.WaitForPackage(PackageName.ToString());
 	TOptional<FAssetPackageData> PackageData = AssetRegistry.GetAssetPackageDataCopy(PackageName);
 	if (!PackageData)
 	{
+		OutErrorMessage = FString::Printf(TEXT("Package %s does not exist in the AssetRegistry"),
+			*PackageName.ToString());
 		return EPackageDigestResult::FileDoesNotExist;
 	}
-	OutPackageDigest = GetPackageDigest(*PackageData, PackageName);
-	return EPackageDigestResult::Success;
+	return GetPackageDigest(OutPackageDigest, OutErrorMessage, *PackageData, PackageName, ClassDigests);
 }
 
 UE::DerivedData::FCacheKey GetEditorDomainPackageKey(const FPackageDigest& PackageDigest)
@@ -90,20 +186,18 @@ UE::DerivedData::FRequest RequestEditorDomainPackage(const FPackagePath& Package
 		UE::DerivedData::ECachePolicy::QueryLocal, CachePriority, MoveTemp(Callback));
 }
 
-bool TrySavePackage(UPackage* Package)
+bool TrySavePackage(UPackage* Package, FClassDigestMap& ClassDigests)
 {
+	FString ErrorMessage;
 	FPackageDigest PackageDigest;
-	EPackageDigestResult FindHashResult = GetPackageDigest(*IAssetRegistry::Get(), Package->GetFName(), PackageDigest);
+	EPackageDigestResult FindHashResult = GetPackageDigest(*IAssetRegistry::Get(), Package->GetFName(), PackageDigest,
+		ErrorMessage, ClassDigests);
 	switch (FindHashResult)
 	{
 	case EPackageDigestResult::Success:
 		break;
-	case EPackageDigestResult::FileDoesNotExist:
-		return false;
-	case EPackageDigestResult::WrongThread:
-		return false;
 	default:
-		check(false);
+		UE_LOG(LogEditorDomain, Warning, TEXT("Could not save package to EditorDomain: %s."), *ErrorMessage)
 		return false;
 	}
 
diff --git a/Engine/Source/Editor/UnrealEd/Private/EditorDomain/EditorDomainUtils.h b/Engine/Source/Editor/UnrealEd/Private/EditorDomain/EditorDomainUtils.h
--- a/Engine/Source/Editor/UnrealEd/Private/EditorDomain/EditorDomainUtils.h
+++ b/Engine/Source/Editor/UnrealEd/Private/EditorDomain/EditorDomainUtils.h
@@ -2,10 +2,12 @@
 
 #pragma once
 
+#include "Containers/ArrayView.h"
 #include "Containers/UnrealString.h"
 #include "DerivedDataCache.h"
 #include "EditorDomain/EditorDomain.h"
 #include "Misc/ConfigCacheIni.h"
+#include "UObject/NameTypes.h"
 
 DECLARE_LOG_CATEGORY_EXTERN(LogEditorDomainSave, Log, All);
 
@@ -14,16 +16,15 @@ class FPackagePath;
 class IAssetRegistry;
 class UPackage;
 
-namespace UE
-{
-namespace EditorDomain
+namespace UE::EditorDomain
 {
 
 enum class EPackageDigestResult
 {
 	Success,
 	FileDoesNotExist,
-	WrongThread,
+	MissingCustomVersion,
+	MissingClass,
 };
 
 /**
@@ -31,7 +32,9 @@ enum class EPackageDigestResult
  * Reads information from the AssetRegistry to compute the digest.
  */
 EPackageDigestResult GetPackageDigest(IAssetRegistry& AssetRegistry, FName PackageName,
-	FPackageDigest& OutPackageDigest);
+	FPackageDigest& OutPackageDigest, FString& OutErrorMessage, FClassDigestMap& ClassDigests);
+/** For any ClassNames not already in ClassDigests, look up their UStruct and add them. */
+void PrecacheClassDigests(TConstArrayView<FName> ClassNames, FClassDigestMap& ClassDigests);
 
 /** Convert the given PackageDigest into CacheKey format. */
 UE::DerivedData::FCacheKey GetEditorDomainPackageKey(const FPackageDigest& PackageDigest);
@@ -42,7 +45,6 @@ UE::DerivedData::FRequest RequestEditorDomainPackage(const FPackagePath& Package
 	UE::DerivedData::FOnCacheGetComplete&& Callback);
 
 /** Save the given package into the EditorDomain. */
-bool TrySavePackage(UPackage* Package);
+bool TrySavePackage(UPackage* Package, FClassDigestMap& ClassDigests);
 
 }
-}
