Prevent long frames after garbage collection.  Previously all viewports had their hit proxies invalidated and were all re-drawn.  Now, the hit-proxies are invalidated, and a redraw only happens when needed.

This is an update to CL 14327142 for https://jira.it.epicgames.com/browse/UE-99374

diff --git a/Engine/Source/Editor/UnrealEd/Private/EditorServer.cpp b/Engine/Source/Editor/UnrealEd/Private/EditorServer.cpp
--- a/Engine/Source/Editor/UnrealEd/Private/EditorServer.cpp
+++ b/Engine/Source/Editor/UnrealEd/Private/EditorServer.cpp
@@ -2870,6 +2870,9 @@ bool UEditorEngine::Map_Load(const TCHAR* Str, FOutputDevice& Ar)
 						GEngine->WorldAdded( Context.World() );
 					}
 
+					// Invalidate all the level viewport hit proxies
+					RedrawLevelEditingViewports();
+
 					// Collect any stale components or other objects that are no longer required after loading the map
 					CollectGarbage(GARBAGE_COLLECTION_KEEPFLAGS, true);
 				}
diff --git a/Engine/Source/Editor/UnrealEd/Private/UnrealEdEngine.cpp b/Engine/Source/Editor/UnrealEd/Private/UnrealEdEngine.cpp
--- a/Engine/Source/Editor/UnrealEd/Private/UnrealEdEngine.cpp
+++ b/Engine/Source/Editor/UnrealEd/Private/UnrealEdEngine.cpp
@@ -700,7 +700,13 @@ void UUnrealEdEngine::OnPostGarbageCollect()
 		}
 	}
 
-	RedrawAllViewports();
+	for (FEditorViewportClient* ViewportClient : GetAllViewportClients())
+	{
+		if (ViewportClient && ViewportClient->Viewport)
+		{
+			ViewportClient->Viewport->InvalidateHitProxy();
+		}
+	}
 }
 
 
