Fallback the editor to calculating a 1x1 scene rect if the viewport ends up 0x0

diff --git a/Engine/Source/Editor/UnrealEd/Private/EditorViewportClient.cpp b/Engine/Source/Editor/UnrealEd/Private/EditorViewportClient.cpp
--- a/Engine/Source/Editor/UnrealEd/Private/EditorViewportClient.cpp
+++ b/Engine/Source/Editor/UnrealEd/Private/EditorViewportClient.cpp
@@ -1170,6 +1170,12 @@ FSceneView* FEditorViewportClient::CalcSceneView(FSceneViewFamily* ViewFamily, c
 		}
 	}
 
+	if (!ViewInitOptions.IsValidViewRectangle())
+	{
+		// Zero sized rects are invalid, so fake to 1x1 to avoid asserts later on
+		ViewInitOptions.SetViewRectangle(FIntRect(0, 0, 1, 1));
+	}
+
 	// Allocate our stereo view state on demand, so that only viewports that actually use stereo features have one
 	const int32 ViewStateIndex = (StereoPass > eSSP_RIGHT_EYE) ? StereoPass - eSSP_RIGHT_EYE : 0;
 	if (bStereoRendering)
