Ensure the inactive world created from the Content Browser is purged when loading that world for editing

UEditorEngine::OnAssetLoaded calls UEditorEngine::InitializeNewlyCreatedInactiveWorld for any world loaded, which can result in a resident world asset that is only partially initialized (inactive worlds disable things like AI).

UEditorEngine::Map_Load finds this resident world and attempts to re-use it, however since bIsWorldInitialized is already set to true it skips running the normal editor world initialization (which enables things like AI), and the end result is an editable world with certain features uninitialized.

UEditorEngine::Map_Load does attempt to purge any resident worlds that have been left lingering, but had made special exceptions for the world that the editor is attempting to load (UE-9631, UE-17205). This change has that logic run on any lingering resident worlds that have already been initialized, thus ensuring that the world being edited will definitely be initialized with the correct set of enabled features.

diff --git a/Engine/Source/Editor/UnrealEd/Private/EditorServer.cpp b/Engine/Source/Editor/UnrealEd/Private/EditorServer.cpp
--- a/Engine/Source/Editor/UnrealEd/Private/EditorServer.cpp
+++ b/Engine/Source/Editor/UnrealEd/Private/EditorServer.cpp
@@ -2483,15 +2483,15 @@ bool UEditorEngine::Map_Load(const TCHAR* Str, FOutputDevice& Ar)
 		if ( FPackageName::TryConvertFilenameToLongPackageName(TempFname, LongTempFname) )
 		{
 			// Is the new world already loaded?
-			UPackage* ExistingPackage = FindPackage(NULL, *LongTempFname);
-			UWorld* ExistingWorld = NULL;
+			UPackage* ExistingPackage = FindPackage(nullptr, *LongTempFname);
+			UWorld* ExistingWorld = nullptr;
 			if (ExistingPackage)
 			{
 				ExistingWorld = FindWorldInPackageOrFollowRedirector(ExistingPackage);
 			}
 
 			FString UnusedAlteredPath;
-			if ( ExistingWorld || FPackageName::DoesPackageExist(LongTempFname, NULL, &UnusedAlteredPath) )
+			if ( ExistingWorld || FPackageName::DoesPackageExist(LongTempFname, nullptr, &UnusedAlteredPath) )
 			{
 				FText NotMapReason;
 				if( !ExistingWorld && !PackageIsAMapFile( *TempFname, NotMapReason ) )
@@ -2534,8 +2534,6 @@ bool UEditorEngine::Map_Load(const TCHAR* Str, FOutputDevice& Ar)
 
 				SlowTask.EnterProgressFrame(10, FText::Format( NSLOCTEXT("UnrealEd", "LoadingMapStatus_CleaningUp", "{0} (Clearing existing world)"), LocalizedLoadingMap ));
 
-				UObject* OldOuter = NULL;
-
 				{
 					// Clear the lighting build results
 					FMessageLog("LightingResults").NewPage(LOCTEXT("LightingBuildNewLogPage", "Lighting Build"));
@@ -2543,8 +2541,6 @@ bool UEditorEngine::Map_Load(const TCHAR* Str, FOutputDevice& Ar)
 					FStatsViewerModule& StatsViewerModule = FModuleManager::Get().LoadModuleChecked<FStatsViewerModule>(TEXT("StatsViewer"));
 					StatsViewerModule.GetPage(EStatsPage::LightingBuildInfo)->Clear();
 
-					OldOuter = Context.World()->GetOuter();
-
 					ResetTransaction( LocalizedLoadingMap );
 
 					// Don't clear errors if we are loading a startup map so we can see all startup load errors
@@ -2556,38 +2552,74 @@ bool UEditorEngine::Map_Load(const TCHAR* Str, FOutputDevice& Ar)
 					}
 
 					// If we are loading the same world again (reloading) then we must not specify that we want to keep this world in memory.
-					// Otherwise, try to keep the existing world in memory since there is not reason to reload it.
+					// Otherwise, try to keep the existing uninitialized world in memory since there is not reason to reload it.
 					UWorld* NewWorld = nullptr;
-					if (!bIsLoadingMapTemplate && ExistingWorld != nullptr && Context.World() != ExistingWorld)
+					if (ExistingWorld && !ExistingWorld->bIsWorldInitialized && Context.World() != ExistingWorld && !bIsLoadingMapTemplate)
 					{
 						NewWorld = ExistingWorld;
 					}
 					EditorDestroyWorld( Context, LocalizedLoadingMap, NewWorld );
 
-					// Unload all other map packages currently loaded, before opening a new map.
+					// Unload all other map packages currently loaded and initialized, before opening a new map.
 					// The world is only initialized correctly as part of the level loading process, so ensure that every map package needs loading.
-					TArray<UPackage*> WorldPackages;
-					for (TObjectIterator<UWorld> It; It; ++It)
 					{
-						UPackage* Package = Cast<UPackage>(It->GetOuter());
-						
-
-						if (Package && Package != GetTransientPackage() && Package->GetPathName() != LongTempFname)
+						TArray<UPackage*> WorldPackages;
+						for (TObjectIterator<UWorld> It; It; ++It)
 						{
-							WorldPackages.AddUnique(Package);
+							UPackage* Package = Cast<UPackage>(It->GetOuter());
+							if (Package && Package != GetTransientPackage() && (It->bIsWorldInitialized || Package->GetPathName() != LongTempFname))
+							{
+								WorldPackages.AddUnique(Package);
+							}
 						}
+						UPackageTools::UnloadPackages(WorldPackages);
 					}
-					UPackageTools::UnloadPackages(WorldPackages);
 
 					// Refresh ExistingPackage and Existing World now that GC has occurred.
-					ExistingPackage = FindPackage(NULL, *LongTempFname);
+					ExistingPackage = FindPackage(nullptr, *LongTempFname);
 					if (ExistingPackage)
 					{
 						ExistingWorld = FindWorldInPackageOrFollowRedirector(ExistingPackage);
+
+						if (!ExistingWorld)
+						{
+							// If we have a world package without a world, try and unload the package so we can hopefully reload it again correctly from disk
+							TArray<UPackage*> WorldPackages;
+							WorldPackages.Add(ExistingPackage);
+							UPackageTools::UnloadPackages(WorldPackages);
+
+							ExistingPackage = FindPackage(nullptr, *LongTempFname);
+						}
 					}
 					else
 					{
-						ExistingWorld = NULL;
+						ExistingWorld = nullptr;
+					}
+
+					// If the existing world has already been initialized then the unload attempts above failed, and we need to 
+					// fatally error and dump any lingering references (like we would when unloading the main editor world).
+					if (ExistingWorld && (ExistingWorld->bIsWorldInitialized || ExistingPackage))
+					{
+						int32 NumFailedToCleanup = 0;
+
+						if (ExistingWorld)
+						{
+							FReferenceChainSearch RefChainSearch(ExistingWorld, EReferenceChainSearchMode::Shortest | EReferenceChainSearchMode::PrintResults);
+							UE_LOG(LogEditorServer, Error, TEXT("Old world %s not cleaned up by garbage collection while loading new map! Referenced by:") LINE_TERMINATOR TEXT("%s"), *ExistingWorld->GetPathName(), *RefChainSearch.GetRootPath());
+							++NumFailedToCleanup;
+						}
+						
+						if (ExistingPackage)
+						{
+							FReferenceChainSearch RefChainSearch(ExistingPackage, EReferenceChainSearchMode::Shortest | EReferenceChainSearchMode::PrintResults);
+							UE_LOG(LogEditorServer, Error, TEXT("Old level package %s not cleaned up by garbage collection while loading new map! Referenced by:") LINE_TERMINATOR TEXT("%s"), *ExistingPackage->GetPathName(), *RefChainSearch.GetRootPath());
+							++NumFailedToCleanup;
+						}
+						
+						if (NumFailedToCleanup > 0)
+						{
+							UE_LOG(LogEditorServer, Fatal, TEXT("World Memory Leaks: %d leaks objects and packages. See The output above."), NumFailedToCleanup);
+						}
 					}
 
 					SlowTask.EnterProgressFrame( 70, LocalizedLoadingMap );
diff --git a/Engine/Source/Editor/UnrealEd/Private/PackageTools.cpp b/Engine/Source/Editor/UnrealEd/Private/PackageTools.cpp
--- a/Engine/Source/Editor/UnrealEd/Private/PackageTools.cpp
+++ b/Engine/Source/Editor/UnrealEd/Private/PackageTools.cpp
@@ -388,7 +388,7 @@ UPackageTools::UPackageTools(const FObjectInitializer& ObjectInitializer)
 				// Flush all pending render commands, as unloading the package may invalidate render resources.
 				FlushRenderingCommands();
 
-				// Close any open asset editors
+				// Close any open asset editors.
 				ForEachObjectWithPackage(PackageBeingUnloaded, [](UObject* Obj)
 				{
 					if (Obj->IsAsset())
@@ -398,6 +398,23 @@ UPackageTools::UPackageTools(const FObjectInitializer& ObjectInitializer)
 					return true;
 				}, false);
 
+				// Notify any Blueprints that are about to be unloaded, and destroy any leftover worlds.
+				ForEachObjectWithPackage(PackageBeingUnloaded, [](UObject* Obj)
+				{
+					if (UBlueprint* BP = Cast<UBlueprint>(Obj))
+					{
+						BP->ClearEditorReferences();
+					}
+					if (UWorld* World = Cast<UWorld>(Obj))
+					{
+						if (World->bIsWorldInitialized)
+						{
+							World->CleanupWorld();
+						}
+					}
+					return true;
+				}, true, RF_Transient, EInternalObjectFlags::PendingKill);
+
 				PackageBeingUnloaded->bHasBeenFullyLoaded = false;
 				PackageBeingUnloaded->ClearFlags(RF_WasLoaded);
 				if ( PackageBeingUnloaded->HasAnyPackageFlags(PKG_ContainsScript) )
@@ -839,15 +856,22 @@ UPackageTools::UPackageTools(const FObjectInitializer& ObjectInitializer)
 		{
 			GEngine->NotifyToolsOfObjectReplacement(InPackageReloadedEvent->GetRepointedObjects());
 
-			// Notify any Blueprints that are about to be unloaded.
-			ForEachObjectWithPackage(InPackageReloadedEvent->GetOldPackage(), [&](UObject* InObject)
+			// Notify any Blueprints that are about to be unloaded, and destroy any leftover worlds.
+			ForEachObjectWithPackage(InPackageReloadedEvent->GetOldPackage(), [](UObject* InObject)
 			{
 				if (UBlueprint* BP = Cast<UBlueprint>(InObject))
 				{
 					BP->ClearEditorReferences();
 				}
+				if (UWorld* World = Cast<UWorld>(InObject))
+				{
+					if (World->bIsWorldInitialized)
+					{
+						World->CleanupWorld();
+					}
+				}
 				return true;
-			}, false, RF_Transient, EInternalObjectFlags::PendingKill);
+			}, true, RF_Transient, EInternalObjectFlags::PendingKill);
 		}
 
 		if (InPackageReloadPhase == EPackageReloadPhase::OnPackageFixup)
