IterationProfile: Add explicit test name parameter to the commandlet, and specify exact test name w/o project or platform to keep telemetry concise.

diff --git a/Engine/Build/IterationProfile.xml b/Engine/Build/IterationProfile.xml
--- a/Engine/Build/IterationProfile.xml
+++ b/Engine/Build/IterationProfile.xml
@@ -61,8 +61,9 @@
 
 	<!-- The main test running macro logic that is expanded with platform and dependency node names -->
 	<Macro Name="RunTest" Arguments="OutputName;PlatformName;ProjectName;NodeName;NodeRequires;Phase;Args">
-		<Property Name="TraceFile" Value="$(TraceDir)/$(ProjectName)IterationProfile$(PlatformName)$(Phase).utrace"/>
-		<Property Name="TelemetryFile" Value="$(TraceDir)/$(ProjectName)IterationProfile$(PlatformName)$(Phase)Telemetry.csv"/>
+		<Property Name="TestName" Value="IterationProfile$(Phase)"/>
+		<Property Name="TraceFile" Value="$(TraceDir)/$(ProjectName)$(PlatformName)$(TestName).utrace"/>
+		<Property Name="TelemetryFile" Value="$(TraceDir)/$(ProjectName)$(PlatformName)$(TestName)Telemetry.csv"/>
 		<Property Name="ReportExportPath" Value="$(AutomationOutputDir)/$(OutputName)"/>
 		<Property Name="ReportParams" Value="-ReportExportPath=&quot;$(ReportExportPath)&quot;"/>
 		<Do If="$(IsBuildMachine)">
@@ -72,10 +73,10 @@
 		<Property Name="LogOutputParams" Value="-logdir=&quot;$(AutomationOutputDir)/$(OutputName)&quot;"/>
 		<Property Name="AllGauntletArgs" Value="-platform=$(PlatformName) -project=$(ProjectName) $(GauntletArgs) $(ReportParams) $(LogOutputParams) -artifactname=&quot;$(Phase)&quot;"/>
 		<Command Name="RunUnreal" Arguments="$(AllGauntletArgs) $(ProfileArgs) $(Args) -tracefile=&quot;$(TraceFile)&quot;"/>
-		<Commandlet Name="SummarizeTrace" Arguments="-NoEnginePlugins -inputfile=&quot;$(TraceFile)&quot; $(SummarizeArgs)"/>
+		<Commandlet Name="SummarizeTrace" Arguments="-NoEnginePlugins -inputfile=&quot;$(TraceFile)&quot; -testname=&quot;$(TestName)&quot; $(SummarizeArgs)"/>
 		<Do If="$(IsBuildMachine)">
 			<!-- Files copied to this folder are automatically uploaded to Horde as artifacts -->
-			<Copy From="$(TraceDir)/$(ProjectName)IterationProfile*.*" To="$(RootDir)/Engine/Programs/AutomationTool/Saved/Logs"/>
+			<Copy From="$(TraceDir)/$(ProjectName)$(TestName)*.*" To="$(RootDir)/Engine/Programs/AutomationTool/Saved/Logs"/>
 			<Command Name="PublishUnrealAutomationTelemetry" Arguments="-CSVFile=&quot;$(TelemetryFile)&quot; -project=$(ProjectName) -changelist=$(Change) -branch=$(Branch)"/>
 		</Do>
 	</Macro>
diff --git a/Engine/Source/Editor/UnrealEd/Private/Commandlets/SummarizeTraceCommandlet.cpp b/Engine/Source/Editor/UnrealEd/Private/Commandlets/SummarizeTraceCommandlet.cpp
--- a/Engine/Source/Editor/UnrealEd/Private/Commandlets/SummarizeTraceCommandlet.cpp
+++ b/Engine/Source/Editor/UnrealEd/Private/Commandlets/SummarizeTraceCommandlet.cpp
@@ -619,6 +619,7 @@ int32 USummarizeTraceCommandlet::Main(const FString& CmdLineParams)
 		UE_LOG(LogSummarizeTrace, Log, TEXT("Options:"));
 		UE_LOG(LogSummarizeTrace, Log, TEXT(" Required: -inputfile=<utrace path>   (The utrace you wish to process)"));
 		UE_LOG(LogSummarizeTrace, Log, TEXT(" Optional: -statsfile=<csv path>      (The csv of statistics to generate)"));
+		UE_LOG(LogSummarizeTrace, Log, TEXT(" Optional: -testname=<string>         (Test name to use in telemetry csv)"));
 		return 0;
 	}
 
@@ -865,6 +866,10 @@ int32 USummarizeTraceCommandlet::Main(const FString& CmdLineParams)
 		IFileHandle* TelemetryCsvHandle = FPlatformFileManager::Get().GetPlatformFile().OpenWrite(*TelemetryCsvFileName);
 		if (TelemetryCsvHandle)
 		{
+			// override the test name
+			FString TestName = TraceFileBasename;
+			FParse::Value(*CmdLineParams, TEXT("testname="), TestName, true);
+
 			// no newline, see row printfs
 			FilePrint(TelemetryCsvHandle, FString::Printf(TEXT("TestName,Context,DataPoint,Measurement")));
 			for (const FSummarizeCpuAnalyzer::FScope& Scope : SortedScopes)
@@ -874,7 +879,7 @@ int32 USummarizeTraceCommandlet::Main(const FString& CmdLineParams)
 				for (const FString& Statistic : Statistics)
 				{
 					// note newline is at the front of every data line to prevent final extraneous newline, per customary for csv
-					FilePrint(TelemetryCsvHandle, FString::Printf(TEXT("\n%s,%s,%s,%s,"), *TraceFileBasename, *Scope.Name, *Statistic, *Scope.GetValue(Statistic)));
+					FilePrint(TelemetryCsvHandle, FString::Printf(TEXT("\n%s,%s,%s,%s,"), *TestName, *Scope.Name, *Statistic, *Scope.GetValue(Statistic)));
 				}
 			}
 			for (const TMap<uint16, FSummarizeCountersAnalyzer::FCounter>::ElementType& Counter : CountersAnalyzer.Counters)
@@ -884,7 +889,7 @@ int32 USummarizeTraceCommandlet::Main(const FString& CmdLineParams)
 				for (const FString& Statistic : Statistics)
 				{
 					// note newline is at the front of every data line to prevent final extraneous newline, per customary for csv
-					FilePrint(TelemetryCsvHandle, FString::Printf(TEXT("\n%s,%s,%s,%s,"), *TraceFileBasename, *Counter.Value.Name, *Statistic, *Counter.Value.GetValue()));
+					FilePrint(TelemetryCsvHandle, FString::Printf(TEXT("\n%s,%s,%s,%s,"), *TestName, *Counter.Value.Name, *Statistic, *Counter.Value.GetValue()));
 				}
 			}
 			TelemetryCsvHandle->Flush();
