World Partition Cook on the Fly:
Set guids for created levels deterministically to allow confirming the cooked packages don't change; this will also eventually be needed for shipped titles to reduce spurious diffs that create large patches.
Keep the UWorld referenced to prevent it from being GC'd after its saved but before its packages are popualted; we need this to enable deferred population.
Suppress cooking of one file per actor packages in cook on the fly as is already done in cook by the book.
#rnx

diff --git a/Engine/Source/Editor/UnrealEd/Classes/CookOnTheSide/CookOnTheFlyServer.h b/Engine/Source/Editor/UnrealEd/Classes/CookOnTheSide/CookOnTheFlyServer.h
--- a/Engine/Source/Editor/UnrealEd/Classes/CookOnTheSide/CookOnTheFlyServer.h
+++ b/Engine/Source/Editor/UnrealEd/Classes/CookOnTheSide/CookOnTheFlyServer.h
@@ -652,7 +652,7 @@ private:
 		const TArray<FString>& IniMapSections, ECookByTheBookOptions FilesToCookFlags,
 		const TArrayView<const ITargetPlatform* const>& TargetPlatforms);
 	/* Collect filespackages that should not be cooked from ini settings and commandline. Does not include checking UAssetManager, which has to be queried later */
-	static TArray<FName> GetNeverCookPackageFileNames(TArrayView<const FString> ExtraNeverCookDirectories = TArrayView<const FString>());
+	TArray<FName> GetNeverCookPackageFileNames(TArrayView<const FString> ExtraNeverCookDirectories = TArrayView<const FString>());
 
 
 	/**
diff --git a/Engine/Source/Editor/UnrealEd/Private/CookOnTheFlyServer.cpp b/Engine/Source/Editor/UnrealEd/Private/CookOnTheFlyServer.cpp
--- a/Engine/Source/Editor/UnrealEd/Private/CookOnTheFlyServer.cpp
+++ b/Engine/Source/Editor/UnrealEd/Private/CookOnTheFlyServer.cpp
@@ -538,6 +538,11 @@ bool UCookOnTheFlyServer::StartNetworkFileServer(const bool BindAnyPort, const T
 
 	CreateSandboxFile();
 	GenerateAssetRegistry();
+	for (FName NeverCookPackage : GetNeverCookPackageFileNames(TArrayView<const FString>()))
+	{
+		PackageTracker->NeverCookPackageList.Add(NeverCookPackage);
+	}
+
 
 	{
 		UE::Cook::FPlatformManager::FReadScopeLock PlatformScopeLock(PlatformManager->ReadLockPlatforms());
@@ -2219,17 +2224,6 @@ UE::Cook::FGeneratorPackage* UCookOnTheFlyServer::CreateGeneratorPackage(UE::Coo
 		return nullptr;
 	}
 
-	// TODO: Add support for cook on the fly
-	if (IsCookOnTheFlyMode())
-	{
-		UWorld* World = Cast<UWorld>(SplitDataObject);
-		if (World && World->GetWorldPartition())
-		{
-			UE_LOG(LogCook, Error, TEXT("Cook on the fly doesn't support World Partition."));
-			return nullptr;
-		}
-	}
-
 	UE_LOG(LogCook, Display, TEXT("Splitting Package %s with class %s acting on object %s."), *PackageData.GetPackageName().ToString(), *Splitter->GetSplitDataClass()->GetName(), *SplitDataObject->GetFullName());
 
 	// Create instance of CookPackageSplitter class
@@ -2310,6 +2304,7 @@ void UCookOnTheFlyServer::SplitPackage(UE::Cook::FGeneratorPackage* GeneratorStr
 		GeneratorStruct->SetClearedOldPackages();
 	}
 
+	UPackage* Owner = GeneratorStruct->GetOwner().GetPackage();
 	if (!GeneratorStruct->HasQueuedGeneratedPackages())
 	{
 		if (Splitter->UseDeferredPopulate())
@@ -2323,7 +2318,6 @@ void UCookOnTheFlyServer::SplitPackage(UE::Cook::FGeneratorPackage* GeneratorStr
 		}
 		else
 		{
-			UPackage* Owner = GeneratorStruct->GetOwner().GetPackage();
 			FPopulatePackageContext Context(GeneratorStruct, Owner);
 			TArrayView<FGeneratorPackage::FGeneratedStruct> PackagesToGenerate = GeneratorStruct->GetPackagesToGenerate();
 			// Start Splitting
@@ -2346,7 +2340,7 @@ void UCookOnTheFlyServer::SplitPackage(UE::Cook::FGeneratorPackage* GeneratorStr
 				FString IntermediateFile = GeneratorStruct->GetIntermediateLocalPath(GeneratedStruct);
 				FSavePackageResultStruct SaveResult = GEditor->Save(GeneratedPackage, nullptr, RF_Standalone,
 					*IntermediateFile, GError, nullptr, /*bForceByteSwapping*/ false,
-					/*bWarnOfLongFilename*/ true, SAVE_None);
+					/*bWarnOfLongFilename*/ true, SAVE_KeepGUID);
 				if (SaveResult.Result != ESavePackageResult::Success)
 				{
 					UE_LOG(LogCook, Error, TEXT("PackageSplitter could not save uncooked generated package. Splitter=%s, Generated=%s."),
@@ -2380,7 +2374,7 @@ void UCookOnTheFlyServer::SplitPackage(UE::Cook::FGeneratorPackage* GeneratorStr
 		SplitterData.Package = FindObject<UPackage>(nullptr, *GeneratedPackageName);
 		if (!SplitterData.Package)
 		{
-			SplitterData.Package = CreatePackage(*GeneratedPackageName);
+			SplitterData.Package = GeneratorStruct->CreateGeneratedUPackage(Owner, *GeneratedPackageName);
 		}
 	}
 
@@ -2451,7 +2445,7 @@ UPackage* UCookOnTheFlyServer::TryPopulateGeneratedPackage(UE::Cook::FPopulatePa
 	}
 	else
 	{
-		GeneratedPackage = CreatePackage(*GeneratedPackageName);
+		GeneratedPackage = GeneratorStruct->CreateGeneratedUPackage(OwnerPackage, *GeneratedPackageName);
 	}
 
 	// Populate package using CookPackageSplitterInstance and pass GeneratedPackage's cooked name for it to
@@ -7971,9 +7965,15 @@ TArray<FName> UCookOnTheFlyServer::GetNeverCookPackageFileNames(TArrayView<const
 
 	};
 	const UProjectPackagingSettings* const PackagingSettings = GetDefault<UProjectPackagingSettings>();
-	AddDirectoryPathArray(PackagingSettings->DirectoriesToNeverCook, TEXT("ProjectSettings -> Project -> Packaging -> Directories to never cook"));
-	AddDirectoryPathArray(PackagingSettings->TestDirectoriesToNotSearch, TEXT("ProjectSettings -> Project -> Packaging -> Test directories to not search"));
-	// Never cook External Actors
+
+	if (IsCookByTheBookMode())
+	{
+		// Respect the packaging settings nevercook directories for CookByTheBook
+		AddDirectoryPathArray(PackagingSettings->DirectoriesToNeverCook, TEXT("ProjectSettings -> Project -> Packaging -> Directories to never cook"));
+		AddDirectoryPathArray(PackagingSettings->TestDirectoriesToNotSearch, TEXT("ProjectSettings -> Project -> Packaging -> Test directories to not search"));
+	}
+
+	// For all modes, never cook External Actors; they are handled by the parent map
 	NeverCookDirectories.Add(FString::Printf(TEXT("/Game/%s"), ULevel::GetExternalActorsFolderName()));
 
 	TArray<FString> NeverCookPackagesPaths;
@@ -8958,4 +8958,4 @@ uint32 UCookOnTheFlyServer::FullLoadAndSave(uint32& CookedPackageCount)
 	return Result;
 }
 
-#undef LOCTEXT_NAMESPACE
\ No newline at end of file
+#undef LOCTEXT_NAMESPACE
diff --git a/Engine/Source/Editor/UnrealEd/Private/Cooker/CookPackageData.cpp b/Engine/Source/Editor/UnrealEd/Private/Cooker/CookPackageData.cpp
--- a/Engine/Source/Editor/UnrealEd/Private/Cooker/CookPackageData.cpp
+++ b/Engine/Source/Editor/UnrealEd/Private/Cooker/CookPackageData.cpp
@@ -1063,6 +1063,17 @@ namespace Cook
 		}
 	}
 
+	UPackage* FGeneratorPackage::CreateGeneratedUPackage(const UPackage* OwnerPackage, const TCHAR* GeneratedPackageName)
+	{
+		UPackage* GeneratedPackage = CreatePackage(GeneratedPackageName);
+		PRAGMA_DISABLE_DEPRECATION_WARNINGS
+		GeneratedPackage->SetGuid(OwnerPackage->GetGuid());
+		PRAGMA_ENABLE_DEPRECATION_WARNINGS
+		GeneratedPackage->SetPersistentGuid(OwnerPackage->GetPersistentGuid());
+		return GeneratedPackage;
+	}
+
+
 	void FGeneratorPackage::GetIntermediateMountPoint(FString& OutPackagePath, FString& OutLocalFilePath) const
 	{
 		const FString OwnerShortName = FPackageName::GetShortName(Owner.GetPackageName().ToString());
diff --git a/Engine/Source/Editor/UnrealEd/Private/Cooker/CookPackageData.h b/Engine/Source/Editor/UnrealEd/Private/Cooker/CookPackageData.h
--- a/Engine/Source/Editor/UnrealEd/Private/Cooker/CookPackageData.h
+++ b/Engine/Source/Editor/UnrealEd/Private/Cooker/CookPackageData.h
@@ -482,6 +482,9 @@ namespace Cook
 		/** Callback during garbage collection */
 		void PostGarbageCollect();
 
+		/** Call CreatePackage and set PackageData correctly for deterministic generated packages */
+		UPackage* CreateGeneratedUPackage(const UPackage* OwnerPackage, const TCHAR* GeneratedPackageName);
+
 	private:
 		/** PackageData for the package that is being split */
 		const UE::Cook::FPackageData& Owner;
diff --git a/Engine/Source/Runtime/Engine/Private/WorldPartition/WorldPartitionCookPackageSplitter.cpp b/Engine/Source/Runtime/Engine/Private/WorldPartition/WorldPartitionCookPackageSplitter.cpp
--- a/Engine/Source/Runtime/Engine/Private/WorldPartition/WorldPartitionCookPackageSplitter.cpp
+++ b/Engine/Source/Runtime/Engine/Private/WorldPartition/WorldPartitionCookPackageSplitter.cpp
@@ -40,6 +40,11 @@ TArray<ICookPackageSplitter::FGeneratedPackage> FWorldPartitionCookPackageSplitt
 	const UWorld* ConstPartitionedWorld = ValidateDataObject(OwnerObject);
 	UWorld* PartitionedWorld = const_cast<UWorld*>(ConstPartitionedWorld);
 
+	// Store the World pointer to declare it to GarbageCollection; we do not want to allow the World to be Garbage Collected
+	// until we have finished all of our TryPopulatePackage calls, because we store information on the World 
+	// that is necessary for populate 
+	ReferencedWorld = PartitionedWorld;
+
 	// World is not initialized
 	ensure(!PartitionedWorld->bIsWorldInitialized);
 
@@ -82,4 +87,9 @@ void FWorldPartitionCookPackageSplitter::PreSaveGeneratorPackage(UPackage* Owner
 	WorldPartition->FinalizeGeneratorPackageForCook(GeneratedPackages);
 }
 
+void FWorldPartitionCookPackageSplitter::AddReferencedObjects(FReferenceCollector& Collector)
+{
+	Collector.AddReferencedObject(ReferencedWorld);
+}
+
 #endif
\ No newline at end of file
diff --git a/Engine/Source/Runtime/Engine/Private/WorldPartition/WorldPartitionCookPackageSplitter.h b/Engine/Source/Runtime/Engine/Private/WorldPartition/WorldPartitionCookPackageSplitter.h
--- a/Engine/Source/Runtime/Engine/Private/WorldPartition/WorldPartitionCookPackageSplitter.h
+++ b/Engine/Source/Runtime/Engine/Private/WorldPartition/WorldPartitionCookPackageSplitter.h
@@ -6,8 +6,9 @@
 
 #include "CookPackageSplitter.h"
 #include "Engine/World.h"
+#include "UObject/GCObject.h"
 
-class FWorldPartitionCookPackageSplitter : public ICookPackageSplitter
+class FWorldPartitionCookPackageSplitter : public ICookPackageSplitter, public FGCObject
 {
 public:
 	//~ Begin of ICookPackageSplitter
@@ -21,9 +22,19 @@ public:
 		const TArray<ICookPackageSplitter::FGeneratedPackageForPreSave>& GeneratedPackages) override;
 	//~ End of ICookPackageSplitter
 
+	//~ Begin of FGCObject
+	virtual void AddReferencedObjects(FReferenceCollector& Collector);
+	virtual FString GetReferencerName() const
+	{
+		return TEXT("FWorldPartitionCookPackageSplitter");
+	}
+	//~ End of FGCObject
+
 private:
 	const UWorld* ValidateDataObject(const UObject* SplitData);
 	UWorld* ValidateDataObject(UObject* SplitData);
+
+	UWorld* ReferencedWorld = nullptr;
 };
 
 #endif
\ No newline at end of file
diff --git a/Engine/Source/Runtime/Engine/Private/WorldPartition/WorldPartitionLevelHelper.cpp b/Engine/Source/Runtime/Engine/Private/WorldPartition/WorldPartitionLevelHelper.cpp
--- a/Engine/Source/Runtime/Engine/Private/WorldPartition/WorldPartitionLevelHelper.cpp
+++ b/Engine/Source/Runtime/Engine/Private/WorldPartition/WorldPartitionLevelHelper.cpp
@@ -10,6 +10,7 @@
 #if WITH_EDITOR
 
 #include "FileHelpers.h"
+#include "Model.h"
 #include "UnrealEngine.h"
 #include "UObject/UObjectHash.h"
 #include "WorldPartition/ActorDescContainer.h"
@@ -180,6 +181,15 @@ ULevel* FWorldPartitionLevelHelper::CreateEmptyLevelForRuntimeCell(const UWorld*
 	check(NewLevel->OwningWorld == NewWorld);
 	check(NewLevel->Model);
 	check(!NewLevel->bIsVisible);
+
+	// Set the guids on the constructed level to something based on the generator rather than allowing indeterminism by
+	// constructing new Guids on every cook
+	// @todo_ow: revisit for static lighting support. We need to base the LevelBuildDataId on the relevant information from the
+	// actor's package.
+	NewLevel->LevelBuildDataId = InWorld->PersistentLevel->LevelBuildDataId;
+	check(InWorld->PersistentLevel->Model && NewLevel->Model);
+	NewLevel->Model->LightingGuid = InWorld->PersistentLevel->Model->LightingGuid;
+
 	return NewLevel;
 }
 
