Fixing issue with ScriptStructs and Enumerations appearing in the blueprint editor that have already been re-instanced.  This resulted in a crash with ScriptStructs.  Also fixed issue with changes to default values in a script struct not being refected in the blueprint nodes.

#rnx

diff --git a/Engine/Source/Editor/BlueprintGraph/Private/BlueprintActionDatabaseRegistrar.cpp b/Engine/Source/Editor/BlueprintGraph/Private/BlueprintActionDatabaseRegistrar.cpp
--- a/Engine/Source/Editor/BlueprintGraph/Private/BlueprintActionDatabaseRegistrar.cpp
+++ b/Engine/Source/Editor/BlueprintGraph/Private/BlueprintActionDatabaseRegistrar.cpp
@@ -213,9 +213,12 @@ int32 FBlueprintActionDatabaseRegistrar::RegisterStructActions(const FMakeStruct
 		if (const UScriptStruct* StructAsset = Cast<UScriptStruct>(RegistrarTarget))
 		{
 			check(IsOpenForRegistration(StructAsset));
-			if (UBlueprintNodeSpawner* NewAction = MakeActionCallback.Execute(StructAsset))
+			if ((StructAsset->StructFlags & STRUCT_NewerVersionExists) == 0)
 			{
-				RegisteredCount += (int32)AddBlueprintAction(StructAsset, NewAction);
+				if (UBlueprintNodeSpawner* NewAction = MakeActionCallback.Execute(StructAsset))
+				{
+					RegisteredCount += (int32)AddBlueprintAction(StructAsset, NewAction);
+				}
 			}
 		}
 		// else, the target is a class or a different asset type... not something pertaining to a struct
@@ -225,9 +228,12 @@ int32 FBlueprintActionDatabaseRegistrar::RegisterStructActions(const FMakeStruct
 		for (TObjectIterator<UScriptStruct> StructIt; StructIt; ++StructIt)
 		{
 			UScriptStruct const* Struct = (*StructIt);
-			if (UBlueprintNodeSpawner* NewAction = MakeActionCallback.Execute(Struct))
+			if ((Struct->StructFlags & STRUCT_NewerVersionExists) == 0)
 			{
-				RegisteredCount += (int32)AddBlueprintAction(Struct, NewAction);
+				if (UBlueprintNodeSpawner* NewAction = MakeActionCallback.Execute(Struct))
+				{
+					RegisteredCount += (int32)AddBlueprintAction(Struct, NewAction);
+				}
 			}
 		}
 	}
@@ -245,6 +251,10 @@ int32 FBlueprintActionDatabaseRegistrar::RegisterEnumActions(const FMakeEnumSpaw
 			for (TFieldIterator<UEnum> EnumIt(TargetClass, EFieldIteratorFlags::ExcludeSuper); EnumIt; ++EnumIt)
 			{
 				UEnum const* EnumToConsider = (*EnumIt);
+				if (EnumToConsider->HasAnyEnumFlags(EEnumFlags::NewerVersionExists))
+				{
+					continue;
+				}
 				if (!UEdGraphSchema_K2::IsAllowableBlueprintVariableType(EnumToConsider))
 				{
 					continue;
@@ -258,7 +268,7 @@ int32 FBlueprintActionDatabaseRegistrar::RegisterEnumActions(const FMakeEnumSpaw
 		}
 		else if (const UEnum* TargetEnum = Cast<UEnum>(RegistrarTarget))
 		{
-			if (UEdGraphSchema_K2::IsAllowableBlueprintVariableType(TargetEnum))
+			if (!TargetEnum->HasAnyEnumFlags(EEnumFlags::NewerVersionExists) && UEdGraphSchema_K2::IsAllowableBlueprintVariableType(TargetEnum))
 			{
 				if (UBlueprintNodeSpawner* NewAction = MakeActionCallback.Execute(TargetEnum))
 				{
@@ -272,6 +282,11 @@ int32 FBlueprintActionDatabaseRegistrar::RegisterEnumActions(const FMakeEnumSpaw
 		for (TObjectIterator<UEnum> EnumIt; EnumIt; ++EnumIt)
 		{
 			UEnum const* EnumToConsider = (*EnumIt);
+			if (EnumToConsider->HasAnyEnumFlags(EEnumFlags::NewerVersionExists))
+			{
+				continue;
+			}
+
 			if (!UEdGraphSchema_K2::IsAllowableBlueprintVariableType(EnumToConsider))
 			{
 				continue;
diff --git a/Engine/Source/Editor/UnrealEd/Private/Kismet2/ReloadUtilities.cpp b/Engine/Source/Editor/UnrealEd/Private/Kismet2/ReloadUtilities.cpp
--- a/Engine/Source/Editor/UnrealEd/Private/Kismet2/ReloadUtilities.cpp
+++ b/Engine/Source/Editor/UnrealEd/Private/Kismet2/ReloadUtilities.cpp
@@ -843,6 +843,20 @@ void FReload::NotifyChange(UScriptStruct* New, UScriptStruct* Old)
 	}
 }
 
+namespace
+{
+	template<typename T>
+	void CollectPackages(TArray<UPackage*>& Packages, const TMap<T*, T*>& Reinstances)
+	{
+		for (const TPair<T*, T*>& Pair : Reinstances)
+		{
+			T* Old = Pair.Key;
+			T* New = Pair.Value;
+			Packages.AddUnique(New ? New->GetPackage() : Old->GetPackage());
+		}
+	}
+}
+
 void FReload::Reinstance()
 {
 	if (Type != EActiveReloadType::Reinstancing)
@@ -852,6 +866,49 @@ void FReload::Reinstance()
 
 	TMap<UClass*, UClass*>& ClassesToReinstance = GetClassesToReinstanceForHotReload();
 
+	// If we have to collect the packages, gather them from the reinstanced objects
+	if (bCollectPackages)
+	{
+		CollectPackages(Packages, ClassesToReinstance);
+		CollectPackages(Packages, ReinstancedStructs);
+		CollectPackages(Packages, ReinstancedEnums);
+	}
+
+	// Remap all native functions (and gather scriptstructs)
+	TArray<UScriptStruct*> ScriptStructs;
+	for (FRawObjectIterator It; It; ++It)
+	{
+		if (UFunction* Function = Cast<UFunction>(static_cast<UObject*>(It->Object)))
+		{
+			if (FNativeFuncPtr NewFunction = FunctionRemap.FindRef(Function->GetNativeFunc()))
+			{
+				++NumFunctionsRemapped;
+				Function->SetNativeFunc(NewFunction);
+			}
+		} 
+		else if (UScriptStruct* ScriptStruct = Cast<UScriptStruct>(static_cast<UObject*>(It->Object)))
+		{
+			if (!ScriptStruct->HasAnyFlags(RF_ClassDefaultObject) && ScriptStruct->GetCppStructOps() && 
+				Packages.ContainsByPredicate([ScriptStruct](UPackage* Package) { return ScriptStruct->IsIn(Package); }))
+			{
+				ScriptStructs.Add(ScriptStruct);
+			}
+		}
+	}
+
+	// now let's set up the script structs...this relies on super behavior, so null them all, then set them all up. Internally this sets them up hierarchically.
+	for (UScriptStruct* Script : ScriptStructs)
+	{
+		Script->ClearCppStructOps();
+	}
+	for (UScriptStruct* Script : ScriptStructs)
+	{
+		Script->PrepareCppStructOps();
+		check(Script->GetCppStructOps());
+	}
+	NumScriptStructsRemapped = ScriptStructs.Num();
+
+	// Collect all the classes being reinstanced
 	TSet<UObject*> ReinstancingObjects;
 	ReinstancingObjects.Reserve(ClassesToReinstance.Num() + ReinstancedStructs.Num() + ReinstancedEnums.Num());
 	for (const TPair<UClass*, UClass*>& Pair : ClassesToReinstance)
@@ -868,23 +925,31 @@ void FReload::Reinstance()
 		BlueprintUpdateInfo.Nodes.Add(Node);
 	};
 
+	// Update all the structures.  We add the unchanging structs to the list to make sure the defaults are updated
 	TMap<UScriptStruct*, UScriptStruct*> ChangedStructs;
 	for (const TPair<UScriptStruct*, UScriptStruct*>& Pair : ReinstancedStructs)
 	{
 		ReinstancingObjects.Add(Pair.Key);
 		if (Pair.Value)
 		{
+			Pair.Key->StructFlags = EStructFlags(Pair.Key->StructFlags | STRUCT_NewerVersionExists);
 			ChangedStructs.Emplace(Pair.Key, Pair.Value);
 		}
+		else
+		{
+			ChangedStructs.Emplace(Pair.Key, Pair.Key);
+		}
 	}
 	FBlueprintEditorUtils::UpdateScriptStructsInNodes(ChangedStructs, OnNodeFoundOrUpdated);
 
+	// Update all the enumeration nodes
 	TMap<UEnum*, UEnum*> ChangedEnums;
 	for (const TPair<UEnum*, UEnum*>& Pair : ReinstancedEnums)
 	{
 		ReinstancingObjects.Add(Pair.Key);
 		if (Pair.Value)
 		{
+			Pair.Key->SetEnumFlags(EEnumFlags::NewerVersionExists);
 			ChangedEnums.Emplace(Pair.Key, Pair.Value);
 		}
 	}
@@ -943,53 +1008,7 @@ void FReload::ReinstanceClass(UClass* NewClass, UClass* OldClass, const TSet<UOb
 
 void FReload::Finalize(bool bRunGC)
 {
-	// If we have to collect the packages
-	if (bCollectPackages)
-	{
-		for (const TPair<UClass*, UClass*>& Pair : ReinstancedClasses)
-		{
-			UClass* OldClass = Pair.Key;
-			UClass* NewClass = Pair.Value;
-			Packages.AddUnique(NewClass ? NewClass->GetPackage() : OldClass->GetPackage());
-		}
-	}
-
-	// Remap all native functions (and gather scriptstructs)
-	TArray<UScriptStruct*> ScriptStructs;
-	for (FRawObjectIterator It; It; ++It)
-	{
-		UObject* a = static_cast<UObject*>(It->Object);
-		UFunction* f = Cast<UFunction>(a);
-
-		if (UFunction* Function = Cast<UFunction>(static_cast<UObject*>(It->Object)))
-		{
-			if (FNativeFuncPtr NewFunction = FunctionRemap.FindRef(Function->GetNativeFunc()))
-			{
-				++NumFunctionsRemapped;
-				Function->SetNativeFunc(NewFunction);
-			}
-		}
-
-		if (UScriptStruct* ScriptStruct = Cast<UScriptStruct>(static_cast<UObject*>(It->Object)))
-		{
-			if (!ScriptStruct->HasAnyFlags(RF_ClassDefaultObject) && ScriptStruct->GetCppStructOps() && Packages.ContainsByPredicate([=](UPackage* Package) { return ScriptStruct->IsIn(Package); }))
-			{
-				ScriptStructs.Add(ScriptStruct);
-			}
-		}
-	}
-	NumScriptStructsRemapped = ScriptStructs.Num();
 
-	// now let's set up the script structs...this relies on super behavior, so null them all, then set them all up. Internally this sets them up hierarchically.
-	for (UScriptStruct* Script : ScriptStructs)
-	{
-		Script->ClearCppStructOps();
-	}
-	for (UScriptStruct* Script : ScriptStructs)
-	{
-		Script->PrepareCppStructOps();
-		check(Script->GetCppStructOps());
-	}
 	// Make sure new classes have the token stream assembled
 	UClass::AssembleReferenceTokenStreams();
 
diff --git a/Engine/Source/Runtime/CoreUObject/Public/UObject/Class.h b/Engine/Source/Runtime/CoreUObject/Public/UObject/Class.h
--- a/Engine/Source/Runtime/CoreUObject/Public/UObject/Class.h
+++ b/Engine/Source/Runtime/CoreUObject/Public/UObject/Class.h
@@ -696,6 +696,9 @@ enum EStructFlags
 	/** If set, this struct has been cleaned and sanitized (trashed) and should not be used */
 	STRUCT_Trashed = 0x00800000,
 
+	/** If set, this structure has been replaced via reinstancing */
+	STRUCT_NewerVersionExists = 0x01000000,
+
 	/** Struct flags that are automatically inherited */
 	STRUCT_Inherit				= STRUCT_HasInstancedReference|STRUCT_Atomic,
 
@@ -2167,6 +2170,11 @@ public:
 		return CppForm;
 	}
 
+	void SetEnumFlags(EEnumFlags FlagsToSet)
+	{
+		EnumFlags |= FlagsToSet;
+	}
+
 	bool HasAnyEnumFlags(EEnumFlags InFlags) const
 	{
 		return EnumHasAnyFlags(EnumFlags, InFlags);
diff --git a/Engine/Source/Runtime/CoreUObject/Public/UObject/ObjectMacros.h b/Engine/Source/Runtime/CoreUObject/Public/UObject/ObjectMacros.h
--- a/Engine/Source/Runtime/CoreUObject/Public/UObject/ObjectMacros.h
+++ b/Engine/Source/Runtime/CoreUObject/Public/UObject/ObjectMacros.h
@@ -560,7 +560,8 @@ enum class EEnumFlags
 {
 	None,
 
-	Flags = 0x00000001 // Whether the UEnum represents a set of flags
+	Flags = 0x00000001, // Whether the UEnum represents a set of flags
+	NewerVersionExists = 0x00000002, // If set, this UEnum has been replaced by a newer version
 };
 
 ENUM_CLASS_FLAGS(EEnumFlags)
