In Distributed HLOD builds, allow -SetupHLODs & -BuildHLODs steps to run without source control
* Will prevent AWS buildmachines from holding an exclusive lock on actor files in case of a crash & if the Horde process doesn't cleanup checked out file properly (often the case)
* All p4 operations were temporary during those two steps those anyway
* Much faster setup step

The last step, -SubmitHLODs will obviously require source control to be enabled

diff --git a/Engine/Build/Graph/Examples/BuildWorldPartitionHLODs.xml b/Engine/Build/Graph/Examples/BuildWorldPartitionHLODs.xml
--- a/Engine/Build/Graph/Examples/BuildWorldPartitionHLODs.xml
+++ b/Engine/Build/Graph/Examples/BuildWorldPartitionHLODs.xml
@@ -12,8 +12,11 @@
 	<EnvVar Name="uebp_CLIENT"/>
 		
 	<Property Name="HLODCommonBuilderArgs" Value="$(MapName) -Builder=WorldPartitionHLODsBuilder -DistributedBuild"/>
-	<Property Name="HLODCommonBuilderArgs" Value="-unattended -buildmachine -SCCProvider=Perforce -P4Port=$(uebp_PORT) -P4User=$(uebp_USER) -P4Client=$(uebp_CLIENT) $(HLODCommonBuilderArgs)" If="$(IsBuildMachine)"/>
+	<Property Name="HLODCommonBuilderArgs" Value="-unattended -buildmachine $(HLODCommonBuilderArgs)" If="$(IsBuildMachine)"/>
 	
+	<Property Name="HLODBuilderSCC_None" Value="-SCCProvider=None"/>
+	<Property Name="HLODBuilderSCC_P4" Value="-SCCProvider=Perforce -P4Port=$(uebp_PORT) -P4User=$(uebp_USER) -P4Client=$(uebp_CLIENT)"/>
+		
 	<Agent Name="HLOD Generation Prerequisites" Type="Win64">
 		<Node Name="Update Version Files">
 			<SetVersion Change="$(Change)" Branch="$(EscapedBranch)" If="$(IsBuildMachine)"/>
@@ -32,7 +35,7 @@
 	
 	<Agent Name="HLOD Setup" Type="Win64">
 		<Node Name="HLOD Setup" Requires="$(HLODSetupDependencies)">
-			<Commandlet Name="WorldPartitionBuilderCommandlet" Project="$(ProjectName)" Arguments="$(HLODCommonBuilderArgs) -SetupHLODs -BuilderCount=$(BuilderCount)"/>
+			<Commandlet Name="WorldPartitionBuilderCommandlet" Project="$(ProjectName)" Arguments="$(HLODCommonBuilderArgs) $(HLODBuilderSCC_None) -SetupHLODs -BuilderCount=$(BuilderCount)"/>
 			<GatherBuildProductsFromFile BuildProductsFile="$(RootDir)/HLODTemp/BuildProducts.txt"/>
 		</Node>
 	</Agent>
@@ -50,7 +53,7 @@
 			
 			<Agent Name="$(HLODBuildAgent)" Type="Win64">
 				<Node Name="$(HLODBuildNode)" Requires="HLOD Setup">
-					<Commandlet Name="WorldPartitionBuilderCommandlet" Project="$(ProjectName)" Arguments="$(HLODCommonBuilderArgs) -AllowCommandletRendering -BuildHLODs -BuilderIdx=$(HLODBuildJobIndex)"/>
+					<Commandlet Name="WorldPartitionBuilderCommandlet" Project="$(ProjectName)" Arguments="$(HLODCommonBuilderArgs) $(HLODBuilderSCC_None) -AllowCommandletRendering -BuildHLODs -BuilderIdx=$(HLODBuildJobIndex)"/>
 				</Node>
 				<Node Name="$(HLODBuildCompleteNode)" After="$(HLODBuildNode)">
 					<GatherBuildProductsFromFile BuildProductsFile="$(RootDir)/HLODTemp/BuildProducts.txt"/>
@@ -61,7 +64,7 @@
 	
 	<Agent Name="HLOD Submit" Type="Win64_NonAWS">
 		<Node Name="HLOD Submit" Requires="$(HLODBuildNodes);$(HLODBuildCompleteNodes);HLOD Setup">
-			<Commandlet Name="WorldPartitionBuilderCommandlet" Project="$(ProjectName)" Arguments="$(HLODCommonBuilderArgs) -SubmitHLODs"/>
+			<Commandlet Name="WorldPartitionBuilderCommandlet" Project="$(ProjectName)" Arguments="$(HLODCommonBuilderArgs) $(HLODBuilderSCC_P4) -SubmitHLODs"/>
 		</Node>
 	</Agent>
 
diff --git a/Engine/Source/Editor/UnrealEd/Private/WorldPartition/WorldPartitionHLODsBuilder.cpp b/Engine/Source/Editor/UnrealEd/Private/WorldPartition/WorldPartitionHLODsBuilder.cpp
--- a/Engine/Source/Editor/UnrealEd/Private/WorldPartition/WorldPartitionHLODsBuilder.cpp
+++ b/Engine/Source/Editor/UnrealEd/Private/WorldPartition/WorldPartitionHLODsBuilder.cpp
@@ -7,6 +7,7 @@
 #include "Logging/LogMacros.h"
 #include "Misc/ConfigCacheIni.h"
 #include "Misc/FileHelper.h"
+#include "Algo/ForEach.h"
 
 #include "EngineUtils.h"
 #include "SourceControlHelpers.h"
@@ -47,67 +48,92 @@ public:
 
 	virtual bool Checkout(UPackage* Package) const override
 	{
-		ModifiedFiles.Add(GetFilename(Package));
-		return PackageHelper.Checkout(Package);
+		bool bCheckedOut = PackageHelper.Checkout(Package);
+		if (bCheckedOut)
+		{
+			ModifiedFiles.Add(FHLODModifiedFiles::EFileOperation::FileEdited, GetFilename(Package));
+		}
+		return bCheckedOut;
 	}
 
 	virtual bool Add(UPackage* Package) const override
 	{
-		ModifiedFiles.Add(GetFilename(Package));
-		return PackageHelper.AddToSourceControl(Package);
+		bool bAdded = PackageHelper.AddToSourceControl(Package);
+		if (bAdded)
+		{
+			ModifiedFiles.Add(FHLODModifiedFiles::EFileOperation::FileAdded, GetFilename(Package));
+		}
+		return bAdded;
 	}
 
 	virtual bool Delete(const FString& PackageName) const override
 	{
-		ModifiedFiles.Add(PackageName);
-		return PackageHelper.Delete(PackageName);
+		bool bDeleted = PackageHelper.Delete(PackageName);
+		if (bDeleted)
+		{
+			ModifiedFiles.Add(FHLODModifiedFiles::EFileOperation::FileDeleted, PackageName);
+		}
+		return bDeleted;
 	}
 
 	virtual bool Delete(UPackage* Package) const override
 	{
-		ModifiedFiles.Add(GetFilename(Package));
-		return PackageHelper.Delete(Package);
+		FString PackageName = GetFilename(Package);
+		bool bDeleted = PackageHelper.Delete(Package);
+		if (bDeleted)
+		{
+			ModifiedFiles.Add(FHLODModifiedFiles::EFileOperation::FileDeleted, PackageName);
+		}
+		return bDeleted;
 	}
 
 	virtual bool Save(UPackage* Package) const override
 	{
-		ModifiedFiles.Add(GetFilename(Package));
+		bool bFileExists = IPlatformFile::GetPlatformPhysical().FileExists(*GetFilename(Package));
 
 		// Checkout package
 		Package->MarkAsFullyLoaded();
 
-		if (!Checkout(Package))
+		if (bFileExists)
 		{
-			UE_LOG(LogWorldPartitionHLODsBuilder, Error, TEXT("Error checking out package %s."), *Package->GetName());
-			return false;
+			if (!Checkout(Package))
+			{
+				UE_LOG(LogWorldPartitionHLODsBuilder, Error, TEXT("Error checking out package %s."), *Package->GetName());
+				return false;
+			}
 		}
 
 		// Save package
+		ESaveFlags SaveFlags = PackageHelper.UseSourceControl() ? ESaveFlags::SAVE_None : ESaveFlags::SAVE_Async;
 		FString PackageFileName = GetFilename(Package);
-		if (!UPackage::SavePackage(Package, nullptr, RF_Standalone, *PackageFileName))
+		if (!UPackage::SavePackage(Package, nullptr, RF_Standalone, *PackageFileName, GError, nullptr, false, true, SaveFlags))
 		{
 			UE_LOG(LogWorldPartitionHLODsBuilder, Error, TEXT("Error saving package %s."), *Package->GetName());
 			return false;
 		}
 
 		// Add new package to source control
-		if (!Add(Package))
+		if (!bFileExists)
 		{
-			UE_LOG(LogWorldPartitionHLODsBuilder, Error, TEXT("Error adding package %s to source control."), *Package->GetName());
-			return false;
+			if (!Add(Package))
+			{
+				UE_LOG(LogWorldPartitionHLODsBuilder, Error, TEXT("Error adding package %s to source control."), *Package->GetName());
+				return false;
+			}
 		}
 
 		return true;
 	}
 
-	const TSet<FString> GetModifiedFiles() const
+	const FHLODModifiedFiles& GetModifiedFiles() const
 	{
+		UPackage::WaitForAsyncFileWrites();
 		return ModifiedFiles;
 	}
 
 private:
 	FPackageSourceControlHelper& PackageHelper;
-	mutable TSet<FString> ModifiedFiles;
+	mutable FHLODModifiedFiles ModifiedFiles;
 };
 
 static const FString DistributedBuildWorkingDirName = TEXT("HLODTemp");
@@ -215,9 +241,9 @@ bool UWorldPartitionHLODsBuilder::ValidateParams() const
 		return false;
 	}
 
-	if (IsDistributedBuild() && !ISourceControlModule::Get().GetProvider().IsEnabled())
+	if (IsDistributedBuild() && bSubmitHLODs && !ISourceControlModule::Get().GetProvider().IsEnabled())
 	{
-		UE_LOG(LogWorldPartitionHLODsBuilder, Error, TEXT("Distributed builds requires that a valid source control provider is enabled, exiting..."), *BuildManifest);
+		UE_LOG(LogWorldPartitionHLODsBuilder, Error, TEXT("Distributed builds submit step requires that a valid source control provider is enabled, exiting..."), *BuildManifest);
 		return false;
 	}
 
@@ -333,20 +359,24 @@ bool UWorldPartitionHLODsBuilder::SetupHLODActors(bool bCreateOnly)
 
 				ModifiedFiles.Append(SourceControlHelper->GetModifiedFiles());
 
-				TArray<TArray<FString>> BuildersFiles;
+				TArray<FHLODModifiedFiles> BuildersFiles;
 				BuildersFiles.SetNum(BuilderCount);
 
-				for (const FString& ModifiedFile : ModifiedFiles)
+				for (int32 i = 0; i < FHLODModifiedFiles::EFileOperation::NumFileOperations; i++)
 				{
-					int32* Idx = FilesToBuilderMap.Find(ModifiedFile);
-					if (Idx)
+					FHLODModifiedFiles::EFileOperation FileOp = (FHLODModifiedFiles::EFileOperation)i;
+					for (const FString& ModifiedFile : ModifiedFiles.Get(FileOp))
 					{
-						BuildersFiles[*Idx].Add(ModifiedFile);
-					}
-					else
-					{
-						// Add general files to the last builder
-						BuildersFiles.Last().Add(ModifiedFile);
+						int32* Idx = FilesToBuilderMap.Find(ModifiedFile);
+						if (Idx)
+						{
+							BuildersFiles[*Idx].Add(FileOp, ModifiedFile);
+						}
+						else
+						{
+							// Add general files to the last builder
+							BuildersFiles.Last().Add(FileOp, ModifiedFile);
+						}
 					}
 				}
 
@@ -442,7 +472,7 @@ bool UWorldPartitionHLODsBuilder::BuildHLODActors()
 
 		TArray<FString> BuildProducts;
 
-		if (!CopyFilesToWorkingDir("ToSubmit", ModifiedFiles.Array(), BuildProducts))
+		if (!CopyFilesToWorkingDir("ToSubmit", ModifiedFiles, BuildProducts))
 		{
 			return false;
 		}
@@ -506,23 +536,23 @@ bool UWorldPartitionHLODsBuilder::SubmitHLODActors()
 	bool bRet = true;
 
 	// Ensure all files modified by the source control helper are taken into account
-	ModifiedFiles.Append(SourceControlHelper->GetModifiedFiles());
+	TArray<FString> FilesToSubmit = SourceControlHelper->GetModifiedFiles().GetAllFiles();
 
 	// Check in all modified files
-	if (ModifiedFiles.Num() > 0)
+	if (FilesToSubmit.Num() > 0)
 	{
 		FText ChangelistDescription = FText::FromString(FString::Printf(TEXT("Rebuilt HLODs for \"%s\" at %s"), *WorldPartition->GetWorld()->GetName(), *FEngineVersion::Current().ToString()));
 
 		TSharedRef<FCheckIn, ESPMode::ThreadSafe> CheckInOperation = ISourceControlOperation::Create<FCheckIn>();
 		CheckInOperation->SetDescription(ChangelistDescription);
-		bRet = ISourceControlModule::Get().GetProvider().Execute(CheckInOperation, ModifiedFiles.Array()) == ECommandResult::Succeeded;
+		bRet = ISourceControlModule::Get().GetProvider().Execute(CheckInOperation, FilesToSubmit) == ECommandResult::Succeeded;
 		if (!bRet)
 		{
-			UE_LOG(LogWorldPartitionHLODsBuilder, Error, TEXT("Failed to submit %d files to source control."), ModifiedFiles.Num());
+			UE_LOG(LogWorldPartitionHLODsBuilder, Error, TEXT("Failed to submit %d files to source control."), FilesToSubmit.Num());
 		}
 		else
 		{
-			UE_LOG(LogWorldPartitionHLODsBuilder, Display, TEXT("#### Submitted %d files to source control ####"), ModifiedFiles.Num());
+			UE_LOG(LogWorldPartitionHLODsBuilder, Display, TEXT("#### Submitted %d files to source control ####"), FilesToSubmit.Num());
 		}
 	}
 	else
@@ -783,74 +813,62 @@ const FName FileAction_Add(TEXT("Add"));
 const FName FileAction_Edit(TEXT("Edit"));
 const FName FileAction_Delete(TEXT("Delete"));
 
-bool UWorldPartitionHLODsBuilder::CopyFilesToWorkingDir(const FString& TargetDir, const TArray<FString>& FilesToCopy, TArray<FString>& BuildProducts)
+bool UWorldPartitionHLODsBuilder::CopyFilesToWorkingDir(const FString& TargetDir, const FHLODModifiedFiles& Files, TArray<FString>& BuildProducts)
 {
-	FString AbsoluteTargetDir = DistributedBuildWorkingDir / TargetDir / TEXT("");
+	const FString AbsoluteTargetDir = DistributedBuildWorkingDir / TargetDir / TEXT("");
 
-	TArray<FString> FilesToDelete;
+	bool bSuccess = true;
 
-	for (const FString& SourceFilename : FilesToCopy)
+	auto CopyFileToWorkingDir = [&](const FString& SourceFilename, const FName FileAction)
 	{
-		FSourceControlState FileState = USourceControlHelpers::QueryFileState(SourceFilename);
-		if (FileState.bIsValid)
-		{
-			bool bShouldCopyToWorkingDir = false;
-			FName FileAction;
+		FString SourceFilenameRelativeToRoot = SourceFilename;
+		FPaths::MakePathRelativeTo(SourceFilenameRelativeToRoot, *FPaths::RootDir());
 
-			if (FileState.bIsAdded)
-			{
-				FileAction = FileAction_Add;
-				FilesToDelete.Add(SourceFilename);
-			}
-			else if (FileState.bIsCheckedOut)
-			{
-				FileAction = FileAction_Edit;
-			}
-			else if (FileState.bIsDeleted)
+		FString TargetFilename = AbsoluteTargetDir / FileAction.ToString() / SourceFilenameRelativeToRoot;
+
+		BuildProducts.Add(TargetFilename);
+
+		if (FileAction != FileAction_Delete)
+		{
+			bool bRet = IFileManager::Get().Copy(*TargetFilename, *SourceFilename, false) == COPY_OK;
+			if (!bRet)
 			{
-				FileAction = FileAction_Delete;
+				UE_LOG(LogWorldPartitionHLODsBuilder, Error, TEXT("Failed to copy file from \"%s\" to \"%s\""), *SourceFilename, *TargetFilename);
+				bSuccess = false;
 			}
-
-			if (!FileAction.IsNone())
+		}
+		else
+		{
+			bool bRet = FFileHelper::SaveStringToFile(TEXT(""), *TargetFilename);
+			if (!bRet)
 			{
-				FString SourceFilenameRelativeToRoot = SourceFilename;
-				FPaths::MakePathRelativeTo(SourceFilenameRelativeToRoot, *FPaths::RootDir());
-
-				FString TargetFilename = AbsoluteTargetDir / FileAction.ToString() / SourceFilenameRelativeToRoot;
-
-				BuildProducts.Add(TargetFilename);
-
-				if (FileAction != FileAction_Delete)
-				{
-					bool bRet = IFileManager::Get().Copy(*TargetFilename, *SourceFilename, false) == COPY_OK;
-					if (!bRet)
-					{
-						UE_LOG(LogWorldPartitionHLODsBuilder, Error, TEXT("Failed to copy file from \"%s\" to \"%s\""), *SourceFilename, *TargetFilename);
-						return false;
-					}
-				}
-				else
-				{
-					bool bRet = FFileHelper::SaveStringToFile(TEXT(""), *TargetFilename);
-					if (!bRet)
-					{
-						UE_LOG(LogWorldPartitionHLODsBuilder, Error, TEXT("Failed to create empty file at \"%s\""), *TargetFilename);
-						return false;
-					}
-				}
+				UE_LOG(LogWorldPartitionHLODsBuilder, Error, TEXT("Failed to create empty file at \"%s\""), *TargetFilename);
+				bSuccess = false;
 			}
 		}
-	}
+	};
 
-	bool bRet = USourceControlHelpers::RevertFiles(FilesToCopy);
-	if (!bRet)
+	Algo::ForEach(Files.Get(FHLODModifiedFiles::EFileOperation::FileAdded), [&](const FString& SourceFilename) { CopyFileToWorkingDir(SourceFilename, FileAction_Add); });
+	Algo::ForEach(Files.Get(FHLODModifiedFiles::EFileOperation::FileEdited), [&](const FString& SourceFilename) { CopyFileToWorkingDir(SourceFilename, FileAction_Edit); });
+	Algo::ForEach(Files.Get(FHLODModifiedFiles::EFileOperation::FileDeleted), [&](const FString& SourceFilename) { CopyFileToWorkingDir(SourceFilename, FileAction_Delete); });
+	if (!bSuccess)
 	{
-		UE_LOG(LogWorldPartitionHLODsBuilder, Error, TEXT("Failed to revert modified files: %s"), *USourceControlHelpers::LastErrorMsg().ToString());
 		return false;
 	}
 
+	// Revert any file changes
+	if (ISourceControlModule::Get().IsEnabled())
+	{
+		bool bRet = USourceControlHelpers::RevertFiles(Files.GetAllFiles());
+		if (!bRet)
+		{
+			UE_LOG(LogWorldPartitionHLODsBuilder, Error, TEXT("Failed to revert modified files: %s"), *USourceControlHelpers::LastErrorMsg().ToString());
+			return false;
+		}
+	}
+
 	// Delete files we added
-	for (const FString& FileToDelete : FilesToDelete)
+	for (const FString& FileToDelete : Files.Get(FHLODModifiedFiles::EFileOperation::FileAdded))
 	{
 		if (!IFileManager::Get().Delete(*FileToDelete, false, true))
 		{
@@ -864,13 +882,27 @@ bool UWorldPartitionHLODsBuilder::CopyFilesToWorkingDir(const FString& TargetDir
 
 bool UWorldPartitionHLODsBuilder::CopyFilesFromWorkingDir(const FString& SourceDir)
 {
-	FString AbsoluteSourceDir = DistributedBuildWorkingDir / SourceDir / TEXT("");
+	const FString AbsoluteSourceDir = DistributedBuildWorkingDir / SourceDir / TEXT("");
+
+	auto CopyFromWorkingDir = [](const TMap<FString, FString>& FilesToCopy) -> bool
+	{
+		for (const auto& Pair : FilesToCopy)
+		{
+			bool bRet = IFileManager::Get().Copy(*Pair.Key, *Pair.Value) == COPY_OK;
+			if (!bRet)
+			{
+				UE_LOG(LogWorldPartitionHLODsBuilder, Error, TEXT("Failed to copy file from \"%s\" to \"%s\""), *Pair.Value, *Pair.Key);
+				return false;
+			}
+		}
+		return true;
+	};
 
 	TArray<FString> Files;
 	IFileManager::Get().FindFilesRecursive(Files, *AbsoluteSourceDir, TEXT("*.*"), true, false);
 
-	TArray<FString>			FilesToAdd;
-	TMap<FString, FString>	FilesToCopy;
+	TMap<FString, FString>	FilesToAdd;
+	TMap<FString, FString>	FilesToEdit;
 	TArray<FString>			FilesToDelete;
 
 	bool bRet = true;
@@ -893,17 +925,11 @@ bool UWorldPartitionHLODsBuilder::CopyFilesFromWorkingDir(const FString& SourceD
 		FName FileAction(FileActionString);
 		if (FileAction == FileAction_Add)
 		{
-			bRet = IFileManager::Get().Copy(*FullPathInRootDirectory, *File, false) == COPY_OK;
-			if (!bRet)
-			{
-				UE_LOG(LogWorldPartitionHLODsBuilder, Error, TEXT("Failed to copy file from \"%s\" to \"%s\""), *File, *FullPathInRootDirectory);
-				return false;
-			}
-			FilesToAdd.Add(*FullPathInRootDirectory);
+			FilesToAdd.Add(*FullPathInRootDirectory, File);
 		}
 		else if (FileAction == FileAction_Edit)
 		{
-			FilesToCopy.Add(FullPathInRootDirectory, File);
+			FilesToEdit.Add(FullPathInRootDirectory, File);
 		}
 		else if (FileAction == FileAction_Delete)
 		{
@@ -918,59 +944,84 @@ bool UWorldPartitionHLODsBuilder::CopyFilesFromWorkingDir(const FString& SourceD
 	// Add
 	if (!FilesToAdd.IsEmpty())
 	{
-		bRet = USourceControlHelpers::MarkFilesForAdd(FilesToAdd);
+		bRet = CopyFromWorkingDir(FilesToAdd);
 		if (!bRet)
 		{
-			UE_LOG(LogWorldPartitionHLODsBuilder, Error, TEXT("Adding files to source control failed: %s"), *USourceControlHelpers::LastErrorMsg().ToString());
-			while (!FPlatformMisc::IsDebuggerPresent())
+			return false;
+		}
+
+		TArray<FString> ToAdd;
+		FilesToAdd.GetKeys(ToAdd);
+
+		if (ISourceControlModule::Get().IsEnabled())
+		{
+			bRet = USourceControlHelpers::MarkFilesForAdd(ToAdd);
+			if (!bRet)
 			{
-				FPlatformProcess::Sleep(0.1f);
+				UE_LOG(LogWorldPartitionHLODsBuilder, Error, TEXT("Adding files to source control failed: %s"), *USourceControlHelpers::LastErrorMsg().ToString());
+				return false;
 			}
-			return false;
 		}
-		ModifiedFiles.Append(FilesToAdd);
+
+		ModifiedFiles.Append(FHLODModifiedFiles::EFileOperation::FileAdded, ToAdd);
 	}
 
 	// Delete
 	if (!FilesToDelete.IsEmpty())
 	{
-		bRet = USourceControlHelpers::MarkFilesForDelete(FilesToDelete);
-		if (!bRet)
+		if (ISourceControlModule::Get().IsEnabled())
 		{
-			UE_LOG(LogWorldPartitionHLODsBuilder, Error, TEXT("Deleting files from source control failed: %s"), *USourceControlHelpers::LastErrorMsg().ToString());
-			return false;
+			bRet = USourceControlHelpers::MarkFilesForDelete(FilesToDelete);
+			if (!bRet)
+			{
+				UE_LOG(LogWorldPartitionHLODsBuilder, Error, TEXT("Deleting files from source control failed: %s"), *USourceControlHelpers::LastErrorMsg().ToString());
+				return false;
+			}
+		}
+		else
+		{
+			for (const FString& FileToDelete : FilesToDelete)
+			{
+				bRet = IFileManager::Get().Delete(*FileToDelete);
+				if (!bRet)
+				{
+					UE_LOG(LogWorldPartitionHLODsBuilder, Error, TEXT("Failed to delete file from disk: %s"), *USourceControlHelpers::LastErrorMsg().ToString());
+					return false;
+				}
+			}
 		}
-		ModifiedFiles.Append(FilesToDelete);
+
+		ModifiedFiles.Append(FHLODModifiedFiles::EFileOperation::FileDeleted, FilesToDelete);
 	}
 
 	// Edit
-	if (!FilesToCopy.IsEmpty())
+	if (!FilesToEdit.IsEmpty())
 	{
-		TArray<FString> FilesToEdit;
-		FilesToCopy.GetKeys(FilesToEdit);
+		TArray<FString> ToEdit;
+		FilesToEdit.GetKeys(ToEdit);
 
-		bRet = USourceControlHelpers::CheckOutFiles(FilesToEdit);
-		if (!bRet)
-		{
-			UE_LOG(LogWorldPartitionHLODsBuilder, Error, TEXT("Checking out files from source control failed: %s"), *USourceControlHelpers::LastErrorMsg().ToString());
-			return false;
-		}
-		ModifiedFiles.Append(FilesToEdit);
-
-		for (const auto& Pair : FilesToCopy)
+		if (ISourceControlModule::Get().IsEnabled())
 		{
-			bRet = IFileManager::Get().Copy(*Pair.Key, *Pair.Value) == COPY_OK;
+			bRet = USourceControlHelpers::CheckOutFiles(ToEdit);
 			if (!bRet)
 			{
-				UE_LOG(LogWorldPartitionHLODsBuilder, Error, TEXT("Failed to copy file from \"%s\" to \"%s\""), *Pair.Value, *Pair.Key);
+				UE_LOG(LogWorldPartitionHLODsBuilder, Error, TEXT("Checking out files from source control failed: %s"), *USourceControlHelpers::LastErrorMsg().ToString());
 				return false;
 			}
 		}
+
+		ModifiedFiles.Append(FHLODModifiedFiles::EFileOperation::FileEdited, ToEdit);
+
+		bRet = CopyFromWorkingDir(FilesToEdit);
+		if (!bRet)
+		{
+			return false;
+		}
 	}
 
 	// Force a rescan of the updated files
 	IAssetRegistry& AssetRegistry = FModuleManager::LoadModuleChecked<FAssetRegistryModule>(TEXT("AssetRegistry")).Get();
-	AssetRegistry.ScanModifiedAssetFiles(ModifiedFiles.Array());
+	AssetRegistry.ScanModifiedAssetFiles(ModifiedFiles.GetAllFiles());
 
 	return true;
 }
diff --git a/Engine/Source/Editor/UnrealEd/Public/WorldPartition/WorldPartitionHLODsBuilder.h b/Engine/Source/Editor/UnrealEd/Public/WorldPartition/WorldPartitionHLODsBuilder.h
--- a/Engine/Source/Editor/UnrealEd/Public/WorldPartition/WorldPartitionHLODsBuilder.h
+++ b/Engine/Source/Editor/UnrealEd/Public/WorldPartition/WorldPartitionHLODsBuilder.h
@@ -5,11 +5,62 @@
 #include "WorldPartition/WorldPartitionBuilder.h"
 #include "WorldPartitionHLODsBuilder.generated.h"
 
+struct FHLODModifiedFiles
+{
+	enum EFileOperation
+	{
+		FileAdded,
+		FileEdited,
+		FileDeleted,
+		NumFileOperations
+	};
+
+	void Add(EFileOperation FileOp, const FString& File)
+	{
+		Files[FileOp].Add(File);
+	}
+
+	const TSet<FString>& Get(EFileOperation FileOp) const
+	{
+		return Files[FileOp];
+	}
+
+	void Append(EFileOperation FileOp, const TArray<FString>& InFiles)
+	{
+		Files[FileOp].Append(InFiles);
+	}
+
+	void Append(const FHLODModifiedFiles& Other)
+	{
+		Files[EFileOperation::FileAdded].Append(Other.Files[EFileOperation::FileAdded]);
+		Files[EFileOperation::FileEdited].Append(Other.Files[EFileOperation::FileEdited]);
+		Files[EFileOperation::FileDeleted].Append(Other.Files[EFileOperation::FileDeleted]);
+	}
+
+	void Empty()
+	{
+		Files[EFileOperation::FileAdded].Empty();
+		Files[EFileOperation::FileEdited].Empty();
+		Files[EFileOperation::FileDeleted].Empty();
+	}
+
+	TArray<FString> GetAllFiles() const
+	{
+		TArray<FString> AllFiles;
+		AllFiles.Append(Files[EFileOperation::FileAdded].Array());
+		AllFiles.Append(Files[EFileOperation::FileEdited].Array());
+		AllFiles.Append(Files[EFileOperation::FileDeleted].Array());
+		return AllFiles;
+	}
+
+private:
+	TSet<FString> Files[NumFileOperations];
+};
+
 UCLASS()
 class UWorldPartitionHLODsBuilder : public UWorldPartitionBuilder
 {
 	GENERATED_UCLASS_BODY()
-
 public:
 	// UWorldPartitionBuilder interface begin
 	virtual bool RequiresCommandletRendering() const override;
@@ -34,7 +85,7 @@ protected:
 	TArray<TArray<FGuid>> GetHLODWorldloads(int32 NumWorkloads) const;
 	bool ValidateWorkload(const TArray<FGuid>& Workload) const;
 
-	bool CopyFilesToWorkingDir(const FString& TargetDir, const TArray<FString>& FilesToCopy, TArray<FString>& BuildProducts);
+	bool CopyFilesToWorkingDir(const FString& TargetDir, const FHLODModifiedFiles& ModifiedFiles, TArray<FString>& BuildProducts);
 	bool CopyFilesFromWorkingDir(const FString& SourceDir);
 
 private:
@@ -53,8 +104,8 @@ private:
 	int32 BuilderIdx;
 	int32 BuilderCount;
 
-	TSet<FString> ModifiedFiles;
-
 	const FString DistributedBuildWorkingDir;
 	const FString DistributedBuildManifest;
+	
+	FHLODModifiedFiles ModifiedFiles;
 };
\ No newline at end of file
diff --git a/Engine/Source/Runtime/Engine/Private/WorldPartition/RuntimeSpatialHash/RuntimeSpatialHashHLOD.cpp b/Engine/Source/Runtime/Engine/Private/WorldPartition/RuntimeSpatialHash/RuntimeSpatialHashHLOD.cpp
--- a/Engine/Source/Runtime/Engine/Private/WorldPartition/RuntimeSpatialHash/RuntimeSpatialHashHLOD.cpp
+++ b/Engine/Source/Runtime/Engine/Private/WorldPartition/RuntimeSpatialHash/RuntimeSpatialHashHLOD.cpp
@@ -237,6 +237,9 @@ static TArray<FGuid> GenerateHLODsForGrid(UWorldPartition* WorldPartition, const
 		}
 	});
 
+	// Ensure all async packages writes are completed before we start processing another HLOD level or grid
+	UPackage::WaitForAsyncFileWrites();
+
 	// Need to collect garbage here since some HLOD actors have been marked pending kill when destroying them
 	// and they may be loaded when generating the next HLOD layer.
 	DoCollectGarbage();
@@ -484,21 +487,34 @@ bool UWorldPartitionRuntimeSpatialHash::GenerateHLOD(ISourceControlHelper* Sourc
 
 	// HLOD creation context
 	FHLODCreationContext Context;
-	TArray<FWorldPartitionHandle> InvalidHLODActors;
+
+	TSet<uint64> InvalidHLODCellHashes;
+	InvalidHLODCellHashes.Add(0);
+
+	TSet<FWorldPartitionHandle> InvalidHLODActors;
 
 	for (UActorDescContainer::TIterator<AWorldPartitionHLOD> HLODIterator(WorldPartition); HLODIterator; ++HLODIterator)
 	{
+		uint64 CellHash = HLODIterator->GetCellHash();
+
 		FWorldPartitionHandle HLODActorHandle(WorldPartition, HLODIterator->GetGuid());
+		FWorldPartitionHandle* DupeHLODActorHandle = Context.HLODActorDescs.Find(CellHash);
 
-		uint64 CellHash = HLODIterator->GetCellHash();
-		if (CellHash != 0)
+		if (DupeHLODActorHandle == nullptr && !InvalidHLODCellHashes.Contains(CellHash))
 		{
-			check(!Context.HLODActorDescs.Contains(CellHash));
 			Context.HLODActorDescs.Emplace(CellHash, MoveTemp(HLODActorHandle));
 		}
 		else
 		{
+			InvalidHLODCellHashes.Add(CellHash);
+
 			InvalidHLODActors.Emplace(MoveTemp(HLODActorHandle));
+			if (DupeHLODActorHandle)
+			{
+				InvalidHLODActors.Emplace(*DupeHLODActorHandle);
+			}
+
+			Context.HLODActorDescs.Remove(CellHash);
 		}
 	}
 
