Fixed issue where some changes to structs could result in an assert during the live coding re-instancing process.

#rnx

diff --git a/Engine/Source/Editor/UnrealEd/Private/Kismet2/ReloadUtilities.cpp b/Engine/Source/Editor/UnrealEd/Private/Kismet2/ReloadUtilities.cpp
--- a/Engine/Source/Editor/UnrealEd/Private/Kismet2/ReloadUtilities.cpp
+++ b/Engine/Source/Editor/UnrealEd/Private/Kismet2/ReloadUtilities.cpp
@@ -965,6 +965,13 @@ void FReload::Reinstance()
 		{
 			FBlueprintEditorUtils::RecombineNestedSubPins(Node);
 		}
+
+		// We must reconstruct the node first other wise some pins might not be 
+		// in a good state for the recompile
+		for (UK2Node* Node : Info.Nodes)
+		{
+			Node->ReconstructNode();
+		}
 	}
 
 	TSet<UBlueprint*> CompiledBlueprints;
@@ -973,7 +980,7 @@ void FReload::Reinstance()
 		ReinstanceClass(Pair.Value, Pair.Key, ReinstancingObjects, CompiledBlueprints);
 	}
 
-	// Reconstruct the nodes (and recompile blueprints if they haven't already been recompiled)
+	// Recompile blueprints if they haven't already been recompiled)
 	for (TPair<UBlueprint*, FBlueprintUpdateInfo>& KVP : ModifiedBlueprints)
 	{
 		UBlueprint* Blueprint = KVP.Key;
@@ -984,11 +991,6 @@ void FReload::Reinstance()
 			EBlueprintCompileOptions Options = EBlueprintCompileOptions::SkipGarbageCollection;
 			FKismetEditorUtilities::CompileBlueprint(Blueprint, Options);
 		}
-
-		for (UK2Node* Node : Info.Nodes)
-		{
-			Node->ReconstructNode();
-		}
 	}
 
 	ReinstancedClasses = MoveTemp(ClassesToReinstance);
