Fixed issue where "unrecognized" tabs could start spawning in the Blueprint Editor after also using other asset editors.

The core issue was that the BP editor was performing deactivation/cleanup work whenever any system asked if it can be closed, whereas other editors were not and sometimes interrupting/canceling the closure process.  This led to situations where a BP editor tab deactivated itself before a different editor tab halted the window closing process, putting it in a dysfunctional state.

diff --git a/Engine/Source/Editor/Kismet/Private/BlueprintEditor.cpp b/Engine/Source/Editor/Kismet/Private/BlueprintEditor.cpp
--- a/Engine/Source/Editor/Kismet/Private/BlueprintEditor.cpp
+++ b/Engine/Source/Editor/Kismet/Private/BlueprintEditor.cpp
@@ -566,6 +566,11 @@ bool FBlueprintEditor::IsInAScriptingMode() const
 }
 
 bool FBlueprintEditor::OnRequestClose()
+{
+	return FWorkflowCentricApplication::OnRequestClose();
+}
+
+void FBlueprintEditor::OnClose()
 {
 	// Also close the Find Results tab if we're not in full edit mode and the option to host Global Find Results is enabled.
 	TSharedPtr<SDockTab> FindResultsTab = TabManager->FindExistingLiveTab(FBlueprintEditorTabs::FindResultsID);
@@ -582,8 +587,8 @@ bool FBlueprintEditor::OnRequestClose()
 	}
 
 	bEditorMarkedAsClosed = true;
-	
-	return FWorkflowCentricApplication::OnRequestClose();
+
+	FWorkflowCentricApplication::OnClose();
 }
 
 bool FBlueprintEditor::InEditingMode() const
diff --git a/Engine/Source/Editor/Kismet/Public/BlueprintEditor.h b/Engine/Source/Editor/Kismet/Public/BlueprintEditor.h
--- a/Engine/Source/Editor/Kismet/Public/BlueprintEditor.h
+++ b/Engine/Source/Editor/Kismet/Public/BlueprintEditor.h
@@ -200,6 +200,7 @@ public:
 public:
 	//~ Begin FAssetEditorToolkit Interface
 	virtual bool OnRequestClose() override;
+	virtual void OnClose() override;
 	// End of FAssetEditorToolkit 
 
 	//~ Begin IToolkit Interface
diff --git a/Engine/Source/Editor/UnrealEd/Private/Toolkits/AssetEditorToolkit.cpp b/Engine/Source/Editor/UnrealEd/Private/Toolkits/AssetEditorToolkit.cpp
--- a/Engine/Source/Editor/UnrealEd/Private/Toolkits/AssetEditorToolkit.cpp
+++ b/Engine/Source/Editor/UnrealEd/Private/Toolkits/AssetEditorToolkit.cpp
@@ -196,6 +196,7 @@ void FAssetEditorToolkit::InitAssetEditor( const EToolkitMode::Type Mode, const
 		( 
 			SAssignNew( NewStandaloneHost, SStandaloneAssetEditorToolkitHost, NewTabManager, AppIdentifier )
 			.OnRequestClose(this, &FAssetEditorToolkit::OnRequestClose)
+			.OnClose(this, &FAssetEditorToolkit::OnClose)
 		);
 
 		// Assign our toolkit host before we setup initial content.  (Important: We must cache this pointer here as SetupInitialContent
@@ -461,6 +462,8 @@ bool FAssetEditorToolkit::CloseWindow()
 {
 	if (OnRequestClose())
 	{
+		OnClose();
+
 		// Close this toolkit
 		FToolkitManager::Get().CloseToolkit( AsShared() );
 	}
diff --git a/Engine/Source/Editor/UnrealEd/Private/Toolkits/SStandaloneAssetEditorToolkitHost.cpp b/Engine/Source/Editor/UnrealEd/Private/Toolkits/SStandaloneAssetEditorToolkitHost.cpp
--- a/Engine/Source/Editor/UnrealEd/Private/Toolkits/SStandaloneAssetEditorToolkitHost.cpp
+++ b/Engine/Source/Editor/UnrealEd/Private/Toolkits/SStandaloneAssetEditorToolkitHost.cpp
@@ -25,6 +25,7 @@ void SStandaloneAssetEditorToolkitHost::Construct( const SStandaloneAssetEditorT
 	ToolbarSlot = nullptr;
 
 	EditorCloseRequest = InArgs._OnRequestClose;
+	EditorClosing = InArgs._OnClose;
 	AppName = InitAppName;
 
 	MyTabManager = InTabManager;
@@ -382,6 +383,8 @@ void SStandaloneAssetEditorToolkitHost::OnTabClosed(TSharedRef<SDockTab> TabClos
 {
 	check(TabClosed == HostTabPtr.Pin());
 
+	EditorClosing.ExecuteIfBound();
+
 	MyTabManager->SetMenuMultiBox(nullptr, nullptr);
 	
 	if(HostedAssetEditorToolkit.IsValid())
diff --git a/Engine/Source/Editor/UnrealEd/Private/Toolkits/SStandaloneAssetEditorToolkitHost.h b/Engine/Source/Editor/UnrealEd/Private/Toolkits/SStandaloneAssetEditorToolkitHost.h
--- a/Engine/Source/Editor/UnrealEd/Private/Toolkits/SStandaloneAssetEditorToolkitHost.h
+++ b/Engine/Source/Editor/UnrealEd/Private/Toolkits/SStandaloneAssetEditorToolkitHost.h
@@ -22,11 +22,13 @@ class SStandaloneAssetEditorToolkitHost :
 {
 
 public:
-	/** SCompoundWidget interface */
 	SLATE_BEGIN_ARGS( SStandaloneAssetEditorToolkitHost ){}
 		SLATE_EVENT(FRequestAssetEditorClose, OnRequestClose)
+		SLATE_EVENT(FAssetEditorClosing, OnClose)
 	SLATE_END_ARGS()
 
+	/** SCompoundWidget interface */
+
 	/**
 	 * Constructs this widget
 	 *
@@ -123,6 +125,9 @@ private:
 	/** Delegate to be called to determine if we are allowed to close this toolkit host */
 	FRequestAssetEditorClose EditorCloseRequest;
 
+	/** Delegate to be called when this toolkit host is closing */
+	FAssetEditorClosing EditorClosing;
+
 	/** The menu extenders to populate the main toolkit host menu with */
 	TArray< TSharedPtr<FExtender> > MenuExtenders;
 
diff --git a/Engine/Source/Editor/UnrealEd/Private/WorkflowCentricApplication.cpp b/Engine/Source/Editor/UnrealEd/Private/WorkflowCentricApplication.cpp
--- a/Engine/Source/Editor/UnrealEd/Private/WorkflowCentricApplication.cpp
+++ b/Engine/Source/Editor/UnrealEd/Private/WorkflowCentricApplication.cpp
@@ -124,11 +124,11 @@ void FWorkflowCentricApplication::PushTabFactories(FWorkflowAllowedTabSet& Facto
 
 bool FWorkflowCentricApplication::OnRequestClose()
 {
-	if (!FSlateApplication::Get().IsNormalExecution())
-	{
-		return false;
-	}
+	return FSlateApplication::Get().IsNormalExecution();
+}
 
+void FWorkflowCentricApplication::OnClose()
+{
 	if (CurrentAppModePtr.IsValid())
 	{
 		check(TabManager.IsValid());
@@ -142,8 +142,6 @@ bool FWorkflowCentricApplication::OnRequestClose()
 		// Unregister tab spawners
 		TabManager->UnregisterAllTabSpawners();
 	}
-
-	return true;
 }
 
 void FWorkflowCentricApplication::AddApplicationMode(FName ModeName, TSharedRef<FApplicationMode> Mode)
diff --git a/Engine/Source/Editor/UnrealEd/Public/Toolkits/AssetEditorToolkit.h b/Engine/Source/Editor/UnrealEd/Public/Toolkits/AssetEditorToolkit.h
--- a/Engine/Source/Editor/UnrealEd/Public/Toolkits/AssetEditorToolkit.h
+++ b/Engine/Source/Editor/UnrealEd/Public/Toolkits/AssetEditorToolkit.h
@@ -26,6 +26,7 @@ class SStandaloneAssetEditorToolkitHost;
 class FWorkspaceItem;
 
 DECLARE_DELEGATE_RetVal( bool, FRequestAssetEditorClose );
+DECLARE_DELEGATE_RetVal( void, FAssetEditorClosing);
 
 
 /**
@@ -331,9 +332,12 @@ protected:
 	/** Called to determine if the user should be prompted for a new file if one is missing during an asset reload */
 	virtual bool ShouldPromptForNewFilesOnReload(const UObject& object) const;
 
-	/** Called when this toolkit would close */
+	/** Called when this toolkit is requested to close. Returns false if closing should be prevented. */
 	virtual bool OnRequestClose() { return true; }
 
+	/** Called when this toolkit is being closed */
+	virtual void OnClose() {}
+
 	/**
 	  * Static: Called when "Switch to Standalone Editor" is clicked for the asset editor
 	  *
diff --git a/Engine/Source/Editor/UnrealEd/Public/WorkflowOrientedApp/WorkflowCentricApplication.h b/Engine/Source/Editor/UnrealEd/Public/WorkflowOrientedApp/WorkflowCentricApplication.h
--- a/Engine/Source/Editor/UnrealEd/Public/WorkflowOrientedApp/WorkflowCentricApplication.h
+++ b/Engine/Source/Editor/UnrealEd/Public/WorkflowOrientedApp/WorkflowCentricApplication.h
@@ -26,6 +26,7 @@ public:
 	
 	// FAssetEditorToolkit interface
 	virtual bool OnRequestClose() override;
+	virtual void OnClose() override;
 	virtual FName GetToolMenuToolbarName(FName& OutParentName) const override;
 	// End of FAssetEditorToolkit interface
 
