Fixed issue where moving actors from multiple sublevels could result in actors being lost

diff --git a/Engine/Source/Editor/UnrealEd/Private/EditorActor.cpp b/Engine/Source/Editor/UnrealEd/Private/EditorActor.cpp
--- a/Engine/Source/Editor/UnrealEd/Private/EditorActor.cpp
+++ b/Engine/Source/Editor/UnrealEd/Private/EditorActor.cpp
@@ -254,7 +254,7 @@ void UUnrealEdEngine::edactPasteSelected(UWorld* InWorld, bool bDuplicate, bool
 
 			for (UActorComponent* PastedComponent : PastedComponents)
 			{
-				GEditor->SelectComponent(PastedComponent, true, false);
+				GEditor->SelectComponent(PastedComponent, true, false, true);
 			}
 
 			ComponentSelection->EndBatchSelectOperation(true);
@@ -275,7 +275,7 @@ void UUnrealEdEngine::edactPasteSelected(UWorld* InWorld, bool bDuplicate, bool
 
 			for (AActor* PastedActor : PastedActors)
 			{
-				GEditor->SelectActor(PastedActor, true, false);
+				GEditor->SelectActor(PastedActor, true, false, true);
 			}
 
 			ActorSelection->EndBatchSelectOperation(true);
diff --git a/Engine/Source/Editor/UnrealEd/Private/EditorServer.cpp b/Engine/Source/Editor/UnrealEd/Private/EditorServer.cpp
--- a/Engine/Source/Editor/UnrealEd/Private/EditorServer.cpp
+++ b/Engine/Source/Editor/UnrealEd/Private/EditorServer.cpp
@@ -3567,8 +3567,17 @@ void UEditorEngine::CopySelectedActorsToClipboard( UWorld* InWorld, bool bShould
 					}
 				}
 
-				// Update the clipboard with the final string
-				FPlatformApplicationMisc::ClipboardCopy( *ClipboardString );
+				if (DestinationData)
+				{
+					// Fill DestinationData with the final string
+					*DestinationData = ClipboardString;
+				}
+				else
+				{
+					// Update the clipboard with the final string
+					FPlatformApplicationMisc::ClipboardCopy(*ClipboardString);
+				}
+				
 
 				// Cleanup.
 				for ( CopyJobMap::TIterator It( CopyJobs ) ; It ; ++It )
