Fix a crash when turning off analytics while assets are being discovered.

diff --git a/Engine/Source/Editor/UnrealEd/Private/UnrealEdMisc.cpp b/Engine/Source/Editor/UnrealEd/Private/UnrealEdMisc.cpp
--- a/Engine/Source/Editor/UnrealEd/Private/UnrealEdMisc.cpp
+++ b/Engine/Source/Editor/UnrealEd/Private/UnrealEdMisc.cpp
@@ -809,7 +809,7 @@ void FUnrealEdMisc::EditorAnalyticsHeartbeat()
 
 void FUnrealEdMisc::TickAssetAnalytics()
 {
-	if( bIsAssetAnalyticsPending )
+	if( bIsAssetAnalyticsPending && FEngineAnalytics::IsAvailable())
 	{
 		TRACE_CPUPROFILER_EVENT_SCOPE(FUnrealEdMisc::TickAssetAnalytics);
 
