IterationProfile: Support for baseline comparison of tracked telemetry.

The baseline telemetry csv files are an artifact of a build that are stashed in revision control used to compare telemetry values against bleeding edge builds to detect regressions.

Add support to IterationProfile.xml for stashing and retrieving a specific build's telemetry data for blessing as the baseline for which builds are compared.
Add SummarizeTraceCommandlet support to ingest the blessed telemetry csv for comparison against the current build.
Add warning and error threshold columns to stats.csv, and add support for percentage and cardinal value ascending overage and descending underage.
Restructure CSV loading to pull some complexity out of the Commandlet::Main.
The baseline telemetry csvs included in this change are from [REDACTED_HANDLE].

diff --git a/Engine/Build/BatchFiles/IterationProfile.bat b/Engine/Build/BatchFiles/IterationProfile.bat
--- a/Engine/Build/BatchFiles/IterationProfile.bat
+++ b/Engine/Build/BatchFiles/IterationProfile.bat
@@ -1 +1 @@
-@call %~dp0RunUAT.bat BuildGraph -Script=Engine/Build/IterationProfile.xml -Target="Run Iteration Profile Tests" -append:ProjectNames="EngineTest" -append:ProjectNames="ShooterGame" -set:Win64=true %*
+@call %~dp0RunUAT.bat BuildGraph -Script=Engine/Build/IterationProfile.xml -Target="Run Iteration Profile Tests" -append:ProjectNames="ShooterGame" -set:Win64=true %*
diff --git a/Engine/Build/BatchFiles/IterationProfileBaseline.bat b/Engine/Build/BatchFiles/IterationProfileBaseline.bat
new file mode 100755
--- /dev/null
+++ b/Engine/Build/BatchFiles/IterationProfileBaseline.bat
@@ -0,0 +1 @@
+@call %~dp0RunUAT.bat BuildGraph -Script=Engine/Build/IterationProfile.xml -Target="Baseline Iteration Profile Results" -append:ProjectNames="ShooterGame" -set:Win64=true %*
diff --git a/Engine/Build/BatchFiles/IterationProfileRunAndBaseline.bat b/Engine/Build/BatchFiles/IterationProfileRunAndBaseline.bat
new file mode 100755
--- /dev/null
+++ b/Engine/Build/BatchFiles/IterationProfileRunAndBaseline.bat
@@ -0,0 +1 @@
+@call %~dp0RunUAT.bat BuildGraph -Script=Engine/Build/IterationProfile.xml -set:RunAndBaseline=true -Target="Run Iteration Profile Tests" -Target="Baseline Iteration Profile Results" -append:ProjectNames="ShooterGame" -set:Win64=true %*
diff --git a/Engine/Build/IterationProfile.xml b/Engine/Build/IterationProfile.xml
--- a/Engine/Build/IterationProfile.xml
+++ b/Engine/Build/IterationProfile.xml
@@ -34,9 +34,24 @@
 	<Property Name="EditorCompileArgs" Value="-BuildPlugin=$(ExtraPluginsPlusDelims)"/>
 
 	<!-- Automation uses its own network output directory -->
-	<Property Name="AutomationNodes" Value="" />
+	<Property Name="AutomationNetworkRootDir" Value="//epicgames.net/root/Builds/Automation/Reports/IterationProfile"/>
+	<Property Name="AutomationNetworkRootDir" Value="/Volumes/Builds/Automation/Reports/IterationProfile" If="'$(HostPlatform)' != 'Win64'"/>
+	<Property Name="AutomationReportRootUrl" Value="http://automation.epicgames.net/reports/iterationprofile"/>
 	<Property Name="AutomationLocalOutputDir" Value="$(RootDir)/LocalBuilds/Reports"/>
 
+	<!-- Outputs for this run -->
+	<Property Name="AutomationNetworkOutputDir" Value="$(AutomationNetworkRootDir)/$(BuildName)"/>
+	<Property Name="AutomationReportBaseUrl" Value="$(AutomationReportRootUrl)/$(BuildName)"/>
+	<Property Name="AutomationOutputDir" Value="$(AutomationLocalOutputDir)"/>
+	<Property Name="AutomationOutputDir" Value="$(AutomationNetworkOutputDir)" If="$(IsBuildMachine)"/>
+
+	<!-- Options for baselining -->
+	<Option Name="RunAndBaseline" DefaultValue="false" Description="Enable graph dependencies for running and setting the baseline in a single build."/>
+	<Option Name="BaselineDir" DefaultValue="$(RootDir)/Engine/Build/Baseline" Description="The revision controlled location of the baseline telemetry files."/>
+	<Option Name="BaselineBuildName" DefaultValue="$(BuildName)" Description="The BuildName of the build to set as the new baseline to compare against."/>
+	<Property Name="BaselineAutomationNetworkOutputDir" Value="$(AutomationNetworkRootDir)/$(BaselineBuildName)"/>
+	<Property Name="BaselineAutomationNetworkOutputDir" Value="$(AutomationOutputDir)" If="$(RunAndBaseline)"/>
+
 	<!-- Arguments for tracing -->
 	<Option Name="ColdPass" DefaultValue="true" Description="Perform cold workspace loading profiles"/>
 	<Option Name="ColdArgs" DefaultValue="-ddc=Cold" Description="Arguments to unreal for the cold pass"/>
@@ -58,12 +73,14 @@
 
 	<!-- Gauntlet arguments for summarize commandlet runs -->
 	<Property Name="SummarizeArgs" Value="-statsfile=&quot;$(RootDir)/Engine/Build/IterationProfileStats.csv&quot;"/>
+	<Property Name="SummarizeArgs" Value="$(SummarizeArgs) -skipbaseline" If="$(RunAndBaseline)"/>
 
 	<!-- The main test running macro logic that is expanded with platform and dependency node names -->
-	<Macro Name="RunTest" Arguments="OutputName;PlatformName;ProjectName;NodeName;NodeRequires;Phase;Args">
+	<Macro Name="RunTest" Arguments="OutputName;PlatformName;ProjectName;NodeName;Phase;Args">
 		<Property Name="TestName" Value="IterationProfile$(Phase)"/>
-		<Property Name="TraceFile" Value="$(TraceDir)/$(ProjectName)$(PlatformName)$(TestName).utrace"/>
-		<Property Name="TelemetryFile" Value="$(TraceDir)/$(ProjectName)$(PlatformName)$(TestName)Telemetry.csv"/>
+		<Property Name="TraceFilePath" Value="$(TraceDir)/$(ProjectName)$(PlatformName)$(TestName).utrace"/>
+		<Property Name="TelemetryFileName" Value="$(ProjectName)$(PlatformName)$(TestName)Telemetry.csv"/>
+		<Property Name="TelemetryFilePath" Value="$(TraceDir)/$(TelemetryFileName)"/>
 		<Property Name="ReportExportPath" Value="$(AutomationOutputDir)/$(OutputName)"/>
 		<Property Name="ReportParams" Value="-ReportExportPath=&quot;$(ReportExportPath)&quot;"/>
 		<Do If="$(IsBuildMachine)">
@@ -71,18 +88,43 @@
 			<Property Name="ReportParams" Value="$(ReportParams) -ReportURL=&quot;$(ReportExportURL)&quot; -WriteTestResultsForHorde -HordeTestDataKey=&quot;$(NodeName) $(Phase)$(PreflightSuffix)&quot;"/>
 		</Do>
 		<Property Name="LogOutputParams" Value="-logdir=&quot;$(AutomationOutputDir)/$(OutputName)&quot;"/>
-		<Property Name="AllGauntletArgs" Value="-platform=$(PlatformName) -project=$(ProjectName) $(GauntletArgs) $(ReportParams) $(LogOutputParams) -artifactname=&quot;$(Phase)&quot;"/>
-		<Command Name="RunUnreal" Arguments="$(AllGauntletArgs) $(ProfileArgs) $(Args) -tracefile=&quot;$(TraceFile)&quot;"/>
-		<Commandlet Name="SummarizeTrace" Arguments="-NoEnginePlugins -inputfile=&quot;$(TraceFile)&quot; -testname=&quot;$(TestName)&quot; $(SummarizeArgs)"/>
+		<Property Name="AllGauntletArgs" Value="-project=$(ProjectName) -platform=$(PlatformName) $(GauntletArgs) $(ReportParams) $(LogOutputParams) -artifactname=&quot;$(Phase)&quot;"/>
+
+		<!-- Run the Gauntlet Iteration Profile Tests -->
+		<Command Name="RunUnreal" Arguments="$(AllGauntletArgs) $(ProfileArgs) $(Args) -tracefile=&quot;$(TraceFilePath)&quot;"/>
+
+		<!-- Translate the utrace file into csv data -->
+		<Commandlet Name="SummarizeTrace" Arguments="-NoEnginePlugins -inputfile=&quot;$(TraceFilePath)&quot; -testname=&quot;$(TestName)&quot; $(SummarizeArgs)"/>
+
+		<!-- Stash the telemetry file in case this build will become the baseline in the future -->
+		<Copy From="$(TelemetryFilePath)" To="$(AutomationOutputDir)/$(TelemetryFileName)"/>
+
 		<Do If="$(IsBuildMachine)">
 			<!-- Files copied to this folder are automatically uploaded to Horde as artifacts -->
-			<Copy From="$(TraceDir)/$(ProjectName)$(TestName)*.*" To="$(RootDir)/Engine/Programs/AutomationTool/Saved/Logs"/>
-			<Command Name="PublishUnrealAutomationTelemetry" Arguments="-CSVFile=&quot;$(TelemetryFile)&quot; -project=$(ProjectName) -changelist=$(Change) -branch=$(Branch)"/>
+			<Copy From="$(TraceDir)/$(ProjectName)$(PlatformName)$(TestName)*.*" To="$(RootDir)/Engine/Programs/AutomationTool/Saved/Logs"/>
+
+			<!-- Publish telemetry to the db -->
+			<Command Name="PublishUnrealAutomationTelemetry" Arguments="-CSVFile=&quot;$(TelemetryFilePath)&quot; -project=$(ProjectName) -changelist=$(Change) -branch=$(Branch)"/>
 		</Do>
 	</Macro>
-	
+
+	<!-- This will be set to the names of all nodes we want to execute based on specified platforms/options -->
+	<Property Name="RunGraphNodes" Value=""/>
+
+	<!-- Mimic of the Runner that retrieves the stashed telemetry file and copies it to the depot location -->
+	<Macro Name="BaselineResult" Arguments="PlatformName;ProjectName;Phase">
+		<Property Name="TestName" Value="IterationProfile$(Phase)"/>
+		<Property Name="TelemetryFileName" Value="$(ProjectName)$(PlatformName)$(TestName)Telemetry.csv"/>
+		<Property Name="LocalTelemetryFilePath" Value="$(BaselineDir)/$(TelemetryFileName)"/>
+		<Property Name="NetworkTelemetryFilePath" Value="$(BaselineAutomationNetworkOutputDir)/$(TelemetryFileName)"/>
+		<Copy From="$(NetworkTelemetryFilePath)" To="$(LocalTelemetryFilePath)" Overwrite="true"/>
+		<Do If="!$(IsBuildMachine)">
+			<Spawn Exe="p4" Arguments="reconcile -ea $(LocalTelemetryFilePath)"/>
+		</Do>
+	</Macro>
+
 	<!-- This will be set to the names of all nodes we want to execute based on specified platforms/options -->
-	<Property Name="GraphNodes" Value="" />	
+	<Property Name="BaselineGraphNodes" Value=""/>
 
 	<!-- Windows -->
 	<Agent Name="Prerequisites Agent Win64" Type="Win64">
@@ -103,17 +145,10 @@
 			<Compile Target="CrashReportClientEditor" Platform="Mac" Configuration="Shipping"/>
 		</Node>
 	</Agent>
-
+	
 	<!-- ForEach Project to run, generate the graph nodes -->
 	<ForEach Name="ProjectName" Values="$(ProjectNames)">
 
-		<!-- Automation uses its own network output directory -->
-		<Property Name="AutomationNetworkOutputDir" Value="P:/Builds/Automation/Reports/$(EscapedBranch)/$(ProjectName)/$(BuildName)"/>
-		<Property Name="AutomationNetworkOutputDir" Value="/Volumes/Builds/Automation/Reports/$(EscapedBranch)/$(ProjectName)/$(BuildName)" If="'$(HostPlatform)' != 'Win64'"/>
-		<Property Name="AutomationReportBaseUrl" Value="http://automation.epicgames.net/reports/$(EscapedBranch)/$(ProjectName)/$(BuildName)"/>
-		<Property Name="AutomationOutputDir" Value="$(AutomationLocalOutputDir)"/>
-		<Property Name="AutomationOutputDir" Value="$(AutomationNetworkOutputDir)" If="$(IsBuildMachine)"/>
-
 		<!-- Windows -->
 		<Agent Name="Prerequisites Agent Win64" Type="Win64">
 			<Node Name="Compile $(ProjectName) Editor Win64" Requires="Compile Prerequisites Win64Host">
@@ -121,20 +156,34 @@
 			</Node>
 		</Agent>
 		<Agent Name="Automated Test Agent $(ProjectName) Win64" Type="TestWin64">
-			<Property Name="PlatformName" Value="Win64" />
-			<Property Name="OutputName" Value="IterationProfile $(PlatformName)" />
-			<Property Name="NodeName" Value="Run $(ProjectName) $(OutputName)" />
-			<Property Name="NodeRequires" Value="Compile Prerequisites $(PlatformName)Host;Compile $(ProjectName) Editor $(PlatformName)" />
+			<Property Name="PlatformName" Value="Win64"/>
+			<Property Name="OutputName" Value="$(ProjectName)$(PlatformName)IterationProfile"/>
+			<Property Name="RunNodeName" Value="Run $(OutputName)"/>
+			<Property Name="RunNodeRequires" Value="Compile Prerequisites $(PlatformName)Host;Compile $(ProjectName) Editor $(PlatformName)"/>
+			<Property Name="BaselineNodeName" Value="Baseline $(OutputName)"/>
+			<Property Name="BaselineNodeRequires" Value=""/>
+			<Property Name="BaselineNodeRequires" Value="$(RunNodeName)" If="$(RunAndBaseline)"/>
 			<Do If="$(Win64)">
-				<Node Name="$(NodeName)" Requires="$(NodeRequires)">
+				<!-- Run the Tests -->
+				<Node Name="$(RunNodeName)" Requires="$(RunNodeRequires)">
+					<Do If="$(ColdPass)">
+						<Expand Name="RunTest" OutputName="$(OutputName)" PlatformName="$(PlatformName)" ProjectName="$(ProjectName)" NodeName="$(RunNodeName)" Phase="Cold" Args="$(ColdArgs)"/>
+					</Do>
+					<Do If="$(WarmPass)">
+						<Expand Name="RunTest" OutputName="$(OutputName)" PlatformName="$(PlatformName)" ProjectName="$(ProjectName)" NodeName="$(RunNodeName)" Phase="Warm" Args="$(WarmArgs)"/>
+					</Do>
+				</Node>
+				<Property Name="RunGraphNodes" Value="$(RunGraphNodes);$(RunNodeName)"/>
+				<!-- Baseline the Results -->
+				<Node Name="$(BaselineNodeName)" Requires="$(BaselineNodeRequires)">
 					<Do If="$(ColdPass)">
-						<Expand Name="RunTest" OutputName="$(OutputName)" PlatformName="$(PlatformName)" ProjectName="$(ProjectName)" NodeName="$(NodeName)" NodeRequires="$(NodeRequires)" Phase="Cold" Args="$(ColdArgs)"/>
+						<Expand Name="BaselineResult" PlatformName="$(PlatformName)" ProjectName="$(ProjectName)" Phase="Cold"/>
 					</Do>
 					<Do If="$(WarmPass)">
-						<Expand Name="RunTest" OutputName="$(OutputName)" PlatformName="$(PlatformName)" ProjectName="$(ProjectName)" NodeName="$(NodeName)" NodeRequires="$(NodeRequires)" Phase="Warm" Args="$(WarmArgs)"/>
+						<Expand Name="BaselineResult" PlatformName="$(PlatformName)" ProjectName="$(ProjectName)" Phase="Warm"/>
 					</Do>
 				</Node>
-				<Property Name="GraphNodes" Value="$(GraphNodes);$(NodeName)"/>
+				<Property Name="BaselineGraphNodes" Value="$(BaselineGraphNodes);$(BaselineNodeName)"/>
 			</Do>
 		</Agent>
 		
@@ -145,20 +194,34 @@
 			</Node>
 		</Agent>
 		<Agent Name="Automated Test Agent $(ProjectName) Mac" Type="TestMac">
-			<Property Name="PlatformName" Value="Mac" />
-			<Property Name="OutputName" Value="IterationProfile $(PlatformName)" />
-			<Property Name="NodeName" Value="Run $(ProjectName) $(OutputName)" />
-			<Property Name="NodeRequires" Value="Compile Prerequisites $(PlatformName);Compile $(ProjectName) Editor $(PlatformName)" />
+			<Property Name="PlatformName" Value="Mac"/>
+			<Property Name="OutputName" Value="$(ProjectName)$(PlatformName)IterationProfile"/>
+			<Property Name="RunNodeName" Value="Run $(OutputName)"/>
+			<Property Name="RunNodeRequires" Value="Compile Prerequisites $(PlatformName);Compile $(ProjectName) Editor $(PlatformName)"/>
+			<Property Name="BaselineNodeName" Value="Baseline $(OutputName)"/>
+			<Property Name="BaselineNodeRequires" Value=""/>
+			<Property Name="BaselineNodeRequires" Value="$(RunNodeName)" If="$(RunAndBaseline)"/>
 			<Do If="$(MacOS)">
-				<Node Name="$(NodeName)" Requires="$(NodeRequires)">
+				<!-- Run the Tests -->
+				<Node Name="$(RunNodeName)" Requires="$(RunNodeRequires)">
+					<Do If="$(ColdPass)">
+						<Expand Name="RunTest" OutputName="$(OutputName)" PlatformName="$(PlatformName)" ProjectName="$(ProjectName)" NodeName="$(RunNodeName)" Phase="Cold" Args="$(ColdArgs)"/>
+					</Do>
+					<Do If="$(WarmPass)">
+						<Expand Name="RunTest" OutputName="$(OutputName)" PlatformName="$(PlatformName)" ProjectName="$(ProjectName)" NodeName="$(RunNodeName)" Phase="Warm" Args="$(WarmArgs)"/>
+					</Do>
+				</Node>
+				<Property Name="RunGraphNodes" Value="$(RunGraphNodes);$(RunNodeName)"/>
+				<!-- Baseline the Results -->
+				<Node Name="$(BaselineNodeName)" Requires="$(BaselineNodeRequires)">
 					<Do If="$(ColdPass)">
-						<Expand Name="RunTest" OutputName="$(OutputName)" PlatformName="$(PlatformName)" ProjectName="$(ProjectName)" NodeName="$(NodeName)" NodeRequires="$(NodeRequires)" Phase="Cold" Args="$(ColdArgs)"/>
+						<Expand Name="BaselineResult" PlatformName="$(PlatformName)" ProjectName="$(ProjectName)" Phase="Cold"/>
 					</Do>
 					<Do If="$(WarmPass)">
-						<Expand Name="RunTest" OutputName="$(OutputName)" PlatformName="$(PlatformName)" ProjectName="$(ProjectName)" NodeName="$(NodeName)" NodeRequires="$(NodeRequires)" Phase="Warm" Args="$(WarmArgs)"/>
+						<Expand Name="BaselineResult" PlatformName="$(PlatformName)" ProjectName="$(ProjectName)" Phase="Warm"/>
 					</Do>
 				</Node>
-				<Property Name="GraphNodes" Value="$(GraphNodes);$(NodeName)"/>
+				<Property Name="BaselineGraphNodes" Value="$(BaselineGraphNodes);$(BaselineNodeName)"/>
 			</Do>
 		</Agent>
 
@@ -168,6 +231,19 @@
 		
 	</ForEach>
 
-	<Aggregate Name="Run Iteration Profile Tests" Requires="$(GraphNodes)"/>
+	<!-- These Nodes run the Editor and gather profile data -->
+	<Aggregate Name="Run Iteration Profile Tests" Requires="$(RunGraphNodes)"/>
+
+	<!-- These Nodes copies the CSV telemetry files to special location in the depot, thus making them the "baselines" -->
+	<Aggregate Name="Baseline Iteration Profile Results" Requires="$(BaselineGraphNodes)"/>
+
+	<!-- This Node submits the baseline CSV files to Perforce -->
+	<!--
+	<Agent Name="Submission Agent Win64" Type="Win64">
+		<Node Name="Submit Baseline To Perforce" Requires="$(BaselineGraphNodes)">
+			<Submit Description="[CL $(Change)] Updated IterationProfile baselines" Files="$(BaselineFiles)"/>
+		</Node>
+	</Agent>
+	-->
 
 </BuildGraph>
diff --git a/Engine/Source/Editor/UnrealEd/Private/Commandlets/SummarizeTraceCommandlet.cpp b/Engine/Source/Editor/UnrealEd/Private/Commandlets/SummarizeTraceCommandlet.cpp
--- a/Engine/Source/Editor/UnrealEd/Private/Commandlets/SummarizeTraceCommandlet.cpp
+++ b/Engine/Source/Editor/UnrealEd/Private/Commandlets/SummarizeTraceCommandlet.cpp
@@ -308,33 +308,43 @@ public:
 		FString Name;
 		uint64 Count = 0;
 		double TotalSeconds = 0.0;
+
 		double FirstStartSeconds = 0.0;
 		double FirstEndSeconds = 0.0;
-		double FirstSeconds = 0.0;
-		double MinSeconds = 1e10;
-		double MaxSeconds = -1e10;
-		double MeanSeconds = 0.0;
+		double FirstDurationSeconds = 0.0;
+
+		double LastStartSeconds = 0.0;
+		double LastEndSeconds = 0.0;
+		double LastDurationSeconds = 0.0;
+
+		double MinDurationSeconds = 1e10;
+		double MaxDurationSeconds = -1e10;
+		double MeanDurationSeconds = 0.0;
 		double VarianceAcc = 0.0; // Accumulator for Welford's
 
 		void AddDuration(double StartSeconds, double EndSeconds)
 		{
 			Count += 1;
 
+			// compute the duration
+			double DurationSeconds = EndSeconds - StartSeconds;
+
 			// only set first for the first sample, compare exact zero
 			if (FirstStartSeconds == 0.0)
 			{
 				FirstStartSeconds = StartSeconds;
 				FirstEndSeconds = EndSeconds;
-				FirstSeconds = FirstEndSeconds - FirstStartSeconds;
+				FirstDurationSeconds = DurationSeconds;
 			}
 
-			// compute the duration
-			double DurationSeconds = EndSeconds - StartSeconds;
+			LastStartSeconds = StartSeconds;
+			LastEndSeconds = EndSeconds;
+			LastDurationSeconds = DurationSeconds;
 
 			// set duration statistics
 			TotalSeconds += DurationSeconds;
-			MinSeconds = FMath::Min(MinSeconds, DurationSeconds);
-			MaxSeconds = FMath::Max(MaxSeconds, DurationSeconds);
+			MinDurationSeconds = FMath::Min(MinDurationSeconds, DurationSeconds);
+			MaxDurationSeconds = FMath::Max(MaxDurationSeconds, DurationSeconds);
 			UpdateVariance(DurationSeconds);
 		}
 
@@ -343,12 +353,12 @@ public:
 			ensure(Count);
 
 			// Welford's increment
-			double OldMeanSeconds = MeanSeconds;
-			MeanSeconds = MeanSeconds + ((DurationSeconds - MeanSeconds) / double(Count));
-			VarianceAcc = VarianceAcc + ((DurationSeconds - MeanSeconds) * (DurationSeconds - OldMeanSeconds));
+			double OldMeanDurationSeconds = MeanDurationSeconds;
+			MeanDurationSeconds = MeanDurationSeconds + ((DurationSeconds - MeanDurationSeconds) / double(Count));
+			VarianceAcc = VarianceAcc + ((DurationSeconds - MeanDurationSeconds) * (DurationSeconds - OldMeanDurationSeconds));
 		}
 
-		double GetDeviationSeconds() const
+		double GetDeviationDurationSeconds() const
 		{
 			if (Count > 1)
 			{
@@ -368,8 +378,8 @@ public:
 		{
 			check(Name == Scope.Name);
 			TotalSeconds += Scope.TotalSeconds;
-			MinSeconds = FMath::Min(MinSeconds, Scope.MinSeconds);
-			MaxSeconds = FMath::Max(MaxSeconds, Scope.MaxSeconds);
+			MinDurationSeconds = FMath::Min(MinDurationSeconds, Scope.MinDurationSeconds);
+			MaxDurationSeconds = FMath::Max(MaxDurationSeconds, Scope.MaxDurationSeconds);
 			Count += Scope.Count;
 		}
 
@@ -395,25 +405,37 @@ public:
 			{
 				return FString::Printf(TEXT("%f"), FirstEndSeconds);
 			}
-			else if (Statistic == TEXT("FirstSeconds"))
+			else if (Statistic == TEXT("FirstDurationSeconds"))
+			{
+				return FString::Printf(TEXT("%f"), FirstDurationSeconds);
+			}
+			else if (Statistic == TEXT("LastStartSeconds"))
+			{
+				return FString::Printf(TEXT("%f"), LastStartSeconds);
+			}
+			else if (Statistic == TEXT("LastEndSeconds"))
 			{
-				return FString::Printf(TEXT("%f"), FirstSeconds);
+				return FString::Printf(TEXT("%f"), LastEndSeconds);
 			}
-			else if (Statistic == TEXT("MinSeconds"))
+			else if (Statistic == TEXT("LastDurationSeconds"))
 			{
-				return FString::Printf(TEXT("%f"), MinSeconds);
+				return FString::Printf(TEXT("%f"), LastDurationSeconds);
 			}
-			else if (Statistic == TEXT("MaxSeconds"))
+			else if (Statistic == TEXT("MinDurationSeconds"))
 			{
-				return FString::Printf(TEXT("%f"), MaxSeconds);
+				return FString::Printf(TEXT("%f"), MinDurationSeconds);
 			}
-			else if (Statistic == TEXT("MeanSeconds"))
+			else if (Statistic == TEXT("MaxDurationSeconds"))
 			{
-				return FString::Printf(TEXT("%f"), MeanSeconds);
+				return FString::Printf(TEXT("%f"), MaxDurationSeconds);
 			}
-			else if (Statistic == TEXT("DeviationSeconds"))
+			else if (Statistic == TEXT("MeanDurationSeconds"))
 			{
-				return FString::Printf(TEXT("%f"), GetDeviationSeconds());
+				return FString::Printf(TEXT("%f"), MeanDurationSeconds);
+			}
+			else if (Statistic == TEXT("DeviationDurationSeconds"))
+			{
+				return FString::Printf(TEXT("%f"), GetDeviationDurationSeconds());
 			}
 			return FString();
 		}
@@ -582,6 +604,282 @@ void FSummarizeCountersAnalyzer::OnCounterFloatValue(const FCounterFloatValue& N
 	}
 }
 
+/*
+* Helper class for the stats csv file
+*/
+
+struct StatisticDefinition
+{
+	StatisticDefinition()
+	{}
+
+	StatisticDefinition(const FString& InName, const FString& InStatistic, const FString& InWarningThreshold, const FString& InErrorThreshold)
+		: Name(InName)
+		, Statistic(InStatistic)
+		, WarningThreshold(InWarningThreshold)
+		, ErrorThreshold(InErrorThreshold)
+	{}
+
+	StatisticDefinition(const StatisticDefinition& InStatistic)
+		: Name(InStatistic.Name)
+		, Statistic(InStatistic.Statistic)
+		, WarningThreshold(InStatistic.WarningThreshold)
+		, ErrorThreshold(InStatistic.ErrorThreshold)
+	{}
+
+	bool operator==(const StatisticDefinition& InStatistic) const
+	{
+		return Name == InStatistic.Name
+			&& Statistic == InStatistic.Statistic
+			&& WarningThreshold == InStatistic.WarningThreshold
+			&& ErrorThreshold == InStatistic.ErrorThreshold;
+	}
+
+	static bool LoadFromCSV(const FString& FilePath, TMultiMap<FString, StatisticDefinition>& NameToDefinitionMap);
+
+	FString Name;
+	FString Statistic;
+	FString WarningThreshold;
+	FString ErrorThreshold;
+};
+
+bool StatisticDefinition::LoadFromCSV(const FString& FilePath, TMultiMap<FString, StatisticDefinition>& NameToDefinitionMap)
+{
+	TArray<FString> ParsedCSVFile;
+	FFileHelper::LoadFileToStringArray(ParsedCSVFile, *FilePath);
+
+	int NameColumn = -1;
+	int StatisticColumn = -1;
+	int WarningThresholdColumn = -1;
+	int ErrorThresholdColumn = -1;
+	struct Column
+	{
+		const TCHAR* Name = nullptr;
+		int* Index = nullptr;
+	}
+	Columns[] =
+	{
+		{ TEXT("Name"), &NameColumn },
+		{ TEXT("Statistic"), &StatisticColumn },
+		{ TEXT("WarningThreshold"), &WarningThresholdColumn },
+		{ TEXT("ErrorThreshold"), &ErrorThresholdColumn },
+	};
+
+	bool bValidColumns = true;
+	for (int CSVIndex = 0; CSVIndex < ParsedCSVFile.Num() && bValidColumns; ++CSVIndex)
+	{
+		const FString& CSVEntry = ParsedCSVFile[CSVIndex];
+		TArray<FString> Fields;
+		UE::String::ParseTokens(CSVEntry.TrimStartAndEnd(), TEXT(','),
+			[&Fields](FStringView Field)
+			{
+				Fields.Add(FString(Field));
+			});
+
+		if (CSVIndex == 0) // is this the header row?
+		{
+			for (struct Column& Column : Columns)
+			{
+				for (int FieldIndex = 0; FieldIndex < Fields.Num(); ++FieldIndex)
+				{
+					if (Fields[FieldIndex] == Column.Name)
+					{
+						(*Column.Index) = FieldIndex;
+						break;
+					}
+				}
+
+				if (*Column.Index == -1)
+				{
+					bValidColumns = false;
+				}
+			}
+		}
+		else // else it is a data row, pull each element from appropriate column
+		{
+			const FString& Name(Fields[NameColumn]);
+			const FString& Statistic(Fields[StatisticColumn]);
+			const FString& WarningThreshold(Fields[WarningThresholdColumn]);
+			const FString& ErrorThreshold(Fields[ErrorThresholdColumn]);
+			NameToDefinitionMap.AddUnique(Name, StatisticDefinition(Name, Statistic, WarningThreshold, ErrorThreshold));
+		}
+	}
+
+	return bValidColumns;
+}
+
+/*
+* Helper class for the telemetry csv file
+*/
+
+struct TelemetryDefinition
+{
+	TelemetryDefinition()
+	{}
+
+	TelemetryDefinition(const FString& InTestName, const FString& InContext, const FString& InDataPoint, const FString& InMeasurement)
+		: TestName(InTestName)
+		, Context(InContext)
+		, DataPoint(InDataPoint)
+		, Measurement(InMeasurement)
+	{}
+
+	TelemetryDefinition(const TelemetryDefinition& InStatistic)
+		: TestName(InStatistic.TestName)
+		, Context(InStatistic.Context)
+		, DataPoint(InStatistic.DataPoint)
+		, Measurement(InStatistic.Measurement)
+	{}
+
+	bool operator==(const TelemetryDefinition& InStatistic) const
+	{
+		return TestName == InStatistic.TestName
+			&& Context == InStatistic.Context
+			&& DataPoint == InStatistic.DataPoint
+			&& Measurement == InStatistic.Measurement;
+	}
+
+	static bool LoadFromCSV(const FString& FilePath, TMap<TPair<FString,FString>, TelemetryDefinition>& ContextAndDataPointToDefinitionMap);
+	static bool MeasurementWithinThreshold(const FString& Value, const FString& BaselineValue, const FString& Threshold);
+
+	FString TestName;
+	FString Context;
+	FString DataPoint;
+	FString Measurement;
+};
+
+bool TelemetryDefinition::LoadFromCSV(const FString& FilePath, TMap<TPair<FString, FString>, TelemetryDefinition>& ContextAndDataPointToDefinitionMap)
+{
+	TArray<FString> ParsedCSVFile;
+	FFileHelper::LoadFileToStringArray(ParsedCSVFile, *FilePath);
+
+	int TestNameColumn = -1;
+	int ContextColumn = -1;
+	int DataPointColumn = -1;
+	int MeasurementColumn = -1;
+	struct Column
+	{
+		const TCHAR* Name = nullptr;
+		int* Index = nullptr;
+	}
+	Columns[] =
+	{
+		{ TEXT("TestName"), &TestNameColumn },
+		{ TEXT("Context"), &ContextColumn },
+		{ TEXT("DataPoint"), &DataPointColumn },
+		{ TEXT("Measurement"), &MeasurementColumn },
+	};
+
+	bool bValidColumns = true;
+	for (int CSVIndex = 0; CSVIndex < ParsedCSVFile.Num() && bValidColumns; ++CSVIndex)
+	{
+		const FString& CSVEntry = ParsedCSVFile[CSVIndex];
+		TArray<FString> Fields;
+		UE::String::ParseTokens(CSVEntry.TrimStartAndEnd(), TEXT(','),
+			[&Fields](FStringView Field)
+			{
+				Fields.Add(FString(Field));
+			});
+
+		if (CSVIndex == 0) // is this the header row?
+		{
+			for (struct Column& Column : Columns)
+			{
+				for (int FieldIndex = 0; FieldIndex < Fields.Num(); ++FieldIndex)
+				{
+					if (Fields[FieldIndex] == Column.Name)
+					{
+						(*Column.Index) = FieldIndex;
+						break;
+					}
+				}
+
+				if (*Column.Index == -1)
+				{
+					bValidColumns = false;
+				}
+			}
+		}
+		else // else it is a data row, pull each element from appropriate column
+		{
+			const FString& TestName(Fields[TestNameColumn]);
+			const FString& Context(Fields[ContextColumn]);
+			const FString& DataPoint(Fields[DataPointColumn]);
+			const FString& Measurement(Fields[MeasurementColumn]);
+			ContextAndDataPointToDefinitionMap.Add(TPair<FString, FString>(Context, DataPoint), TelemetryDefinition(TestName, Context, DataPoint, Measurement));
+		}
+	}
+
+	return bValidColumns;
+}
+
+bool TelemetryDefinition::MeasurementWithinThreshold(const FString& MeasurementValue, const FString& BaselineValue, const FString& Threshold)
+{
+	if (Threshold.IsEmpty())
+	{
+		return true;
+	}
+
+	// detect threshold as delta percentage
+	int32 PercentIndex = INDEX_NONE;
+	if (Threshold.FindChar(TEXT('%'), PercentIndex))
+	{
+		FString ThresholdWithoutPercentSign = Threshold;
+		ThresholdWithoutPercentSign.RemoveAt(PercentIndex);
+
+		double Factor = 1.0 + (FCString::Atod(*ThresholdWithoutPercentSign) / 100.0);
+		double RationalValue = FCString::Atod(*MeasurementValue);
+		double RationalBaselineValue = FCString::Atod(*BaselineValue);
+		if (Factor >= 1.0)
+		{
+			return RationalValue < (RationalBaselineValue * Factor);
+		}
+		else
+		{
+			return RationalValue > (RationalBaselineValue * Factor);
+		}
+	}
+	else // threshold as delta cardinal value
+	{
+		// rational number, use float math
+		if (Threshold.Contains(TEXT(".")))
+		{
+			double Delta = FCString::Atod(*Threshold);
+			double RationalValue = FCString::Atod(*MeasurementValue);
+			double RationalBaselineValue = FCString::Atod(*BaselineValue);
+			if (Delta > 0.0)
+			{
+				return RationalValue <= (RationalBaselineValue + Delta);
+			}
+			else if (Delta < 0.0)
+			{
+				return RationalValue >= (RationalBaselineValue + Delta);
+			}
+			else
+			{
+				return fabs(RationalBaselineValue - RationalValue) < FLT_EPSILON;
+			}
+		}
+		else // natural number, use int math
+		{
+			int64 Delta = FCString::Strtoi64(*Threshold, nullptr, 10);
+			int64 NaturalValue = FCString::Strtoi64(*MeasurementValue, nullptr, 10);
+			int64 NaturalBaselineValue = FCString::Strtoi64(*BaselineValue, nullptr, 10);
+			if (Delta > 0)
+			{
+				return NaturalValue <= (NaturalBaselineValue + Delta);
+			}
+			else if (Delta < 0)
+			{
+				return NaturalValue >= (NaturalBaselineValue + Delta);
+			}
+			else
+			{
+				return NaturalValue == NaturalBaselineValue;
+			}
+		}
+	}
+}
 
 /*
  * SummarizeTrace commandlet ingests a utrace file and summarizes the
@@ -597,13 +895,6 @@ USummarizeTraceCommandlet::USummarizeTraceCommandlet(const FObjectInitializer& O
 {
 }
 
-// helper for USummarizeTraceCommandlet::Main
-static void FilePrint(IFileHandle* Handle, const FString& String)
-{
-	const auto& UTF8String = StringCast<ANSICHAR>(*String);
-	Handle->Write(reinterpret_cast<const uint8*>(UTF8String.Get()), UTF8String.Length());
-}
-
 int32 USummarizeTraceCommandlet::Main(const FString& CmdLineParams)
 {
 	TArray<FString> Tokens;
@@ -620,6 +911,7 @@ int32 USummarizeTraceCommandlet::Main(const FString& CmdLineParams)
 		UE_LOG(LogSummarizeTrace, Log, TEXT(" Required: -inputfile=<utrace path>   (The utrace you wish to process)"));
 		UE_LOG(LogSummarizeTrace, Log, TEXT(" Optional: -statsfile=<csv path>      (The csv of statistics to generate)"));
 		UE_LOG(LogSummarizeTrace, Log, TEXT(" Optional: -testname=<string>         (Test name to use in telemetry csv)"));
+		UE_LOG(LogSummarizeTrace, Log, TEXT(" Optional: -skipbaseline              (Skip comparing values vs. baselines)"));
 		return 0;
 	}
 
@@ -637,65 +929,12 @@ int32 USummarizeTraceCommandlet::Main(const FString& CmdLineParams)
 	// load the stats file to know which event name and statistic name to generate in the telementry csv
 	// the telemetry csv is ingested completely, so this just whitelists specific data elements we want to track
 	FString StatisticsFileName;
-	TMultiMap<FString, FString> NameToStatisticMap;
+	TMultiMap<FString, StatisticDefinition> NameToDefinitionMap;
 	if (FParse::Value(*CmdLineParams, TEXT("statsfile="), StatisticsFileName, true))
 	{
-		UE_LOG(LogSummarizeTrace, Display, TEXT("Generating statistics from %s"), *StatisticsFileName);
-
-		TArray<FString> ParsedCSVFile;
-		FFileHelper::LoadFileToStringArray(ParsedCSVFile, *StatisticsFileName);
-
-		int NameColumn = -1;
-		int StatisticColumn = -1;
-		struct Column
-		{
-			const TCHAR* Name = nullptr;
-			int* Index = nullptr;
-		}
-		Columns[] =
-		{
-			{ TEXT("Name"), &NameColumn },
-			{ TEXT("Statistic"), &StatisticColumn },
-		};
-
-		bool bValidColumns = true;
-		for (int CSVIndex = 0; CSVIndex < ParsedCSVFile.Num() && bValidColumns; ++CSVIndex)
-		{
-			const FString& CSVEntry = ParsedCSVFile[CSVIndex];
-			TArray<FString> Fields;
-			UE::String::ParseTokens(CSVEntry.TrimStartAndEnd(), TEXT(','),
-				[&Fields](FStringView Field)
-				{
-					if (!Field.IsEmpty())
-					{
-						Fields.Add(FString(Field));
-					}
-				});
-
-			for (struct Column& Column : Columns)
-			{
-				if (CSVIndex == 0) // is this the header row?
-				{
-					for (int FieldIndex = 0; FieldIndex < Fields.Num(); ++FieldIndex)
-					{
-						if (Fields[FieldIndex] == Column.Name)
-						{
-							(*Column.Index) = FieldIndex;
-							break;
-						}
-					}
-
-					if (*Column.Index == -1)
-					{
-						bValidColumns = false;
-					}
-				}
-				else // else it is a data row, pull each element from appropriate column
-				{
-					NameToStatisticMap.AddUnique(FString(Fields[NameColumn]), FString(Fields[StatisticColumn]));
-				}
-			}
-		}
+		UE_LOG(LogSummarizeTrace, Display, TEXT("Loading statistics from %s"), *StatisticsFileName);
+		bool bCSVOk = StatisticDefinition::LoadFromCSV(StatisticsFileName, NameToDefinitionMap);
+		check(bCSVOk);
 	}
 
 	bool bFound;
@@ -736,17 +975,17 @@ int32 USummarizeTraceCommandlet::Main(const FString& CmdLineParams)
 	}
 
 	// setup analysis context with analyzers
-	UE::Trace::FAnalysisContext Context;
+	UE::Trace::FAnalysisContext AnalysisContext;
 	FSummarizeCpuAnalyzer CpuAnalyzer;
-	Context.AddAnalyzer(CpuAnalyzer);
+	AnalysisContext.AddAnalyzer(CpuAnalyzer);
 	FSummarizeCountersAnalyzer CountersAnalyzer;
-	Context.AddAnalyzer(CountersAnalyzer);
+	AnalysisContext.AddAnalyzer(CountersAnalyzer);
 
 	// kick processing on a thread
-	UE::Trace::FAnalysisProcessor Processor = Context.Process(DataStream);
+	UE::Trace::FAnalysisProcessor AnalysisProcessor = AnalysisContext.Process(DataStream);
 
 	// sync on completion
-	Processor.Wait();
+	AnalysisProcessor.Wait();
 
 	TSet<FSummarizeCpuAnalyzer::FScope> DeduplicatedScopes;
 
@@ -827,6 +1066,13 @@ int32 USummarizeTraceCommandlet::Main(const FString& CmdLineParams)
 	}
 	SortedScopes.Sort();
 
+	// csv is UTF-8, so encode every string we print
+	auto WriteUTF8 = [](IFileHandle* Handle, const FString& String)
+	{
+		const auto& UTF8String = StringCast<ANSICHAR>(*String);
+		Handle->Write(reinterpret_cast<const uint8*>(UTF8String.Get()), UTF8String.Length());
+	};
+
 	// generate a summary csv, always
 	FString CsvFileName = FPaths::SetExtension(TraceFileName, "csv");
 	UE_LOG(LogSummarizeTrace, Display, TEXT("Writing summary to %s..."), *CsvFileName);
@@ -834,16 +1080,16 @@ int32 USummarizeTraceCommandlet::Main(const FString& CmdLineParams)
 	if (CsvHandle)
 	{
 		// no newline, see row printfs
-		FilePrint(CsvHandle, FString::Printf(TEXT("Name,Count,TotalSeconds,FirstStartSeconds,FirstEndSeconds,FirstSeconds,MinSeconds,MaxSeconds,MeanSeconds,DeviationSeconds")));
+		WriteUTF8(CsvHandle, FString::Printf(TEXT("Name,Count,TotalSeconds,FirstStartSeconds,FirstEndSeconds,FirstDurationSeconds,LastStartSeconds,LastEndSeconds,LastDurationSeconds,MinDurationSeconds,MaxDurationSeconds,MeanDurationSeconds,DeviationDurationSeconds,")));
 		for (const FSummarizeCpuAnalyzer::FScope& Scope : SortedScopes)
 		{
 			// note newline is at the front of every data line to prevent final extraneous newline, per customary for csv
-			FilePrint(CsvHandle, FString::Printf(TEXT("\n%s,%llu,%f,%f,%f,%f,%f,%f,%f,%f,"), *Scope.Name, Scope.Count, Scope.TotalSeconds, Scope.FirstStartSeconds, Scope.FirstEndSeconds, Scope.FirstSeconds, Scope.MinSeconds, Scope.MaxSeconds, Scope.MeanSeconds, Scope.GetDeviationSeconds()));
+			WriteUTF8(CsvHandle, FString::Printf(TEXT("\n%s,%llu,%f,%f,%f,%f,%f,%f,%f,%f,%f,%f,%f,"), *Scope.Name, Scope.Count, Scope.TotalSeconds, Scope.FirstStartSeconds, Scope.FirstEndSeconds, Scope.FirstDurationSeconds, Scope.FirstStartSeconds, Scope.FirstEndSeconds, Scope.FirstDurationSeconds, Scope.MinDurationSeconds, Scope.MaxDurationSeconds, Scope.MeanDurationSeconds, Scope.GetDeviationDurationSeconds()));
 		}
 		for (const TMap<uint16, FSummarizeCountersAnalyzer::FCounter>::ElementType& Counter : CountersAnalyzer.Counters)
 		{
 			// note newline is at the front of every data line to prevent final extraneous newline, per customary for csv
-			FilePrint(CsvHandle, FString::Printf(TEXT("\n%s,%s,,,,,,,,,"), *Counter.Value.Name, *Counter.Value.GetValue()));
+			WriteUTF8(CsvHandle, FString::Printf(TEXT("\n%s,%s,,,,,,,,,,,,"), *Counter.Value.Name, *Counter.Value.GetValue()));
 		}
 		CsvHandle->Flush();
 		delete CsvHandle;
@@ -856,42 +1102,53 @@ int32 USummarizeTraceCommandlet::Main(const FString& CmdLineParams)
 	}
 
 	// if we were asked to generate a telememtry file, generate it
-	if (!NameToStatisticMap.IsEmpty())
+	if (!NameToDefinitionMap.IsEmpty())
 	{
 		FString TracePath = FPaths::GetPath(TraceFileName);
 		FString TraceFileBasename = FPaths::GetBaseFilename(TraceFileName);
 		FString TelemetryCsvFileName = TraceFileBasename + TEXT("Telemetry");
 		TelemetryCsvFileName = FPaths::Combine(TracePath, FPaths::SetExtension(TelemetryCsvFileName, "csv"));
+
 		UE_LOG(LogSummarizeTrace, Display, TEXT("Writing telemetry to %s..."), *TelemetryCsvFileName);
 		IFileHandle* TelemetryCsvHandle = FPlatformFileManager::Get().GetPlatformFile().OpenWrite(*TelemetryCsvFileName);
+		TArray<TelemetryDefinition> TelemetryWritten;
 		if (TelemetryCsvHandle)
 		{
 			// override the test name
 			FString TestName = TraceFileBasename;
 			FParse::Value(*CmdLineParams, TEXT("testname="), TestName, true);
 
-			// no newline, see row printfs
-			FilePrint(TelemetryCsvHandle, FString::Printf(TEXT("TestName,Context,DataPoint,Measurement")));
+			// resolve scopes to telemetry
 			for (const FSummarizeCpuAnalyzer::FScope& Scope : SortedScopes)
 			{
-				TArray<FString> Statistics;
-				NameToStatisticMap.MultiFind(Scope.Name, Statistics, true);
-				for (const FString& Statistic : Statistics)
+				TArray<StatisticDefinition> Statistics;
+				NameToDefinitionMap.MultiFind(Scope.Name, Statistics, true);
+				for (const StatisticDefinition& Statistic : Statistics)
 				{
-					// note newline is at the front of every data line to prevent final extraneous newline, per customary for csv
-					FilePrint(TelemetryCsvHandle, FString::Printf(TEXT("\n%s,%s,%s,%s,"), *TestName, *Scope.Name, *Statistic, *Scope.GetValue(Statistic)));
+					TelemetryWritten.Add(TelemetryDefinition(TestName, Scope.Name, Statistic.Statistic, Scope.GetValue(Statistic.Statistic)));
 				}
 			}
+
+			// resolve counters to telemetry
 			for (const TMap<uint16, FSummarizeCountersAnalyzer::FCounter>::ElementType& Counter : CountersAnalyzer.Counters)
 			{
-				TArray<FString> Statistics;
-				NameToStatisticMap.MultiFind(Counter.Value.Name, Statistics, true);
-				for (const FString& Statistic : Statistics)
+				TArray<StatisticDefinition> Statistics;
+				NameToDefinitionMap.MultiFind(Counter.Value.Name, Statistics, true);
+				ensure(Statistics.Num() <= 1); // there should only be one, the counter value
+				for (const StatisticDefinition& Statistic : Statistics)
 				{
-					// note newline is at the front of every data line to prevent final extraneous newline, per customary for csv
-					FilePrint(TelemetryCsvHandle, FString::Printf(TEXT("\n%s,%s,%s,%s,"), *TestName, *Counter.Value.Name, *Statistic, *Counter.Value.GetValue()));
+					TelemetryWritten.Add(TelemetryDefinition(TestName, Counter.Value.Name, Statistic.Statistic, Counter.Value.GetValue()));
 				}
 			}
+
+			// no newline, see row printfs
+			WriteUTF8(TelemetryCsvHandle, FString::Printf(TEXT("TestName,Context,DataPoint,Measurement,")));
+			for (const TelemetryDefinition& Telemetry : TelemetryWritten)
+			{
+				// note newline is at the front of every data line to prevent final extraneous newline, per customary for csv
+				WriteUTF8(TelemetryCsvHandle, FString::Printf(TEXT("\n%s,%s,%s,%s,"), *Telemetry.TestName, *Telemetry.Context, *Telemetry.DataPoint, *Telemetry.Measurement));
+			}
+
 			TelemetryCsvHandle->Flush();
 			delete TelemetryCsvHandle;
 			TelemetryCsvHandle = nullptr;
@@ -901,6 +1158,80 @@ int32 USummarizeTraceCommandlet::Main(const FString& CmdLineParams)
 			UE_LOG(LogSummarizeTrace, Error, TEXT("Unable to open telemetry csv '%s' for write"), *TelemetryCsvFileName);
 			return 1;
 		}
+
+		// compare vs. baseline telemetry file, if it exists
+		// note this does assume that the tracefile basename is directly comparable to a file in the baseline folder
+		FString BaselineTelemetryCsvFilePath = FPaths::Combine(FPaths::EngineDir(), TEXT("Build"), TEXT("Baseline"), FPaths::SetExtension(TraceFileBasename + TEXT("Telemetry"), "csv"));
+		if (FParse::Param(*CmdLineParams, TEXT("skipbaseline")))
+		{
+			BaselineTelemetryCsvFilePath.Empty();
+		}
+		if (FPaths::FileExists(BaselineTelemetryCsvFilePath))
+		{
+			UE_LOG(LogSummarizeTrace, Display, TEXT("Comparing telemetry to baseline telemetry %s..."), *TelemetryCsvFileName);
+
+			// each context (scope name or coutner name) and data point (statistic name) pair form a key, an item to check
+			TMap<TPair<FString, FString>, TelemetryDefinition> ContextAndDataPointToDefinitionMap;
+			bool bCSVOk = TelemetryDefinition::LoadFromCSV(*BaselineTelemetryCsvFilePath, ContextAndDataPointToDefinitionMap);
+			check(bCSVOk);
+
+			// for every telemetry item we wrote for this trace...
+			for (const TelemetryDefinition& Telemetry : TelemetryWritten)
+			{
+				// find the corresponding keyed telemetry item in the baseline telemetry file...
+				TelemetryDefinition* BaselineTelemetry = ContextAndDataPointToDefinitionMap.Find(TPair<FString, FString>(Telemetry.Context, Telemetry.DataPoint));
+				if (BaselineTelemetry)
+				{
+					// the threshold is defined along with the original statistic map
+					const StatisticDefinition* RelatedStatistic = nullptr;
+
+					// find the statistic definition
+					TArray<StatisticDefinition> Statistics;
+					NameToDefinitionMap.MultiFind(Telemetry.Context, Statistics, true);
+					for (const StatisticDefinition& Statistic : Statistics)
+					{
+						// the find will match on name, here we just need to find the right statistic for that named item
+						if (Statistic.Statistic == Telemetry.DataPoint)
+						{
+							// we found it!
+							RelatedStatistic = &Statistic;
+							break;
+						}
+					}
+
+					// do we still have the statistic definition in our current stats file? (if we don't that's fine, we don't care about it anymore)
+					if (RelatedStatistic)
+					{
+						// verify that this telemetry measurement is within the allowed threshold as defined in the current stats file
+						if (TelemetryDefinition::MeasurementWithinThreshold(Telemetry.Measurement, BaselineTelemetry->Measurement, RelatedStatistic->WarningThreshold))
+						{
+							UE_LOG(LogSummarizeTrace, Display, TEXT("Telemetry %s,%s,%s,%s within baseline value %s using warning threshold %s"),
+								*Telemetry.TestName, *Telemetry.Context, *Telemetry.DataPoint, *Telemetry.Measurement,
+								*BaselineTelemetry->Measurement, *RelatedStatistic->WarningThreshold);
+						}
+						else
+						{
+							if (TelemetryDefinition::MeasurementWithinThreshold(Telemetry.Measurement, BaselineTelemetry->Measurement, RelatedStatistic->ErrorThreshold))
+							{
+								UE_LOG(LogSummarizeTrace, Warning, TEXT("Telemetry %s,%s,%s,%s beyond baseline value %s using warning threshold %s"),
+									*Telemetry.TestName, *Telemetry.Context, *Telemetry.DataPoint, *Telemetry.Measurement,
+									*BaselineTelemetry->Measurement, *RelatedStatistic->WarningThreshold);
+							}
+							else
+							{
+								UE_LOG(LogSummarizeTrace, Error, TEXT("Telemetry %s,%s,%s,%s beyond baseline value %s using error threshold %s"),
+									*Telemetry.TestName, *Telemetry.Context, *Telemetry.DataPoint, *Telemetry.Measurement,
+									*BaselineTelemetry->Measurement, *RelatedStatistic->ErrorThreshold);
+							}
+						}
+					}
+				}
+				else
+				{
+					UE_LOG(LogSummarizeTrace, Display, TEXT("Telemetry for %s,%s has no baseline measurement, skipping..."), *Telemetry.Context, *Telemetry.DataPoint);
+				}
+			}
+		}
 	}
 
 	return 0;
