WorldPartition - Non-unity build fixes

diff --git a/Engine/Source/Developer/MeshMergeUtilities/Private/MeshMergeUtilities.h b/Engine/Source/Developer/MeshMergeUtilities/Private/MeshMergeUtilities.h
--- a/Engine/Source/Developer/MeshMergeUtilities/Private/MeshMergeUtilities.h
+++ b/Engine/Source/Developer/MeshMergeUtilities/Private/MeshMergeUtilities.h
@@ -4,6 +4,7 @@
 
 #include "IMeshMergeUtilities.h"
 #include "SceneTypes.h"
+#include "MaterialShared.h"
 
 class FMeshMergeDataTracker;
 class FMaterialUpdateContext;
diff --git a/Engine/Source/Editor/UnrealEd/Classes/Commandlets/WorldPartitionConvertCommandlet.h b/Engine/Source/Editor/UnrealEd/Classes/Commandlets/WorldPartitionConvertCommandlet.h
--- a/Engine/Source/Editor/UnrealEd/Classes/Commandlets/WorldPartitionConvertCommandlet.h
+++ b/Engine/Source/Editor/UnrealEd/Classes/Commandlets/WorldPartitionConvertCommandlet.h
@@ -17,6 +17,7 @@ UNREALED_API DECLARE_LOG_CATEGORY_EXTERN(LogWorldPartitionConvertCommandlet, Log
 class ULevelStreaming;
 class UWorldPartition;
 class UWorldComposition;
+class ULevelStreaming;
 
 USTRUCT()
 struct FHLODLayerActorMapping
diff --git a/Engine/Source/Runtime/Engine/Public/WorldPartition/HLOD/HLODLayer.h b/Engine/Source/Runtime/Engine/Public/WorldPartition/HLOD/HLODLayer.h
--- a/Engine/Source/Runtime/Engine/Public/WorldPartition/HLOD/HLODLayer.h
+++ b/Engine/Source/Runtime/Engine/Public/WorldPartition/HLOD/HLODLayer.h
@@ -8,6 +8,8 @@
 #include "UObject/UObjectGlobals.h"
 #include "UObject/Object.h"
 
+#include "EngineDefines.h"
+
 #if WITH_EDITOR
 #include "Engine/MeshMerging.h"
 #include "Engine/MeshMerging.h"
