IterationProfile polish

Update TotalSeconds -> TotalDurationSeconds to make it unambiguously a time span (vs. time stamp)
Add some Telemetry points for Reverb PIE stalls
Don't push telemetry from Preflight builds

Tested with local UAT runs

diff --git a/Engine/Build/IterationProfile.xml b/Engine/Build/IterationProfile.xml
--- a/Engine/Build/IterationProfile.xml
+++ b/Engine/Build/IterationProfile.xml
@@ -104,8 +104,10 @@
 			<!-- Files copied to this folder are automatically uploaded to Horde as artifacts -->
 			<Copy From="$(TraceDir)/$(ProjectName)$(PlatformName)$(TestName)*.*" To="$(RootDir)/Engine/Programs/AutomationTool/Saved/Logs"/>
 
-			<!-- Publish telemetry to the db -->
-			<Command Name="PublishUnrealAutomationTelemetry" Arguments="-CSVFile=&quot;$(TelemetryFilePath)&quot; -project=$(ProjectName) -changelist=$(Change) -branch=$(Branch)"/>
+			<Do If="!$(IsPreflight)">
+				<!-- Publish telemetry to the db -->
+				<Command Name="PublishUnrealAutomationTelemetry" Arguments="-CSVFile=&quot;$(TelemetryFilePath)&quot; -project=$(ProjectName) -changelist=$(Change) -branch=$(Branch)"/>
+			</Do>
 		</Do>
 	</Macro>
 
diff --git a/Engine/Source/Editor/UnrealEd/Private/Commandlets/SummarizeTraceCommandlet.cpp b/Engine/Source/Editor/UnrealEd/Private/Commandlets/SummarizeTraceCommandlet.cpp
--- a/Engine/Source/Editor/UnrealEd/Private/Commandlets/SummarizeTraceCommandlet.cpp
+++ b/Engine/Source/Editor/UnrealEd/Private/Commandlets/SummarizeTraceCommandlet.cpp
@@ -307,7 +307,7 @@ public:
 	{
 		FString Name;
 		uint64 Count = 0;
-		double TotalSeconds = 0.0;
+		double TotalDurationSeconds = 0.0;
 
 		double FirstStartSeconds = 0.0;
 		double FirstEndSeconds = 0.0;
@@ -342,7 +342,7 @@ public:
 			LastDurationSeconds = DurationSeconds;
 
 			// set duration statistics
-			TotalSeconds += DurationSeconds;
+			TotalDurationSeconds += DurationSeconds;
 			MinDurationSeconds = FMath::Min(MinDurationSeconds, DurationSeconds);
 			MaxDurationSeconds = FMath::Max(MaxDurationSeconds, DurationSeconds);
 			UpdateVariance(DurationSeconds);
@@ -377,7 +377,7 @@ public:
 		void Merge(const FScope& Scope)
 		{
 			check(Name == Scope.Name);
-			TotalSeconds += Scope.TotalSeconds;
+			TotalDurationSeconds += Scope.TotalDurationSeconds;
 			MinDurationSeconds = FMath::Min(MinDurationSeconds, Scope.MinDurationSeconds);
 			MaxDurationSeconds = FMath::Max(MaxDurationSeconds, Scope.MaxDurationSeconds);
 			Count += Scope.Count;
@@ -393,9 +393,9 @@ public:
 			{
 				return FString::Printf(TEXT("%llu"), Count);
 			}
-			else if (Statistic == TEXT("TotalSeconds"))
+			else if (Statistic == TEXT("TotalDurationSeconds"))
 			{
-				return FString::Printf(TEXT("%f"), TotalSeconds);
+				return FString::Printf(TEXT("%f"), TotalDurationSeconds);
 			}
 			else if (Statistic == TEXT("FirstStartSeconds"))
 			{
@@ -449,7 +449,7 @@ public:
 		// for sorting descending
 		bool operator<(const FScope& Scope) const
 		{
-			return TotalSeconds > Scope.TotalSeconds;
+			return TotalDurationSeconds > Scope.TotalDurationSeconds;
 		}
 	};
 
@@ -1080,11 +1080,11 @@ int32 USummarizeTraceCommandlet::Main(const FString& CmdLineParams)
 	if (CsvHandle)
 	{
 		// no newline, see row printfs
-		WriteUTF8(CsvHandle, FString::Printf(TEXT("Name,Count,TotalSeconds,FirstStartSeconds,FirstEndSeconds,FirstDurationSeconds,LastStartSeconds,LastEndSeconds,LastDurationSeconds,MinDurationSeconds,MaxDurationSeconds,MeanDurationSeconds,DeviationDurationSeconds,")));
+		WriteUTF8(CsvHandle, FString::Printf(TEXT("Name,Count,TotalDurationSeconds,FirstStartSeconds,FirstEndSeconds,FirstDurationSeconds,LastStartSeconds,LastEndSeconds,LastDurationSeconds,MinDurationSeconds,MaxDurationSeconds,MeanDurationSeconds,DeviationDurationSeconds,")));
 		for (const FSummarizeCpuAnalyzer::FScope& Scope : SortedScopes)
 		{
 			// note newline is at the front of every data line to prevent final extraneous newline, per customary for csv
-			WriteUTF8(CsvHandle, FString::Printf(TEXT("\n%s,%llu,%f,%f,%f,%f,%f,%f,%f,%f,%f,%f,%f,"), *Scope.Name, Scope.Count, Scope.TotalSeconds, Scope.FirstStartSeconds, Scope.FirstEndSeconds, Scope.FirstDurationSeconds, Scope.FirstStartSeconds, Scope.FirstEndSeconds, Scope.FirstDurationSeconds, Scope.MinDurationSeconds, Scope.MaxDurationSeconds, Scope.MeanDurationSeconds, Scope.GetDeviationDurationSeconds()));
+			WriteUTF8(CsvHandle, FString::Printf(TEXT("\n%s,%llu,%f,%f,%f,%f,%f,%f,%f,%f,%f,%f,%f,"), *Scope.Name, Scope.Count, Scope.TotalDurationSeconds, Scope.FirstStartSeconds, Scope.FirstEndSeconds, Scope.FirstDurationSeconds, Scope.FirstStartSeconds, Scope.FirstEndSeconds, Scope.FirstDurationSeconds, Scope.MinDurationSeconds, Scope.MaxDurationSeconds, Scope.MeanDurationSeconds, Scope.GetDeviationDurationSeconds()));
 		}
 		for (const TMap<uint16, FSummarizeCountersAnalyzer::FCounter>::ElementType& Counter : CountersAnalyzer.Counters)
 		{
