Allow actors which manage ISM components to customize or disable per-instance editing

The static mesh instance (SMInstance) element type used to directly edit the data on the ISM component, which can cause issues if the ISM component is owned by an actor that has additional bookkeeping data.

To solve this there is now an extra level of indirection that the actor owning a component can use to customize the default per-instance editing behavior, or even to disable per-instance editing entirely:
 - FSMInstanceManager is a handle that replaces direct use of FSMInstanceId within the SMInstance element implementation, and promotes performing actions via the instance manager interface rather than on the ISM component directly.
 - ISMInstanceManager is an interface that actors owning an ISM component can implement in order to customize the default SMInstance editing behavior. The ISM component also implements this interface for the cases where it manages itself.
 - ISMInstanceManagerProvider is an interface that actors owning an ISM component can implement to either redirect the request to get the instance manager instance, or to disable per-instance editing entirely by returning a null instance manager instance.

The instance manager instance used is resolved from the ISM component using the following logic:
 - If the owner actor implements the ISMInstanceManagerProvider interface, then that API is used to provide the instance manager. No further fallback happens, and a null instance manager instance will disable per-instance editing.
 - If the owner actor implements the ISMInstanceManager interface, then the owner actor is used as the instance manager.
 - Otherwise the ISM component is used as the instance manager.

ALightWeightInstanceStaticMeshManager provides an example of using ISMInstanceManager to keep extra bookkeeping data up-to-date, and AISMPartitionActor provides an example of using ISMInstanceManagerProvider to disable per-instance editing.

AISMPartitionActor will be updated to correctly support SMInstance editing in a future change, by forwarding the request on to the client that owns the respective instance.

diff --git a/Engine/Source/Editor/LevelEditor/Private/Elements/SMInstance/SMInstanceElementLevelEditorSelectionCustomization.cpp b/Engine/Source/Editor/LevelEditor/Private/Elements/SMInstance/SMInstanceElementLevelEditorSelectionCustomization.cpp
--- a/Engine/Source/Editor/LevelEditor/Private/Elements/SMInstance/SMInstanceElementLevelEditorSelectionCustomization.cpp
+++ b/Engine/Source/Editor/LevelEditor/Private/Elements/SMInstance/SMInstanceElementLevelEditorSelectionCustomization.cpp
@@ -11,7 +11,11 @@ DEFINE_LOG_CATEGORY_STATIC(LogSMInstanceLevelEditorSelection, Log, All);
 
 bool FSMInstanceElementLevelEditorSelectionCustomization::CanSelectElement(const TTypedElement<UTypedElementSelectionInterface>& InElementSelectionHandle, const FTypedElementSelectionOptions& InSelectionOptions)
 {
-	const FSMInstanceId SMInstance = SMInstanceElementDataUtil::GetSMInstanceFromHandleChecked(InElementSelectionHandle);
+	const FSMInstanceManager SMInstance = SMInstanceElementDataUtil::GetSMInstanceFromHandleChecked(InElementSelectionHandle);
+	if (!SMInstance)
+	{
+		return false;
+	}
 	
 	// Bail if global selection is locked
 	return !GEdSelectionLock;
@@ -19,7 +23,11 @@ bool FSMInstanceElementLevelEditorSelectionCustomization::CanSelectElement(const
 
 bool FSMInstanceElementLevelEditorSelectionCustomization::CanDeselectElement(const TTypedElement<UTypedElementSelectionInterface>& InElementSelectionHandle, const FTypedElementSelectionOptions& InSelectionOptions)
 {
-	const FSMInstanceId SMInstance = SMInstanceElementDataUtil::GetSMInstanceFromHandleChecked(InElementSelectionHandle);
+	const FSMInstanceManager SMInstance = SMInstanceElementDataUtil::GetSMInstanceFromHandleChecked(InElementSelectionHandle);
+	if (!SMInstance)
+	{
+		return false;
+	}
 
 	// Bail if global selection is locked
 	return !GEdSelectionLock;
@@ -27,34 +35,42 @@ bool FSMInstanceElementLevelEditorSelectionCustomization::CanDeselectElement(con
 
 bool FSMInstanceElementLevelEditorSelectionCustomization::SelectElement(const TTypedElement<UTypedElementSelectionInterface>& InElementSelectionHandle, UTypedElementList* InSelectionSet, const FTypedElementSelectionOptions& InSelectionOptions)
 {
-	const FSMInstanceId SMInstance = SMInstanceElementDataUtil::GetSMInstanceFromHandleChecked(InElementSelectionHandle);
+	const FSMInstanceManager SMInstance = SMInstanceElementDataUtil::GetSMInstanceFromHandleChecked(InElementSelectionHandle);
+	if (!SMInstance)
+	{
+		return false;
+	}
 
 	if (!InElementSelectionHandle.SelectElement(InSelectionSet, InSelectionOptions))
 	{
 		return false;
 	}
 
-	UE_LOG(LogSMInstanceLevelEditorSelection, Verbose, TEXT("Selected SMInstance: %s (%s), Index %d"), *SMInstance.ISMComponent->GetPathName(), *SMInstance.ISMComponent->GetClass()->GetName(), SMInstance.InstanceIndex);
+	UE_LOG(LogSMInstanceLevelEditorSelection, Verbose, TEXT("Selected SMInstance: %s (%s), Index %d"), *SMInstance.GetISMComponent()->GetPathName(), *SMInstance.GetISMComponent()->GetClass()->GetName(), SMInstance.GetISMInstanceIndex());
 
 	// Update the internal selection state for viewport selection rendering
-	SMInstance.ISMComponent->SelectInstance(/*bIsSelected*/true, SMInstance.InstanceIndex);
+	SMInstance.GetISMComponent()->SelectInstance(/*bIsSelected*/true, SMInstance.GetISMInstanceIndex());
 	
 	return true;
 }
 
 bool FSMInstanceElementLevelEditorSelectionCustomization::DeselectElement(const TTypedElement<UTypedElementSelectionInterface>& InElementSelectionHandle, UTypedElementList* InSelectionSet, const FTypedElementSelectionOptions& InSelectionOptions)
 {
-	const FSMInstanceId SMInstance = SMInstanceElementDataUtil::GetSMInstanceFromHandleChecked(InElementSelectionHandle);
+	const FSMInstanceManager SMInstance = SMInstanceElementDataUtil::GetSMInstanceFromHandleChecked(InElementSelectionHandle);
+	if (!SMInstance)
+	{
+		return false;
+	}
 
 	if (!InElementSelectionHandle.DeselectElement(InSelectionSet, InSelectionOptions))
 	{
 		return false;
 	}
 
-	UE_LOG(LogSMInstanceLevelEditorSelection, Verbose, TEXT("Deselected SMInstance: %s (%s), Index %d"), *SMInstance.ISMComponent->GetPathName(), *SMInstance.ISMComponent->GetClass()->GetName(), SMInstance.InstanceIndex);
+	UE_LOG(LogSMInstanceLevelEditorSelection, Verbose, TEXT("Deselected SMInstance: %s (%s), Index %d"), *SMInstance.GetISMComponent()->GetPathName(), *SMInstance.GetISMComponent()->GetClass()->GetName(), SMInstance.GetISMInstanceIndex());
 
 	// Update the internal selection state for viewport selection rendering
-	SMInstance.ISMComponent->SelectInstance(/*bIsSelected*/false, SMInstance.InstanceIndex);
+	SMInstance.GetISMComponent()->SelectInstance(/*bIsSelected*/false, SMInstance.GetISMInstanceIndex());
 
 	return true;
 }
diff --git a/Engine/Source/Editor/UnrealEd/Private/Elements/SMInstance/SMInstanceElementDetailsProxyObject.cpp b/Engine/Source/Editor/UnrealEd/Private/Elements/SMInstance/SMInstanceElementDetailsProxyObject.cpp
--- a/Engine/Source/Editor/UnrealEd/Private/Elements/SMInstance/SMInstanceElementDetailsProxyObject.cpp
+++ b/Engine/Source/Editor/UnrealEd/Private/Elements/SMInstance/SMInstanceElementDetailsProxyObject.cpp
@@ -3,6 +3,8 @@
 #include "Elements/SMInstance/SMInstanceElementDetailsProxyObject.h"
 #include "Components/InstancedStaticMeshComponent.h"
 
+#include "Elements/SMInstance/SMInstanceElementData.h"
+
 #include "UnrealEdGlobals.h"
 #include "Editor/UnrealEdEngine.h"
 
@@ -33,26 +35,16 @@ void USMInstanceElementDetailsProxyObject::Shutdown()
 	SyncProxyStateFromInstance();
 }
 
-bool USMInstanceElementDetailsProxyObject::Modify(bool bAlwaysMarkDirty)
-{
-	if (FSMInstanceId InstanceId = GetInstanceId())
-	{
-		return InstanceId.ISMComponent->Modify(bAlwaysMarkDirty);
-	}
-
-	return false;
-}
-
 void USMInstanceElementDetailsProxyObject::PostEditChangeChainProperty(FPropertyChangedChainEvent& PropertyChangedEvent)
 {
 	if (PropertyChangedEvent.Property)
 	{
 		if (PropertyChangedEvent.Property->GetFName() == GET_MEMBER_NAME_CHECKED(USMInstanceElementDetailsProxyObject, Transform))
 		{
-			if (FSMInstanceId InstanceId = GetInstanceId())
+			if (FSMInstanceManager SMInstance = GetSMInstance())
 			{
 				// TODO: Need flag for local/world space, like FComponentTransformDetails
-				InstanceId.ISMComponent->UpdateInstanceTransform(InstanceId.InstanceIndex, Transform, /*bWorldSpace*/false, /*bMarkRenderStateDirty*/true);
+				SMInstance.SetSMInstanceTransform(Transform, /*bWorldSpace*/false, /*bMarkRenderStateDirty*/true);
 				
 				GUnrealEd->UpdatePivotLocationForSelection();
 				GUnrealEd->RedrawLevelEditingViewports();
@@ -65,10 +57,10 @@ void USMInstanceElementDetailsProxyObject::PostEditChangeChainProperty(FProperty
 
 void USMInstanceElementDetailsProxyObject::SyncProxyStateFromInstance()
 {
-	if (FSMInstanceId InstanceId = GetInstanceId())
+	if (FSMInstanceManager SMInstance = GetSMInstance())
 	{
 		// TODO: Need flag for local/world space, like FComponentTransformDetails
-		InstanceId.ISMComponent->GetInstanceTransform(InstanceId.InstanceIndex, Transform, /*bWorldSpace*/false);
+		SMInstance.GetSMInstanceTransform(Transform, /*bWorldSpace*/false);
 	}
 	else
 	{
@@ -76,7 +68,8 @@ void USMInstanceElementDetailsProxyObject::SyncProxyStateFromInstance()
 	}
 }
 
-FSMInstanceId USMInstanceElementDetailsProxyObject::GetInstanceId() const
+FSMInstanceManager USMInstanceElementDetailsProxyObject::GetSMInstance() const
 {
-	return FSMInstanceElementIdMap::Get().GetSMInstanceIdFromSMInstanceElementId(FSMInstanceElementId{ ISMComponent.Get(), ISMInstanceId });
+	const FSMInstanceId SMInstanceId = FSMInstanceElementIdMap::Get().GetSMInstanceIdFromSMInstanceElementId(FSMInstanceElementId{ ISMComponent.Get(), ISMInstanceId });
+	return FSMInstanceManager(SMInstanceId, SMInstanceElementDataUtil::GetSMInstanceManager(SMInstanceId));
 }
diff --git a/Engine/Source/Editor/UnrealEd/Private/Elements/SMInstance/SMInstanceElementDetailsProxyObject.h b/Engine/Source/Editor/UnrealEd/Private/Elements/SMInstance/SMInstanceElementDetailsProxyObject.h
--- a/Engine/Source/Editor/UnrealEd/Private/Elements/SMInstance/SMInstanceElementDetailsProxyObject.h
+++ b/Engine/Source/Editor/UnrealEd/Private/Elements/SMInstance/SMInstanceElementDetailsProxyObject.h
@@ -3,7 +3,7 @@
 #pragma once
 
 #include "UObject/Object.h"
-#include "Elements/SMInstance/SMInstanceElementId.h"
+#include "Elements/SMInstance/SMInstanceManager.h"
 #include "SMInstanceElementDetailsProxyObject.generated.h"
 
 UCLASS(Transient)
@@ -16,7 +16,6 @@ public:
 	void Shutdown();
 
 	//~ UObject interface
-	virtual bool Modify(bool bAlwaysMarkDirty = true) override;
 	virtual void PostEditChangeChainProperty(FPropertyChangedChainEvent& PropertyChangedEvent) override;
 
 	UPROPERTY(EditAnywhere, Category="Transform", meta=(ShowOnlyInnerProperties))
@@ -25,7 +24,7 @@ public:
 private:
 	void SyncProxyStateFromInstance();
 
-	FSMInstanceId GetInstanceId() const;
+	FSMInstanceManager GetSMInstance() const;
 
 	TWeakObjectPtr<UInstancedStaticMeshComponent> ISMComponent;
 	uint64 ISMInstanceId = 0;
diff --git a/Engine/Source/Editor/UnrealEd/Private/Elements/SMInstance/SMInstanceElementEditorSelectionInterface.cpp b/Engine/Source/Editor/UnrealEd/Private/Elements/SMInstance/SMInstanceElementEditorSelectionInterface.cpp
--- a/Engine/Source/Editor/UnrealEd/Private/Elements/SMInstance/SMInstanceElementEditorSelectionInterface.cpp
+++ b/Engine/Source/Editor/UnrealEd/Private/Elements/SMInstance/SMInstanceElementEditorSelectionInterface.cpp
@@ -18,9 +18,9 @@ private:
 
 	virtual FTypedElementHandle GetElementImpl() const override
 	{
-		FSMInstanceElementId SMInstanceElementId{ ISMComponentPtr.Get(/*bEvenIfPendingKill*/true), InstanceId };
-		return FSMInstanceElementIdMap::Get().GetSMInstanceIdFromSMInstanceElementId(SMInstanceElementId)
-			? UEngineElementsLibrary::AcquireEditorSMInstanceElementHandle(SMInstanceElementId)
+		FSMInstanceId SMInstanceId = FSMInstanceElementIdMap::Get().GetSMInstanceIdFromSMInstanceElementId(FSMInstanceElementId{ ISMComponentPtr.Get(/*bEvenIfPendingKill*/true), InstanceId });
+		return SMInstanceId
+			? UEngineElementsLibrary::AcquireEditorSMInstanceElementHandle(SMInstanceId)
 			: FTypedElementHandle();
 	}
 
@@ -43,7 +43,7 @@ private:
 
 bool USMInstanceElementEditorSelectionInterface::IsElementSelected(const FTypedElementHandle& InElementHandle, const UTypedElementList* InSelectionSet, const FTypedElementIsSelectedOptions& InSelectionOptions)
 {
-	if (const FSMInstanceId SMInstance = SMInstanceElementDataUtil::GetSMInstanceFromHandle(InElementHandle))
+	if (const FSMInstanceManager SMInstance = SMInstanceElementDataUtil::GetSMInstanceFromHandle(InElementHandle))
 	{
 		if (InSelectionSet->Num() == 0)
 		{
@@ -57,7 +57,7 @@ bool USMInstanceElementEditorSelectionInterface::IsElementSelected(const FTypedE
 
 		if (InSelectionOptions.AllowIndirect())
 		{
-			if (FTypedElementHandle ISMComponentElement = UEngineElementsLibrary::AcquireEditorComponentElementHandle(SMInstance.ISMComponent, /*bAllowCreate*/false))
+			if (FTypedElementHandle ISMComponentElement = UEngineElementsLibrary::AcquireEditorComponentElementHandle(SMInstance.GetISMComponent(), /*bAllowCreate*/false))
 			{
 				return InSelectionSet->Contains(ISMComponentElement);
 			}
@@ -69,8 +69,8 @@ bool USMInstanceElementEditorSelectionInterface::IsElementSelected(const FTypedE
 
 bool USMInstanceElementEditorSelectionInterface::ShouldPreventTransactions(const FTypedElementHandle& InElementHandle)
 {
-	const FSMInstanceId SMInstance = SMInstanceElementDataUtil::GetSMInstanceFromHandle(InElementHandle);
-	return SMInstance && UObjectElementEditorSelectionInterface::ShouldObjectPreventTransactions(SMInstance.ISMComponent);
+	const FSMInstanceManager SMInstance = SMInstanceElementDataUtil::GetSMInstanceFromHandle(InElementHandle);
+	return SMInstance && UObjectElementEditorSelectionInterface::ShouldObjectPreventTransactions(SMInstance.GetISMComponent());
 }
 
 TUniquePtr<ITypedElementTransactedElement> USMInstanceElementEditorSelectionInterface::CreateTransactedElementImpl()
diff --git a/Engine/Source/Editor/UnrealEd/Private/Elements/SMInstance/SMInstanceElementEditorWorldInterface.cpp b/Engine/Source/Editor/UnrealEd/Private/Elements/SMInstance/SMInstanceElementEditorWorldInterface.cpp
--- a/Engine/Source/Editor/UnrealEd/Private/Elements/SMInstance/SMInstanceElementEditorWorldInterface.cpp
+++ b/Engine/Source/Editor/UnrealEd/Private/Elements/SMInstance/SMInstanceElementEditorWorldInterface.cpp
@@ -9,19 +9,18 @@
 
 bool USMInstanceElementEditorWorldInterface::CanDeleteElement(const FTypedElementHandle& InElementHandle)
 {
-	const FSMInstanceId SMInstance = SMInstanceElementDataUtil::GetSMInstanceFromHandle(InElementHandle);
-	return SMInstance && CanEditSMInstance(SMInstance);
+	const FSMInstanceManager SMInstance = SMInstanceElementDataUtil::GetSMInstanceFromHandle(InElementHandle);
+	return SMInstance && SMInstance.CanDeleteSMInstance();
 }
 
 bool USMInstanceElementEditorWorldInterface::DeleteElement(const FTypedElementHandle& InElementHandle, UWorld* InWorld, UTypedElementSelectionSet* InSelectionSet, const FTypedElementDeletionOptions& InDeletionOptions)
 {
-	if (const FSMInstanceId SMInstance = SMInstanceElementDataUtil::GetSMInstanceFromHandle(InElementHandle))
+	if (const FSMInstanceManager SMInstance = SMInstanceElementDataUtil::GetSMInstanceFromHandle(InElementHandle))
 	{
-		if (CanEditSMInstance(SMInstance))
+		if (SMInstance.CanDeleteSMInstance())
 		{
-			SMInstance.ISMComponent->Modify();
 			InSelectionSet->DeselectElement(InElementHandle, FTypedElementSelectionOptions());
-			return SMInstance.ISMComponent->RemoveInstance(SMInstance.InstanceIndex);
+			return SMInstance.DeleteSMInstance();
 		}
 	}
 
@@ -30,33 +29,31 @@ bool USMInstanceElementEditorWorldInterface::DeleteElement(const FTypedElementHa
 
 bool USMInstanceElementEditorWorldInterface::DeleteElements(TArrayView<const FTypedElementHandle> InElementHandles, UWorld* InWorld, UTypedElementSelectionSet* InSelectionSet, const FTypedElementDeletionOptions& InDeletionOptions)
 {
-	const TArray<FSMInstanceId> SMInstancesToDelete = SMInstanceElementDataUtil::GetSMInstancesFromHandles(InElementHandles);
+	const TArray<FSMInstanceManager> SMInstancesToDelete = SMInstanceElementDataUtil::GetSMInstancesFromHandles(InElementHandles);
 
 	if (SMInstancesToDelete.Num() > 0)
 	{
-		// Batch by the ISM component
-		TMap<UInstancedStaticMeshComponent*, TArray<int32>> BatchedInstancesToDelete;
-		for (const FSMInstanceId& SMInstance : SMInstancesToDelete)
+		// Batch by the ISM manager
+		TMap<TScriptInterface<ISMInstanceManager>, TArray<FSMInstanceId>> BatchedInstancesToDelete;
+		for (const FSMInstanceManager& SMInstance : SMInstancesToDelete)
 		{
-			if (CanEditSMInstance(SMInstance))
+			if (SMInstance.CanDeleteSMInstance())
 			{
-				TArray<int32>& InstanceIndices = BatchedInstancesToDelete.FindOrAdd(SMInstance.ISMComponent);
-				InstanceIndices.Add(SMInstance.InstanceIndex);
+				TArray<FSMInstanceId>& InstanceIds = BatchedInstancesToDelete.FindOrAdd(SMInstance.GetInstanceManager());
+				InstanceIds.Add(SMInstance.GetInstanceId());
 			}
 		}
 
 		bool bDidDelete = false;
 
 		FTypedElementListLegacySyncScopedBatch LegacySyncBatch(InSelectionSet->GetElementList());
-		for (TTuple<UInstancedStaticMeshComponent*, TArray<int32>>& BatchedInstancesToDeletePair : BatchedInstancesToDelete)
+		for (TTuple<TScriptInterface<ISMInstanceManager>, TArray<FSMInstanceId>>& BatchedInstancesToDeletePair : BatchedInstancesToDelete)
 		{
-			BatchedInstancesToDeletePair.Key->Modify();
-
-			for (const int32 InstanceIndex : BatchedInstancesToDeletePair.Value)
+			for (const FSMInstanceId& InstanceId : BatchedInstancesToDeletePair.Value)
 			{
-				InSelectionSet->DeselectElement(UEngineElementsLibrary::AcquireEditorSMInstanceElementHandle(FSMInstanceId{ BatchedInstancesToDeletePair.Key, InstanceIndex }), FTypedElementSelectionOptions());
+				InSelectionSet->DeselectElement(UEngineElementsLibrary::AcquireEditorSMInstanceElementHandle(InstanceId), FTypedElementSelectionOptions());
 			}
-			bDidDelete |= BatchedInstancesToDeletePair.Key->RemoveInstances(BatchedInstancesToDeletePair.Value);
+			bDidDelete |= BatchedInstancesToDeletePair.Key->DeleteSMInstances(BatchedInstancesToDeletePair.Value);
 		}
 
 		return bDidDelete;
@@ -67,22 +64,33 @@ bool USMInstanceElementEditorWorldInterface::DeleteElements(TArrayView<const FTy
 
 bool USMInstanceElementEditorWorldInterface::CanDuplicateElement(const FTypedElementHandle& InElementHandle)
 {
-	const FSMInstanceId SMInstance = SMInstanceElementDataUtil::GetSMInstanceFromHandle(InElementHandle);
-	return SMInstance && CanEditSMInstance(SMInstance);
+	const FSMInstanceManager SMInstance = SMInstanceElementDataUtil::GetSMInstanceFromHandle(InElementHandle);
+	return SMInstance && SMInstance.CanDuplicateSMInstance();
 }
 
 FTypedElementHandle USMInstanceElementEditorWorldInterface::DuplicateElement(const FTypedElementHandle& InElementHandle, UWorld* InWorld, const FVector& InLocationOffset)
 {
-	if (const FSMInstanceId SMInstance = SMInstanceElementDataUtil::GetSMInstanceFromHandle(InElementHandle))
+	if (const FSMInstanceManager SMInstance = SMInstanceElementDataUtil::GetSMInstanceFromHandle(InElementHandle))
 	{
-		if (CanEditSMInstance(SMInstance))
+		if (SMInstance.CanDuplicateSMInstance())
 		{
-			FTransform NewInstanceTransform = FTransform::Identity;
-			SMInstance.ISMComponent->GetInstanceTransform(SMInstance.InstanceIndex, NewInstanceTransform);
-			NewInstanceTransform.SetTranslation(NewInstanceTransform.GetTranslation() + InLocationOffset);
+			FSMInstanceId NewInstanceId;
+			if (SMInstance.DuplicateSMInstance(NewInstanceId))
+			{
+				const bool bOffsetIsZero = InLocationOffset.IsZero();
+
+				if (!bOffsetIsZero)
+				{
+					const TScriptInterface<ISMInstanceManager>& InstanceManager = SMInstance.GetInstanceManager();
+
+					FTransform NewInstanceTransform = FTransform::Identity;
+					InstanceManager->GetSMInstanceTransform(NewInstanceId, NewInstanceTransform, /*bWorldSpace*/false);
+					NewInstanceTransform.SetTranslation(NewInstanceTransform.GetTranslation() + InLocationOffset);
+					InstanceManager->SetSMInstanceTransform(NewInstanceId, NewInstanceTransform, /*bWorldSpace*/false, /*bMarkRenderStateDirty*/true);
+				}
 
-			const int32 NewInstanceIndex = SMInstance.ISMComponent->AddInstance(NewInstanceTransform);
-			return UEngineElementsLibrary::AcquireEditorSMInstanceElementHandle(FSMInstanceId{ SMInstance.ISMComponent, NewInstanceIndex });
+				return UEngineElementsLibrary::AcquireEditorSMInstanceElementHandle(NewInstanceId);
+			}
 		}
 	}
 
@@ -91,37 +99,41 @@ FTypedElementHandle USMInstanceElementEditorWorldInterface::DuplicateElement(con
 
 void USMInstanceElementEditorWorldInterface::DuplicateElements(TArrayView<const FTypedElementHandle> InElementHandles, UWorld* InWorld, const FVector& InLocationOffset, TArray<FTypedElementHandle>& OutNewElements)
 {
-	const TArray<FSMInstanceId> SMInstancesToDuplicate = SMInstanceElementDataUtil::GetSMInstancesFromHandles(InElementHandles);
+	const TArray<FSMInstanceManager> SMInstancesToDuplicate = SMInstanceElementDataUtil::GetSMInstancesFromHandles(InElementHandles);
 
 	if (SMInstancesToDuplicate.Num() > 0)
 	{
-		// Batch by the ISM component
-		TMap<UInstancedStaticMeshComponent*, TArray<int32>> BatchedInstancesToDuplicate;
-		for (const FSMInstanceId& SMInstance : SMInstancesToDuplicate)
+		// Batch by the ISM manager
+		TMap<TScriptInterface<ISMInstanceManager>, TArray<FSMInstanceId>> BatchedInstancesToDuplicate;
+		for (const FSMInstanceManager& SMInstance : SMInstancesToDuplicate)
 		{
-			if (CanEditSMInstance(SMInstance))
+			if (SMInstance.CanDuplicateSMInstance())
 			{
-				TArray<int32>& InstanceIndices = BatchedInstancesToDuplicate.FindOrAdd(SMInstance.ISMComponent);
-				InstanceIndices.Add(SMInstance.InstanceIndex);
+				TArray<FSMInstanceId>& InstanceIds = BatchedInstancesToDuplicate.FindOrAdd(SMInstance.GetInstanceManager());
+				InstanceIds.Add(SMInstance.GetInstanceId());
 			}
 		}
 
-		for (const TTuple<UInstancedStaticMeshComponent*, TArray<int32>>& BatchedInstancesToDuplicatePair : BatchedInstancesToDuplicate)
+		for (const TTuple<TScriptInterface<ISMInstanceManager>, TArray<FSMInstanceId>>& BatchedInstancesToDuplicatePair : BatchedInstancesToDuplicate)
 		{
-			TArray<FTransform> NewInstanceTransforms;
-			NewInstanceTransforms.Reserve(BatchedInstancesToDuplicatePair.Value.Num());
-			for (const int32 InstanceIndex : BatchedInstancesToDuplicatePair.Value)
-			{
-				FTransform& NewInstanceTransform = NewInstanceTransforms.Add_GetRef(FTransform::Identity);
-				BatchedInstancesToDuplicatePair.Key->GetInstanceTransform(InstanceIndex, NewInstanceTransform);
-				NewInstanceTransform.SetTranslation(NewInstanceTransform.GetTranslation() + InLocationOffset);
-			}
-
-			const TArray<int32> NewInstanceIndices = BatchedInstancesToDuplicatePair.Key->AddInstances(NewInstanceTransforms, /*bShouldReturnIndices*/true);
-			OutNewElements.Reserve(OutNewElements.Num() + NewInstanceIndices.Num());
-			for (const int32 NewInstanceIndex : NewInstanceIndices)
+			TArray<FSMInstanceId> NewInstanceIds;
+			if (BatchedInstancesToDuplicatePair.Key->DuplicateSMInstances(BatchedInstancesToDuplicatePair.Value, NewInstanceIds))
 			{
-				OutNewElements.Add(UEngineElementsLibrary::AcquireEditorSMInstanceElementHandle(FSMInstanceId{ BatchedInstancesToDuplicatePair.Key, NewInstanceIndex }));
+				const bool bOffsetIsZero = InLocationOffset.IsZero();
+
+				OutNewElements.Reserve(OutNewElements.Num() + NewInstanceIds.Num());
+				for (const FSMInstanceId& NewInstanceId : BatchedInstancesToDuplicatePair.Value)
+				{
+					if (!bOffsetIsZero)
+					{
+						FTransform NewInstanceTransform = FTransform::Identity;
+						BatchedInstancesToDuplicatePair.Key->GetSMInstanceTransform(NewInstanceId, NewInstanceTransform, /*bWorldSpace*/false);
+						NewInstanceTransform.SetTranslation(NewInstanceTransform.GetTranslation() + InLocationOffset);
+						BatchedInstancesToDuplicatePair.Key->SetSMInstanceTransform(NewInstanceId, NewInstanceTransform, /*bWorldSpace*/false, /*bMarkRenderStateDirty*/true);
+					}
+
+					OutNewElements.Add(UEngineElementsLibrary::AcquireEditorSMInstanceElementHandle(NewInstanceId));
+				}
 			}
 		}
 	}
diff --git a/Engine/Source/Runtime/Engine/Classes/Components/InstancedStaticMeshComponent.h b/Engine/Source/Runtime/Engine/Classes/Components/InstancedStaticMeshComponent.h
--- a/Engine/Source/Runtime/Engine/Classes/Components/InstancedStaticMeshComponent.h
+++ b/Engine/Source/Runtime/Engine/Classes/Components/InstancedStaticMeshComponent.h
@@ -10,7 +10,7 @@
 #include "Misc/Guid.h"
 #include "Engine/TextureStreamingTypes.h"
 #include "Components/StaticMeshComponent.h"
-
+#include "Elements/SMInstance/SMInstanceManager.h"
 #include "InstancedStaticMeshComponent.generated.h"
 
 class FLightingBuildOptions;
@@ -121,7 +121,7 @@ struct FInstancedStaticMeshMappingInfo
 
 /** A component that efficiently renders multiple instances of the same StaticMesh. */
 UCLASS(ClassGroup = Rendering, meta = (BlueprintSpawnableComponent), Blueprintable)
-class ENGINE_API UInstancedStaticMeshComponent : public UStaticMeshComponent
+class ENGINE_API UInstancedStaticMeshComponent : public UStaticMeshComponent, public ISMInstanceManager
 {
 	GENERATED_UCLASS_BODY()
 
@@ -454,6 +454,13 @@ protected:
 	/** Creates rendering buffer from serialized data, if any */
 	virtual void OnPostLoadPerInstanceData();
 
+	//~ ISMInstanceManager interface
+	virtual bool CanEditSMInstance(const FSMInstanceId& InstanceId) const override;
+	virtual bool GetSMInstanceTransform(const FSMInstanceId& InstanceId, FTransform& OutInstanceTransform, bool bWorldSpace = false) const override;
+	virtual bool SetSMInstanceTransform(const FSMInstanceId& InstanceId, const FTransform& InstanceTransform, bool bWorldSpace = false, bool bMarkRenderStateDirty = false, bool bTeleport = false) override;
+	virtual bool DeleteSMInstances(TArrayView<const FSMInstanceId> InstanceIds) override;
+	virtual bool DuplicateSMInstances(TArrayView<const FSMInstanceId> InstanceIds, TArray<FSMInstanceId>& OutNewInstanceIds) override;
+
 	friend FStaticLightingTextureMapping_InstancedStaticMesh;
 	friend FInstancedLightMap2D;
 	friend FInstancedShadowMap2D;
diff --git a/Engine/Source/Runtime/Engine/Classes/GameFramework/LightWeightInstanceStaticMeshManager.h b/Engine/Source/Runtime/Engine/Classes/GameFramework/LightWeightInstanceStaticMeshManager.h
--- a/Engine/Source/Runtime/Engine/Classes/GameFramework/LightWeightInstanceStaticMeshManager.h
+++ b/Engine/Source/Runtime/Engine/Classes/GameFramework/LightWeightInstanceStaticMeshManager.h
@@ -5,13 +5,14 @@
 #include "CoreMinimal.h"
 #include "Engine/StaticMesh.h"
 #include "GameFramework/LightWeightInstanceManager.h"
+#include "Elements/SMInstance/SMInstanceManager.h"
 
 #include "LightWeightInstanceStaticMeshManager.generated.h"
 
 DECLARE_DYNAMIC_DELEGATE_RetVal_OneParam(FActorInstanceHandle, FOnActorReady, FActorInstanceHandle, InHandle);
 
 UCLASS(BlueprintType, Blueprintable)
-class ENGINE_API ALightWeightInstanceStaticMeshManager : public ALightWeightInstanceManager
+class ENGINE_API ALightWeightInstanceStaticMeshManager : public ALightWeightInstanceManager, public ISMInstanceManager
 {
 	GENERATED_UCLASS_BODY()
 
@@ -46,6 +47,18 @@ public:
 
 	virtual void GetLifetimeReplicatedProps(TArray<FLifetimeProperty>& OutLifetimeProps) const override;
 
+protected:
+
+	//~ ISMInstanceManager interface
+	virtual bool CanEditSMInstance(const FSMInstanceId& InstanceId) const override;
+	virtual bool GetSMInstanceTransform(const FSMInstanceId& InstanceId, FTransform& OutInstanceTransform, bool bWorldSpace = false) const override;
+	virtual bool SetSMInstanceTransform(const FSMInstanceId& InstanceId, const FTransform& InstanceTransform, bool bWorldSpace = false, bool bMarkRenderStateDirty = false, bool bTeleport = false) override;
+	virtual bool DeleteSMInstances(TArrayView<const FSMInstanceId> InstanceIds) override;
+	virtual bool DuplicateSMInstances(TArrayView<const FSMInstanceId> InstanceIds, TArray<FSMInstanceId>& OutNewInstanceIds) override;
+
+	virtual void DuplicateLWIInstances(TArrayView<const int32> DataIndices, TArray<int32>& OutNewDataIndices);
+	void GetLWIDataIndices(TArrayView<const FSMInstanceId> InstanceIds, TArray<int32>& OutDataIndices) const;
+
 protected:
 
 	UPROPERTY(ReplicatedUsing = OnRep_StaticMesh)
diff --git a/Engine/Source/Runtime/Engine/Private/Elements/Framework/EngineElementsLibrary.cpp b/Engine/Source/Runtime/Engine/Private/Elements/Framework/EngineElementsLibrary.cpp
--- a/Engine/Source/Runtime/Engine/Private/Elements/Framework/EngineElementsLibrary.cpp
+++ b/Engine/Source/Runtime/Engine/Private/Elements/Framework/EngineElementsLibrary.cpp
@@ -378,15 +378,17 @@ TTypedElementOwner<FSMInstanceElementData> UEngineElementsLibrary::CreateSMInsta
 {
 	const FSMInstanceElementId SMInstanceElementId = FSMInstanceElementIdMap::Get().GetSMInstanceElementIdFromSMInstanceId(InSMInstanceId);
 	checkf(SMInstanceElementId, TEXT("Static Mesh Instance Index failed to map to a valid Static Mesh Instance Element ID!"));
-	return CreateSMInstanceElement(SMInstanceElementId);
+	return CreateSMInstanceElementImpl(InSMInstanceId, SMInstanceElementId);
 }
 
-TTypedElementOwner<FSMInstanceElementData> UEngineElementsLibrary::CreateSMInstanceElement(const FSMInstanceElementId& InSMInstanceElementId)
+TTypedElementOwner<FSMInstanceElementData> UEngineElementsLibrary::CreateSMInstanceElementImpl(const FSMInstanceId& InSMInstanceId, const FSMInstanceElementId& InSMInstanceElementId)
 {
+	checkf(InSMInstanceId.ISMComponent == InSMInstanceElementId.ISMComponent, TEXT("Component mismatch in IDs!"));
+
 	UTypedElementRegistry* Registry = UTypedElementRegistry::GetInstance();
 
 	TTypedElementOwner<FSMInstanceElementData> SMInstanceElement;
-	if (ensureAlways(Registry) && SMInstanceElementDataUtil::SMInstanceElementsEnabled() && SMInstanceElementDataUtil::IsValidComponentForSMInstanceElements(InSMInstanceElementId.ISMComponent))
+	if (ensureAlways(Registry) && SMInstanceElementDataUtil::SMInstanceElementsEnabled() && SMInstanceElementDataUtil::GetSMInstanceManager(InSMInstanceId))
 	{
 		SMInstanceElement = Registry->CreateElement<FSMInstanceElementData>(NAME_SMInstance);
 		if (SMInstanceElement)
@@ -414,7 +416,7 @@ void UEngineElementsLibrary::CreateEditorSMInstanceElement(const FSMInstanceId&
 	{
 		const FSMInstanceElementId SMInstanceElementId = FSMInstanceElementIdMap::Get().GetSMInstanceElementIdFromSMInstanceId(SMInstanceId);
 		checkf(SMInstanceElementId, TEXT("Static Mesh Instance Index failed to map to a valid Static Mesh Instance Element ID!"));
-		GSMInstanceElementOwnerStore.RegisterElementOwner(SMInstanceElementId, CreateSMInstanceElement(SMInstanceElementId));
+		GSMInstanceElementOwnerStore.RegisterElementOwner(SMInstanceElementId, CreateSMInstanceElementImpl(SMInstanceId, SMInstanceElementId));
 	}
 }
 
@@ -437,18 +439,9 @@ FTypedElementHandle UEngineElementsLibrary::AcquireEditorSMInstanceElementHandle
 	{
 		const FSMInstanceElementId SMInstanceElementId = FSMInstanceElementIdMap::Get().GetSMInstanceElementIdFromSMInstanceId(SMInstanceId, bAllowCreate);
 		checkf(!bAllowCreate || SMInstanceElementId, TEXT("Static Mesh Instance Index failed to map to a valid Static Mesh Instance Element ID!"));
-		return AcquireEditorSMInstanceElementHandle(SMInstanceElementId, bAllowCreate);
-	}
-
-	return FTypedElementHandle();
-}
-
-FTypedElementHandle UEngineElementsLibrary::AcquireEditorSMInstanceElementHandle(const FSMInstanceElementId& SMInstanceElementId, const bool bAllowCreate)
-{
-	if (GIsEditor)
-	{
+		
 		TTypedElementOwnerScopedAccess<FSMInstanceElementData> SMInstanceElement = bAllowCreate
-			? GSMInstanceElementOwnerStore.FindOrRegisterElementOwner(SMInstanceElementId, [&SMInstanceElementId]() { return CreateSMInstanceElement(SMInstanceElementId); })
+			? GSMInstanceElementOwnerStore.FindOrRegisterElementOwner(SMInstanceElementId, [&SMInstanceId, &SMInstanceElementId]() { return CreateSMInstanceElementImpl(SMInstanceId, SMInstanceElementId); })
 			: GSMInstanceElementOwnerStore.FindElementOwner(SMInstanceElementId);
 
 		if (SMInstanceElement)
@@ -460,4 +453,3 @@ FTypedElementHandle UEngineElementsLibrary::AcquireEditorSMInstanceElementHandle
 	return FTypedElementHandle();
 }
 #endif
-
diff --git a/Engine/Source/Runtime/Engine/Private/Elements/SMInstance/SMInstanceElementAssetDataInterface.cpp b/Engine/Source/Runtime/Engine/Private/Elements/SMInstance/SMInstanceElementAssetDataInterface.cpp
--- a/Engine/Source/Runtime/Engine/Private/Elements/SMInstance/SMInstanceElementAssetDataInterface.cpp
+++ b/Engine/Source/Runtime/Engine/Private/Elements/SMInstance/SMInstanceElementAssetDataInterface.cpp
@@ -8,11 +8,11 @@
 
 FAssetData USMInstanceElementAssetDataInterface::GetAssetData(const FTypedElementHandle& InElementHandle)
 {
-	if (FSMInstanceId SMInstance = SMInstanceElementDataUtil::GetSMInstanceFromHandle(InElementHandle))
+	if (FSMInstanceManager SMInstance = SMInstanceElementDataUtil::GetSMInstanceFromHandle(InElementHandle))
 	{
-		if (SMInstance.ISMComponent->GetStaticMesh())
+		if (UStaticMesh* StaticMesh = SMInstance.GetISMComponent()->GetStaticMesh())
 		{
-			return FAssetData(SMInstance.ISMComponent->GetStaticMesh());
+			return FAssetData(StaticMesh);
 		}
 	}
 
diff --git a/Engine/Source/Runtime/Engine/Private/Elements/SMInstance/SMInstanceElementData.cpp b/Engine/Source/Runtime/Engine/Private/Elements/SMInstance/SMInstanceElementData.cpp
--- a/Engine/Source/Runtime/Engine/Private/Elements/SMInstance/SMInstanceElementData.cpp
+++ b/Engine/Source/Runtime/Engine/Private/Elements/SMInstance/SMInstanceElementData.cpp
@@ -22,43 +22,44 @@ bool SMInstanceElementsEnabled()
 	return GEnableSMInstanceElements != 0;
 }
 
-bool IsValidComponentForSMInstanceElements(const UInstancedStaticMeshComponent* InComponent)
+TScriptInterface<ISMInstanceManager> GetSMInstanceManager(const FSMInstanceId& InstanceId)
 {
-	if (!InComponent)
+	if (!InstanceId)
 	{
-		return false;
+		return nullptr;
 	}
 
-	if (const AActor* OwnerActor = InComponent->GetOwner())
+	if (AActor* OwnerActor = InstanceId.ISMComponent->GetOwner())
 	{
-		// Foliage actors have extra bookkeeping data which isn't correctly updated by static mesh instance elements
-		// Disable being able to create static mesh instance elements for foliage actors until this is resolved...
-		// Note: This test is by name as we cannot link directly to AInstancedFoliageActor
-		static const FName NAME_InstancedFoliageActor = "InstancedFoliageActor";
-		for (const UClass* ActorClass = OwnerActor->GetClass(); ActorClass; ActorClass = ActorClass->GetSuperClass())
+		// If the owner actor is an instance manager provider, then just ask that for the instance manager
+		if (ISMInstanceManagerProvider* InstanceManagerProvider = Cast<ISMInstanceManagerProvider>(OwnerActor))
 		{
-			if (ActorClass->GetFName() == NAME_InstancedFoliageActor)
-			{
-				return false;
-			}
+			return InstanceManagerProvider->GetSMInstanceManager(InstanceId);
+		}
+
+		// If the owner actor is an instance manager, then just use that
+		if (ISMInstanceManager* InstanceManager = Cast<ISMInstanceManager>(OwnerActor))
+		{
+			return OwnerActor;
 		}
 	}
 
-	return true;
+	// Otherwise, allow the ISM component to manage itself
+	return InstanceId.ISMComponent;
 }
 
-FSMInstanceId GetSMInstanceFromHandle(const FTypedElementHandle& InHandle, const bool bSilent)
+FSMInstanceManager GetSMInstanceFromHandle(const FTypedElementHandle& InHandle, const bool bSilent)
 {
 	const FSMInstanceElementData* SMInstanceElement = InHandle.GetData<FSMInstanceElementData>(bSilent);
-	return SMInstanceElement ? FSMInstanceElementIdMap::Get().GetSMInstanceIdFromSMInstanceElementId(SMInstanceElement->InstanceElementId) : FSMInstanceId();
+	const FSMInstanceId SMInstanceId = SMInstanceElement ? FSMInstanceElementIdMap::Get().GetSMInstanceIdFromSMInstanceElementId(SMInstanceElement->InstanceElementId) : FSMInstanceId();
+	return FSMInstanceManager(SMInstanceId, GetSMInstanceManager(SMInstanceId));
 }
 
-FSMInstanceId GetSMInstanceFromHandleChecked(const FTypedElementHandle& InHandle)
+FSMInstanceManager GetSMInstanceFromHandleChecked(const FTypedElementHandle& InHandle)
 {
 	const FSMInstanceElementData& SMInstanceElement = InHandle.GetDataChecked<FSMInstanceElementData>();
-	FSMInstanceId SMInstanceId = FSMInstanceElementIdMap::Get().GetSMInstanceIdFromSMInstanceElementId(SMInstanceElement.InstanceElementId);
-	checkf(SMInstanceId, TEXT("Static Mesh Instance Element ID failed to map to a valid Static Mesh Instance Index!"));
-	return SMInstanceId;
+	const FSMInstanceId SMInstanceId = FSMInstanceElementIdMap::Get().GetSMInstanceIdFromSMInstanceElementId(SMInstanceElement.InstanceElementId);
+	return FSMInstanceManager(SMInstanceId, GetSMInstanceManager(SMInstanceId));
 }
 
 } // namespace SMInstanceElementDataUtil
diff --git a/Engine/Source/Runtime/Engine/Private/Elements/SMInstance/SMInstanceElementWorldInterface.cpp b/Engine/Source/Runtime/Engine/Private/Elements/SMInstance/SMInstanceElementWorldInterface.cpp
--- a/Engine/Source/Runtime/Engine/Private/Elements/SMInstance/SMInstanceElementWorldInterface.cpp
+++ b/Engine/Source/Runtime/Engine/Private/Elements/SMInstance/SMInstanceElementWorldInterface.cpp
@@ -11,21 +11,21 @@
 
 bool USMInstanceElementWorldInterface::CanEditElement(const FTypedElementHandle& InElementHandle)
 {
-	const FSMInstanceId SMInstance = SMInstanceElementDataUtil::GetSMInstanceFromHandle(InElementHandle);
-	return SMInstance && CanEditSMInstance(SMInstance);
+	const FSMInstanceManager SMInstance = SMInstanceElementDataUtil::GetSMInstanceFromHandle(InElementHandle);
+	return SMInstance && SMInstance.CanEditSMInstance();
 }
 
 bool USMInstanceElementWorldInterface::IsTemplateElement(const FTypedElementHandle& InElementHandle)
 {
-	const FSMInstanceId SMInstance = SMInstanceElementDataUtil::GetSMInstanceFromHandle(InElementHandle);
-	return SMInstance && SMInstance.ISMComponent->IsTemplate();
+	const FSMInstanceManager SMInstance = SMInstanceElementDataUtil::GetSMInstanceFromHandle(InElementHandle);
+	return SMInstance && SMInstance.GetISMComponent()->IsTemplate();
 }
 
 ULevel* USMInstanceElementWorldInterface::GetOwnerLevel(const FTypedElementHandle& InElementHandle)
 {
-	if (const FSMInstanceId SMInstance = SMInstanceElementDataUtil::GetSMInstanceFromHandle(InElementHandle))
+	if (const FSMInstanceManager SMInstance = SMInstanceElementDataUtil::GetSMInstanceFromHandle(InElementHandle))
 	{
-		if (const AActor* ComponentOwner = SMInstance.ISMComponent->GetOwner())
+		if (const AActor* ComponentOwner = SMInstance.GetISMComponent()->GetOwner())
 		{
 			return ComponentOwner->GetLevel();
 		}
@@ -36,15 +36,15 @@ ULevel* USMInstanceElementWorldInterface::GetOwnerLevel(const FTypedElementHandl
 
 UWorld* USMInstanceElementWorldInterface::GetOwnerWorld(const FTypedElementHandle& InElementHandle)
 {
-	const FSMInstanceId SMInstance = SMInstanceElementDataUtil::GetSMInstanceFromHandle(InElementHandle);
-	return SMInstance ? SMInstance.ISMComponent->GetWorld() : nullptr;
+	const FSMInstanceManager SMInstance = SMInstanceElementDataUtil::GetSMInstanceFromHandle(InElementHandle);
+	return SMInstance ? SMInstance.GetISMComponent()->GetWorld() : nullptr;
 }
 
 bool USMInstanceElementWorldInterface::GetBounds(const FTypedElementHandle& InElementHandle, FBoxSphereBounds& OutBounds)
 {
-	if (const FSMInstanceId SMInstance = SMInstanceElementDataUtil::GetSMInstanceFromHandle(InElementHandle))
+	if (const FSMInstanceManager SMInstance = SMInstanceElementDataUtil::GetSMInstanceFromHandle(InElementHandle))
 	{
-		if (const UStaticMesh* StaticMesh = SMInstance.ISMComponent->GetStaticMesh())
+		if (const UStaticMesh* StaticMesh = SMInstance.GetISMComponent()->GetStaticMesh())
 		{
 			OutBounds = StaticMesh->GetBounds();
 		}
@@ -54,7 +54,7 @@ bool USMInstanceElementWorldInterface::GetBounds(const FTypedElementHandle& InEl
 		}
 
 		FTransform InstanceTransform;
-		SMInstance.ISMComponent->GetInstanceTransform(SMInstance.InstanceIndex, InstanceTransform, /*bWorldSpace*/true);
+		SMInstance.GetSMInstanceTransform(InstanceTransform, /*bWorldSpace*/true);
 
 		OutBounds = OutBounds.TransformBy(InstanceTransform);
 		return true;
@@ -65,29 +65,24 @@ bool USMInstanceElementWorldInterface::GetBounds(const FTypedElementHandle& InEl
 
 bool USMInstanceElementWorldInterface::GetWorldTransform(const FTypedElementHandle& InElementHandle, FTransform& OutTransform)
 {
-	const FSMInstanceId SMInstance = SMInstanceElementDataUtil::GetSMInstanceFromHandle(InElementHandle);
-	return SMInstance && SMInstance.ISMComponent->GetInstanceTransform(SMInstance.InstanceIndex, OutTransform, /*bWorldSpace*/true);
+	const FSMInstanceManager SMInstance = SMInstanceElementDataUtil::GetSMInstanceFromHandle(InElementHandle);
+	return SMInstance && SMInstance.GetSMInstanceTransform(OutTransform, /*bWorldSpace*/true);
 }
 
 bool USMInstanceElementWorldInterface::SetWorldTransform(const FTypedElementHandle& InElementHandle, const FTransform& InTransform)
 {
-	const FSMInstanceId SMInstance = SMInstanceElementDataUtil::GetSMInstanceFromHandle(InElementHandle);
-	return SMInstance && SMInstance.ISMComponent->UpdateInstanceTransform(SMInstance.InstanceIndex, InTransform, /*bWorldSpace*/true, /*bMarkRenderStateDirty*/true);
+	const FSMInstanceManager SMInstance = SMInstanceElementDataUtil::GetSMInstanceFromHandle(InElementHandle);
+	return SMInstance && SMInstance.SetSMInstanceTransform(InTransform, /*bWorldSpace*/true, /*bMarkRenderStateDirty*/true);
 }
 
 bool USMInstanceElementWorldInterface::GetRelativeTransform(const FTypedElementHandle& InElementHandle, FTransform& OutTransform)
 {
-	const FSMInstanceId SMInstance = SMInstanceElementDataUtil::GetSMInstanceFromHandle(InElementHandle);
-	return SMInstance && SMInstance.ISMComponent->GetInstanceTransform(SMInstance.InstanceIndex, OutTransform, /*bWorldSpace*/false);
+	const FSMInstanceManager SMInstance = SMInstanceElementDataUtil::GetSMInstanceFromHandle(InElementHandle);
+	return SMInstance && SMInstance.GetSMInstanceTransform(OutTransform, /*bWorldSpace*/false);
 }
 
 bool USMInstanceElementWorldInterface::SetRelativeTransform(const FTypedElementHandle& InElementHandle, const FTransform& InTransform)
 {
-	const FSMInstanceId SMInstance = SMInstanceElementDataUtil::GetSMInstanceFromHandle(InElementHandle);
-	return SMInstance && SMInstance.ISMComponent->UpdateInstanceTransform(SMInstance.InstanceIndex, InTransform, /*bWorldSpace*/false, /*bMarkRenderStateDirty*/true);
-}
-
-bool USMInstanceElementWorldInterface::CanEditSMInstance(const FSMInstanceId& InSMInstanceId)
-{
-	return InSMInstanceId.ISMComponent->IsEditableWhenInherited();
+	const FSMInstanceManager SMInstance = SMInstanceElementDataUtil::GetSMInstanceFromHandle(InElementHandle);
+	return SMInstance && SMInstance.SetSMInstanceTransform(InTransform, /*bWorldSpace*/false, /*bMarkRenderStateDirty*/true);
 }
diff --git a/Engine/Source/Runtime/Engine/Private/ISMPartition/ISMPartitionActor.cpp b/Engine/Source/Runtime/Engine/Private/ISMPartition/ISMPartitionActor.cpp
--- a/Engine/Source/Runtime/Engine/Private/ISMPartition/ISMPartitionActor.cpp
+++ b/Engine/Source/Runtime/Engine/Private/ISMPartition/ISMPartitionActor.cpp
@@ -650,3 +650,9 @@ void AISMPartitionActor::GetClientComponents(const FISMClientHandle& Handle, TAr
 }
 
 #endif
+
+TScriptInterface<ISMInstanceManager> AISMPartitionActor::GetSMInstanceManager(const FSMInstanceId& InstanceId) const
+{
+	// Disable per-instance editing for now...
+	return nullptr;
+}
diff --git a/Engine/Source/Runtime/Engine/Private/InstancedStaticMesh.cpp b/Engine/Source/Runtime/Engine/Private/InstancedStaticMesh.cpp
--- a/Engine/Source/Runtime/Engine/Private/InstancedStaticMesh.cpp
+++ b/Engine/Source/Runtime/Engine/Private/InstancedStaticMesh.cpp
@@ -158,7 +158,15 @@ FTypedElementHandle HInstancedStaticMeshInstance::GetElementHandle() const
 #if WITH_EDITOR
 	if (Component)
 	{
-		return UEngineElementsLibrary::AcquireEditorSMInstanceElementHandle(FSMInstanceId{ Component, InstanceIndex });
+		// Prefer per-instance selection if available
+		// This may fail to return a handle if the feature is disabled, or if per-instance editing is disabled for this component
+		if (FTypedElementHandle ElementHandle = UEngineElementsLibrary::AcquireEditorSMInstanceElementHandle(Component, InstanceIndex))
+		{
+			return ElementHandle;
+		}
+
+		// If per-instance selection isn't possible, fallback to general per-component selection (which may choose to select the owner actor instead)
+		return UEngineElementsLibrary::AcquireEditorComponentElementHandle(Component);
 	}
 #endif	// WITH_EDITOR
 	return FTypedElementHandle();
@@ -3920,6 +3928,61 @@ void UInstancedStaticMeshComponent::ClearInstanceSelection()
 #endif
 }
 
+bool UInstancedStaticMeshComponent::CanEditSMInstance(const FSMInstanceId& InstanceId) const
+{
+	check(InstanceId.ISMComponent == this);
+	return IsEditableWhenInherited();
+}
+
+bool UInstancedStaticMeshComponent::GetSMInstanceTransform(const FSMInstanceId& InstanceId, FTransform& OutInstanceTransform, bool bWorldSpace) const
+{
+	check(InstanceId.ISMComponent == this);
+	return GetInstanceTransform(InstanceId.InstanceIndex, OutInstanceTransform, bWorldSpace);
+}
+
+bool UInstancedStaticMeshComponent::SetSMInstanceTransform(const FSMInstanceId& InstanceId, const FTransform& InstanceTransform, bool bWorldSpace, bool bMarkRenderStateDirty, bool bTeleport)
+{
+	check(InstanceId.ISMComponent == this);
+	return UpdateInstanceTransform(InstanceId.InstanceIndex, InstanceTransform, bWorldSpace, bMarkRenderStateDirty, bTeleport);
+}
+
+bool UInstancedStaticMeshComponent::DeleteSMInstances(TArrayView<const FSMInstanceId> InstanceIds)
+{
+	TArray<int32> InstanceIndices;
+	InstanceIndices.Reserve(InstanceIds.Num());
+	for (const FSMInstanceId& InstanceId : InstanceIds)
+	{
+		check(InstanceId.ISMComponent == this);
+		InstanceIndices.Add(InstanceId.InstanceIndex);
+	}
+
+	Modify();
+	return RemoveInstances(InstanceIndices);
+}
+
+bool UInstancedStaticMeshComponent::DuplicateSMInstances(TArrayView<const FSMInstanceId> InstanceIds, TArray<FSMInstanceId>& OutNewInstanceIds)
+{
+	TArray<FTransform> NewInstanceTransforms;
+	NewInstanceTransforms.Reserve(InstanceIds.Num());
+	for (const FSMInstanceId& InstanceId : InstanceIds)
+	{
+		check(InstanceId.ISMComponent == this);
+		FTransform& NewInstanceTransform = NewInstanceTransforms.Add_GetRef(FTransform::Identity);
+		GetInstanceTransform(InstanceId.InstanceIndex, NewInstanceTransform);
+	}
+
+	Modify();
+	const TArray<int32> NewInstanceIndices = AddInstances(NewInstanceTransforms, /*bShouldReturnIndices*/true);
+
+	OutNewInstanceIds.Reset(NewInstanceIndices.Num());
+	for (int32 NewInstanceIndex : NewInstanceIndices)
+	{
+		OutNewInstanceIds.Add(FSMInstanceId{ this, NewInstanceIndex });
+	}
+
+	return true;
+}
+
 static TAutoConsoleVariable<int32> CVarCullAllInVertexShader(
 	TEXT("foliage.CullAllInVertexShader"),
 	0,
diff --git a/Engine/Source/Runtime/Engine/Private/LightWeightInstanceStaticMeshManager.cpp b/Engine/Source/Runtime/Engine/Private/LightWeightInstanceStaticMeshManager.cpp
--- a/Engine/Source/Runtime/Engine/Private/LightWeightInstanceStaticMeshManager.cpp
+++ b/Engine/Source/Runtime/Engine/Private/LightWeightInstanceStaticMeshManager.cpp
@@ -194,4 +194,108 @@ void ALightWeightInstanceStaticMeshManager::ClearStaticMesh()
 {
 	StaticMesh = nullptr;
 	OnStaticMeshSet();
-}
\ No newline at end of file
+}
+
+bool ALightWeightInstanceStaticMeshManager::CanEditSMInstance(const FSMInstanceId& InstanceId) const
+{
+	check(InstanceId.ISMComponent == InstancedStaticMeshComponent);
+
+	ISMInstanceManager* InstanceManager = InstancedStaticMeshComponent;
+	return InstanceManager->CanEditSMInstance(InstanceId);
+}
+
+bool ALightWeightInstanceStaticMeshManager::GetSMInstanceTransform(const FSMInstanceId& InstanceId, FTransform& OutInstanceTransform, bool bWorldSpace) const
+{
+	check(InstanceId.ISMComponent == InstancedStaticMeshComponent);
+
+	if (RenderingIndicesToDataIndices.IsValidIndex(InstanceId.InstanceIndex))
+	{
+		const int32 DataIndex = RenderingIndicesToDataIndices[InstanceId.InstanceIndex];
+		OutInstanceTransform = bWorldSpace ? InstanceTransforms[DataIndex] : InstanceTransforms[DataIndex].GetRelativeTransform(InstancedStaticMeshComponent->GetComponentTransform());
+		return true;
+	}
+	
+	return false;
+}
+
+bool ALightWeightInstanceStaticMeshManager::SetSMInstanceTransform(const FSMInstanceId& InstanceId, const FTransform& InstanceTransform, bool bWorldSpace, bool bMarkRenderStateDirty, bool bTeleport)
+{
+	check(InstanceId.ISMComponent == InstancedStaticMeshComponent);
+
+	if (RenderingIndicesToDataIndices.IsValidIndex(InstanceId.InstanceIndex) && InstancedStaticMeshComponent->UpdateInstanceTransform(InstanceId.InstanceIndex, InstanceTransform, bWorldSpace, bMarkRenderStateDirty, bTeleport))
+	{
+		const int32 DataIndex = RenderingIndicesToDataIndices[InstanceId.InstanceIndex];
+		InstancedStaticMeshComponent->GetInstanceTransform(InstanceId.InstanceIndex, InstanceTransforms[DataIndex], /*bWorldSpace*/true);
+		return true;
+	}
+
+	return false;
+}
+
+bool ALightWeightInstanceStaticMeshManager::DeleteSMInstances(TArrayView<const FSMInstanceId> InstanceIds)
+{
+	TArray<int32> DataIndices;
+	GetLWIDataIndices(InstanceIds, DataIndices);
+
+	// Sort so Remove doesn't alter the indices of items still to remove
+	DataIndices.Sort(TGreater<int32>());
+
+	Modify();
+	InstancedStaticMeshComponent->Modify();
+
+	for (int32 DataIndex : DataIndices)
+	{
+		RemoveInstance(DataIndex);
+	}
+
+	return true;
+}
+
+bool ALightWeightInstanceStaticMeshManager::DuplicateSMInstances(TArrayView<const FSMInstanceId> InstanceIds, TArray<FSMInstanceId>& OutNewInstanceIds)
+{
+	TArray<int32> DataIndices;
+	GetLWIDataIndices(InstanceIds, DataIndices);
+
+	Modify();
+	InstancedStaticMeshComponent->Modify();
+
+	TArray<int32> NewDataIndices;
+	DuplicateLWIInstances(DataIndices, NewDataIndices);
+
+	OutNewInstanceIds.Reset(NewDataIndices.Num());
+	for (int32 NewDataIndex : NewDataIndices)
+	{
+		OutNewInstanceIds.Add(FSMInstanceId{ InstancedStaticMeshComponent, DataIndicesToRenderingIndices[NewDataIndex] });
+	}
+
+	return true;
+}
+
+void ALightWeightInstanceStaticMeshManager::DuplicateLWIInstances(TArrayView<const int32> DataIndices, TArray<int32>& OutNewDataIndices)
+{
+	OutNewDataIndices.Reset(DataIndices.Num());
+	for (int32 DataIndex : DataIndices)
+	{
+		FLWIData InitData;
+		InitData.Transform = InstanceTransforms[DataIndex];
+
+		const int32 NewDataIndex = AddNewInstance(&InitData);
+		if (DataIndex != INDEX_NONE)
+		{
+			OutNewDataIndices.Add(NewDataIndex);
+		}
+	}
+}
+
+void ALightWeightInstanceStaticMeshManager::GetLWIDataIndices(TArrayView<const FSMInstanceId> InstanceIds, TArray<int32>& OutDataIndices) const
+{
+	OutDataIndices.Reset(InstanceIds.Num());
+	for (const FSMInstanceId& InstanceId : InstanceIds)
+	{
+		check(InstanceId.ISMComponent == InstancedStaticMeshComponent);
+		if (RenderingIndicesToDataIndices.IsValidIndex(InstanceId.InstanceIndex))
+		{
+			OutDataIndices.Add(RenderingIndicesToDataIndices[InstanceId.InstanceIndex]);
+		}
+	}
+}
diff --git a/Engine/Source/Runtime/Engine/Public/Elements/Framework/EngineElementsLibrary.h b/Engine/Source/Runtime/Engine/Public/Elements/Framework/EngineElementsLibrary.h
--- a/Engine/Source/Runtime/Engine/Public/Elements/Framework/EngineElementsLibrary.h
+++ b/Engine/Source/Runtime/Engine/Public/Elements/Framework/EngineElementsLibrary.h
@@ -60,7 +60,6 @@ public:
 #endif
 
 	static TTypedElementOwner<FSMInstanceElementData> CreateSMInstanceElement(const FSMInstanceId& InSMInstanceId);
-	static TTypedElementOwner<FSMInstanceElementData> CreateSMInstanceElement(const FSMInstanceElementId& InSMInstanceElementId);
 	static void DestroySMInstanceElement(const FSMInstanceElementId& InSMInstanceElementId, TTypedElementOwner<FSMInstanceElementData>& InOutSMInstanceElement);
 #if WITH_EDITOR
 	static void CreateEditorSMInstanceElement(const FSMInstanceId& SMInstanceId);
@@ -68,10 +67,10 @@ public:
 	UFUNCTION(BlueprintCallable, BlueprintPure=false, Category="TypedElementFramework|StaticMeshInstance", meta=(ScriptMethod="AcquireEditorElementHandle"))
 	static FTypedElementHandle AcquireEditorSMInstanceElementHandle(const UInstancedStaticMeshComponent* ISMComponent, const int32 InstanceIndex, const bool bAllowCreate = true);
 	static FTypedElementHandle AcquireEditorSMInstanceElementHandle(const FSMInstanceId& SMInstanceId, const bool bAllowCreate = true);
-	static FTypedElementHandle AcquireEditorSMInstanceElementHandle(const FSMInstanceElementId& SMInstanceElementId, const bool bAllowCreate = true);
 #endif
 
 private:
+	static TTypedElementOwner<FSMInstanceElementData> CreateSMInstanceElementImpl(const FSMInstanceId& InSMInstanceId, const FSMInstanceElementId& InSMInstanceElementId);
 #if WITH_EDITOR
 	static void DestroyUnreachableEditorObjectElements();
 	static void OnObjectsReplaced(const TMap<UObject*, UObject*>& InReplacementObjects);
diff --git a/Engine/Source/Runtime/Engine/Public/Elements/SMInstance/SMInstanceElementData.h b/Engine/Source/Runtime/Engine/Public/Elements/SMInstance/SMInstanceElementData.h
--- a/Engine/Source/Runtime/Engine/Public/Elements/SMInstance/SMInstanceElementData.h
+++ b/Engine/Source/Runtime/Engine/Public/Elements/SMInstance/SMInstanceElementData.h
@@ -3,7 +3,7 @@
 #pragma once
 
 #include "Elements/Framework/TypedElementData.h"
-#include "Elements/SMInstance/SMInstanceElementId.h"
+#include "Elements/SMInstance/SMInstanceManager.h"
 
 struct FTypedElementHandle;
 
@@ -27,40 +27,41 @@ namespace SMInstanceElementDataUtil
 ENGINE_API bool SMInstanceElementsEnabled();
 
 /**
- * Test whether the given ISM component is valid to be used with static mesh instance elements.
+ * Get the static mesh instance manager for the given instance.
+ * @return The static mesh instance manager, or null if this instance cannot be managed.
  */
-ENGINE_API bool IsValidComponentForSMInstanceElements(const UInstancedStaticMeshComponent* InComponent);
+ENGINE_API TScriptInterface<ISMInstanceManager> GetSMInstanceManager(const FSMInstanceId& InstanceId);
 
 /**
  * Attempt to get the static mesh instance ID from the given element handle.
  * @note This is not typically something you'd want to use outside of data access within an interface implementation.
- * @return The static mesh instance ID if the element handle contains FSMInstanceElementData, otherwise an invalid ID.
+ * @return The static mesh instance ID if the element handle contains FSMInstanceElementData which resolved to a valid FSMInstanceManager, otherwise an invalid ID.
  */
-ENGINE_API FSMInstanceId GetSMInstanceFromHandle(const FTypedElementHandle& InHandle, const bool bSilent = false);
+ENGINE_API FSMInstanceManager GetSMInstanceFromHandle(const FTypedElementHandle& InHandle, const bool bSilent = false);
 
 /**
  * Attempt to get the static mesh instance ID from the given element handle, asserting if the element handle doesn't contain FSMInstanceElementData.
  * @note This is not typically something you'd want to use outside of data access within an interface implementation.
- * @return The static mesh instance ID.
+ * @return The static mesh instance ID, or an invalid ID if the FSMInstanceElementData didn't resolve to a valid FSMInstanceManager.
  */
-ENGINE_API FSMInstanceId GetSMInstanceFromHandleChecked(const FTypedElementHandle& InHandle);
+ENGINE_API FSMInstanceManager GetSMInstanceFromHandleChecked(const FTypedElementHandle& InHandle);
 
 /**
  * Attempt to get the static mesh instance IDs from the given element handles.
  * @note This is not typically something you'd want to use outside of data access within an interface implementation.
- * @return The static mesh instance IDs of any element handles that contain FSMInstanceElementData, skipping any that don't.
+ * @return The static mesh instance IDs of any element handles that contain FSMInstanceElementData which resolves to a valid FSMInstanceManager, skipping any that don't.
  */
 template <typename ElementHandleType>
-TArray<FSMInstanceId> GetSMInstancesFromHandles(TArrayView<const ElementHandleType> InHandles, const bool bSilent = false)
+TArray<FSMInstanceManager> GetSMInstancesFromHandles(TArrayView<const ElementHandleType> InHandles, const bool bSilent = false)
 {
-	TArray<FSMInstanceId> SMInstanceIds;
+	TArray<FSMInstanceManager> SMInstanceIds;
 	SMInstanceIds.Reserve(InHandles.Num());
 
 	for (const FTypedElementHandle& Handle : InHandles)
 	{
-		if (FSMInstanceId SMInstanceId = GetSMInstanceFromHandle(Handle, bSilent))
+		if (FSMInstanceManager SMInstanceId = GetSMInstanceFromHandle(Handle, bSilent))
 		{
-			SMInstanceIds.Add(SMInstanceId);
+			SMInstanceIds.Add(MoveTemp(SMInstanceId));
 		}
 	}
 
@@ -68,19 +69,22 @@ TArray<FSMInstanceId> GetSMInstancesFromHandles(TArrayView<const ElementHandleTy
 }
 
 /**
- * Attempt to get the static mesh instance IDs from the given element handles, asserting if any element handle doesn't contain FSMInstanceElementData.
+ * Attempt to get the static mesh instance IDs from the given element handles, asserting if any element handle doesn't contain FSMInstanceElementData, and skipping any that don't resolve to a valid FSMInstanceManager.
  * @note This is not typically something you'd want to use outside of data access within an interface implementation.
  * @return The static mesh instance IDs.
  */
 template <typename ElementHandleType>
-TArray<FSMInstanceId> GetSMInstancesFromHandlesChecked(TArrayView<const ElementHandleType> InHandles)
+TArray<FSMInstanceManager> GetSMInstancesFromHandlesChecked(TArrayView<const ElementHandleType> InHandles)
 {
-	TArray<FSMInstanceId> SMInstanceIds;
+	TArray<FSMInstanceManager> SMInstanceIds;
 	SMInstanceIds.Reserve(InHandles.Num());
 
 	for (const FTypedElementHandle& Handle : InHandles)
 	{
-		SMInstanceIds.Add(GetSMInstanceFromHandleChecked(Handle));
+		if (FSMInstanceManager SMInstanceId = GetSMInstanceFromHandleChecked(Handle))
+		{
+			SMInstanceIds.Add(MoveTemp(SMInstanceId));
+		}
 	}
 
 	return SMInstanceIds;
diff --git a/Engine/Source/Runtime/Engine/Public/Elements/SMInstance/SMInstanceElementWorldInterface.h b/Engine/Source/Runtime/Engine/Public/Elements/SMInstance/SMInstanceElementWorldInterface.h
--- a/Engine/Source/Runtime/Engine/Public/Elements/SMInstance/SMInstanceElementWorldInterface.h
+++ b/Engine/Source/Runtime/Engine/Public/Elements/SMInstance/SMInstanceElementWorldInterface.h
@@ -5,8 +5,6 @@
 #include "Elements/Interfaces/TypedElementWorldInterface.h"
 #include "SMInstanceElementWorldInterface.generated.h"
 
-struct FSMInstanceId;
-
 UCLASS()
 class ENGINE_API USMInstanceElementWorldInterface : public UTypedElementWorldInterface
 {
@@ -22,6 +20,4 @@ public:
 	virtual bool SetWorldTransform(const FTypedElementHandle& InElementHandle, const FTransform& InTransform) override;
 	virtual bool GetRelativeTransform(const FTypedElementHandle& InElementHandle, FTransform& OutTransform) override;
 	virtual bool SetRelativeTransform(const FTypedElementHandle& InElementHandle, const FTransform& InTransform) override;
-
-	static bool CanEditSMInstance(const FSMInstanceId& InSMInstanceId);
 };
diff --git a/Engine/Source/Runtime/Engine/Public/Elements/SMInstance/SMInstanceManager.h b/Engine/Source/Runtime/Engine/Public/Elements/SMInstance/SMInstanceManager.h
new file mode 100644
--- /dev/null
+++ b/Engine/Source/Runtime/Engine/Public/Elements/SMInstance/SMInstanceManager.h
@@ -0,0 +1,162 @@
+// Copyright Epic Games, Inc. All Rights Reserved.
+
+#pragma once
+
+#include "CoreMinimal.h"
+#include "UObject/Interface.h"
+#include "Containers/ArrayView.h"
+#include "Elements/SMInstance/SMInstanceElementId.h"
+#include "SMInstanceManager.generated.h"
+
+/**
+ * An interface for actors that manage static mesh instances.
+ * This exists so that actors that directly manage instances can inject custom logic around the manipulation of the underlying ISM component.
+ * @note The static mesh instances given to this API must be valid and belong to this instance manager. The instance manager implementation is free to assert or crash if that contract is broken.
+ */
+UINTERFACE(BlueprintType, meta=(CannotImplementInterfaceInBlueprint))
+class ENGINE_API USMInstanceManager : public UInterface
+{
+	GENERATED_BODY()
+};
+class ENGINE_API ISMInstanceManager
+{
+	GENERATED_BODY()
+
+public:
+	/**
+	 * Can the given static mesh instance be edited?
+	 * @return True if it can be edited, false otherwise.
+	 */
+	virtual bool CanEditSMInstance(const FSMInstanceId& InstanceId) const = 0;
+
+	/**
+	 * Attempt to get the transform of the given static mesh instance.
+	 * @note The transform should be in the local space of the owner ISM component, unless bWorldSpace is set.
+	 * @return True if the transform was retrieved, false otherwise.
+	 */
+	virtual bool GetSMInstanceTransform(const FSMInstanceId& InstanceId, FTransform& OutInstanceTransform, bool bWorldSpace = false) const = 0;
+
+	/**
+	 * Attempt to set the transform of the given static mesh instance.
+	 * @note The transform should be in the local space of the owner ISM component, unless bWorldSpace is set.
+	 * @return True if the transform was updated, false otherwise.
+	 */
+	virtual bool SetSMInstanceTransform(const FSMInstanceId& InstanceId, const FTransform& InstanceTransform, bool bWorldSpace = false, bool bMarkRenderStateDirty = false, bool bTeleport = false) = 0;
+
+	/**
+	 * Can the given static mesh instance be deleted?
+	 * @return True if it can be deleted, false otherwise.
+	 */
+	virtual bool CanDeleteSMInstance(const FSMInstanceId& InstanceId) const
+	{
+		return CanEditSMInstance(InstanceId);
+	}
+
+	/**
+	 * Attempt to delete the given static mesh instances.
+	 * @return True if any instances were deleted, false otherwise.
+	 */
+	virtual bool DeleteSMInstances(TArrayView<const FSMInstanceId> InstanceIds) = 0;
+
+	/**
+	 * Can the given static mesh instance be duplicated?
+	 * @return True if it can be duplicated, false otherwise.
+	 */
+	virtual bool CanDuplicateSMInstance(const FSMInstanceId& InstanceId) const
+	{
+		return CanEditSMInstance(InstanceId);
+	}
+
+	/**
+	 * Attempt to duplicate the given static mesh instances, retrieving the IDs of any new instances.
+	 * @return True if any instances were duplicated, false otherwise.
+	 */
+	virtual bool DuplicateSMInstances(TArrayView<const FSMInstanceId> InstanceIds, TArray<FSMInstanceId>& OutNewInstanceIds) = 0;
+};
+
+/**
+ * An interface for actors that can provide a manager for static mesh instances.
+ * This exists so that actors that indirectly manage instances can provide the correct underlying manager for the specific ISM component and instance.
+ * @note The static mesh instances given to this API must be valid. The instance manager provider implementation is free to assert or crash if that contract is broken.
+ */
+UINTERFACE(BlueprintType, meta=(CannotImplementInterfaceInBlueprint))
+class ENGINE_API USMInstanceManagerProvider : public UInterface
+{
+	GENERATED_BODY()
+};
+class ENGINE_API ISMInstanceManagerProvider
+{
+	GENERATED_BODY()
+
+public:
+	/**
+	 * Attempt to get the instance manager associated with the given static mesh instance, if any.
+	 * @return The instance manager, or null if there is no instance manager associated with the given instance.
+	 */
+	virtual TScriptInterface<ISMInstanceManager> GetSMInstanceManager(const FSMInstanceId& InstanceId) const = 0;
+};
+
+/**
+ * A static mesh instance manager, tied to a given static mesh instance ID.
+ */
+struct ENGINE_API FSMInstanceManager
+{
+public:
+	FSMInstanceManager() = default;
+
+	FSMInstanceManager(const FSMInstanceId& InInstanceId, const TScriptInterface<ISMInstanceManager>& InInstanceManager)
+		: InstanceId(InInstanceId)
+		, InstanceManager(InInstanceManager)
+	{
+	}
+
+	explicit operator bool() const
+	{
+		return InstanceId
+			&& InstanceManager;
+	}
+
+	bool operator==(const FSMInstanceManager& InRHS) const
+	{
+		return InstanceId == InRHS.InstanceId
+			&& InstanceManager == InRHS.InstanceManager;
+	}
+
+	bool operator!=(const FSMInstanceManager& InRHS) const
+	{
+		return !(*this == InRHS);
+	}
+
+	friend inline uint32 GetTypeHash(const FSMInstanceManager& InId)
+	{
+		return GetTypeHash(InId.InstanceId);
+	}
+
+	const FSMInstanceId& GetInstanceId() const { return InstanceId; }
+	const TScriptInterface<ISMInstanceManager>& GetInstanceManager() const { return InstanceManager; }
+
+	UInstancedStaticMeshComponent* GetISMComponent() const { return InstanceId.ISMComponent; }
+	int32 GetISMInstanceIndex() const { return InstanceId.InstanceIndex; }
+
+	//~ ISMInstanceManager interface
+	bool CanEditSMInstance() const { return InstanceManager->CanEditSMInstance(InstanceId); }
+	bool GetSMInstanceTransform(FTransform& OutInstanceTransform, bool bWorldSpace = false) const { return InstanceManager->GetSMInstanceTransform(InstanceId, OutInstanceTransform, bWorldSpace); }
+	bool SetSMInstanceTransform(const FTransform& InstanceTransform, bool bWorldSpace = false, bool bMarkRenderStateDirty = false, bool bTeleport = false) const { return InstanceManager->SetSMInstanceTransform(InstanceId, InstanceTransform, bWorldSpace, bMarkRenderStateDirty, bTeleport); }
+	bool CanDeleteSMInstance() const { return InstanceManager->CanDeleteSMInstance(InstanceId); }
+	bool DeleteSMInstance() const { return InstanceManager->DeleteSMInstances(MakeArrayView(&InstanceId, 1)); }
+	bool CanDuplicateSMInstance() const { return InstanceManager->CanDuplicateSMInstance(InstanceId); }
+	bool DuplicateSMInstance(FSMInstanceId& OutNewInstanceId) const
+	{
+		TArray<FSMInstanceId> NewInstanceIds;
+		if (InstanceManager->DuplicateSMInstances(MakeArrayView(&InstanceId, 1), NewInstanceIds))
+		{
+			OutNewInstanceId = NewInstanceIds[0];
+			return true;
+		}
+		return false;
+	}
+
+private:
+	FSMInstanceId InstanceId;
+	TScriptInterface<ISMInstanceManager> InstanceManager;
+};
diff --git a/Engine/Source/Runtime/Engine/Public/ISMPartition/ISMPartitionActor.h b/Engine/Source/Runtime/Engine/Public/ISMPartition/ISMPartitionActor.h
--- a/Engine/Source/Runtime/Engine/Public/ISMPartition/ISMPartitionActor.h
+++ b/Engine/Source/Runtime/Engine/Public/ISMPartition/ISMPartitionActor.h
@@ -8,6 +8,7 @@
 #include "ISMPartition/ISMComponentData.h"
 #include "Templates/Tuple.h"
 #include "Containers/SortedMap.h"
+#include "Elements/SMInstance/SMInstanceManager.h"
 #include "ISMPartitionActor.generated.h"
 
 struct FISMClientHandle;
@@ -17,7 +18,7 @@ DECLARE_LOG_CATEGORY_EXTERN(LogISMPartition, Log, All);
 /** Actor base class for instance containers placed on a grid.
 	See UActorPartitionSubsystem. */
 UCLASS(Abstract)
-class ENGINE_API AISMPartitionActor : public APartitionActor
+class ENGINE_API AISMPartitionActor : public APartitionActor, public ISMInstanceManagerProvider
 {
 	GENERATED_UCLASS_BODY()
 
@@ -49,6 +50,7 @@ public:
 	void UpdateHISMTrees(bool bAsync, bool bForce);
 	void GetClientComponents(const FISMClientHandle& Handle, TArray<UInstancedStaticMeshComponent*>& OutComponents);
 	void OutputStats() const;
+
 protected:
 
 private:
@@ -64,6 +66,11 @@ private:
 	void ModifyActor();
 #endif
 
+protected:
+	//~ ISMInstanceManagerProvider interface
+	virtual TScriptInterface<ISMInstanceManager> GetSMInstanceManager(const FSMInstanceId& InstanceId) const override;
+
+private:
 #if WITH_EDITORONLY_DATA
 	UPROPERTY()
 	TArray<FGuid> Clients;
