IterationProfile polish

Add warning to advise significant improvements submit a new baseline or update threshold.

Tested with local UAT runs

diff --git a/Engine/Source/Editor/UnrealEd/Private/Commandlets/SummarizeTraceCommandlet.cpp b/Engine/Source/Editor/UnrealEd/Private/Commandlets/SummarizeTraceCommandlet.cpp
--- a/Engine/Source/Editor/UnrealEd/Private/Commandlets/SummarizeTraceCommandlet.cpp
+++ b/Engine/Source/Editor/UnrealEd/Private/Commandlets/SummarizeTraceCommandlet.cpp
@@ -741,6 +741,7 @@ struct TelemetryDefinition
 
 	static bool LoadFromCSV(const FString& FilePath, TMap<TPair<FString,FString>, TelemetryDefinition>& ContextAndDataPointToDefinitionMap);
 	static bool MeasurementWithinThreshold(const FString& Value, const FString& BaselineValue, const FString& Threshold);
+	static FString SignFlipThreshold(const FString& Threshold);
 
 	FString TestName;
 	FString Context;
@@ -881,6 +882,22 @@ bool TelemetryDefinition::MeasurementWithinThreshold(const FString& MeasurementV
 	}
 }
 
+FString TelemetryDefinition::SignFlipThreshold(const FString& Threshold)
+{
+	FString SignFlipped;
+
+	if (Threshold.StartsWith(TEXT("-")))
+	{
+		SignFlipped = Threshold.RightChop(1);
+	}
+	else
+	{
+		SignFlipped = FString(TEXT("-")) + Threshold;
+	}
+
+	return SignFlipped;
+}
+
 /*
  * SummarizeTrace commandlet ingests a utrace file and summarizes the
  * cpu scope events within it, and summarizes each event to a csv. It
@@ -1199,27 +1216,50 @@ int32 USummarizeTraceCommandlet::Main(const FString& CmdLineParams)
 						}
 					}
 
+					// let's only report on statistics that have an assigned threshold, to keep things concise
+					bool bHasThreshold = !RelatedStatistic->WarningThreshold.IsEmpty() || !RelatedStatistic->ErrorThreshold.IsEmpty();
+
 					// do we still have the statistic definition in our current stats file? (if we don't that's fine, we don't care about it anymore)
-					if (RelatedStatistic)
+					if (RelatedStatistic && bHasThreshold)
 					{
 						// verify that this telemetry measurement is within the allowed threshold as defined in the current stats file
 						if (TelemetryDefinition::MeasurementWithinThreshold(Telemetry.Measurement, BaselineTelemetry->Measurement, RelatedStatistic->WarningThreshold))
 						{
-							UE_LOG(LogSummarizeTrace, Display, TEXT("Telemetry %s,%s,%s,%s within baseline value %s using warning threshold %s"),
-								*Telemetry.TestName, *Telemetry.Context, *Telemetry.DataPoint, *Telemetry.Measurement,
-								*BaselineTelemetry->Measurement, *RelatedStatistic->WarningThreshold);
+							FString SignFlippedWarningThreshold = TelemetryDefinition::SignFlipThreshold(RelatedStatistic->WarningThreshold);
+
+							// check if it's beyond the threshold the other way and needs lowering in the stats csv
+							if (!TelemetryDefinition::MeasurementWithinThreshold(Telemetry.Measurement, BaselineTelemetry->Measurement, SignFlippedWarningThreshold))
+							{
+								FString BaselineRelPath = FPaths::ConvertRelativePathToFull(BaselineTelemetryCsvFilePath);
+								FPaths::MakePathRelativeTo(BaselineRelPath, *FPaths::RootDir());
+
+								FString StatisticsRelPath = FPaths::ConvertRelativePathToFull(StatisticsFileName);
+								FPaths::MakePathRelativeTo(StatisticsRelPath, *FPaths::RootDir());
+
+								UE_LOG(LogSummarizeTrace, Warning, TEXT("Telemetry %s,%s,%s,%s significantly within baseline value %s using warning threshold %s. Please submit a new baseline to %s or adjust the threshold in %s."),
+									*Telemetry.TestName, *Telemetry.Context, *Telemetry.DataPoint, *Telemetry.Measurement,
+									*BaselineTelemetry->Measurement, *RelatedStatistic->WarningThreshold,
+									*BaselineRelPath, *StatisticsRelPath);
+							}
+							else // it's within tolerance, just report that it's ok
+							{
+								UE_LOG(LogSummarizeTrace, Verbose, TEXT("Telemetry %s,%s,%s,%s within baseline value %s using warning threshold %s"),
+									*Telemetry.TestName, *Telemetry.Context, *Telemetry.DataPoint, *Telemetry.Measurement,
+									*BaselineTelemetry->Measurement, *RelatedStatistic->WarningThreshold);
+							}
 						}
 						else
 						{
+							// it's outside warning threshold, check if it's inside the error threshold to just issue a warning
 							if (TelemetryDefinition::MeasurementWithinThreshold(Telemetry.Measurement, BaselineTelemetry->Measurement, RelatedStatistic->ErrorThreshold))
 							{
-								UE_LOG(LogSummarizeTrace, Warning, TEXT("Telemetry %s,%s,%s,%s beyond baseline value %s using warning threshold %s"),
+								UE_LOG(LogSummarizeTrace, Warning, TEXT("Telemetry %s,%s,%s,%s beyond baseline value %s using warning threshold %s. This could be a performance regression!"),
 									*Telemetry.TestName, *Telemetry.Context, *Telemetry.DataPoint, *Telemetry.Measurement,
 									*BaselineTelemetry->Measurement, *RelatedStatistic->WarningThreshold);
 							}
-							else
+							else // it's outside the error threshold, hard error
 							{
-								UE_LOG(LogSummarizeTrace, Error, TEXT("Telemetry %s,%s,%s,%s beyond baseline value %s using error threshold %s"),
+								UE_LOG(LogSummarizeTrace, Error, TEXT("Telemetry %s,%s,%s,%s beyond baseline value %s using error threshold %s. This could be a performance regression!"),
 									*Telemetry.TestName, *Telemetry.Context, *Telemetry.DataPoint, *Telemetry.Measurement,
 									*BaselineTelemetry->Measurement, *RelatedStatistic->ErrorThreshold);
 							}
