Remove unnecssary forward declarations of SCS editor

#rnx

diff --git a/Engine/Source/Editor/LevelEditor/Private/SInViewportDetails.h b/Engine/Source/Editor/LevelEditor/Private/SInViewportDetails.h
--- a/Engine/Source/Editor/LevelEditor/Private/SInViewportDetails.h
+++ b/Engine/Source/Editor/LevelEditor/Private/SInViewportDetails.h
@@ -19,7 +19,6 @@ class FTabManager;
 class FUICommandList;
 class IDetailsView;
 class SBox;
-class SSCSEditor;
 class SSplitter;
 class UBlueprint;
 class FDetailsViewObjectFilter;
diff --git a/Engine/Source/Editor/UnrealEd/Public/Kismet2/ChildActorComponentEditorUtils.h b/Engine/Source/Editor/UnrealEd/Public/Kismet2/ChildActorComponentEditorUtils.h
--- a/Engine/Source/Editor/UnrealEd/Public/Kismet2/ChildActorComponentEditorUtils.h
+++ b/Engine/Source/Editor/UnrealEd/Public/Kismet2/ChildActorComponentEditorUtils.h
@@ -6,7 +6,6 @@
 #include "Components/ChildActorComponent.h"
 
 class UToolMenu;
-class SSCSEditor;
 class FSCSEditorTreeNode;
 
 class UNREALED_API FChildActorComponentEditorUtils
