Fix some static analysis warnings

diff --git a/Engine/Source/Editor/UnrealEd/Private/Commandlets/SummarizeTraceCommandlet.cpp b/Engine/Source/Editor/UnrealEd/Private/Commandlets/SummarizeTraceCommandlet.cpp
--- a/Engine/Source/Editor/UnrealEd/Private/Commandlets/SummarizeTraceCommandlet.cpp
+++ b/Engine/Source/Editor/UnrealEd/Private/Commandlets/SummarizeTraceCommandlet.cpp
@@ -1216,52 +1216,53 @@ int32 USummarizeTraceCommandlet::Main(const FString& CmdLineParams)
 						}
 					}
 
-					// let's only report on statistics that have an assigned threshold, to keep things concise
-					bool bHasThreshold = !RelatedStatistic->WarningThreshold.IsEmpty() || !RelatedStatistic->ErrorThreshold.IsEmpty();
-
 					// do we still have the statistic definition in our current stats file? (if we don't that's fine, we don't care about it anymore)
-					if (RelatedStatistic && bHasThreshold)
+					if (RelatedStatistic)
 					{
-						// verify that this telemetry measurement is within the allowed threshold as defined in the current stats file
-						if (TelemetryDefinition::MeasurementWithinThreshold(Telemetry.Measurement, BaselineTelemetry->Measurement, RelatedStatistic->WarningThreshold))
-						{
-							FString SignFlippedWarningThreshold = TelemetryDefinition::SignFlipThreshold(RelatedStatistic->WarningThreshold);
-
-							// check if it's beyond the threshold the other way and needs lowering in the stats csv
-							if (!TelemetryDefinition::MeasurementWithinThreshold(Telemetry.Measurement, BaselineTelemetry->Measurement, SignFlippedWarningThreshold))
-							{
-								FString BaselineRelPath = FPaths::ConvertRelativePathToFull(BaselineTelemetryCsvFilePath);
-								FPaths::MakePathRelativeTo(BaselineRelPath, *FPaths::RootDir());
-
-								FString StatisticsRelPath = FPaths::ConvertRelativePathToFull(StatisticsFileName);
-								FPaths::MakePathRelativeTo(StatisticsRelPath, *FPaths::RootDir());
-
-								UE_LOG(LogSummarizeTrace, Warning, TEXT("Telemetry %s,%s,%s,%s significantly within baseline value %s using warning threshold %s. Please submit a new baseline to %s or adjust the threshold in %s."),
-									*Telemetry.TestName, *Telemetry.Context, *Telemetry.DataPoint, *Telemetry.Measurement,
-									*BaselineTelemetry->Measurement, *RelatedStatistic->WarningThreshold,
-									*BaselineRelPath, *StatisticsRelPath);
-							}
-							else // it's within tolerance, just report that it's ok
-							{
-								UE_LOG(LogSummarizeTrace, Verbose, TEXT("Telemetry %s,%s,%s,%s within baseline value %s using warning threshold %s"),
-									*Telemetry.TestName, *Telemetry.Context, *Telemetry.DataPoint, *Telemetry.Measurement,
-									*BaselineTelemetry->Measurement, *RelatedStatistic->WarningThreshold);
-							}
-						}
-						else
+						// let's only report on statistics that have an assigned threshold, to keep things concise
+						if (!RelatedStatistic->WarningThreshold.IsEmpty() || !RelatedStatistic->ErrorThreshold.IsEmpty())
 						{
-							// it's outside warning threshold, check if it's inside the error threshold to just issue a warning
-							if (TelemetryDefinition::MeasurementWithinThreshold(Telemetry.Measurement, BaselineTelemetry->Measurement, RelatedStatistic->ErrorThreshold))
+							// verify that this telemetry measurement is within the allowed threshold as defined in the current stats file
+							if (TelemetryDefinition::MeasurementWithinThreshold(Telemetry.Measurement, BaselineTelemetry->Measurement, RelatedStatistic->WarningThreshold))
 							{
-								UE_LOG(LogSummarizeTrace, Warning, TEXT("Telemetry %s,%s,%s,%s beyond baseline value %s using warning threshold %s. This could be a performance regression!"),
-									*Telemetry.TestName, *Telemetry.Context, *Telemetry.DataPoint, *Telemetry.Measurement,
-									*BaselineTelemetry->Measurement, *RelatedStatistic->WarningThreshold);
+								FString SignFlippedWarningThreshold = TelemetryDefinition::SignFlipThreshold(RelatedStatistic->WarningThreshold);
+
+								// check if it's beyond the threshold the other way and needs lowering in the stats csv
+								if (!TelemetryDefinition::MeasurementWithinThreshold(Telemetry.Measurement, BaselineTelemetry->Measurement, SignFlippedWarningThreshold))
+								{
+									FString BaselineRelPath = FPaths::ConvertRelativePathToFull(BaselineTelemetryCsvFilePath);
+									FPaths::MakePathRelativeTo(BaselineRelPath, *FPaths::RootDir());
+
+									FString StatisticsRelPath = FPaths::ConvertRelativePathToFull(StatisticsFileName);
+									FPaths::MakePathRelativeTo(StatisticsRelPath, *FPaths::RootDir());
+
+									UE_LOG(LogSummarizeTrace, Warning, TEXT("Telemetry %s,%s,%s,%s significantly within baseline value %s using warning threshold %s. Please submit a new baseline to %s or adjust the threshold in %s."),
+										*Telemetry.TestName, *Telemetry.Context, *Telemetry.DataPoint, *Telemetry.Measurement,
+										*BaselineTelemetry->Measurement, *RelatedStatistic->WarningThreshold,
+										*BaselineRelPath, *StatisticsRelPath);
+								}
+								else // it's within tolerance, just report that it's ok
+								{
+									UE_LOG(LogSummarizeTrace, Verbose, TEXT("Telemetry %s,%s,%s,%s within baseline value %s using warning threshold %s"),
+										*Telemetry.TestName, *Telemetry.Context, *Telemetry.DataPoint, *Telemetry.Measurement,
+										*BaselineTelemetry->Measurement, *RelatedStatistic->WarningThreshold);
+								}
 							}
-							else // it's outside the error threshold, hard error
+							else
 							{
-								UE_LOG(LogSummarizeTrace, Error, TEXT("Telemetry %s,%s,%s,%s beyond baseline value %s using error threshold %s. This could be a performance regression!"),
-									*Telemetry.TestName, *Telemetry.Context, *Telemetry.DataPoint, *Telemetry.Measurement,
-									*BaselineTelemetry->Measurement, *RelatedStatistic->ErrorThreshold);
+								// it's outside warning threshold, check if it's inside the error threshold to just issue a warning
+								if (TelemetryDefinition::MeasurementWithinThreshold(Telemetry.Measurement, BaselineTelemetry->Measurement, RelatedStatistic->ErrorThreshold))
+								{
+									UE_LOG(LogSummarizeTrace, Warning, TEXT("Telemetry %s,%s,%s,%s beyond baseline value %s using warning threshold %s. This could be a performance regression!"),
+										*Telemetry.TestName, *Telemetry.Context, *Telemetry.DataPoint, *Telemetry.Measurement,
+										*BaselineTelemetry->Measurement, *RelatedStatistic->WarningThreshold);
+								}
+								else // it's outside the error threshold, hard error
+								{
+									UE_LOG(LogSummarizeTrace, Error, TEXT("Telemetry %s,%s,%s,%s beyond baseline value %s using error threshold %s. This could be a performance regression!"),
+										*Telemetry.TestName, *Telemetry.Context, *Telemetry.DataPoint, *Telemetry.Measurement,
+										*BaselineTelemetry->Measurement, *RelatedStatistic->ErrorThreshold);
+								}
 							}
 						}
 					}
