Make the InScope struct const on the KismetNameValidator

#rnx

diff --git a/Engine/Source/Editor/UnrealEd/Private/Kismet2/Kismet2NameValidators.cpp b/Engine/Source/Editor/UnrealEd/Private/Kismet2/Kismet2NameValidators.cpp
--- a/Engine/Source/Editor/UnrealEd/Private/Kismet2/Kismet2NameValidators.cpp
+++ b/Engine/Source/Editor/UnrealEd/Private/Kismet2/Kismet2NameValidators.cpp
@@ -98,17 +98,16 @@ namespace BlueprintNameConstants
 	 int32 NameMaxLength = 100;
 }
 
-FKismetNameValidator::FKismetNameValidator(const class UBlueprint* Blueprint, FName InExistingName/* = NAME_None*/, UStruct* InScope/* = NULL*/)
+FKismetNameValidator::FKismetNameValidator(const class UBlueprint* Blueprint, FName InExistingName/* = NAME_None*/, const UStruct* InScope/* = nullptr */)
+	: BlueprintObject(Blueprint)
+	, ExistingName(InExistingName)
+	, Scope(InScope)
 {
-	ExistingName = InExistingName;
-	BlueprintObject = Blueprint;
 	FBlueprintEditorUtils::GetClassVariableList(BlueprintObject, Names, true);
 	FBlueprintEditorUtils::GetFunctionNameList(BlueprintObject, Names);
 	FBlueprintEditorUtils::GetAllGraphNames(BlueprintObject, Names);
 	FBlueprintEditorUtils::GetSCSVariableNameList(Blueprint, Names);
 	FBlueprintEditorUtils::GetImplementingBlueprintsFunctionNameList(Blueprint, Names);
-
-	Scope = InScope;
 }
 
 int32 FKismetNameValidator::GetMaximumNameLength()
diff --git a/Engine/Source/Editor/UnrealEd/Public/Kismet2/Kismet2NameValidators.h b/Engine/Source/Editor/UnrealEd/Public/Kismet2/Kismet2NameValidators.h
--- a/Engine/Source/Editor/UnrealEd/Public/Kismet2/Kismet2NameValidators.h
+++ b/Engine/Source/Editor/UnrealEd/Public/Kismet2/Kismet2NameValidators.h
@@ -83,7 +83,7 @@ public:
 class UNREALED_API FKismetNameValidator : public INameValidatorInterface
 {
 public:
-	FKismetNameValidator(const class UBlueprint* Blueprint, FName InExistingName = NAME_None, UStruct* InScope = NULL);
+	FKismetNameValidator(const class UBlueprint* Blueprint, FName InExistingName = NAME_None, const UStruct* InScope = nullptr);
 	~FKismetNameValidator() {}
 
 	/** Return the name validator maximum string length */
@@ -101,7 +101,7 @@ private:
 	/** The current name of the object being validated */
 	FName ExistingName;
 	/** Scope to check against local variables (or other members) */
-	UStruct* Scope;
+	const UStruct* Scope;
 };
 
 /////////////////////////////////////////////////////
