Removing attached debugger logic from if assets should be saved for optional restoring
PR #8038:  Proposed fix for UE-72962: Added an option to reopen asset tabs when closing and opening Unreal Engine, while the debugger is attached (Contributed by [REDACTED_PERSON])

diff --git a/Engine/Source/Editor/UnrealEd/Private/Subsystems/AssetEditorSubsystem.cpp b/Engine/Source/Editor/UnrealEd/Private/Subsystems/AssetEditorSubsystem.cpp
--- a/Engine/Source/Editor/UnrealEd/Private/Subsystems/AssetEditorSubsystem.cpp
+++ b/Engine/Source/Editor/UnrealEd/Private/Subsystems/AssetEditorSubsystem.cpp
@@ -845,14 +845,12 @@ void UAssetEditorSubsystem::OnCancelRestorePreviouslyOpenAssets()
 	}
 }
 
-void UAssetEditorSubsystem::SaveOpenAssetEditors(const bool bOnShutdown, const bool bCancelIfDebugger)
+void UAssetEditorSubsystem::SaveOpenAssetEditors(const bool bOnShutdown)
 {
 	if (!bSavingOnShutdown && !bAutoRestoreAndDisableSaving)
 	{
 		TArray<FString> OpenAssets;
-
-		// If bCancelIfDebugger = true, don't save a list of assets to restore if we are running under a debugger
-		if (!bCancelIfDebugger || !FPlatformMisc::IsDebuggerPresent())
+		if (bOnShutdown || !FPlatformMisc::IsDebuggerPresent())
 		{
 			for (const TPair<IAssetEditorInstance*, UObject*>& EditorPair : OpenedEditors)
 			{
@@ -872,13 +870,17 @@ void UAssetEditorSubsystem::SaveOpenAssetEditors(const bool bOnShutdown, const b
 				}
 			}
 		}
-
 		GConfig->SetArray(TEXT("AssetEditorSubsystem"), TEXT("OpenAssetsAtExit"), OpenAssets, GEditorPerProjectIni);
 		GConfig->SetBool(TEXT("AssetEditorSubsystem"), TEXT("CleanShutdown"), bOnShutdown, GEditorPerProjectIni);
 		GConfig->Flush(false, GEditorPerProjectIni);
 	}
 }
 
+void UAssetEditorSubsystem::SaveOpenAssetEditors(const bool bOnShutdown, const bool bCancelIfDebugger)
+{
+	SaveOpenAssetEditors(bOnShutdown);
+}
+
 void UAssetEditorSubsystem::HandlePackageReloaded(const EPackageReloadPhase InPackageReloadPhase, FPackageReloadedEvent* InPackageReloadedEvent)
 {
 	static TArray<UObject*> PendingAssetsToOpen;
diff --git a/Engine/Source/Editor/UnrealEd/Public/Subsystems/AssetEditorSubsystem.h b/Engine/Source/Editor/UnrealEd/Public/Subsystems/AssetEditorSubsystem.h
--- a/Engine/Source/Editor/UnrealEd/Public/Subsystems/AssetEditorSubsystem.h
+++ b/Engine/Source/Editor/UnrealEd/Public/Subsystems/AssetEditorSubsystem.h
@@ -211,9 +211,18 @@ public:
 
 	/**
 	 * Saves a list of open asset editors so they can be restored on editor restart.
+	 * @param bOnShutdown If true, this is handled as if the engine is shutting down right now.
+	 */
+
+	void SaveOpenAssetEditors(const bool bOnShutdown);
+	
+	/**
+	 * Saves a list of open asset editors so they can be restored on editor restart.
+	 * @param bOnShutdown If true, this is handled as if the engine is shutting down right now.
 	 * @param bCancelIfDebugger If true, don't save a list of assets to restore if we are running under a debugger.
 	 */
-	void SaveOpenAssetEditors(const bool bOnShutdown, const bool bCancelIfDebugger = true);
+	UE_DEPRECATED(5.0, "Please use the version of SaveOpenAssetEditors with only one argument, bOnShutdown.")
+	void SaveOpenAssetEditors(const bool bOnShutdown, const bool bCancelIfDebugger);
 
 	/** Restore the assets that were previously open when the editor was last closed. */
 	void RestorePreviouslyOpenAssets();
