Updated SetSelectedLevelActors in EditorActorSubsystem to batch the selection, avoiding performance issues when working with large numbers of actors

diff --git a/Engine/Source/Editor/UnrealEd/Private/Subsystems/EditorActorSubsystem.cpp b/Engine/Source/Editor/UnrealEd/Private/Subsystems/EditorActorSubsystem.cpp
--- a/Engine/Source/Editor/UnrealEd/Private/Subsystems/EditorActorSubsystem.cpp
+++ b/Engine/Source/Editor/UnrealEd/Private/Subsystems/EditorActorSubsystem.cpp
@@ -393,9 +393,11 @@ void UEditorActorSubsystem::SetSelectedLevelActors(const TArray<class AActor*>&
 		return;
 	}
 
-	GEditor->GetSelectedActors()->Modify();
+	USelection* SelectedActors = GEditor->GetSelectedActors();
+	SelectedActors->Modify();
 	if (ActorsToSelect.Num() > 0)
 	{
+		SelectedActors->BeginBatchSelectOperation();
 		GEditor->SelectNone(false, true, false);
 		for (AActor* Actor : ActorsToSelect)
 		{
@@ -409,6 +411,7 @@ void UEditorActorSubsystem::SetSelectedLevelActors(const TArray<class AActor*>&
 				GEditor->SelectActor(Actor, true, false);
 			}
 		}
+		SelectedActors->EndBatchSelectOperation();
 		GEditor->NoteSelectionChange();
 	}
 	else
