Fix for build.

diff --git a/Engine/Source/Editor/UnrealEd/Public/UnrealEd.h b/Engine/Source/Editor/UnrealEd/Public/UnrealEd.h
--- a/Engine/Source/Editor/UnrealEd/Public/UnrealEd.h
+++ b/Engine/Source/Editor/UnrealEd/Public/UnrealEd.h
@@ -21,7 +21,6 @@ MONOLITHIC_HEADER_BOILERPLATE()
 // @todo Classes.h
 #include "ActorFactories/ActorFactory.h"
 #include "ActorFactories/ActorFactoryAmbientSound.h"
-#include "ActorFactories/ActorFactoryAtmosphericFog.h"
 #include "ActorFactories/ActorFactoryBlueprint.h"
 #include "ActorFactories/ActorFactoryBoxReflectionCapture.h"
 #include "ActorFactories/ActorFactoryBoxVolume.h"
