Fix launch on so it uses the per platform config the user selected instead of the current editor build type

diff --git a/Engine/Source/Editor/UnrealEd/Private/PlayLevelViaLauncher.cpp b/Engine/Source/Editor/UnrealEd/Private/PlayLevelViaLauncher.cpp
--- a/Engine/Source/Editor/UnrealEd/Private/PlayLevelViaLauncher.cpp
+++ b/Engine/Source/Editor/UnrealEd/Private/PlayLevelViaLauncher.cpp
@@ -83,11 +83,13 @@ void UEditorEngine::StartPlayUsingLauncherSession(FRequestPlaySessionParams& InR
 	FString LaunchPlatformName = LastPlayUsingLauncherDeviceId.Left(LastPlayUsingLauncherDeviceId.Find(TEXT("@")));
 	FString LaunchPlatformNameFromID = LastPlayUsingLauncherDeviceId.Right(LastPlayUsingLauncherDeviceId.Find(TEXT("@")));
 	ITargetPlatform* LaunchPlatform = GetTargetPlatformManagerRef().FindTargetPlatform(LaunchPlatformName);
+	FString IniPlatformName = LaunchPlatformName;
 
 	// create a temporary device group and launcher profile
 	ILauncherDeviceGroupRef DeviceGroup = LauncherServicesModule.CreateDeviceGroup(FGuid::NewGuid(), TEXT("PlayOnDevices"));
 	if (LaunchPlatform != nullptr)
 	{
+		IniPlatformName = LaunchPlatform->IniPlatformName();
 		if (LaunchPlatformNameFromID.Equals(LaunchPlatformName))
 		{
 			// create a temporary list of devices for the target platform
@@ -131,7 +133,7 @@ void UEditorEngine::StartPlayUsingLauncherSession(FRequestPlaySessionParams& InR
 	}
 
 	// set the build/launch configuration 
-	EBuildConfiguration BuildConfiguration;
+	EBuildConfiguration BuildConfiguration = EBuildConfiguration::Development;
 	const ULevelEditorPlaySettings* EditorPlaySettings = PlaySessionRequest->EditorPlaySettings;
 	switch (EditorPlaySettings->LaunchConfiguration)
 	{
@@ -148,10 +150,36 @@ void UEditorEngine::StartPlayUsingLauncherSession(FRequestPlaySessionParams& InR
 		BuildConfiguration = EBuildConfiguration::Shipping;
 		break;
 	default:
-		// same as the running editor
-		BuildConfiguration = FApp::GetBuildConfiguration();
+	{
+		const UProjectPackagingSettings* AllPlatformPackagingSettings = GetDefault<UProjectPackagingSettings>();
+
+		EProjectPackagingBuildConfigurations BuildConfig = AllPlatformPackagingSettings->GetBuildConfigurationForPlatform(*IniPlatformName);
+		// if PPBC_MAX is set, then the project default should be used instead of the per platform build config
+		if (BuildConfig == EProjectPackagingBuildConfigurations::PPBC_MAX)
+		{
+			BuildConfig = AllPlatformPackagingSettings->BuildConfiguration;
+		}
+
+		switch (BuildConfig)
+		{
+		case EProjectPackagingBuildConfigurations::PPBC_Debug:
+		case EProjectPackagingBuildConfigurations::PPBC_DebugGame:
+			BuildConfiguration = EBuildConfiguration::Debug;
+			break;
+		case EProjectPackagingBuildConfigurations::PPBC_Development:
+			BuildConfiguration = EBuildConfiguration::Development;
+			break;
+		case EProjectPackagingBuildConfigurations::PPBC_Test:
+			BuildConfiguration = EBuildConfiguration::Test;
+			break;
+		case EProjectPackagingBuildConfigurations::PPBC_Shipping:
+			BuildConfiguration = EBuildConfiguration::Shipping;
+			break;
+		}
+
 		break;
 	}
+	}
 
 	// does the project have any code?
 	FGameProjectGenerationModule& GameProjectModule = FModuleManager::LoadModuleChecked<FGameProjectGenerationModule>(TEXT("GameProjectGeneration"));
