Fix spam when playing demos in PIE:
LogPackageName: Warning: TryConvertFilenameToLongPackageName was passed an ObjectPath...
#rnx

diff --git a/Engine/Source/Editor/UnrealEd/Private/EditorEngine.cpp b/Engine/Source/Editor/UnrealEd/Private/EditorEngine.cpp
--- a/Engine/Source/Editor/UnrealEd/Private/EditorEngine.cpp
+++ b/Engine/Source/Editor/UnrealEd/Private/EditorEngine.cpp
@@ -6852,10 +6852,10 @@ FORCEINLINE bool NetworkRemapPath_local(FWorldContext& Context, FString& Str, bo
 		if (bIsReplay)
 		{
 			const FString AssetName = Path.GetAssetName();
-			const FString ShortName = FPackageName::GetShortName(Path.GetLongPackageName());
 
-			FString PackageNameOnly = Path.ToString();
+			FString PackageNameOnly = Path.GetLongPackageName();
 			FPackageName::TryConvertFilenameToLongPackageName(PackageNameOnly, PackageNameOnly);
+			const FString ShortName = FPackageName::GetShortName(PackageNameOnly);
 
 			const FString PrefixedFullName = UWorld::ConvertToPIEPackageName(Str, Context.PIEInstance);
 			const FString PrefixedPackageName = UWorld::ConvertToPIEPackageName(PackageNameOnly, Context.PIEInstance);
