Switched a few world teardown to use UWorld::DestroyWorld()

diff --git a/Engine/Source/Editor/UnrealEd/Private/Commandlets/ContentCommandlets.cpp b/Engine/Source/Editor/UnrealEd/Private/Commandlets/ContentCommandlets.cpp
--- a/Engine/Source/Editor/UnrealEd/Private/Commandlets/ContentCommandlets.cpp
+++ b/Engine/Source/Editor/UnrealEd/Private/Commandlets/ContentCommandlets.cpp
@@ -2071,9 +2071,8 @@ void UResavePackagesCommandlet::PerformAdditionalOperations(class UWorld* World,
 	}
 
 	// Cleanup
-	World->RemoveFromRoot();
-	World->ClearWorldComponents();
-	World->CleanupWorld();
+	const bool bBroadcastWorldDestroyedEvent = false;
+	World->DestroyWorld(bBroadcastWorldDestroyedEvent);
 	
 	// Restore previous world
 	WorldContext.SetCurrentWorld(PrevGWorld);
diff --git a/Engine/Source/Editor/UnrealEd/Private/Commandlets/DerivedDataCacheCommandlet.cpp b/Engine/Source/Editor/UnrealEd/Private/Commandlets/DerivedDataCacheCommandlet.cpp
--- a/Engine/Source/Editor/UnrealEd/Private/Commandlets/DerivedDataCacheCommandlet.cpp
+++ b/Engine/Source/Editor/UnrealEd/Private/Commandlets/DerivedDataCacheCommandlet.cpp
@@ -272,9 +272,8 @@ void UDerivedDataCacheCommandlet::CacheWorldPackages(UWorld* World, uint8 Packag
 		});
 	}
 
-	World->ClearWorldComponents();
-	World->CleanupWorld();
-	World->RemoveFromRoot();
+	const bool bBroadcastWorldDestroyedEvent = false;
+	World->DestroyWorld(bBroadcastWorldDestroyedEvent);
 }
 
 int32 UDerivedDataCacheCommandlet::Main( const FString& Params )
diff --git a/Engine/Source/Editor/UnrealEd/Private/Commandlets/PackageUtilities.cpp b/Engine/Source/Editor/UnrealEd/Private/Commandlets/PackageUtilities.cpp
--- a/Engine/Source/Editor/UnrealEd/Private/Commandlets/PackageUtilities.cpp
+++ b/Engine/Source/Editor/UnrealEd/Private/Commandlets/PackageUtilities.cpp
@@ -2255,8 +2255,8 @@ int32 UReplaceActorCommandlet::Main(const FString& Params)
 			// clean up any previous world
 			if (World != NULL)
 			{
-				World->CleanupWorld();
-				World->RemoveFromRoot();
+				const bool bBroadcastWorldDestroyedEvent = false;
+				World->DestroyWorld(bBroadcastWorldDestroyedEvent);
 			}
 
 			// load the package
@@ -2399,8 +2399,8 @@ int32 UReplaceActorCommandlet::Main(const FString& Params)
 				}
 
 				// clear GWorld by removing it from the root set and replacing it with a new one
-				World->CleanupWorld();
-				World->RemoveFromRoot();
+				const bool bBroadcastWorldDestroyedEvent = false;
+				World->DestroyWorld(bBroadcastWorldDestroyedEvent);
 				World = GWorld = NULL;
 			}
 		}
diff --git a/Engine/Source/Editor/UnrealEd/Private/Commandlets/WorldPartitionConvertCommandlet.cpp b/Engine/Source/Editor/UnrealEd/Private/Commandlets/WorldPartitionConvertCommandlet.cpp
--- a/Engine/Source/Editor/UnrealEd/Private/Commandlets/WorldPartitionConvertCommandlet.cpp
+++ b/Engine/Source/Editor/UnrealEd/Private/Commandlets/WorldPartitionConvertCommandlet.cpp
@@ -854,9 +854,8 @@ int32 UWorldPartitionConvertCommandlet::Main(const FString& Params)
 
 	ON_SCOPE_EXIT
 	{
-		MainWorld->ClearWorldComponents();
-		MainWorld->CleanupWorld();
-		MainWorld->RemoveFromRoot();
+		const bool bBroadcastWorldDestroyedEvent = false;
+		MainWorld->DestroyWorld(bBroadcastWorldDestroyedEvent);
 	};
 
 	UPackage* MainPackage = MainLevel->GetPackage();
