Fixed fatal error that could happen if the editor attempted to keep a resident world alive during a map load

diff --git a/Engine/Source/Editor/UnrealEd/Private/EditorServer.cpp b/Engine/Source/Editor/UnrealEd/Private/EditorServer.cpp
--- a/Engine/Source/Editor/UnrealEd/Private/EditorServer.cpp
+++ b/Engine/Source/Editor/UnrealEd/Private/EditorServer.cpp
@@ -2598,12 +2598,16 @@ bool UEditorEngine::Map_Load(const TCHAR* Str, FOutputDevice& Ar)
 
 					// If the existing world has already been initialized then the unload attempts above failed, and we need to 
 					// fatally error and dump any lingering references (like we would when unloading the main editor world).
-					if (ExistingWorld && (ExistingWorld->bIsWorldInitialized || ExistingPackage))
+					if ((ExistingWorld && ExistingWorld->bIsWorldInitialized) || (ExistingPackage && !ExistingWorld))
 					{
 						int32 NumFailedToCleanup = 0;
 
 						if (ExistingWorld)
 						{
+							// If this fires then something in the logic of deciding whether to keep this world 
+							// alive is out-of-sync with ensuring that only uninitialized worlds are kept alive
+							check(ExistingWorld != NewWorld);
+
 							FReferenceChainSearch RefChainSearch(ExistingWorld, EReferenceChainSearchMode::Shortest | EReferenceChainSearchMode::PrintResults);
 							UE_LOG(LogEditorServer, Error, TEXT("Old world %s not cleaned up by garbage collection while loading new map! Referenced by:") LINE_TERMINATOR TEXT("%s"), *ExistingWorld->GetPathName(), *RefChainSearch.GetRootPath());
 							++NumFailedToCleanup;
