- Fixed how the Editor detects if the engine was installed from the Epic Game Launcher in a more robust way (no need to parse the engine identifier).

diff --git a/Engine/Source/Editor/UnrealEd/Private/EditorEngine.cpp b/Engine/Source/Editor/UnrealEd/Private/EditorEngine.cpp
--- a/Engine/Source/Editor/UnrealEd/Private/EditorEngine.cpp
+++ b/Engine/Source/Editor/UnrealEd/Private/EditorEngine.cpp
@@ -7008,8 +7008,9 @@ void UEditorEngine::UpdateIsVanillaProduct()
 		FString Identifier = DesktopPlatform->GetCurrentEngineIdentifier();
 		if (Identifier.Len() > 0)
 		{
-			FEngineVersion Version;
-			if (DesktopPlatform->TryParseStockEngineVersion(Identifier, Version))
+			TMap<FString, FString> Installations;
+			DesktopPlatform->EnumerateLauncherEngineInstallations(Installations);
+			if (Installations.Contains(Identifier))
 			{
 				// Check if we have any marketplace plugins enabled
 				bool bHasMarketplacePlugin = false;
