Fixed batch duplication of static mesh instance elements returning the original (rather than the new) elements

#rnx

diff --git a/Engine/Source/Editor/UnrealEd/Private/Elements/SMInstance/SMInstanceElementEditorWorldInterface.cpp b/Engine/Source/Editor/UnrealEd/Private/Elements/SMInstance/SMInstanceElementEditorWorldInterface.cpp
--- a/Engine/Source/Editor/UnrealEd/Private/Elements/SMInstance/SMInstanceElementEditorWorldInterface.cpp
+++ b/Engine/Source/Editor/UnrealEd/Private/Elements/SMInstance/SMInstanceElementEditorWorldInterface.cpp
@@ -122,7 +122,7 @@ void USMInstanceElementEditorWorldInterface::DuplicateElements(TArrayView<const
 				const bool bOffsetIsZero = InLocationOffset.IsZero();
 
 				OutNewElements.Reserve(OutNewElements.Num() + NewInstanceIds.Num());
-				for (const FSMInstanceId& NewInstanceId : BatchedInstancesToDuplicatePair.Value)
+				for (const FSMInstanceId& NewInstanceId : NewInstanceIds)
 				{
 					if (!bOffsetIsZero)
 					{
