Added multiple packages checkout at the same time, making it fail completely or not at all.

diff --git a/Engine/Source/Editor/UnrealEd/Private/PackageSourceControlHelper.cpp b/Engine/Source/Editor/UnrealEd/Private/PackageSourceControlHelper.cpp
--- a/Engine/Source/Editor/UnrealEd/Private/PackageSourceControlHelper.cpp
+++ b/Engine/Source/Editor/UnrealEd/Private/PackageSourceControlHelper.cpp
@@ -235,48 +235,113 @@ bool FPackageSourceControlHelper::AddToSourceControl(UPackage* Package) const
 
 bool FPackageSourceControlHelper::Checkout(UPackage* Package) const
 {
-	if (UseSourceControl())
+	return !Package || Checkout({ Package->GetName() });
+}
+
+bool FPackageSourceControlHelper::Checkout(const TArray<FString>& PackageNames) const
+{
+	const bool bUseSourceControl = UseSourceControl();
+
+	// Convert package names to package filenames
+	TArray<FString> PackageFilenames = SourceControlHelpers::PackageFilenames(PackageNames);
+
+	// Two-pass checkout mechanism
+	TArray<FString> PackagesToCheckout;
+	PackagesToCheckout.Reserve(PackageFilenames.Num());
+	bool bSomethingFailed = false;
+
+	// In the first pass, we will gather the packages to be checked out, or flag errors and return if we've found any
+	if (bUseSourceControl)
 	{
-		FString PackageFilename = SourceControlHelpers::PackageFilename(Package);
-		FSourceControlStatePtr SourceControlState = GetSourceControlProvider().GetState(PackageFilename, EStateCacheUsage::ForceUpdate);
+		TArray<FSourceControlStateRef> SourceControlStates;
+		ECommandResult::Type UpdateState = GetSourceControlProvider().GetState(PackageFilenames, SourceControlStates, EStateCacheUsage::ForceUpdate);
 
-		if (SourceControlState.IsValid())
+		if (UpdateState != ECommandResult::Succeeded)
 		{
+			UE_LOG(LogCommandletPackageHelper, Error, TEXT("Could not get source control state for packages"));
+			return false;
+		}
+
+		for(FSourceControlStateRef& SourceControlState : SourceControlStates)
+		{
+			const FString& PackageFilename = SourceControlState->GetFilename();
+
 			FString OtherCheckedOutUser;
 			if (SourceControlState->IsCheckedOutOther(&OtherCheckedOutUser))
 			{
-				UE_LOG(LogCommandletPackageHelper, Error, TEXT("Overwriting package %s already checked out by %s, will not submit"), *PackageFilename, *OtherCheckedOutUser);
-				return false;
+				UE_LOG(LogCommandletPackageHelper, Error, TEXT("Overwriting package %s already checked out by %s, will not checkout"), *PackageFilename, *OtherCheckedOutUser);
+				bSomethingFailed = true;
 			}
 			else if (!SourceControlState->IsCurrent())
 			{
-				UE_LOG(LogCommandletPackageHelper, Error, TEXT("Overwriting package %s (not at head revision), will not submit"), *PackageFilename);
-				return false;
+				UE_LOG(LogCommandletPackageHelper, Error, TEXT("Overwriting package %s (not at head revision), will not checkout"), *PackageFilename);
+				bSomethingFailed = true;
 			}
 			else if (SourceControlState->IsCheckedOut() || SourceControlState->IsAdded())
 			{
-				UE_LOG(LogCommandletPackageHelper, Log, TEXT("Skipping package %s (already checked out)"), *PackageFilename);
-				return true;
+				// Nothing to do
 			}
 			else if (SourceControlState->IsSourceControlled())
 			{
-				UE_LOG(LogCommandletPackageHelper, Log, TEXT("Checking out package %s from source control"), *PackageFilename);
-				return GetSourceControlProvider().Execute(ISourceControlOperation::Create<FCheckOut>(), Package) == ECommandResult::Succeeded;
+				PackagesToCheckout.Add(PackageFilename);
 			}
 		}
 	}
 	else
 	{
-		FString PackageFilename = SourceControlHelpers::PackageFilename(Package);
-		if (IPlatformFile::GetPlatformPhysical().IsReadOnly(*PackageFilename))
+		for (const FString& PackageFilename : PackageFilenames)
 		{
-			if (!IPlatformFile::GetPlatformPhysical().SetReadOnly(*PackageFilename, false))
+			if (!IPlatformFile::GetPlatformPhysical().FileExists(*PackageFilename))
 			{
-				UE_LOG(LogCommandletPackageHelper, Error, TEXT("Error setting %s writable"), *PackageFilename);
-				return false;
+				UE_LOG(LogCommandletPackageHelper, Error, TEXT("File %s cannot be checked out as it does not exist"), *PackageFilename);
+				bSomethingFailed = true;
+			}
+			else if (IPlatformFile::GetPlatformPhysical().IsReadOnly(*PackageFilename))
+			{
+				PackagesToCheckout.Add(PackageFilename);
 			}
 		}
 	}
 
-	return true;
-}
+	// Any error up to here will be an early out
+	if (bSomethingFailed)
+	{
+		return false;
+	}
+
+	// In the second pass, we will perform the checkout operation
+	if (PackagesToCheckout.Num() == 0)
+	{
+		return true;
+	}
+	else if (bUseSourceControl)
+	{
+		return GetSourceControlProvider().Execute(ISourceControlOperation::Create<FCheckOut>(), PackagesToCheckout) == ECommandResult::Succeeded;
+	}
+	else
+	{
+		int PackageIndex = 0;
+
+		for (; PackageIndex < PackagesToCheckout.Num(); ++PackageIndex)
+		{
+			if (!IPlatformFile::GetPlatformPhysical().SetReadOnly(*PackagesToCheckout[PackageIndex], false))
+			{
+				UE_LOG(LogCommandletPackageHelper, Error, TEXT("Error setting %s writable"), *PackagesToCheckout[PackageIndex]);
+				bSomethingFailed = true;
+				--PackageIndex;
+				break;
+			}
+		}
+
+		// If a file couldn't be made writeable, put back the files to their original state
+		if (bSomethingFailed)
+		{
+			for (; PackageIndex >= 0; --PackageIndex)
+			{
+				IPlatformFile::GetPlatformPhysical().SetReadOnly(*PackagesToCheckout[PackageIndex], true);
+			}
+		}
+
+		return bSomethingFailed;
+	}
+}
\ No newline at end of file
diff --git a/Engine/Source/Editor/UnrealEd/Public/PackageSourceControlHelper.h b/Engine/Source/Editor/UnrealEd/Public/PackageSourceControlHelper.h
--- a/Engine/Source/Editor/UnrealEd/Public/PackageSourceControlHelper.h
+++ b/Engine/Source/Editor/UnrealEd/Public/PackageSourceControlHelper.h
@@ -16,6 +16,7 @@ public:
 	bool Delete(const TArray<UPackage*>& Packages) const;
 	bool AddToSourceControl(UPackage* Package) const;
 	bool Checkout(UPackage* Package) const;
+	bool Checkout(const TArray<FString>& PackageNames) const;
 
 private:
 	ISourceControlProvider& GetSourceControlProvider() const;
