Fix inverted return condition

diff --git a/Engine/Source/Editor/UnrealEd/Private/PackageSourceControlHelper.cpp b/Engine/Source/Editor/UnrealEd/Private/PackageSourceControlHelper.cpp
--- a/Engine/Source/Editor/UnrealEd/Private/PackageSourceControlHelper.cpp
+++ b/Engine/Source/Editor/UnrealEd/Private/PackageSourceControlHelper.cpp
@@ -248,7 +248,7 @@ bool FPackageSourceControlHelper::Checkout(const TArray<FString>& PackageNames)
 	// Two-pass checkout mechanism
 	TArray<FString> PackagesToCheckout;
 	PackagesToCheckout.Reserve(PackageFilenames.Num());
-	bool bSomethingFailed = false;
+	bool bSuccess = true;
 
 	// In the first pass, we will gather the packages to be checked out, or flag errors and return if we've found any
 	if (bUseSourceControl)
@@ -270,12 +270,12 @@ bool FPackageSourceControlHelper::Checkout(const TArray<FString>& PackageNames)
 			if (SourceControlState->IsCheckedOutOther(&OtherCheckedOutUser))
 			{
 				UE_LOG(LogCommandletPackageHelper, Error, TEXT("Overwriting package %s already checked out by %s, will not checkout"), *PackageFilename, *OtherCheckedOutUser);
-				bSomethingFailed = true;
+				bSuccess = false;
 			}
 			else if (!SourceControlState->IsCurrent())
 			{
 				UE_LOG(LogCommandletPackageHelper, Error, TEXT("Overwriting package %s (not at head revision), will not checkout"), *PackageFilename);
-				bSomethingFailed = true;
+				bSuccess = false;
 			}
 			else if (SourceControlState->IsCheckedOut() || SourceControlState->IsAdded())
 			{
@@ -294,7 +294,7 @@ bool FPackageSourceControlHelper::Checkout(const TArray<FString>& PackageNames)
 			if (!IPlatformFile::GetPlatformPhysical().FileExists(*PackageFilename))
 			{
 				UE_LOG(LogCommandletPackageHelper, Error, TEXT("File %s cannot be checked out as it does not exist"), *PackageFilename);
-				bSomethingFailed = true;
+				bSuccess = false;
 			}
 			else if (IPlatformFile::GetPlatformPhysical().IsReadOnly(*PackageFilename))
 			{
@@ -304,7 +304,7 @@ bool FPackageSourceControlHelper::Checkout(const TArray<FString>& PackageNames)
 	}
 
 	// Any error up to here will be an early out
-	if (bSomethingFailed)
+	if (!bSuccess)
 	{
 		return false;
 	}
@@ -327,14 +327,14 @@ bool FPackageSourceControlHelper::Checkout(const TArray<FString>& PackageNames)
 			if (!IPlatformFile::GetPlatformPhysical().SetReadOnly(*PackagesToCheckout[PackageIndex], false))
 			{
 				UE_LOG(LogCommandletPackageHelper, Error, TEXT("Error setting %s writable"), *PackagesToCheckout[PackageIndex]);
-				bSomethingFailed = true;
+				bSuccess = false;
 				--PackageIndex;
 				break;
 			}
 		}
 
 		// If a file couldn't be made writeable, put back the files to their original state
-		if (bSomethingFailed)
+		if (!bSuccess)
 		{
 			for (; PackageIndex >= 0; --PackageIndex)
 			{
@@ -342,6 +342,6 @@ bool FPackageSourceControlHelper::Checkout(const TArray<FString>& PackageNames)
 			}
 		}
 
-		return bSomethingFailed;
+		return bSuccess;
 	}
 }
\ No newline at end of file
