World Partition Builder - Stop processing if PreWorldInitialization() fails

diff --git a/Engine/Source/Editor/UnrealEd/Private/WorldPartition/WorldPartitionBuilder.cpp b/Engine/Source/Editor/UnrealEd/Private/WorldPartition/WorldPartitionBuilder.cpp
--- a/Engine/Source/Editor/UnrealEd/Private/WorldPartition/WorldPartitionBuilder.cpp
+++ b/Engine/Source/Editor/UnrealEd/Private/WorldPartition/WorldPartitionBuilder.cpp
@@ -61,7 +61,11 @@ bool UWorldPartitionBuilder::RunBuilder(UWorldPartitionBuilder* Builder, UWorld*
 	FPackageSourceControlHelper SCCHelper;
 
 	// Perform builder pre world initialisation
-	Builder->PreWorldInitialization(SCCHelper);
+	if (!Builder->PreWorldInitialization(SCCHelper))
+	{
+		UE_LOG(LogWorldPartitionBuilder, Error, TEXT("PreWorldInitialization failed"));
+		return false;
+	}
 
 	bool bWorldWasRooted = World->IsRooted();
 	if (!bWorldWasRooted)
