Fix Non-Unity Editor build - Missing includes

diff --git a/Engine/Source/Editor/UnrealEd/Private/ImportUtils/InternalImportUtils.cpp b/Engine/Source/Editor/UnrealEd/Private/ImportUtils/InternalImportUtils.cpp
--- a/Engine/Source/Editor/UnrealEd/Private/ImportUtils/InternalImportUtils.cpp
+++ b/Engine/Source/Editor/UnrealEd/Private/ImportUtils/InternalImportUtils.cpp
@@ -13,6 +13,7 @@
 #include "UObject/MetaData.h"
 #include "UObject/NameTypes.h"
 #include "UObject/Object.h"
+#include "UObject/Package.h"
 
 
 bool InternalImportUtils::IsUsingMaterialSlotNameWorkflow(UAssetImportData* AssetImportData)
diff --git a/Engine/Source/Editor/UnrealEd/Private/ImportUtils/StaticMeshImportUtils.cpp b/Engine/Source/Editor/UnrealEd/Private/ImportUtils/StaticMeshImportUtils.cpp
--- a/Engine/Source/Editor/UnrealEd/Private/ImportUtils/StaticMeshImportUtils.cpp
+++ b/Engine/Source/Editor/UnrealEd/Private/ImportUtils/StaticMeshImportUtils.cpp
@@ -10,11 +10,23 @@
 #include "EditorFramework/ThumbnailInfo.h"
 #include "Engine/StaticMeshSocket.h"
 #include "Engine/StaticMesh.h"
+#include "Factories/FbxMeshImportData.h"
+#include "Factories/FbxStaticMeshImportData.h"
+#include "FbxImporter.h"
 #include "IMeshReductionInterfaces.h"
 #include "IMeshReductionManagerModule.h"
 #include "ImportUtils/InternalImportUtils.h"
+#include "Interfaces/ITargetPlatformManagerModule.h"
+#include "Interfaces/ITargetPlatform.h"
 #include "MeshDescription.h"
+#include "Model.h"
 #include "Modules/ModuleManager.h"
+#include "PhysicsEngine/AggregateGeom.h"
+#include "PhysicsEngine/BodySetup.h"
+#include "PhysicsEngine/BoxElem.h"
+#include "PhysicsEngine/ConvexElem.h"
+#include "PhysicsEngine/SphereElem.h"
+#include "PhysicsEngine/SphylElem.h"
 #include "UObject/MetaData.h"
 
 DEFINE_LOG_CATEGORY_STATIC(LogStaticMeshImportUtils, Log, All);
diff --git a/Engine/Source/Editor/UnrealEd/Public/ImportUtils/SkeletalMeshImportUtils.h b/Engine/Source/Editor/UnrealEd/Public/ImportUtils/SkeletalMeshImportUtils.h
--- a/Engine/Source/Editor/UnrealEd/Public/ImportUtils/SkeletalMeshImportUtils.h
+++ b/Engine/Source/Editor/UnrealEd/Public/ImportUtils/SkeletalMeshImportUtils.h
@@ -7,9 +7,11 @@
 
 struct FExistingSkelMeshData;
 struct FReferenceSkeleton;
+struct FSkeletalMaterial;
 class FSkeletalMeshImportData;
 class FSkeletalMeshLODModel;
 class USkeletalMesh;
+class USkeleton;
 
 namespace SkeletalMeshImportUtils
 {
diff --git a/Engine/Source/Editor/UnrealEd/Public/ImportUtils/StaticMeshImportUtils.h b/Engine/Source/Editor/UnrealEd/Public/ImportUtils/StaticMeshImportUtils.h
--- a/Engine/Source/Editor/UnrealEd/Public/ImportUtils/StaticMeshImportUtils.h
+++ b/Engine/Source/Editor/UnrealEd/Public/ImportUtils/StaticMeshImportUtils.h
@@ -19,6 +19,11 @@ class UThumbnailInfo;
 class UModel;
 class UBodySetup;
 
+namespace UnFbx
+{
+	struct FBXImportOptions;
+}
+
 // LOD data to copy over
 struct FExistingLODMeshData
 {
