-Allow remote execution of DDC2 build jobs when the "-DDC2RemoteExecution" commandline is used.

diff --git a/Engine/Source/Editor/UnrealEd/Private/BaseTextureBuildWorkerFactory.cpp b/Engine/Source/Editor/UnrealEd/Private/BaseTextureBuildWorkerFactory.cpp
new file mode 100644
--- /dev/null
+++ b/Engine/Source/Editor/UnrealEd/Private/BaseTextureBuildWorkerFactory.cpp
@@ -0,0 +1,167 @@
+// Copyright Epic Games, Inc. All Rights Reserved.
+
+#include "DerivedDataBuildWorker.h"
+#include "Features/IModularFeatures.h"
+#include "HAL/FileManager.h"
+#include "Misc/Guid.h"
+#include "Misc/Paths.h"
+#include "Misc/PathViews.h"
+#include "Misc/StringBuilder.h"
+#include "Templates/Tuple.h"
+
+// TODO: These explicit Win64 worker factories will be replaced a worker discovery mechanism possibly using the target receipt files produced by the build system.
+
+class FBaseTextureBuildWorkerFactory final : public UE::DerivedData::IBuildWorkerFactory
+{
+public:
+	FBaseTextureBuildWorkerFactory()
+	: EngineDir(FPaths::EngineDir())
+	, bEnabled(true)
+	{
+		GetAllExecutablePaths(ExecutablePaths);
+
+		int32 MetaIndex = 0;
+		ExecutableMeta.AddDefaulted(ExecutablePaths.Num());
+		for (const FWorkerPath& ExecutablePath : ExecutablePaths)
+		{
+			FWorkerPathMeta& Meta = ExecutableMeta[MetaIndex];
+
+			FString LocalExecutablePath = ExecutablePath.LocalPath();
+
+			if (TUniquePtr<FArchive> Ar{IFileManager::Get().CreateFileReader(*LocalExecutablePath, FILEREAD_Silent)})
+			{
+				const int64 TotalSize = Ar->TotalSize();
+				FUniqueBuffer MutableBuffer = FUniqueBuffer::Alloc(uint64(TotalSize));
+				Ar->Serialize(MutableBuffer.GetData(), TotalSize);
+				Meta.Key = FIoHash::HashBuffer(MutableBuffer.GetView());
+				Meta.Value = uint64(TotalSize);
+				Ar->Close();
+			}
+			else
+			{
+				bEnabled = false;
+			}
+			++MetaIndex;
+		}
+
+		if (bEnabled)
+		{
+			IModularFeatures::Get().RegisterModularFeature(IBuildWorkerFactory::GetFeatureName(), this);
+		}
+	}
+
+	virtual ~FBaseTextureBuildWorkerFactory()
+	{
+		if (bEnabled)
+		{
+			IModularFeatures::Get().UnregisterModularFeature(IBuildWorkerFactory::GetFeatureName(), this);
+		}
+	}
+
+	void Build(UE::DerivedData::FBuildWorkerBuilder& Builder) final
+	{
+		Builder.SetName(TEXT("BaseTextureBuildWorker"));
+		Builder.SetPath(ExecutablePaths[0].RemotePath());
+		Builder.SetHostPlatform(TEXT("Win64"));
+		Builder.SetBuildSystemVersion(FGuid(TEXT("ac0574e5-62bd-4c2e-84ec-f2efe48c0fef")));
+
+		// Hard coded to match current contents of TextureFormatUncompressed.  Will be fetched from target receipt in the future.
+		Builder.AddFunction(TEXT("UncompressedTexture"), FGuid(TEXT("c04fe27a-53f6-402e-85b3-648ac6b1ad87")));
+		// Hard coded to match current contents of TextureFormatOodle.  Will be fetched from target receipt in the future.
+		Builder.AddFunction(TEXT("OodleTexture"), FGuid(TEXT("e6b8884f-923a-44a1-8da1-298fb48865b2")));
+
+		for (int32 ExecutableIndex = 0; ExecutableIndex < ExecutablePaths.Num(); ++ExecutableIndex)
+		{
+			Builder.AddExecutable(ExecutablePaths[ExecutableIndex].RemotePath(), ExecutableMeta[ExecutableIndex].Key, ExecutableMeta[ExecutableIndex].Value);
+		}
+	}
+
+	UE::DerivedData::FRequest FindFileData(TConstArrayView<FIoHash> RawHashes, UE::DerivedData::EPriority Priority, UE::DerivedData::FOnBuildWorkerFileDataComplete&& OnComplete) final
+	{
+		if (OnComplete)
+		{
+			TArray<FCompressedBuffer> FileDataBuffers;
+			UE::DerivedData::FBuildWorkerFileDataCompleteParams CompleteParams;
+			for (const FIoHash& RawHash : RawHashes)
+			{
+				for (int32 ExecutableIndex = 0; ExecutableIndex < ExecutablePaths.Num(); ++ExecutableIndex)
+				{
+					FWorkerPathMeta& Meta = ExecutableMeta[ExecutableIndex];
+					if (Meta.Key == RawHash)
+					{
+						if (TUniquePtr<FArchive> Ar{IFileManager::Get().CreateFileReader(*ExecutablePaths[ExecutableIndex].LocalPath(), FILEREAD_Silent)})
+						{
+							const int64 TotalSize = Ar->TotalSize();
+							FUniqueBuffer MutableBuffer = FUniqueBuffer::Alloc(uint64(TotalSize));
+							Ar->Serialize(MutableBuffer.GetData(), TotalSize);
+							if (Ar->Close())
+							{
+								FileDataBuffers.Emplace(FCompressedBuffer::Compress(NAME_None, MutableBuffer.MoveToShared()));
+							}
+						}
+					}
+				}
+			}
+
+			if (FileDataBuffers.Num() == RawHashes.Num())
+			{
+				CompleteParams.Status = UE::DerivedData::EStatus::Ok;
+				CompleteParams.Files = FileDataBuffers;
+				OnComplete(MoveTemp(CompleteParams));
+			}
+			else
+			{
+				CompleteParams.Status = UE::DerivedData::EStatus::Error;
+				OnComplete(MoveTemp(CompleteParams));
+			}
+		}
+
+		return UE::DerivedData::FRequest();
+	}
+
+	struct FWorkerPath
+	{
+		FStringView LocalPathRoot;
+		FStringView LocalPathSuffix;
+		FStringView RemotePathRoot;
+		FStringView RemotePathSuffix;
+
+		FWorkerPath(FStringView InLocalPathRoot, FStringView InLocalPathSuffix, FStringView InRemotePathRoot = FStringView(), FStringView InRemotePathSuffix = FStringView())
+		: LocalPathRoot(InLocalPathRoot)
+		, LocalPathSuffix(InLocalPathSuffix)
+		, RemotePathRoot(InRemotePathRoot.IsEmpty() ? InLocalPathRoot : InRemotePathRoot)
+		, RemotePathSuffix(InRemotePathSuffix.IsEmpty() ? InLocalPathSuffix : InRemotePathSuffix)
+		{
+		}
+
+		FString LocalPath() const
+		{
+			TStringBuilder<128> PathBuilder;
+			FPathViews::Append(PathBuilder, LocalPathRoot, LocalPathSuffix);
+			return PathBuilder.ToString();
+		}
+
+		FString RemotePath() const
+		{
+			TStringBuilder<128> PathBuilder;
+			FPathViews::Append(PathBuilder, RemotePathRoot, RemotePathSuffix);
+			return PathBuilder.ToString();
+		}
+	};
+
+	void GetAllExecutablePaths(TArray<FWorkerPath>& OutPaths)
+	{
+		OutPaths.Emplace(EngineDir, TEXT("Binaries/Win64/BaseTextureBuildWorker.exe"), TEXT("Engine"));
+		OutPaths.Emplace(EngineDir, TEXT("Binaries/ThirdParty/nvTextureTools/Win64/AVX2/nvtt_64.dll"), TEXT("Engine"));
+		OutPaths.Emplace(EngineDir, TEXT("Binaries/ThirdParty/nvTextureTools/Win64/nvtt_64.dll"), TEXT("Engine"));
+	}
+
+private:
+	const FString EngineDir;
+	TArray<FWorkerPath> ExecutablePaths;
+	using FWorkerPathMeta = TTuple<FIoHash, uint64>;
+	TArray<FWorkerPathMeta> ExecutableMeta;
+	bool bEnabled;
+};
+
+FBaseTextureBuildWorkerFactory GBaseTextureBuildWorkerFactory;
diff --git a/Engine/Source/Editor/UnrealEd/Private/DerivedDataBuildRemoteExecutor.cpp b/Engine/Source/Editor/UnrealEd/Private/DerivedDataBuildRemoteExecutor.cpp
new file mode 100644
--- /dev/null
+++ b/Engine/Source/Editor/UnrealEd/Private/DerivedDataBuildRemoteExecutor.cpp
@@ -0,0 +1,633 @@
+// Copyright Epic Games, Inc. All Rights Reserved.
+
+#include "DerivedDataBuildRemoteExecutor.h"
+
+#include "DerivedDataBuild.h"
+#include "DerivedDataBuildAction.h"
+#include "DerivedDataBuildInputs.h"
+#include "DerivedDataBuildOutput.h"
+#include "DerivedDataBuildWorker.h"
+#include "Features/IModularFeatures.h"
+#include "IContentAddressableStorage.h"
+#include "IRemoteExecutor.h"
+#include "Messages.h"
+#include "Misc/CommandLine.h"
+#include "Misc/Optional.h"
+#include "Misc/PathViews.h"
+#include "Modules/ModuleManager.h"
+#include "Serialization/CompactBinaryValidation.h"
+#include "Serialization/CompactBinaryWriter.h"
+
+static uint32 GetTypeHash(const FDigest& Digest)
+{
+	return GetTypeHash(Digest.Hash);
+}
+
+static bool operator==(const FDigest& A, const FDigest& B)
+{
+	return A.Hash == B.Hash;
+}
+
+namespace UE::DerivedData
+{
+
+DEFINE_LOG_CATEGORY_STATIC(LogDerivedDataBuildRemoteExecutor, Log, All);
+
+class FRemoteBuildWorkerExecutor final: public IBuildWorkerExecutor
+{
+public:
+	FRemoteBuildWorkerExecutor()
+	: InstanceName(TEXT("dev.test.reapi"))
+	, Salt(TEXT("807c6a49-0657-49f3-b498-fd457213c0a7"))
+	, BaseDirectoryPath(TEXT("Engine/Binaries/Win64"))
+	, RemoteExecutor(nullptr)
+	, ContentAddressableStorage(nullptr)
+	, bEnabled(false)
+	{
+		const FName RemoteExecutionFeatureName(TEXT("RemoteExecution"));
+		IModularFeatures& ModularFeatures = IModularFeatures::Get();
+
+		if (FParse::Param(FCommandLine::Get(), TEXT("DDC2RemoteExecution")))
+		{
+			FModuleManager::Get().LoadModule("BazelExecutor");
+			if (ModularFeatures.IsModularFeatureAvailable(RemoteExecutionFeatureName))
+			{
+				RemoteExecutor = &ModularFeatures.GetModularFeature<IRemoteExecutor>(RemoteExecutionFeatureName);
+				if (RemoteExecutor)
+				{
+					ContentAddressableStorage = RemoteExecutor->GetContentAddressableStorage();
+				}
+			}
+		}
+
+		bEnabled = (RemoteExecutor != nullptr) && (ContentAddressableStorage != nullptr);
+		if (bEnabled)
+		{
+			ModularFeatures.RegisterModularFeature(IBuildWorkerExecutor::GetFeatureName(), this);
+		}
+	}
+
+	virtual ~FRemoteBuildWorkerExecutor()
+	{
+		if (bEnabled)
+		{
+			IModularFeatures::Get().UnregisterModularFeature(IBuildWorkerExecutor::GetFeatureName(), this);
+		}
+	}
+
+	enum class ENodeType
+	{
+		File,
+		Directory,
+		//Symlink,
+	};
+
+	enum class EFileType
+	{
+		Worker,
+		Input,
+		BuildAction,
+	};
+
+	struct FVariantIndex
+	{
+		ENodeType NodeType;
+		int32 Index;
+
+		FVariantIndex(ENodeType InNodeType, int32 InIndex)
+		: NodeType(InNodeType)
+		, Index(InIndex)
+		{
+		}
+	};
+
+	struct FMerkleTreeFileBuilder
+	{
+		FFileNode File;
+		EFileType Type;
+		FCompositeBuffer ContentBytes;
+	};
+
+	struct FMerkleTreeDirectoryBuilder
+	{
+		FStringView Name;
+		FDirectory Directory;
+		TOptional<FDigest> Digest;
+		TArray<char> ContentBytes;
+		TArray<int32> SubDirIndices;
+	};
+
+	struct FRemoteExecutionState
+	{
+		const FBuildAction& BuildAction;
+		const FOptionalBuildInputs& BuildInputs;
+		const FBuildWorker& BuildWorker;
+		EBuildPolicy BuildPolicy;
+		EPriority BuildPriority;
+
+		// Unordered arrays that are indexed into
+		TArray<FMerkleTreeDirectoryBuilder> Directories;
+		TArray<FMerkleTreeFileBuilder> Files;
+
+		// Lookup tables for indexing in different scenarios
+		TMultiMap<FDigest, FVariantIndex> DigestFilesystemIndex;
+		TMap<FStringView, int32> PathToDirectoryIndex;
+		TMap<int32, FStringView> FileIndexToInputKey;
+
+		// Unique items in the tree
+		FCommand Command;
+		TArray<char> CommandContentBytes;
+		FAction Action;
+		TArray<char> ActionContentBytes;
+		FDigest ActionDigest;
+		FCompositeBuffer BuildActionContentBytes {};
+		FDigest BuildActionDigest;
+		TArray<FString> InputPaths;
+
+		// Step 1: Find missing blobs
+		FFindMissingBlobsRequest FindMissingBlobsRequest;
+		FFindMissingBlobsResponse FindMissingBlobsResponse;
+
+		// Step 3: Batch update blobs (upload)
+		FBatchUpdateBlobsRequest BatchUpdateBlobsRequest;
+		FBatchUpdateBlobsResponse BatchUpdateBlobsResponse;
+
+		// Step 4: Execute
+		FExecuteRequest ExecuteRequest;
+		FExecuteResponse ExecuteResponse;
+
+		// Step 5: Batch read blobs (download)
+		FBatchReadBlobsRequest BatchReadBlobsRequest;
+		FBatchReadBlobsResponse BatchReadBlobsResponse;
+	};
+
+	static FMerkleTreeDirectoryBuilder& GetOrAddMerkleTreeDirectory(FRemoteExecutionState& State, FStringView Path, int32& OutDirectoryBuilderIndex)
+	{
+		int32& DirectoryBuilderIndex = State.PathToDirectoryIndex.FindOrAdd(Path, INDEX_NONE);
+		if (DirectoryBuilderIndex == INDEX_NONE)
+		{
+			DirectoryBuilderIndex = State.Directories.Num();
+			FMerkleTreeDirectoryBuilder& NewNode = State.Directories.AddDefaulted_GetRef();
+			NewNode.Name = Path.IsEmpty() ? Path : FPathViews::GetCleanFilename(Path);
+		}
+
+		OutDirectoryBuilderIndex = DirectoryBuilderIndex;
+
+		int32 ContainingDirectoryIndex = DirectoryBuilderIndex;
+		while (!Path.IsEmpty())
+		{
+			int32 SubDirIndex = ContainingDirectoryIndex;
+			Path = FPathViews::GetPath(Path);
+			GetOrAddMerkleTreeDirectory(State, Path, ContainingDirectoryIndex).SubDirIndices.AddUnique(SubDirIndex);
+		}
+
+		return State.Directories[OutDirectoryBuilderIndex];
+	}
+
+	static int32 AddMerkleTreeFile(FRemoteExecutionState& State, FStringView Path, const FIoHash& RawHash, uint64 RawSize, bool bIsExecutable, EFileType FileType, FCompositeBuffer ContentBytes = FCompositeBuffer::Null)
+	{
+		FStringView ContainingDirectoryPath = FPathViews::GetPath(Path);
+		int32 ContainingDirectoryIndex = INDEX_NONE;
+		FMerkleTreeDirectoryBuilder* ContainingDirectoryBuilder = &GetOrAddMerkleTreeDirectory(State, ContainingDirectoryPath, ContainingDirectoryIndex);
+		FFileNode& NewNode = ContainingDirectoryBuilder->Directory.Files.AddDefaulted_GetRef();
+		NewNode.Name = FPathViews::GetCleanFilename(Path);
+		NewNode.Digest.Hash = RawHash;
+		NewNode.Digest.SizeBytes = RawSize;
+		NewNode.IsExecutable = bIsExecutable;
+
+		int32 NewFileIndex = State.Files.Num();
+		State.DigestFilesystemIndex.Add(NewNode.Digest, FVariantIndex(ENodeType::File, NewFileIndex));
+		FMerkleTreeFileBuilder& FileBuilder = State.Files.AddDefaulted_GetRef();
+		FileBuilder.File = NewNode; // Duplicates the node in the state's file array
+		FileBuilder.Type = FileType;
+		if (ContentBytes)
+		{
+			FileBuilder.ContentBytes = ContentBytes;
+		}
+
+		return NewFileIndex;
+	}
+
+	const FDigest& BuildMerkleTreeDirectoryDigest(FRemoteExecutionState& State, int32 Index)
+	{
+		FMerkleTreeDirectoryBuilder& DirBuilder = State.Directories[Index];
+
+		if (DirBuilder.Digest.IsSet())
+		{
+			return DirBuilder.Digest.GetValue();
+		}
+
+		for (int32 SubDirIndex : DirBuilder.SubDirIndices)
+		{
+			FDirectoryNode& SubDirNode = DirBuilder.Directory.Directories.AddDefaulted_GetRef();
+			SubDirNode.Name = State.Directories[SubDirIndex].Name;
+			SubDirNode.Digest = BuildMerkleTreeDirectoryDigest(State, SubDirIndex);
+			State.DigestFilesystemIndex.Add(SubDirNode.Digest, FVariantIndex(ENodeType::Directory, SubDirIndex));
+		}
+
+		DirBuilder.Directory.Directories.Sort([] (const FDirectoryNode& A, const FDirectoryNode& B)
+		{
+			return A.Name < B.Name;
+		});
+
+		DirBuilder.Directory.Files.Sort([] (const FFileNode& A, const FFileNode& B)
+		{
+			return A.Name < B.Name;
+		});
+
+		ContentAddressableStorage->ToBlob(DirBuilder.Directory, DirBuilder.ContentBytes, DirBuilder.Digest.Emplace());
+
+		return DirBuilder.Digest.GetValue();
+	}
+
+	void BuildMerkleTreeNodes(FRemoteExecutionState& State)
+	{
+		TArray<FIoHash> WorkerFileHashes;
+		TArray<TTuple<FStringView, bool>> WorkerFileMeta;
+		State.BuildWorker.IterateExecutables([&WorkerFileHashes, &WorkerFileMeta] (FStringView Path, const FIoHash& RawHash, uint64 RawSize)
+			{
+				WorkerFileHashes.Emplace(RawHash);
+				WorkerFileMeta.Emplace(Path, true);
+			});
+
+		State.BuildWorker.IterateFiles([&WorkerFileHashes, &WorkerFileMeta] (FStringView Path, const FIoHash& RawHash, uint64 RawSize)
+			{
+				WorkerFileHashes.Emplace(RawHash);
+				WorkerFileMeta.Emplace(Path, false);
+			});
+
+		State.BuildWorker.FindFileData(WorkerFileHashes, EPriority::Normal,
+			[&State, &WorkerFileMeta] (FBuildWorkerFileDataCompleteParams&& Params)
+			{
+				uint32 MetaIndex = 0;
+				for (const FCompressedBuffer& Buffer : Params.Files)
+				{
+					const TTuple<FStringView, bool>& Meta = WorkerFileMeta[MetaIndex];
+					FCompositeBuffer DecompressedComposite = Buffer.DecompressToComposite();
+					AddMerkleTreeFile(State, Meta.Key, FIoHash::HashBuffer(DecompressedComposite), DecompressedComposite.GetSize(), Meta.Value, EFileType::Worker, Buffer.DecompressToComposite());
+					++MetaIndex;
+				}
+			}).Wait();
+
+		if (State.BuildInputs.IsValid())
+		{
+			State.BuildInputs.Get().IterateInputs([&State] (FStringView Key, const FCompressedBuffer& Buffer)
+				{
+					TStringBuilder<128> InputPath;
+					InputPath << TEXT("Inputs/") << FIoHash(Buffer.GetRawHash());
+					const FString& NewInputPath = State.InputPaths.Emplace_GetRef(InputPath);
+					int32 FileIndex = AddMerkleTreeFile(State, NewInputPath, FIoHash::HashBuffer(Buffer.GetCompressed()), Buffer.GetCompressedSize(), false, EFileType::Input, Buffer.GetCompressed());
+					State.FileIndexToInputKey.Add(FileIndex, Key);
+				});
+		}
+
+		// This base directory must be created as worker executables (even those that don't exist in this directory) will attempt to change directories into it during startup.
+		int32 BaseDirectoryIndex = INDEX_NONE;
+		GetOrAddMerkleTreeDirectory(State, BaseDirectoryPath, BaseDirectoryIndex);
+
+		FCbWriter BuildActionWriter;
+		State.BuildAction.Save(BuildActionWriter);
+		FUniqueBuffer UncompressedBuildActionContentBytes = FUniqueBuffer::Alloc(BuildActionWriter.GetSaveSize());
+		BuildActionWriter.Save(UncompressedBuildActionContentBytes);
+		State.BuildActionContentBytes = FCompositeBuffer(UncompressedBuildActionContentBytes.MoveToShared());
+		State.BuildActionDigest.Hash = FIoHash::HashBuffer(State.BuildActionContentBytes);
+		State.BuildActionDigest.SizeBytes = State.BuildActionContentBytes.GetSize();
+		AddMerkleTreeFile(State, TEXT("build.uddba"), State.BuildActionDigest.Hash, State.BuildActionDigest.SizeBytes, false, EFileType::BuildAction, State.BuildActionContentBytes);
+
+
+		if (!State.PathToDirectoryIndex.IsEmpty())
+		{
+			int32 RootDirectoryIndex = State.PathToDirectoryIndex.FindChecked(TEXT(""));
+			State.Action.InputRootDigest = BuildMerkleTreeDirectoryDigest(State, RootDirectoryIndex);
+			State.DigestFilesystemIndex.Add(State.Action.InputRootDigest, FVariantIndex(ENodeType::Directory, RootDirectoryIndex));
+		}
+
+		State.Command.OutputPaths.Add("Outputs");
+		State.Command.OutputPaths.Add("build.uddbo");
+		State.Command.Arguments.Add(FString(State.BuildWorker.GetPath()));
+		State.Command.Arguments.Add("-b=build.uddba");
+		State.BuildWorker.IterateEnvironment([&State] (FStringView Name, FStringView Value)
+			{
+				FCommand::FEnvironmentVariable& EnvVar = State.Command.EnvironmentVariables.AddDefaulted_GetRef();
+				EnvVar.Name = Name;
+				EnvVar.Value = Value;
+			});
+		State.Command.EnvironmentVariables.Sort([] (const FCommand::FEnvironmentVariable& A, const FCommand::FEnvironmentVariable& B)
+			{
+				if (A.Name < B.Name)
+				{
+					return true;
+				}
+				else if (B.Name < A.Name)
+				{
+					return false;
+				}
+				return A.Value < B.Value;
+			});
+		ContentAddressableStorage->ToBlob(State.Command, State.CommandContentBytes, State.Action.CommandDigest);
+
+		State.Action.DoNotCache = true;
+		State.Action.Salt = Salt;
+		ContentAddressableStorage->ToBlob(State.Action, State.ActionContentBytes, State.ActionDigest);
+	}
+
+	// Each step represents a synchronous unit of work that ends with an operation we want to have async in the future
+	FStatus DetermineMissingBlobs(FRemoteExecutionState& State)
+	{
+		BuildMerkleTreeNodes(State);
+
+		State.FindMissingBlobsRequest.InstanceName = InstanceName;
+		State.FindMissingBlobsRequest.BlobDigests.Add(State.ActionDigest);
+		State.FindMissingBlobsRequest.BlobDigests.Add(State.Action.CommandDigest);
+		for (const TPair<FDigest, FVariantIndex>& FilesystemItem : State.DigestFilesystemIndex)
+		{
+			State.FindMissingBlobsRequest.BlobDigests.AddUnique(FilesystemItem.Key);
+		}
+		return ContentAddressableStorage->FindMissingBlobs(State.FindMissingBlobsRequest, State.FindMissingBlobsResponse);
+	}
+
+	void GatherMissingInputFileBlobs(FRemoteExecutionState& State, TArray<FStringView>& OutMissingInputs)
+	{
+		for (const FDigest& MissingItem : State.FindMissingBlobsResponse.MissingBlobDigests)
+		{
+			TArray<FVariantIndex> DigestFilesystemEntries;
+			State.DigestFilesystemIndex.MultiFind(MissingItem, DigestFilesystemEntries);
+
+			for (FVariantIndex& VariantIndex : DigestFilesystemEntries)
+			{
+				if (VariantIndex.NodeType == ENodeType::File)
+				{
+					const FMerkleTreeFileBuilder& File = State.Files[VariantIndex.Index];
+					if ((File.Type == EFileType::Input) && File.ContentBytes.IsNull())
+					{
+						OutMissingInputs.Add(State.FileIndexToInputKey[VariantIndex.Index]);
+					}
+				}
+			}
+		}
+	}
+
+	void LoadMissingWorkerFileBlobs(FRemoteExecutionState& State)
+	{
+		TArray<FIoHash> WorkerFileHashes;
+		TMultiMap<FIoHash, uint32> WorkerFileMapping;
+		for (const FDigest& MissingItem : State.FindMissingBlobsResponse.MissingBlobDigests)
+		{
+			TArray<FVariantIndex> DigestFilesystemEntries;
+			State.DigestFilesystemIndex.MultiFind(MissingItem, DigestFilesystemEntries);
+
+			for (FVariantIndex& VariantIndex : DigestFilesystemEntries)
+			{
+				if (VariantIndex.NodeType == ENodeType::File)
+				{
+					const FMerkleTreeFileBuilder& File = State.Files[VariantIndex.Index];
+					if ((File.Type == EFileType::Worker) && File.ContentBytes.IsNull())
+					{
+						WorkerFileHashes.Emplace(MissingItem.Hash);
+						WorkerFileMapping.Add(MissingItem.Hash, VariantIndex.Index);
+					}
+				}
+			}
+		}
+
+		State.BuildWorker.FindFileData(WorkerFileHashes, EPriority::Normal,
+			[&State, &WorkerFileMapping] (FBuildWorkerFileDataCompleteParams&& Params)
+			{
+				uint32 MetaIndex = 0;
+				for (const FCompressedBuffer& Buffer : Params.Files)
+				{
+					FCompositeBuffer UncompressedWorkerFile = Buffer.DecompressToComposite();
+					TArray<uint32> WorkerFileIndices;
+					WorkerFileMapping.MultiFind(Buffer.GetRawHash(), WorkerFileIndices);
+
+					for (int32 FileIndex : WorkerFileIndices)
+					{
+						State.Files[FileIndex].ContentBytes = UncompressedWorkerFile;
+					}
+				}
+			}).Wait();
+	}
+
+	FStatus UploadMissingBlobs(FRemoteExecutionState& State)
+	{
+		State.BatchUpdateBlobsRequest.InstanceName = InstanceName;
+		for (const FDigest& MissingItem : State.FindMissingBlobsResponse.MissingBlobDigests)
+		{
+			if (MissingItem == State.ActionDigest)
+			{
+				FBatchUpdateBlobsRequest::FRequest& NewRequest = State.BatchUpdateBlobsRequest.Requests.AddDefaulted_GetRef();
+				NewRequest.Digest = State.ActionDigest;
+				NewRequest.Data = MoveTemp(State.ActionContentBytes);
+			}
+			else if (MissingItem == State.Action.CommandDigest)
+			{
+				FBatchUpdateBlobsRequest::FRequest& NewRequest = State.BatchUpdateBlobsRequest.Requests.AddDefaulted_GetRef();
+				NewRequest.Digest = State.Action.CommandDigest;
+				NewRequest.Data = MoveTemp(State.CommandContentBytes);
+			}
+			else
+			{
+				const FVariantIndex VariantIndex = State.DigestFilesystemIndex.FindChecked(MissingItem);
+				FBatchUpdateBlobsRequest::FRequest& NewRequest = State.BatchUpdateBlobsRequest.Requests.AddDefaulted_GetRef();
+				NewRequest.Digest = MissingItem;
+				switch (VariantIndex.NodeType)
+				{
+				case ENodeType::Directory:
+					NewRequest.Digest = State.Directories[VariantIndex.Index].Digest.GetValue();
+					NewRequest.Data = MoveTemp(State.Directories[VariantIndex.Index].ContentBytes);
+					break;
+				case ENodeType::File:
+					{
+						NewRequest.Digest.Hash = FIoHash::HashBuffer(State.Files[VariantIndex.Index].ContentBytes);
+						NewRequest.Digest.SizeBytes = State.Files[VariantIndex.Index].ContentBytes.GetSize();
+						FCompositeBuffer& FileBuffer = State.Files[VariantIndex.Index].ContentBytes;
+						check(!FileBuffer.IsNull());
+						NewRequest.Data.Reserve(FileBuffer.GetSize());
+						for (const FSharedBuffer& Segment : State.Files[VariantIndex.Index].ContentBytes.GetSegments())
+						{
+							NewRequest.Data.Append((const char *)Segment.GetData(), Segment.GetSize());
+						}
+						FileBuffer.Reset();
+					}
+					break;
+				default:
+					checkNoEntry();
+					break;
+				}
+			}
+		}
+
+		return ContentAddressableStorage->BatchUpdateBlobs(State.BatchUpdateBlobsRequest, State.BatchUpdateBlobsResponse);
+	}
+
+	bool ExecuteBuild(FRemoteExecutionState& State)
+	{
+		State.ExecuteRequest.InstanceName = InstanceName;
+		State.ExecuteRequest.ActionDigest = State.ActionDigest;
+		State.ExecuteRequest.SkipCacheLookup = true;
+		return RemoteExecutor->GetExecution()->Execute(State.ExecuteRequest, State.ExecuteResponse);
+	}
+
+	FStatus DownloadResults(FRemoteExecutionState& State)
+	{
+		State.BatchReadBlobsRequest.InstanceName = InstanceName;
+		//State.BatchReadBlobsRequest.Digests.Add(State.ExecuteResponse.Result.StdoutDigest);
+		//State.BatchReadBlobsRequest.Digests.Add(State.ExecuteResponse.Result.StderrDigest);
+
+		for (const FOutputFile& OutputFile : State.ExecuteResponse.Result.OutputFiles)
+		{
+			State.BatchReadBlobsRequest.Digests.Add(OutputFile.Digest);
+		}
+
+		return ContentAddressableStorage->BatchReadBlobs(State.BatchReadBlobsRequest, State.BatchReadBlobsResponse);
+	}
+
+	FOptionalBuildOutput ComposeBuildOutput(FRemoteExecutionState& State, EStatus& OutStatus)
+	{
+		OutStatus = EStatus::Error;
+
+		TOptional<FDigest> BuildOutputDigest;
+		for (const FOutputFile& ExecuteOutputFile : State.ExecuteResponse.Result.OutputFiles)
+		{
+			if (ExecuteOutputFile.Path == TEXT("build.uddbo"))
+			{
+				BuildOutputDigest = ExecuteOutputFile.Digest;
+				break;
+			}
+		}
+
+		if (!BuildOutputDigest.IsSet())
+		{
+			UE_LOG(LogDerivedDataBuildRemoteExecutor, Warning, TEXT("Worker error: build output structure not produced!"));
+			return FOptionalBuildOutput();
+		}
+
+		FOptionalBuildOutput RemoteBuildOutput;
+
+		TMap<FIoHash, FCompressedBuffer> PayloadResponses;
+		for (FBatchReadBlobsResponse::FResponse& OutputReadResponse : State.BatchReadBlobsResponse.Responses)
+		{
+			if (!OutputReadResponse.Status.Ok())
+			{
+				UE_LOG(LogDerivedDataBuildRemoteExecutor, Warning, TEXT("Remote execution system error: output not downloaded!"));
+				return FOptionalBuildOutput();
+			}
+
+			if (OutputReadResponse.Digest == BuildOutputDigest)
+			{
+				FSharedBuffer BuildOutputBuffer = MakeSharedBufferFromArray(MoveTemp(OutputReadResponse.Data));
+
+				if (ValidateCompactBinaryRange(BuildOutputBuffer, ECbValidateMode::Default) != ECbValidateError::None)
+				{
+					UE_LOG(LogDerivedDataBuildRemoteExecutor, Warning, TEXT("Worker error: build output structure not valid!"));
+					return FOptionalBuildOutput();
+				}
+
+				RemoteBuildOutput = GetDerivedDataBuildRef().LoadOutput(State.BuildAction.GetName(), State.BuildAction.GetFunction(), FCbObject(BuildOutputBuffer));
+			}
+			else
+			{
+				FCompressedBuffer NewBuffer = FCompressedBuffer::FromCompressed(MakeSharedBufferFromArray(MoveTemp(OutputReadResponse.Data)));
+				PayloadResponses.Add(NewBuffer.GetRawHash(), MoveTemp(NewBuffer));
+			}
+		}
+
+		if (RemoteBuildOutput.IsNull())
+		{
+			UE_LOG(LogDerivedDataBuildRemoteExecutor, Warning, TEXT("Remote execution system error: build output blob missing!"));
+			return FOptionalBuildOutput();
+		}
+
+		FBuildOutputBuilder OutputBuilder = GetDerivedDataBuildRef().CreateOutput(State.BuildAction.GetName(), State.BuildAction.GetFunction());
+		
+		RemoteBuildOutput.Get().IterateDiagnostics( [&OutputBuilder](const FBuildDiagnostic& Diagnostic)
+		{
+			if (Diagnostic.Level == EBuildDiagnosticLevel::Warning)
+			{
+				OutputBuilder.AddWarning(Diagnostic.Category, Diagnostic.Message);
+			}
+			else if (Diagnostic.Level == EBuildDiagnosticLevel::Error)
+			{
+				OutputBuilder.AddError(Diagnostic.Category, Diagnostic.Message);
+			}
+		});
+
+		for (const FPayload& Payload : RemoteBuildOutput.Get().GetPayloads())
+		{
+			FCompressedBuffer* BufferForPayload = PayloadResponses.Find(Payload.GetRawHash());
+			if (!BufferForPayload)
+			{
+				UE_LOG(LogDerivedDataBuildRemoteExecutor, Warning, TEXT("Remote execution system error: payload blob missing!"));
+				return FOptionalBuildOutput();
+			}
+
+			OutputBuilder.AddPayload(FPayload(Payload.GetId(), *BufferForPayload));
+		}
+
+		OutStatus = State.ExecuteResponse.Result.ExitCode == 0 ? EStatus::Ok : EStatus::Error;
+		return OutputBuilder.Build();
+	}
+
+	FRequest BuildAction(
+		const FBuildAction& Action,
+		const FOptionalBuildInputs& Inputs,
+		const FBuildWorker& Worker,
+		EBuildPolicy Policy,
+		EPriority Priority,
+		FOnBuildWorkerActionComplete&& OnComplete) final
+	{
+		FRemoteExecutionState State { Action, Inputs, Worker, Policy, Priority };
+
+		DetermineMissingBlobs(State);
+
+		TArray<FStringView> MissingInputs;
+		GatherMissingInputFileBlobs(State, MissingInputs);
+		if (!MissingInputs.IsEmpty())
+		{
+			OnComplete({State.BuildAction.GetKey(), {}, MissingInputs, EStatus::Ok});
+			return FRequest();
+		}
+
+		LoadMissingWorkerFileBlobs(State);
+		if (!State.FindMissingBlobsResponse.MissingBlobDigests.IsEmpty())
+		{
+			UploadMissingBlobs(State);
+		}
+
+		if (ExecuteBuild(State))
+		{
+			DownloadResults(State);
+
+			EStatus BuildStatus = EStatus::Error;
+			FOptionalBuildOutput BuildOutput = ComposeBuildOutput(State, BuildStatus);
+			OnComplete({State.BuildAction.GetKey(), MoveTemp(BuildOutput), {}, BuildStatus});
+		}
+		else
+		{
+			UE_LOG(LogDerivedDataBuildRemoteExecutor, Warning, TEXT("Remote execution system error: Failed to execute build operation!"));
+			OnComplete({State.BuildAction.GetKey(), {}, {}, EStatus::Error});
+		}
+
+		return FRequest();
+	}
+
+	TConstArrayView<FStringView> GetHostPlatforms() const final
+	{
+		static constexpr FStringView HostPlatforms[]{TEXT("Win64"_SV)};
+		return HostPlatforms;
+	}
+private:
+	const FStringView InstanceName;
+	const FStringView Salt;
+	const FString BaseDirectoryPath;
+	IRemoteExecutor* RemoteExecutor;
+	IContentAddressableStorage* ContentAddressableStorage;
+	bool bEnabled;
+};
+
+} // namespace UE::DerivedData
+
+void InitDerivedDataBuildRemoteExecutor()
+{
+	static UE::DerivedData::FRemoteBuildWorkerExecutor GRemoteBuildWorkerExecutor;
+}
diff --git a/Engine/Source/Editor/UnrealEd/Private/DerivedDataBuildRemoteExecutor.h b/Engine/Source/Editor/UnrealEd/Private/DerivedDataBuildRemoteExecutor.h
new file mode 100644
--- /dev/null
+++ b/Engine/Source/Editor/UnrealEd/Private/DerivedDataBuildRemoteExecutor.h
@@ -0,0 +1,5 @@
+// Copyright Epic Games, Inc. All Rights Reserved.
+
+#pragma once
+
+void InitDerivedDataBuildRemoteExecutor();
diff --git a/Engine/Source/Editor/UnrealEd/Private/EditorEngine.cpp b/Engine/Source/Editor/UnrealEd/Private/EditorEngine.cpp
--- a/Engine/Source/Editor/UnrealEd/Private/EditorEngine.cpp
+++ b/Engine/Source/Editor/UnrealEd/Private/EditorEngine.cpp
@@ -237,6 +237,7 @@
 #include "Engine/TextureCube.h"
 #include "Misc/PackageAccessTracking.h"
 #include "Misc/PackageAccessTrackingOps.h"
+#include "DerivedDataBuildRemoteExecutor.h"
 
 #if WITH_CHAOS
 #include "ChaosSolversModule.h"
@@ -626,6 +627,10 @@ static bool GetDisplayMultiboxHooks()
 
 void UEditorEngine::InitEditor(IEngineLoop* InEngineLoop)
 {
+	// Allow remote execution of derived data builds from this point
+	// TODO: This needs to be enabled earlier to allow early data builds to be remote executed.
+	InitDerivedDataBuildRemoteExecutor();
+
 	// Call base.
 	UEngine::Init(InEngineLoop);
 
diff --git a/Engine/Source/Editor/UnrealEd/UnrealEd.Build.cs b/Engine/Source/Editor/UnrealEd/UnrealEd.Build.cs
--- a/Engine/Source/Editor/UnrealEd/UnrealEd.Build.cs
+++ b/Engine/Source/Editor/UnrealEd/UnrealEd.Build.cs
@@ -67,6 +67,7 @@ public class UnrealEd : ModuleRules
 				"NavigationSystem",
 				"Media",
 				"VirtualTexturingEditor",
+				"RemoteExecution"
 			}
 		);
 
@@ -201,6 +202,7 @@ public class UnrealEd : ModuleRules
 				"AnimationModifiers",
 				"AnimationBlueprintLibrary",
 				"MaterialBaking",
+				"BazelExecutor",
 			}
 		);
 
