Harvest soft object path package on save for cooker to handle save permutations such as the world partition grid cell split:
* Added a function to RedirectCollector to ingest soft package references
* Push the harevest soft package references from save to the redirect collector on a succesful save
* Moved the processing of soft package reference associated with a package during cook to after that package has been saved

diff --git a/Engine/Source/Editor/UnrealEd/Private/CookOnTheFlyServer.cpp b/Engine/Source/Editor/UnrealEd/Private/CookOnTheFlyServer.cpp
--- a/Engine/Source/Editor/UnrealEd/Private/CookOnTheFlyServer.cpp
+++ b/Engine/Source/Editor/UnrealEd/Private/CookOnTheFlyServer.cpp
@@ -4218,40 +4218,6 @@ void UCookOnTheFlyServer::SaveCookedPackage(UE::Cook::FPackageData& PackageData,
 		}
 	}
 
-	// Don't resolve, just add to request list as needed
-	TSet<FName> SoftObjectPackages;
-
-	if (!IsCookByTheBookMode() || !CookByTheBookOptions->bSkipSoftReferences)
-	{
-		GRedirectCollector.ProcessSoftObjectPathPackageList(Package->GetFName(), false, SoftObjectPackages);
-
-		for (FName SoftObjectPackage : SoftObjectPackages)
-		{
-			TMap<FName, FName> RedirectedPaths;
-
-			// If this is a redirector, extract destination from asset registry
-			if (ContainsRedirector(SoftObjectPackage, RedirectedPaths))
-			{
-				for (TPair<FName, FName>& RedirectedPath : RedirectedPaths)
-				{
-					GRedirectCollector.AddAssetPathRedirection(RedirectedPath.Key, RedirectedPath.Value);
-				}
-			}
-
-			// Verify package actually exists
-
-			if (IsCookByTheBookMode())
-			{
-				UE::Cook::FPackageData* SoftObjectPackageData = PackageDatas->TryAddPackageDataByPackageName(SoftObjectPackage);
-				if (SoftObjectPackageData)
-				{
-					bool bIsUrgent = false;
-					SoftObjectPackageData->UpdateRequestData(PackageData.GetRequestedPlatforms(), bIsUrgent, UE::Cook::FCompletionCallback());
-				}
-			}
-		}
-	}
-
 	if (Filename.Len() != 0 )
 	{
 		if (Package->HasAnyPackageFlags(PKG_ReloadingForCooker))
@@ -4439,6 +4405,38 @@ void UCookOnTheFlyServer::SaveCookedPackage(UE::Cook::FPackageData& PackageData,
 		}
 	}
 
+	// Don't resolve, just add to request list as needed, do this after save to catch save permutations
+	TSet<FName> SoftObjectPackages;
+	if (!IsCookByTheBookMode() || !CookByTheBookOptions->bSkipSoftReferences)
+	{
+		GRedirectCollector.ProcessSoftObjectPathPackageList(Package->GetFName(), false, SoftObjectPackages);
+
+		for (FName SoftObjectPackage : SoftObjectPackages)
+		{
+			TMap<FName, FName> RedirectedPaths;
+
+			// If this is a redirector, extract destination from asset registry
+			if (ContainsRedirector(SoftObjectPackage, RedirectedPaths))
+			{
+				for (TPair<FName, FName>& RedirectedPath : RedirectedPaths)
+				{
+					GRedirectCollector.AddAssetPathRedirection(RedirectedPath.Key, RedirectedPath.Value);
+				}
+			}
+
+			// Verify package actually exists
+			if (IsCookByTheBookMode())
+			{
+				UE::Cook::FPackageData* SoftObjectPackageData = PackageDatas->TryAddPackageDataByPackageName(SoftObjectPackage);
+				if (SoftObjectPackageData)
+				{
+					bool bIsUrgent = false;
+					SoftObjectPackageData->UpdateRequestData(PackageData.GetRequestedPlatforms(), bIsUrgent, UE::Cook::FCompletionCallback());
+				}
+			}
+		}
+	}
+
 	check(bIsSavingPackage == true);
 	bIsSavingPackage = false;
 
diff --git a/Engine/Source/Runtime/CoreUObject/Private/Misc/RedirectCollector.cpp b/Engine/Source/Runtime/CoreUObject/Private/Misc/RedirectCollector.cpp
--- a/Engine/Source/Runtime/CoreUObject/Private/Misc/RedirectCollector.cpp
+++ b/Engine/Source/Runtime/CoreUObject/Private/Misc/RedirectCollector.cpp
@@ -1,6 +1,7 @@
 // Copyright Epic Games, Inc. All Rights Reserved.
 
 #include "Misc/RedirectCollector.h"
+#include "Algo/Transform.h"
 #include "Misc/CoreDelegates.h"
 #include "UObject/UObjectGlobals.h"
 #include "UObject/Object.h"
@@ -47,6 +48,18 @@ void FRedirectCollector::OnSoftObjectPathLoaded(const FSoftObjectPath& InPath, F
 	SoftObjectPathMap.FindOrAdd(PackageName).Add(SoftObjectPathProperty);
 }
 
+void FRedirectCollector::CollectSavedSoftPackageReferences(FName ReferencingPackage, const TSet<FName>& PackageNames, bool bEditorOnlyReferences)
+{
+	TArray<FSoftObjectPathProperty, TInlineAllocator<4>> SoftObjectPathArray;
+	Algo::Transform(PackageNames, SoftObjectPathArray, [ReferencingPackage, bEditorOnlyReferences](const FName& PackageName)
+		{
+			return FSoftObjectPathProperty(PackageName, NAME_None, bEditorOnlyReferences);
+		});
+
+	FScopeLock ScopeLock(&CriticalSection);
+	SoftObjectPathMap.FindOrAdd(ReferencingPackage).Append(SoftObjectPathArray);
+}
+
 void FRedirectCollector::ResolveAllSoftObjectPaths(FName FilterPackage)
 {	
 	auto LoadSoftObjectPathLambda = [this](const FSoftObjectPathProperty& SoftObjectPathProperty)
diff --git a/Engine/Source/Runtime/CoreUObject/Private/UObject/SavePackage.cpp b/Engine/Source/Runtime/CoreUObject/Private/UObject/SavePackage.cpp
--- a/Engine/Source/Runtime/CoreUObject/Private/UObject/SavePackage.cpp
+++ b/Engine/Source/Runtime/CoreUObject/Private/UObject/SavePackage.cpp
@@ -19,6 +19,7 @@
 #include "Serialization/BufferArchive.h"
 #include "Misc/AssetRegistryInterface.h"
 #include "Misc/FeedbackContext.h"
+#include "Misc/RedirectCollector.h"
 #include "Misc/ScopedSlowTask.h"
 #include "Misc/ObjectThumbnail.h"
 #include "UObject/ObjectMacros.h"
@@ -2264,6 +2265,9 @@ FSavePackageResultStruct UPackage::Save(UPackage* InOuter, UObject* Base, EObjec
 		// structure to track what every export needs to import (native only)
 		TMap<UObject*, TArray<UObject*> > NativeObjectDependencies;
 
+		// Used for the asset registry and the redirect collector
+		TSet<FName> SoftPackagesUsedInGame;
+
 		// Size of serialized out package in bytes. This is before compression.
 		int32 PackageSize = INDEX_NONE;
 		TPimplPtr<FLinkerSave> Linker = nullptr;
@@ -2581,7 +2585,6 @@ FSavePackageResultStruct UPackage::Save(UPackage* InOuter, UObject* Base, EObjec
 				// Import objects & names.
 				TSet<UPackage*> PrestreamPackages;
 				TSet<UObject*> ImportsUsedInGame;
-				TSet<FName> SoftPackagesUsedInGame;
 				{
 					SCOPED_SAVETIMER(UPackage_Save_TagImports);
 					
@@ -4685,6 +4688,15 @@ FSavePackageResultStruct UPackage::Save(UPackage* InOuter, UObject* Base, EObjec
 			// Package has been save, so unmark NewlyCreated flag.
 			InOuter->ClearPackageFlags(PKG_NewlyCreated);
 
+			// Notify the soft reference collector about our harvested soft references during save. 
+			// This is currently needed only for cooking which does not require editor-only references 
+#if WITH_EDITOR
+			if (bIsCooking)
+			{
+				GRedirectCollector.CollectSavedSoftPackageReferences(InOuter->GetFName(), SoftPackagesUsedInGame, false);
+			}
+#endif
+
 			if (Linker != nullptr)
 			{
 				Linker->OnPostSave(TargetPackagePath, FObjectPostSaveContext(ObjectSaveContext));
diff --git a/Engine/Source/Runtime/CoreUObject/Private/UObject/SavePackage2.cpp b/Engine/Source/Runtime/CoreUObject/Private/UObject/SavePackage2.cpp
--- a/Engine/Source/Runtime/CoreUObject/Private/UObject/SavePackage2.cpp
+++ b/Engine/Source/Runtime/CoreUObject/Private/UObject/SavePackage2.cpp
@@ -15,6 +15,7 @@
 #include "Misc/FeedbackContext.h"
 #include "Misc/MessageDialog.h"
 #include "Misc/PackageName.h"
+#include "Misc/RedirectCollector.h"
 #include "Misc/ScopedSlowTask.h"
 #include "Misc/SecureHash.h"
 #include "ProfilingDebugging/CookStats.h"
@@ -49,6 +50,9 @@ COREUOBJECT_API extern bool GOutputCookingWarnings;
 // bring the UObectGlobal declaration visible to non editor
 bool IsEditorOnlyObject(const UObject* InObject, bool bCheckRecursive, bool bCheckMarks);
 
+namespace
+{
+
 ESavePackageResult ReturnSuccessOrCancel()
 {
 	return !GWarn->ReceivedUserCancel() ? ESavePackageResult::Success : ESavePackageResult::Canceled;
@@ -1951,6 +1955,16 @@ void PostSavePackage(FSaveContext& SaveContext)
 	// Package has been saved, so unmark the NewlyCreated flag.
 	Package->ClearPackageFlags(PKG_NewlyCreated);
 
+	// Notify the soft reference collector about our harvested soft references during save. 
+	// This is currently needed only for cooking which does not require editor-only references 
+#if WITH_EDITOR
+	if (SaveContext.IsCooking())
+	{
+		GRedirectCollector.CollectSavedSoftPackageReferences(Package->GetFName(), SaveContext.GetSoftPackagesUsedInGame(), false);
+	}
+#endif
+
+	// Call the linker post save callbacks
 	if (SaveContext.Linker != nullptr)
 	{
 		SaveContext.Linker->OnPostSave(SaveContext.GetTargetPackagePath(), FObjectPostSaveContext(SaveContext.GetObjectSaveContext()));
@@ -2180,6 +2194,8 @@ FText GetSlowTaskStatusMessage(const FSaveContext& SaveContext)
 	return FText::Format(NSLOCTEXT("Core", "SavingFile", "Saving file: {CleanFilename}..."), Args);
 }
 
+} // end namespace
+
 FSavePackageResultStruct UPackage::Save2(UPackage* InPackage, UObject* InAsset, const TCHAR* InFilename, FSavePackageArgs& SaveArgs)
 {
 	//COOK_STAT(FScopedDurationTimer FuncSaveTimer(FSavePackageStats::SavePackageTimeSec));
diff --git a/Engine/Source/Runtime/CoreUObject/Public/Misc/RedirectCollector.h b/Engine/Source/Runtime/CoreUObject/Public/Misc/RedirectCollector.h
--- a/Engine/Source/Runtime/CoreUObject/Public/Misc/RedirectCollector.h
+++ b/Engine/Source/Runtime/CoreUObject/Public/Misc/RedirectCollector.h
@@ -70,6 +70,14 @@ public:
 	 */
 	void OnSoftObjectPathLoaded(const struct FSoftObjectPath& InPath, FArchive* InArchive);
 
+	/**
+	 * Called at the end of Package Save to record soft package references that might have been created by save transformations
+	 * @param ReferencingPackage The package on which we are recording the references
+	 * @param PackageNames List of of soft package references needed by the referencing package
+	 * @param bEditorOnlyReferences if the PackageNames list are references made by editor only properties
+	 */
+	void CollectSavedSoftPackageReferences(FName ReferencingPackage, const TSet<FName>& PackageNames, bool bEditorOnlyReferences);
+
 	/**
 	 * Load all soft object paths to resolve them, add that to the remap table, and empty the array
 	 * @param FilterPackage If set, only load references that were created by FilterPackage. If empty, resolve  all of them
