World Partition HLOD Builder - Ensure CopyFilesFromWorkingDir() can overwrite read-only files when running with -SCCProvider=None

diff --git a/Engine/Source/Editor/UnrealEd/Private/WorldPartition/WorldPartitionHLODsBuilder.cpp b/Engine/Source/Editor/UnrealEd/Private/WorldPartition/WorldPartitionHLODsBuilder.cpp
--- a/Engine/Source/Editor/UnrealEd/Private/WorldPartition/WorldPartitionHLODsBuilder.cpp
+++ b/Engine/Source/Editor/UnrealEd/Private/WorldPartition/WorldPartitionHLODsBuilder.cpp
@@ -888,7 +888,9 @@ bool UWorldPartitionHLODsBuilder::CopyFilesFromWorkingDir(const FString& SourceD
 	{
 		for (const auto& Pair : FilesToCopy)
 		{
-			bool bRet = IFileManager::Get().Copy(*Pair.Key, *Pair.Value) == COPY_OK;
+			const bool bReplace = true;
+			const bool bEvenIfReadOnly = true;
+			bool bRet = IFileManager::Get().Copy(*Pair.Key, *Pair.Value, bReplace, bEvenIfReadOnly) == COPY_OK;
 			if (!bRet)
 			{
 				UE_LOG(LogWorldPartitionHLODsBuilder, Error, TEXT("Failed to copy file from \"%s\" to \"%s\""), *Pair.Value, *Pair.Key);
@@ -982,7 +984,9 @@ bool UWorldPartitionHLODsBuilder::CopyFilesFromWorkingDir(const FString& SourceD
 		{
 			for (const FString& FileToDelete : FilesToDelete)
 			{
-				bRet = IFileManager::Get().Delete(*FileToDelete);
+				const bool bRequireExists = false;
+				const bool bEvenIfReadOnly = true;
+				bRet = IFileManager::Get().Delete(*FileToDelete, bRequireExists, bEvenIfReadOnly);
 				if (!bRet)
 				{
 					UE_LOG(LogWorldPartitionHLODsBuilder, Error, TEXT("Failed to delete file from disk: %s"), *USourceControlHelpers::LastErrorMsg().ToString());
