Ensure GizmoManipulationStarted and GizmoManipulationStopped are called in pairs

The editor viewport would previously skip the GizmoManipulationStopped call if nothing had actually moved, but this causes an imbalance in the APIs which call NotifyMovementStarted and NotifyMovementEnded.

GizmoManipulationStopped is now always called, but with an extra argument saying if anything actually moved.

diff --git a/Engine/Source/Editor/EditorFramework/Private/Elements/Framework/TypedElementViewportInteraction.cpp b/Engine/Source/Editor/EditorFramework/Private/Elements/Framework/TypedElementViewportInteraction.cpp
--- a/Engine/Source/Editor/EditorFramework/Private/Elements/Framework/TypedElementViewportInteraction.cpp
+++ b/Engine/Source/Editor/EditorFramework/Private/Elements/Framework/TypedElementViewportInteraction.cpp
@@ -77,7 +77,7 @@ void FTypedElementViewportInteractionCustomization::GizmoManipulationDeltaUpdate
 	}
 }
 
-void FTypedElementViewportInteractionCustomization::GizmoManipulationStopped(const TTypedElement<UTypedElementWorldInterface>& InElementWorldHandle, const UE::Widget::EWidgetMode InWidgetMode)
+void FTypedElementViewportInteractionCustomization::GizmoManipulationStopped(const TTypedElement<UTypedElementWorldInterface>& InElementWorldHandle, const UE::Widget::EWidgetMode InWidgetMode, const ETypedElementViewportInteractionGizmoManipulationType InManipulationType)
 {
 	InElementWorldHandle.NotifyMovementEnded();
 }
@@ -192,14 +192,14 @@ void UTypedElementViewportInteraction::UpdateGizmoManipulation(const UTypedEleme
 	});
 }
 
-void UTypedElementViewportInteraction::EndGizmoManipulation(const UTypedElementList* InElementsToMove, const UE::Widget::EWidgetMode InWidgetMode)
+void UTypedElementViewportInteraction::EndGizmoManipulation(const UTypedElementList* InElementsToMove, const UE::Widget::EWidgetMode InWidgetMode, const ETypedElementViewportInteractionGizmoManipulationType InManipulationType)
 {
-	InElementsToMove->ForEachElementHandle([this, InWidgetMode](const FTypedElementHandle& InElementToMove)
+	InElementsToMove->ForEachElementHandle([this, InWidgetMode, InManipulationType](const FTypedElementHandle& InElementToMove)
 	{
 		FTypedElementViewportInteractionElement ViewportInteractionElement = ResolveViewportInteractionElement(InElementToMove);
 		if (ViewportInteractionElement)
 		{
-			ViewportInteractionElement.GizmoManipulationStopped(InWidgetMode);
+			ViewportInteractionElement.GizmoManipulationStopped(InWidgetMode, InManipulationType);
 		}
 		return true;
 	});
diff --git a/Engine/Source/Editor/EditorFramework/Public/Elements/Framework/TypedElementViewportInteraction.h b/Engine/Source/Editor/EditorFramework/Public/Elements/Framework/TypedElementViewportInteraction.h
--- a/Engine/Source/Editor/EditorFramework/Public/Elements/Framework/TypedElementViewportInteraction.h
+++ b/Engine/Source/Editor/EditorFramework/Public/Elements/Framework/TypedElementViewportInteraction.h
@@ -17,6 +17,14 @@ enum class ETypedElementViewportInteractionWorldType : uint8
 	PlayInEditor,
 };
 
+enum class ETypedElementViewportInteractionGizmoManipulationType : uint8
+{
+	/** The gizmo manipulation resulted in a drag operation */
+	Drag,
+	/** The gizmo manipulation resulted in a click operation */
+	Click,
+};
+
 /**
  * Customization used to allow asset editors (such as the level editor) to override the base behavior of viewport interaction.
  */
@@ -31,7 +39,7 @@ public:
 	virtual void PreGizmoManipulationStarted(TArrayView<const FTypedElementHandle> InElementHandles, const UE::Widget::EWidgetMode InWidgetMode);
 	virtual void GizmoManipulationStarted(const TTypedElement<UTypedElementWorldInterface>& InElementWorldHandle, const UE::Widget::EWidgetMode InWidgetMode);
 	virtual void GizmoManipulationDeltaUpdate(const TTypedElement<UTypedElementWorldInterface>& InElementWorldHandle, const UE::Widget::EWidgetMode InWidgetMode, const EAxisList::Type InDragAxis, const FInputDeviceState& InInputState, const FTransform& InDeltaTransform, const FVector& InPivotLocation);
-	virtual void GizmoManipulationStopped(const TTypedElement<UTypedElementWorldInterface>& InElementWorldHandle, const UE::Widget::EWidgetMode InWidgetMode);
+	virtual void GizmoManipulationStopped(const TTypedElement<UTypedElementWorldInterface>& InElementWorldHandle, const UE::Widget::EWidgetMode InWidgetMode, const ETypedElementViewportInteractionGizmoManipulationType InManipulationType);
 	virtual void PostGizmoManipulationStopped(TArrayView<const FTypedElementHandle> InElementHandles, const UE::Widget::EWidgetMode InWidgetMode);
 	virtual void MirrorElement(const TTypedElement<UTypedElementWorldInterface>& InElementWorldHandle, const FVector& InMirrorScale, const FVector& InPivotLocation);
 };
@@ -72,7 +80,7 @@ public:
 	bool GetGizmoPivotLocation(const UE::Widget::EWidgetMode InWidgetMode, FVector& OutPivotLocation) const { return ViewportInteractionCustomization->GetGizmoPivotLocation(ElementWorldHandle, InWidgetMode, OutPivotLocation); }
 	void GizmoManipulationStarted(const UE::Widget::EWidgetMode InWidgetMode) const { ViewportInteractionCustomization->GizmoManipulationStarted(ElementWorldHandle, InWidgetMode); }
 	void GizmoManipulationDeltaUpdate(const UE::Widget::EWidgetMode InWidgetMode, const EAxisList::Type InDragAxis, const FInputDeviceState& InInputState, const FTransform& InDeltaTransform, const FVector& InPivotLocation) const { ViewportInteractionCustomization->GizmoManipulationDeltaUpdate(ElementWorldHandle, InWidgetMode, InDragAxis, InInputState, InDeltaTransform, InPivotLocation); }
-	void GizmoManipulationStopped(const UE::Widget::EWidgetMode InWidgetMode) const { ViewportInteractionCustomization->GizmoManipulationStopped(ElementWorldHandle, InWidgetMode); }
+	void GizmoManipulationStopped(const UE::Widget::EWidgetMode InWidgetMode, const ETypedElementViewportInteractionGizmoManipulationType InManipulationType) const { ViewportInteractionCustomization->GizmoManipulationStopped(ElementWorldHandle, InWidgetMode, InManipulationType); }
 	void MirrorElement(const FVector& InMirrorScale, const FVector& InPivotLocation) const { ViewportInteractionCustomization->MirrorElement(ElementWorldHandle, InMirrorScale, InPivotLocation); }
 
 private:
@@ -108,7 +116,7 @@ public:
 	/**
 	 * Notify that the gizmo has finished manipulating the transform of the given set of elements (calculated from calling GetSelectedElementsToMove).
 	 */
-	void EndGizmoManipulation(const UTypedElementList* InElementsToMove, const UE::Widget::EWidgetMode InWidgetMode);
+	void EndGizmoManipulation(const UTypedElementList* InElementsToMove, const UE::Widget::EWidgetMode InWidgetMode, const ETypedElementViewportInteractionGizmoManipulationType InManipulationType);
 
 	/**
 	 * Apply the given delta to the specified element.
diff --git a/Engine/Source/Editor/LevelEditor/Private/Elements/Actor/ActorElementLevelEditorViewportInteractionCustomization.cpp b/Engine/Source/Editor/LevelEditor/Private/Elements/Actor/ActorElementLevelEditorViewportInteractionCustomization.cpp
--- a/Engine/Source/Editor/LevelEditor/Private/Elements/Actor/ActorElementLevelEditorViewportInteractionCustomization.cpp
+++ b/Engine/Source/Editor/LevelEditor/Private/Elements/Actor/ActorElementLevelEditorViewportInteractionCustomization.cpp
@@ -97,10 +97,16 @@ void FActorElementLevelEditorViewportInteractionCustomization::GizmoManipulation
 	GetMutableLevelEditorViewportClient()->UpdateLockedActorViewports(Actor, true);
 }
 
-void FActorElementLevelEditorViewportInteractionCustomization::GizmoManipulationStopped(const TTypedElement<UTypedElementWorldInterface>& InElementWorldHandle, const UE::Widget::EWidgetMode InWidgetMode)
+void FActorElementLevelEditorViewportInteractionCustomization::GizmoManipulationStopped(const TTypedElement<UTypedElementWorldInterface>& InElementWorldHandle, const UE::Widget::EWidgetMode InWidgetMode, const ETypedElementViewportInteractionGizmoManipulationType InManipulationType)
 {
 	AActor* Actor = ActorElementDataUtil::GetActorFromHandleChecked(InElementWorldHandle);
 
+	if (InManipulationType != ETypedElementViewportInteractionGizmoManipulationType::Drag)
+	{
+		// Don't trigger an update if nothing actually moved, to preserve the old gizmo behavior
+		return;
+	}
+
 	// Broadcast Post Edit change notification, we can't call PostEditChangeProperty directly on Actor or ActorComponent from here since it wasn't pair with a proper PreEditChange
 	if (FProperty* TransformProperty = FComponentElementLevelEditorViewportInteractionCustomization::GetEditTransformProperty(InWidgetMode))
 	{
diff --git a/Engine/Source/Editor/LevelEditor/Private/Elements/Component/ComponentElementLevelEditorViewportInteractionCustomization.cpp b/Engine/Source/Editor/LevelEditor/Private/Elements/Component/ComponentElementLevelEditorViewportInteractionCustomization.cpp
--- a/Engine/Source/Editor/LevelEditor/Private/Elements/Component/ComponentElementLevelEditorViewportInteractionCustomization.cpp
+++ b/Engine/Source/Editor/LevelEditor/Private/Elements/Component/ComponentElementLevelEditorViewportInteractionCustomization.cpp
@@ -111,10 +111,16 @@ void FComponentElementLevelEditorViewportInteractionCustomization::GizmoManipula
 	}
 }
 
-void FComponentElementLevelEditorViewportInteractionCustomization::GizmoManipulationStopped(const TTypedElement<UTypedElementWorldInterface>& InElementWorldHandle, const UE::Widget::EWidgetMode InWidgetMode)
+void FComponentElementLevelEditorViewportInteractionCustomization::GizmoManipulationStopped(const TTypedElement<UTypedElementWorldInterface>& InElementWorldHandle, const UE::Widget::EWidgetMode InWidgetMode, const ETypedElementViewportInteractionGizmoManipulationType InManipulationType)
 {
 	UActorComponent* Component = ComponentElementDataUtil::GetComponentFromHandleChecked(InElementWorldHandle);
 
+	if (InManipulationType != ETypedElementViewportInteractionGizmoManipulationType::Drag)
+	{
+		// Don't trigger an update if nothing actually moved, to preserve the old gizmo behavior
+		return;
+	}
+
 	if (USceneComponent* SceneComponent = Cast<USceneComponent>(Component))
 	{
 		// Broadcast Post Edit change notification, we can't call PostEditChangeProperty directly on Actor or ActorComponent from here since it wasn't pair with a proper PreEditChange
diff --git a/Engine/Source/Editor/LevelEditor/Public/Elements/Actor/ActorElementLevelEditorViewportInteractionCustomization.h b/Engine/Source/Editor/LevelEditor/Public/Elements/Actor/ActorElementLevelEditorViewportInteractionCustomization.h
--- a/Engine/Source/Editor/LevelEditor/Public/Elements/Actor/ActorElementLevelEditorViewportInteractionCustomization.h
+++ b/Engine/Source/Editor/LevelEditor/Public/Elements/Actor/ActorElementLevelEditorViewportInteractionCustomization.h
@@ -11,7 +11,7 @@ public:
 	virtual void GetElementsToMove(const TTypedElement<UTypedElementWorldInterface>& InElementWorldHandle, const ETypedElementViewportInteractionWorldType InWorldType, const UTypedElementSelectionSet* InSelectionSet, UTypedElementList* OutElementsToMove) override;
 	virtual void GizmoManipulationStarted(const TTypedElement<UTypedElementWorldInterface>& InElementWorldHandle, const UE::Widget::EWidgetMode InWidgetMode) override;
 	virtual void GizmoManipulationDeltaUpdate(const TTypedElement<UTypedElementWorldInterface>& InElementWorldHandle, const UE::Widget::EWidgetMode InWidgetMode, const EAxisList::Type InDragAxis, const FInputDeviceState& InInputState, const FTransform& InDeltaTransform, const FVector& InPivotLocation) override;
-	virtual void GizmoManipulationStopped(const TTypedElement<UTypedElementWorldInterface>& InElementWorldHandle, const UE::Widget::EWidgetMode InWidgetMode) override;
+	virtual void GizmoManipulationStopped(const TTypedElement<UTypedElementWorldInterface>& InElementWorldHandle, const UE::Widget::EWidgetMode InWidgetMode, const ETypedElementViewportInteractionGizmoManipulationType InManipulationType) override;
 	virtual void PostGizmoManipulationStopped(TArrayView<const FTypedElementHandle> InElementHandles, const UE::Widget::EWidgetMode InWidgetMode) override;
 
 	static bool CanMoveActorInViewport(const AActor* InActor, const ETypedElementViewportInteractionWorldType InWorldType);
diff --git a/Engine/Source/Editor/LevelEditor/Public/Elements/Component/ComponentElementLevelEditorViewportInteractionCustomization.h b/Engine/Source/Editor/LevelEditor/Public/Elements/Component/ComponentElementLevelEditorViewportInteractionCustomization.h
--- a/Engine/Source/Editor/LevelEditor/Public/Elements/Component/ComponentElementLevelEditorViewportInteractionCustomization.h
+++ b/Engine/Source/Editor/LevelEditor/Public/Elements/Component/ComponentElementLevelEditorViewportInteractionCustomization.h
@@ -13,7 +13,7 @@ public:
 	virtual void GetElementsToMove(const TTypedElement<UTypedElementWorldInterface>& InElementWorldHandle, const ETypedElementViewportInteractionWorldType InWorldType, const UTypedElementSelectionSet* InSelectionSet, UTypedElementList* OutElementsToMove) override;
 	virtual void GizmoManipulationStarted(const TTypedElement<UTypedElementWorldInterface>& InElementWorldHandle, const UE::Widget::EWidgetMode InWidgetMode) override;
 	virtual void GizmoManipulationDeltaUpdate(const TTypedElement<UTypedElementWorldInterface>& InElementWorldHandle, const UE::Widget::EWidgetMode InWidgetMode, const EAxisList::Type InDragAxis, const FInputDeviceState& InInputState, const FTransform& InDeltaTransform, const FVector& InPivotLocation) override;
-	virtual void GizmoManipulationStopped(const TTypedElement<UTypedElementWorldInterface>& InElementWorldHandle, const UE::Widget::EWidgetMode InWidgetMode) override;
+	virtual void GizmoManipulationStopped(const TTypedElement<UTypedElementWorldInterface>& InElementWorldHandle, const UE::Widget::EWidgetMode InWidgetMode, const ETypedElementViewportInteractionGizmoManipulationType InManipulationType) override;
 
 	static void ValidateScale(const FVector& InOriginalPreDragScale, const EAxisList::Type InDragAxis, const FVector& InCurrentScale, const FVector& InBoxExtent, FVector& InOutScaleDelta, bool bInCheckSmallExtent);
 	static FProperty* GetEditTransformProperty(const UE::Widget::EWidgetMode InWidgetMode);
diff --git a/Engine/Source/Editor/UnrealEd/Private/LevelEditorViewport.cpp b/Engine/Source/Editor/UnrealEd/Private/LevelEditorViewport.cpp
--- a/Engine/Source/Editor/UnrealEd/Private/LevelEditorViewport.cpp
+++ b/Engine/Source/Editor/UnrealEd/Private/LevelEditorViewport.cpp
@@ -1717,6 +1717,7 @@ FLevelEditorViewportClient::FLevelEditorViewportClient(const TSharedPtr<SLevelVi
 	, bOnlyMovedPivot(false)
 	, bLockedCameraView(true)
 	, bNeedToRestoreComponentBeingMovedFlag(false)
+	, bHasBegunGizmoManipulation(false)
 	, bReceivedFocusRecently(false)
 	, bAlwaysShowModeWidgetAfterSelectionChanges(true)
 	, CachedElementsToManipulate(UTypedElementRegistry::GetInstance()->CreateElementList())
@@ -2146,6 +2147,16 @@ void FLevelEditorViewportClient::ProcessClick(FSceneView& View, HHitProxy* HitPr
 {
 	UBrushEditingSubsystem* BrushSubsystem = GEditor->GetEditorSubsystem<UBrushEditingSubsystem>();
 
+	// We may have started gizmo manipulation if hot-keys were pressed when we started this click
+	// If so, we need to end that now before we potentially update the selection below, 
+	// otherwise the usual call in TrackingStopped would include the newly selected element
+	if (bHasBegunGizmoManipulation)
+	{
+		const UTypedElementList* ElementsToManipulate = GetElementsToManipulate();
+		ViewportInteraction->EndGizmoManipulation(ElementsToManipulate, GetWidgetMode(), ETypedElementViewportInteractionGizmoManipulationType::Click);
+		bHasBegunGizmoManipulation = false;
+	}
+
 	const FViewportClick Click(&View,this,Key,Event,HitX,HitY);
 	if (Click.GetKey() == EKeys::MiddleMouseButton && !Click.IsAltDown() && !Click.IsShiftDown())
 	{
@@ -2612,6 +2623,13 @@ bool FLevelEditorViewportClient::InputWidgetDelta(FViewport* InViewport, EAxisLi
 								// Exclusively select the new elements, so that future gizmo interaction manipulates those items instead
 								if (DuplicatedElements.Num() > 0)
 								{
+									// We need to end the gizmo manipulation on the current selection, as they won't be moved again this drag...
+									if (bHasBegunGizmoManipulation)
+									{
+										ViewportInteraction->EndGizmoManipulation(ElementsToManipulate, GetWidgetMode(), ETypedElementViewportInteractionGizmoManipulationType::Click);
+										// Note: We don't reset bHasBegunGizmoManipulation here, as we test it again below to begin manipulating the duplicated selection
+									}
+
 									UTypedElementSelectionSet* SelectionSet = GetMutableSelectionSet();
 
 									const FTypedElementSelectionOptions SelectionOptions = FTypedElementSelectionOptions()
@@ -2621,6 +2639,13 @@ bool FLevelEditorViewportClient::InputWidgetDelta(FViewport* InViewport, EAxisLi
 
 									// Force the cached manipulation list to update, as we don't notify about the change above
 									ResetElementsToManipulate();
+
+									// We need to start gizmo manipulation on the duplicated selection, as they may continue to move as part of this drag duplicate
+									if (bHasBegunGizmoManipulation)
+									{
+										ElementsToManipulate = GetElementsToManipulate();
+										ViewportInteraction->BeginGizmoManipulation(ElementsToManipulate, GetWidgetMode());
+									}
 								}
 
 								RedrawAllViewportsIntoThisScene();
@@ -2942,6 +2967,7 @@ void FLevelEditorViewportClient::TrackingStarted( const FInputEventState& InInpu
 
 	bOnlyMovedPivot = false;
 	bNeedToRestoreComponentBeingMovedFlag = false;
+	bHasBegunGizmoManipulation = false;
 
 	PreDragActorTransforms.Empty();
 
@@ -2963,6 +2989,7 @@ void FLevelEditorViewportClient::TrackingStarted( const FInputEventState& InInpu
 
 		const UTypedElementList* ElementsToManipulate = GetElementsToManipulate();
 		ViewportInteraction->BeginGizmoManipulation(ElementsToManipulate, GetWidgetMode());
+		bHasBegunGizmoManipulation = true;
 
 		if (!bDuplicateActorsInProgress)
 		{
@@ -3082,12 +3109,17 @@ void FLevelEditorViewportClient::TrackingStopped()
 
 	// Notify the selected actors that they have been moved.
 	// Don't do this if AddDelta was never called.
-	if( bDidAnythingActuallyChange && MouseDeltaTracker->HasReceivedDelta() )
 	{
-		const UTypedElementList* ElementsToManipulate = GetElementsToManipulate();
-		ViewportInteraction->EndGizmoManipulation(ElementsToManipulate, GetWidgetMode());
+		const bool bDidMove = bDidAnythingActuallyChange && MouseDeltaTracker->HasReceivedDelta();
+
+		if (bHasBegunGizmoManipulation)
+		{
+			const UTypedElementList* ElementsToManipulate = GetElementsToManipulate();
+			ViewportInteraction->EndGizmoManipulation(ElementsToManipulate, GetWidgetMode(), bDidMove ? ETypedElementViewportInteractionGizmoManipulationType::Drag : ETypedElementViewportInteractionGizmoManipulationType::Click);
+			bHasBegunGizmoManipulation = false;
+		}
 
-		if (!GUnrealEd->IsPivotMovedIndependently())
+		if (bDidMove && !GUnrealEd->IsPivotMovedIndependently())
 		{
 			GUnrealEd->UpdatePivotLocationForSelection();
 		}
diff --git a/Engine/Source/Editor/UnrealEd/Public/LevelEditorViewport.h b/Engine/Source/Editor/UnrealEd/Public/LevelEditorViewport.h
--- a/Engine/Source/Editor/UnrealEd/Public/LevelEditorViewport.h
+++ b/Engine/Source/Editor/UnrealEd/Public/LevelEditorViewport.h
@@ -889,6 +889,9 @@ public:
 	/** true if the viewport needs to restore the flag when tracking ends */
 	bool					bNeedToRestoreComponentBeingMovedFlag;
 
+	/** true if gizmo manipulation was started from a tracking event */
+	bool					bHasBegunGizmoManipulation;
+
 	/** Whether this viewport recently received focus. Used to determine whether component selection is permissible. */
 	bool bReceivedFocusRecently;
 
