Network Emulation settings are disabled when PIE is not configured for multiplayer.

diff --git a/Engine/Source/Editor/UnrealEd/Classes/Settings/LevelEditorPlayNetworkEmulationSettings.h b/Engine/Source/Editor/UnrealEd/Classes/Settings/LevelEditorPlayNetworkEmulationSettings.h
--- a/Engine/Source/Editor/UnrealEd/Classes/Settings/LevelEditorPlayNetworkEmulationSettings.h
+++ b/Engine/Source/Editor/UnrealEd/Classes/Settings/LevelEditorPlayNetworkEmulationSettings.h
@@ -99,6 +99,12 @@ public:
 
 	bool HandleAreSettingsEnabled() const;
 
+	/*
+	 * Checks to see if separate server was launched, the number of player clients
+	 *	is greater than one, or if the current net play mode is not standalone
+	 */
+	bool IsNetworkEmulationAvailable() const;
+
 private:
 
 	// Property Handles
@@ -107,4 +113,4 @@ private:
 	TSharedPtr<IPropertyHandle> CurrentProfileHandle;
 	TSharedPtr<IPropertyHandle> OutPacketsHandle;
 	TSharedPtr<IPropertyHandle> InPacketsHandle;
-};
\ No newline at end of file
+};
diff --git a/Engine/Source/Editor/UnrealEd/Private/Settings/LevelEditorPlayNetworkEmulationSettings.cpp b/Engine/Source/Editor/UnrealEd/Private/Settings/LevelEditorPlayNetworkEmulationSettings.cpp
--- a/Engine/Source/Editor/UnrealEd/Private/Settings/LevelEditorPlayNetworkEmulationSettings.cpp
+++ b/Engine/Source/Editor/UnrealEd/Private/Settings/LevelEditorPlayNetworkEmulationSettings.cpp
@@ -10,6 +10,8 @@
 #include "Engine/NetworkSettings.h"
 #include "Framework/MultiBox/MultiBoxBuilder.h"
 
+#include "Settings/LevelEditorPlaySettings.h"
+
 namespace NetworkEmulationSettingsHelper
 {
 	const FString& GetCustomProfileName()
@@ -71,7 +73,8 @@ void FLevelEditorPlayNetworkEmulationSettingsDetail::CustomizeHeader(TSharedRef<
 	.ValueContent()
 	[
 		IsNetworkEmulationEnabledHandle->CreatePropertyValueWidget()
-	];
+	]
+	.IsEnabled(TAttribute<bool>(this, &FLevelEditorPlayNetworkEmulationSettingsDetail::IsNetworkEmulationAvailable));
 }
 
 
@@ -211,6 +214,19 @@ FText FLevelEditorPlayNetworkEmulationSettingsDetail::GetSelectedNetworkEmulatio
 	return FText::FromString(SelectedProfile);
 }
 
+bool FLevelEditorPlayNetworkEmulationSettingsDetail::IsNetworkEmulationAvailable() const
+{
+	const ULevelEditorPlaySettings* PlayerLevelEditorSettings = GetDefault<ULevelEditorPlaySettings>();
+
+	int32 NumberOfClients;
+	PlayerLevelEditorSettings->GetPlayNumberOfClients(NumberOfClients);
+
+	EPlayNetMode NetMode;
+	PlayerLevelEditorSettings->GetPlayNetMode(NetMode);
+
+	return (NumberOfClients > 1) || NetMode != PIE_Standalone || PlayerLevelEditorSettings->bLaunchSeparateServer;
+}
+
 //-------------------------------------------------------------------------------------------------
 // FLevelEditorPlayNetworkEmulationSettings
 
