WorldPartitionConvertCommandlet: properly remove loaded sublevels from the world's list before cleanup world.

diff --git a/Engine/Source/Editor/UnrealEd/Private/Commandlets/WorldPartitionConvertCommandlet.cpp b/Engine/Source/Editor/UnrealEd/Private/Commandlets/WorldPartitionConvertCommandlet.cpp
--- a/Engine/Source/Editor/UnrealEd/Private/Commandlets/WorldPartitionConvertCommandlet.cpp
+++ b/Engine/Source/Editor/UnrealEd/Private/Commandlets/WorldPartitionConvertCommandlet.cpp
@@ -1184,7 +1184,7 @@ int32 UWorldPartitionConvertCommandlet::Main(const FString& Params)
 
 			// Spawn the level instance actor
 			ULevelStreaming* SubLevelStreaming = nullptr;
-			for (ULevelStreaming* LevelStreaming : MainLevel->GetWorld()->GetStreamingLevels())
+			for (ULevelStreaming* LevelStreaming : MainWorld->GetStreamingLevels())
 			{
 				if (LevelStreaming->GetLoadedLevel() == SubLevel)
 				{
@@ -1300,12 +1300,13 @@ int32 UWorldPartitionConvertCommandlet::Main(const FString& Params)
 	}
 
 	// Clear streaming levels
-	for (ULevelStreaming* LevelStreaming: MainLevel->GetWorld()->GetStreamingLevels())
+	for (ULevelStreaming* LevelStreaming: MainWorld->GetStreamingLevels())
 	{
 		LevelStreaming->MarkPendingKill();
 		ULevelStreaming::RemoveLevelAnnotation(LevelStreaming->GetLoadedLevel());
+		MainWorld->RemoveLevel(LevelStreaming->GetLoadedLevel());
 	}
-	MainLevel->GetWorld()->ClearStreamingLevels();
+	MainWorld->ClearStreamingLevels();
 
 	// Fixup SoftObjectPaths
 	FixupSoftObjectPaths(MainPackage);
