WorldPartitionConvertCommandlet: don't try to call FPackageSourceControlHelper::Checkout for newly created packages, as the behavior of that function recenbtly changed and it's now an error when calling it on a non-existing package.

diff --git a/Engine/Source/Editor/UnrealEd/Private/Commandlets/WorldPartitionConvertCommandlet.cpp b/Engine/Source/Editor/UnrealEd/Private/Commandlets/WorldPartitionConvertCommandlet.cpp
--- a/Engine/Source/Editor/UnrealEd/Private/Commandlets/WorldPartitionConvertCommandlet.cpp
+++ b/Engine/Source/Editor/UnrealEd/Private/Commandlets/WorldPartitionConvertCommandlet.cpp
@@ -1423,9 +1423,13 @@ int32 UWorldPartitionConvertCommandlet::Main(const FString& Params)
 			UE_LOG(LogWorldPartitionConvertCommandlet, Log, TEXT("Checking out %d actor packages."), PackagesToSave.Num());
 			for(UPackage* Package: PackagesToSave)
 			{
-				if (!PackageHelper.Checkout(Package))
+				FString PackageFileName = SourceControlHelpers::PackageFilename(Package);
+				if (FPlatformFileManager::Get().GetPlatformFile().FileExists(*PackageFileName))
 				{
-					return 1;
+					if (!PackageHelper.Checkout(Package))
+					{
+						return 1;
+					}
 				}
 			}
 		}
