PlayFromHere: An actor can now be source of PlayFromHere (Location,Rotation) and get notified (bCanPlayFromHere, OnPlayFromHere)

diff --git a/Engine/Source/Editor/LevelEditor/Private/LevelEditor.cpp b/Engine/Source/Editor/LevelEditor/Private/LevelEditor.cpp
--- a/Engine/Source/Editor/LevelEditor/Private/LevelEditor.cpp
+++ b/Engine/Source/Editor/LevelEditor/Private/LevelEditor.cpp
@@ -720,6 +720,14 @@ void FLevelEditorModule::BindGlobalLevelEditorCommands()
 		FExecuteAction::CreateStatic( &FLevelEditorActionCallbacks::GoHere_Clicked, (const FVector*)nullptr )
 		);
 
+	ActionList.MapAction(
+		Commands.PlayFromHere,
+		FExecuteAction::CreateStatic( &FLevelEditorActionCallbacks::PlayFromHere_Clicked ), 
+		FCanExecuteAction(),
+		FIsActionChecked(),
+		FIsActionButtonVisible::CreateStatic( &FLevelEditorActionCallbacks::PlayFromHere_IsVisible )
+		);
+
 	ActionList.MapAction( 
 		Commands.SnapCameraToObject,
 		FExecuteAction::CreateStatic( &FLevelEditorActionCallbacks::ExecuteExecCommand, FString( TEXT("CAMERA SNAP") ) )
diff --git a/Engine/Source/Editor/LevelEditor/Private/LevelEditorActions.cpp b/Engine/Source/Editor/LevelEditor/Private/LevelEditorActions.cpp
--- a/Engine/Source/Editor/LevelEditor/Private/LevelEditorActions.cpp
+++ b/Engine/Source/Editor/LevelEditor/Private/LevelEditorActions.cpp
@@ -29,6 +29,7 @@
 #include "Engine/Light.h"
 #include "Engine/BlueprintGeneratedClass.h"
 #include "Kismet2/ComponentEditorUtils.h"
+#include "Kismet2/DebuggerCommands.h"
 #include "Engine/Selection.h"
 #include "Misc/ConfigCacheIni.h"
 #include "UObject/UObjectIterator.h"
@@ -1311,6 +1312,31 @@ void FLevelEditorActionCallbacks::CreateNewOutlinerFolder_Clicked()
 	FActorFolders::Get().CreateFolderContainingSelection(*GetWorld(), NewFolderName);
 }
 
+void FLevelEditorActionCallbacks::PlayFromHere_Clicked()
+{
+	if (GEditor->GetSelectedActorCount() == 1)
+	{
+		if (AActor* Actor = Cast<AActor>(*GEditor->GetSelectedActorIterator()))
+		{
+			Actor->GetWorld()->PersistentLevel->PlayFromHereActor = Actor;
+			FPlayWorldCommandCallbacks::StartPlayFromHere(Actor->GetActorLocation(), Actor->GetActorRotation());
+		}
+	}
+}
+
+bool FLevelEditorActionCallbacks::PlayFromHere_IsVisible()
+{
+	if (GEditor->GetSelectedActorCount() == 1)
+	{
+		if (AActor* Actor = Cast<AActor>(*GEditor->GetSelectedActorIterator()))
+		{
+			return Actor->CanPlayFromHere();
+		}
+	}
+
+	return false;
+}
+
 void FLevelEditorActionCallbacks::GoHere_Clicked( const FVector* Point )
 {
 	if( GCurrentLevelEditingViewportClient )
@@ -3320,6 +3346,7 @@ void FLevelEditorCommands::RegisterCommands()
 	UI_COMMAND( EditAssetNoConfirmMultiple, "Edit Asset", "Edits the asset associated with the selected actor", EUserInterfaceActionType::Button, FInputChord( EKeys::E, EModifierKey::Control | EModifierKey::Shift ) );
 
 	UI_COMMAND( GoHere, "Go Here", "Moves the camera to the current mouse position", EUserInterfaceActionType::Button, FInputChord() );
+	UI_COMMAND( PlayFromHere, "Play From Here", "Start PIE session from this actor", EUserInterfaceActionType::Button, FInputChord()) ;
 
 	UI_COMMAND( SnapCameraToObject, "Snap View to Object", "Snaps the view to the selected object", EUserInterfaceActionType::Button, FInputChord() );
 	UI_COMMAND( SnapObjectToCamera, "Snap Object to View", "Snaps the selected object to the view", EUserInterfaceActionType::Button, FInputChord() );
diff --git a/Engine/Source/Editor/LevelEditor/Private/LevelEditorContextMenu.cpp b/Engine/Source/Editor/LevelEditor/Private/LevelEditorContextMenu.cpp
--- a/Engine/Source/Editor/LevelEditor/Private/LevelEditorContextMenu.cpp
+++ b/Engine/Source/Editor/LevelEditor/Private/LevelEditorContextMenu.cpp
@@ -217,6 +217,7 @@ void FLevelEditorContextMenu::RegisterComponentContextMenu()
 			Section.AddMenuEntry(FLevelEditorCommands::Get().GoHere);
 			Section.AddMenuEntry(FLevelEditorCommands::Get().SnapCameraToObject);
 			Section.AddMenuEntry(FLevelEditorCommands::Get().SnapObjectToCamera);
+			Section.AddMenuEntry(FLevelEditorCommands::Get().PlayFromHere);
 			Section.AddMenuEntry(FLevelEditorCommands::Get().CopyActorFilePathtoClipboard);
 
 			FLevelEditorContextMenuImpl::AddSourceControlMenu(Section);
@@ -305,6 +306,7 @@ void FLevelEditorContextMenu::RegisterActorContextMenu()
 			Section.AddMenuEntry(FLevelEditorCommands::Get().GoHere);
 			Section.AddMenuEntry(FLevelEditorCommands::Get().SnapCameraToObject);
 			Section.AddMenuEntry(FLevelEditorCommands::Get().SnapObjectToCamera);
+			Section.AddMenuEntry(FLevelEditorCommands::Get().PlayFromHere);
 			Section.AddMenuEntry(FLevelEditorCommands::Get().CopyActorFilePathtoClipboard);
 			
 			FLevelEditorContextMenuImpl::AddSourceControlMenu(Section);
@@ -767,17 +769,21 @@ void FLevelEditorContextMenu::SummonMenu(const TSharedRef< SLevelEditor >& Level
 	{
 		static void ExtendMenu( FMenuBuilder& MenuBuilder )
 		{
-			// one extra entry when summoning the menu this way
-			MenuBuilder.BeginSection("ActorPreview", LOCTEXT("PreviewHeading", "Preview") );
+			// Only extend if other PlayFromHere option isn't available
+			if (!FLevelEditorActionCallbacks::PlayFromHere_IsVisible())
 			{
-				// Note: not using a command for play from here since it requires a mouse click
-				FUIAction PlayFromHereAction( 
-					FExecuteAction::CreateStatic( &FPlayWorldCommandCallbacks::StartPlayFromHere ) );
+				// one extra entry when summoning the menu this way
+				MenuBuilder.BeginSection("ActorPreview", LOCTEXT("PreviewHeading", "Preview"));
+				{
+					// Note: not using a command for play from here since it requires a mouse click
+					FUIAction PlayFromHereAction(
+						FExecuteAction::CreateStatic(&FPlayWorldCommandCallbacks::StartPlayFromHere));
 
-				const FText PlayFromHereLabel = GEditor->OnlyLoadEditorVisibleLevelsInPIE() ? LOCTEXT("PlayFromHereVisible", "Play From Here (visible levels)") : LOCTEXT("PlayFromHere", "Play From Here");
-				MenuBuilder.AddMenuEntry( PlayFromHereLabel, LOCTEXT("PlayFromHere_ToolTip", "Starts a game preview from the clicked location"),FSlateIcon(), PlayFromHereAction );
+					const FText PlayFromHereLabel = GEditor->OnlyLoadEditorVisibleLevelsInPIE() ? LOCTEXT("PlayFromHereVisible", "Play From Here (visible levels)") : LOCTEXT("PlayFromHere", "Play From Here");
+					MenuBuilder.AddMenuEntry(PlayFromHereLabel, LOCTEXT("PlayFromHere_ToolTip", "Starts a game preview from the clicked location"), FSlateIcon(), PlayFromHereAction);
+				}
+				MenuBuilder.EndSection();
 			}
-			MenuBuilder.EndSection();
 		}
 	};
 	
diff --git a/Engine/Source/Editor/LevelEditor/Public/LevelEditorActions.h b/Engine/Source/Editor/LevelEditor/Public/LevelEditorActions.h
--- a/Engine/Source/Editor/LevelEditor/Public/LevelEditorActions.h
+++ b/Engine/Source/Editor/LevelEditor/Public/LevelEditorActions.h
@@ -125,6 +125,9 @@ public:
 	/** Moves the camera to the current mouse position */
 	TSharedPtr< FUICommandInfo > GoHere;
 
+	/** Play from here on selected actor. */
+	TSharedPtr< FUICommandInfo > PlayFromHere;
+
 	/** Snaps the camera to the selected object. */
 	TSharedPtr< FUICommandInfo > SnapCameraToObject;
 
@@ -924,6 +927,10 @@ public:
 	 */
 	static void GoHere_Clicked( const FVector* Point );
 
+	/** Called when selected actor can be used to start a play session */
+	static void PlayFromHere_Clicked();
+	static bool PlayFromHere_IsVisible();
+
 	/** Called when 'Go to Code for Actor' is clicked */
 	static void GoToCodeForActor_Clicked();
 
diff --git a/Engine/Source/Editor/UnrealEd/Private/Kismet2/DebuggerCommands.cpp b/Engine/Source/Editor/UnrealEd/Private/Kismet2/DebuggerCommands.cpp
--- a/Engine/Source/Editor/UnrealEd/Private/Kismet2/DebuggerCommands.cpp
+++ b/Engine/Source/Editor/UnrealEd/Private/Kismet2/DebuggerCommands.cpp
@@ -870,6 +870,26 @@ TSharedRef< SWidget > FPlayWorldCommands::GeneratePlayMenuContent(TSharedRef<FUI
 // FPlayWorldCommandCallbacks
 
 void FPlayWorldCommandCallbacks::StartPlayFromHere()
+{
+	// Figure out the start location of the player
+	UClass* const PlayerStartClass = GUnrealEd->PlayFromHerePlayerStartClass ? (UClass*)GUnrealEd->PlayFromHerePlayerStartClass : APlayerStart::StaticClass();
+	UCapsuleComponent* DefaultCollisionComponent = CastChecked<UCapsuleComponent>(PlayerStartClass->GetDefaultObject<AActor>()->GetRootComponent());
+	FVector	CollisionExtent = FVector(DefaultCollisionComponent->GetScaledCapsuleRadius(), DefaultCollisionComponent->GetScaledCapsuleRadius(), DefaultCollisionComponent->GetScaledCapsuleHalfHeight());
+	FVector StartLocation = GEditor->UnsnappedClickLocation + GEditor->ClickPlane * (FVector::BoxPushOut(GEditor->ClickPlane, CollisionExtent) + 0.1f);
+	TOptional<FRotator> StartRotation;
+	FLevelEditorModule& LevelEditorModule = FModuleManager::GetModuleChecked<FLevelEditorModule>("LevelEditor");
+
+	TSharedPtr<IAssetViewport> ActiveLevelViewport = LevelEditorModule.GetFirstActiveViewport();
+
+	if (ActiveLevelViewport.IsValid() && ActiveLevelViewport->GetAssetViewportClient().IsPerspective())
+	{
+		StartRotation = ActiveLevelViewport->GetAssetViewportClient().GetViewRotation();
+	}
+
+	StartPlayFromHere(StartLocation, StartRotation);
+}
+
+void FPlayWorldCommandCallbacks::StartPlayFromHere(const TOptional<FVector>& Location, const TOptional<FRotator>& Rotation)
 {
 	// Is a PIE session already running?  If so we close it first
 	if (GUnrealEd->PlayWorld != NULL)
@@ -878,24 +898,17 @@ void FPlayWorldCommandCallbacks::StartPlayFromHere()
 	}
 
 	FRequestPlaySessionParams SessionParams;
-
-	UClass* const PlayerStartClass = GUnrealEd->PlayFromHerePlayerStartClass ? (UClass*)GUnrealEd->PlayFromHerePlayerStartClass : APlayerStart::StaticClass();
-
-	// Figure out the start location of the player
-	UCapsuleComponent*	DefaultCollisionComponent = CastChecked<UCapsuleComponent>(PlayerStartClass->GetDefaultObject<AActor>()->GetRootComponent());
-	FVector	CollisionExtent = FVector(DefaultCollisionComponent->GetScaledCapsuleRadius(), DefaultCollisionComponent->GetScaledCapsuleRadius(), DefaultCollisionComponent->GetScaledCapsuleHalfHeight());
-	SessionParams.StartLocation = GEditor->UnsnappedClickLocation + GEditor->ClickPlane * (FVector::BoxPushOut(GEditor->ClickPlane, CollisionExtent) + 0.1f);
+	SessionParams.StartLocation = Location;
+	SessionParams.StartRotation = Rotation;
 
 	FLevelEditorModule& LevelEditorModule = FModuleManager::GetModuleChecked<FLevelEditorModule>("LevelEditor");
 
 	TSharedPtr<IAssetViewport> ActiveLevelViewport = LevelEditorModule.GetFirstActiveViewport();
 
-
 	if (ActiveLevelViewport.IsValid() && ActiveLevelViewport->GetAssetViewportClient().IsPerspective())
 	{
 		// If there is no level viewport, a new window will be spawned to play in.
 		SessionParams.DestinationSlateViewport = ActiveLevelViewport;
-		SessionParams.StartRotation = ActiveLevelViewport->GetAssetViewportClient().GetViewRotation();
 	}
 
 	GUnrealEd->RequestPlaySession(SessionParams);
diff --git a/Engine/Source/Editor/UnrealEd/Private/PlayLevel.cpp b/Engine/Source/Editor/UnrealEd/Private/PlayLevel.cpp
--- a/Engine/Source/Editor/UnrealEd/Private/PlayLevel.cpp
+++ b/Engine/Source/Editor/UnrealEd/Private/PlayLevel.cpp
@@ -1391,6 +1391,12 @@ bool UEditorEngine::SpawnPlayFromHereStart( UWorld* World, AActor*& PlayerStart,
 		NavPlayerStart->bIsPIEPlayerStart = true;
 	}
 
+	// If PlayFromHere originated from a specific actor
+	if (World->PersistentLevel->PlayFromHereActor)
+	{
+		World->PersistentLevel->PlayFromHereActor->OnPlayFromHere();
+	}
+
 	return true;
 }
 
diff --git a/Engine/Source/Editor/UnrealEd/Public/Kismet2/DebuggerCommands.h b/Engine/Source/Editor/UnrealEd/Public/Kismet2/DebuggerCommands.h
--- a/Engine/Source/Editor/UnrealEd/Public/Kismet2/DebuggerCommands.h
+++ b/Engine/Source/Editor/UnrealEd/Public/Kismet2/DebuggerCommands.h
@@ -138,6 +138,7 @@ public:
 	 * Called from the context menu to start previewing the game at the clicked location                   
 	 */
 	static void StartPlayFromHere();
+	static void StartPlayFromHere(const TOptional<FVector>& Location, const TOptional<FRotator>& Rotation);
 
 	static void ResumePlaySession_Clicked();
 	static void PausePlaySession_Clicked();
diff --git a/Engine/Source/Runtime/Engine/Classes/Engine/Level.h b/Engine/Source/Runtime/Engine/Classes/Engine/Level.h
--- a/Engine/Source/Runtime/Engine/Classes/Engine/Level.h
+++ b/Engine/Source/Runtime/Engine/Classes/Engine/Level.h
@@ -452,6 +452,8 @@ public:
 	TArray<AActor*> ActorsForGC;
 
 #if WITH_EDITORONLY_DATA
+	AActor* PlayFromHereActor;
+
 	/** List of modified, unsaved actors that needs to be duplicated for PIE */
 	UPROPERTY(NonPIEDuplicateTransient)
 	TMap<FName, AActor*> ActorsModifiedForPIE;
diff --git a/Engine/Source/Runtime/Engine/Classes/GameFramework/Actor.h b/Engine/Source/Runtime/Engine/Classes/GameFramework/Actor.h
--- a/Engine/Source/Runtime/Engine/Classes/GameFramework/Actor.h
+++ b/Engine/Source/Runtime/Engine/Classes/GameFramework/Actor.h
@@ -853,6 +853,10 @@ public:
 	/** Set the bLockLocation flag */
 	void SetLockLocation(bool bInLockLocation) { bLockLocation = bInLockLocation; }
 
+	/** Called on actor which initiated the PIE session */
+	virtual void OnPlayFromHere();
+
+	bool CanPlayFromHere() const { return bCanPlayFromHere; }
 private:
 	friend class FActorDescList;
 	friend class UActorDescContainer;
@@ -938,6 +942,9 @@ protected:
 	UPROPERTY(EditDefaultsOnly, AdvancedDisplay, Category=Cooking, meta=(DisplayName="Generate Optimized Blueprint Component Data"))
 	uint8 bOptimizeBPComponentData:1;
 
+	/** Whether the actor can be used as a PlayFromHere origin (OnPlayFromHere() will be called on that actor) */
+	UPROPERTY()
+	uint8 bCanPlayFromHere : 1;
 private:
 	/** Whether this actor is temporarily hidden within the editor; used for show/hide/etc functionality w/o dirtying the actor. */
 	UPROPERTY(Transient)
diff --git a/Engine/Source/Runtime/Engine/Private/ActorEditor.cpp b/Engine/Source/Runtime/Engine/Private/ActorEditor.cpp
--- a/Engine/Source/Runtime/Engine/Private/ActorEditor.cpp
+++ b/Engine/Source/Runtime/Engine/Private/ActorEditor.cpp
@@ -809,6 +809,11 @@ EActorGridPlacement AActor::GetDefaultGridPlacement() const
 	return EActorGridPlacement::None;
 }
 
+void AActor::OnPlayFromHere()
+{
+	check(bCanPlayFromHere);
+}
+
 TUniquePtr<FWorldPartitionActorDesc> AActor::CreateClassActorDesc() const
 {
 	return TUniquePtr<FWorldPartitionActorDesc>(new FWorldPartitionActorDesc());
diff --git a/Engine/Source/Runtime/Engine/Private/Level.cpp b/Engine/Source/Runtime/Engine/Private/Level.cpp
--- a/Engine/Source/Runtime/Engine/Private/Level.cpp
+++ b/Engine/Source/Runtime/Engine/Private/Level.cpp
@@ -330,6 +330,7 @@ ULevel::ULevel( const FObjectInitializer& ObjectInitializer )
 	FixupOverrideVertexColorsCount = 0;
 	bUseExternalActors = false;
 	bContainsStableActorGUIDs = true;
+	PlayFromHereActor = nullptr;
 #endif	
 	bActorClusterCreated = false;
 	bIsPartitioned = false;
@@ -590,6 +591,17 @@ void ULevel::Serialize( FArchive& Ar )
 		FWorldTileInfo Info;
 		Ar << Info;
 	}
+
+#if WITH_EDITORONLY_DATA
+	if (Ar.GetPortFlags() & PPF_DuplicateForPIE)
+	{
+		Ar << PlayFromHereActor;
+		if (Ar.IsSaving())
+		{
+			PlayFromHereActor = nullptr;
+		}
+	}
+#endif
 }
 
 void ULevel::CreateReplicatedDestructionInfo(AActor* const Actor)
