Fix CIS build

diff --git a/Engine/Source/Editor/UnrealEd/Classes/Editor/FixupLazyObjectPtrForPIEArchive.h b/Engine/Source/Editor/UnrealEd/Classes/Editor/FixupLazyObjectPtrForPIEArchive.h
--- a/Engine/Source/Editor/UnrealEd/Classes/Editor/FixupLazyObjectPtrForPIEArchive.h
+++ b/Engine/Source/Editor/UnrealEd/Classes/Editor/FixupLazyObjectPtrForPIEArchive.h
@@ -8,6 +8,7 @@
 
 #include "Engine/World.h"
 #include "UObject/LazyObjectPtr.h"
+#include "Components/InstancedStaticMeshComponent.h"
 
 class FFixupLazyObjectPtrForPIEArchive : public FArchiveUObject
 {
