EditorDomain: Fix race condition assertion in FEditorDomainAsyncReadFileHandle: IAsyncReadRequests must not be deleted until they are completed. Read requests need to delete the memory in their destructor if not claimed.
Repro case was during launch of the editor for a large project withe editor domain already populated.
#rnx

diff --git a/Engine/Source/Editor/UnrealEd/Private/EditorDomain/EditorDomainArchive.cpp b/Engine/Source/Editor/UnrealEd/Private/EditorDomain/EditorDomainArchive.cpp
--- a/Engine/Source/Editor/UnrealEd/Private/EditorDomain/EditorDomainArchive.cpp
+++ b/Engine/Source/Editor/UnrealEd/Private/EditorDomain/EditorDomainArchive.cpp
@@ -320,6 +320,19 @@ public:
 		SetComplete();
 	}
 
+	virtual ~FAsyncReadRequestConstant()
+	{
+		if (Memory)
+		{
+			// this can happen with a race on cancel, it is ok, they didn't take the memory, free it now
+			if (!bUserSuppliedMemory)
+			{
+				FMemory::Free(Memory);
+			}
+			Memory = nullptr;
+		}
+	}
+
 protected:
 	virtual void WaitCompletionImpl(float TimeLimitSeconds) override
 	{
