CleanupPhysicsSceneThatWasInitializedForSave: Pass NewWorld=null to CleanupWorld so that the WorldPartition->CleanupWorldPartition gets called on Worlds that we InitWorld on during cooking. Not doing this previously was causing the WorldPartition (which has some GCObjects internally) to keep itself referenced for the rest of the cook.
#rnx

diff --git a/Engine/Source/Editor/UnrealEd/Private/EditorEngine.cpp b/Engine/Source/Editor/UnrealEd/Private/EditorEngine.cpp
--- a/Engine/Source/Editor/UnrealEd/Private/EditorEngine.cpp
+++ b/Engine/Source/Editor/UnrealEd/Private/EditorEngine.cpp
@@ -4285,7 +4285,7 @@ void UEditorEngine::CleanupPhysicsSceneThatWasInitializedForSave(UWorld* World,
 
 	if(bForceInitialized)
 	{
-		World->CleanupWorld(true, true, World);
+		World->CleanupWorld();
 	}
 
 	World->SetPhysicsScene(nullptr);
