CookPackageSplitter: Add garbage collection of the generator package after the save of all generated packages, and optionally allow it before then any time after the save of the generator. This prevents running out of memory during the cook of large-map world partition projects.
#rnx

diff --git a/Engine/Source/Editor/UnrealEd/Private/CookOnTheFlyServer.cpp b/Engine/Source/Editor/UnrealEd/Private/CookOnTheFlyServer.cpp
--- a/Engine/Source/Editor/UnrealEd/Private/CookOnTheFlyServer.cpp
+++ b/Engine/Source/Editor/UnrealEd/Private/CookOnTheFlyServer.cpp
@@ -1914,7 +1914,7 @@ struct FPopulatePackageContext
 	// Input variables
 	FGeneratorPackage* GeneratorStruct = nullptr;
 	const UPackage* OwnerPackage = nullptr;
-	const FGeneratorPackage::FGeneratedStruct* GeneratedStruct = nullptr;
+	FGeneratorPackage::FGeneratedStruct* GeneratedStruct = nullptr;
 
 	// Cache of initialized-on-demand variables
 	UObject* OwnerObject = nullptr;
@@ -2325,7 +2325,7 @@ void UCookOnTheFlyServer::SplitPackage(UE::Cook::FGeneratorPackage* GeneratorStr
 			int32& NextIndex = GeneratorStruct->GetNextPopulateIndex();
 			while (NextIndex < PackagesToGenerate.Num())
 			{
-				const FGeneratorPackage::FGeneratedStruct& GeneratedStruct = PackagesToGenerate[NextIndex++];
+				FGeneratorPackage::FGeneratedStruct& GeneratedStruct = PackagesToGenerate[NextIndex++];
 				FPackageData* GeneratedPackageData = GeneratedStruct.PackageData;
 				GeneratedPackageData->ClearCookedPlatformData();
 
@@ -2365,7 +2365,7 @@ void UCookOnTheFlyServer::SplitPackage(UE::Cook::FGeneratorPackage* GeneratorStr
 	TArrayView<FGeneratorPackage::FGeneratedStruct> PackagesToGenerate = GeneratorStruct->GetPackagesToGenerate();
 	TArray<ICookPackageSplitter::FGeneratedPackageForPreSave> SplitterDatas;
 	SplitterDatas.Reserve(PackagesToGenerate.Num());
-	for (const FGeneratorPackage::FGeneratedStruct& GeneratedStruct : PackagesToGenerate)
+	for (FGeneratorPackage::FGeneratedStruct& GeneratedStruct : PackagesToGenerate)
 	{
 		ICookPackageSplitter::FGeneratedPackageForPreSave& SplitterData = SplitterDatas.Emplace_GetRef();
 		SplitterData.RelativePath = GeneratedStruct.RelativePath;
@@ -2375,7 +2375,7 @@ void UCookOnTheFlyServer::SplitPackage(UE::Cook::FGeneratorPackage* GeneratorStr
 		SplitterData.Package = FindObject<UPackage>(nullptr, *GeneratedPackageName);
 		if (!SplitterData.Package)
 		{
-			SplitterData.Package = GeneratorStruct->CreateGeneratedUPackage(Owner, *GeneratedPackageName);
+			SplitterData.Package = GeneratorStruct->CreateGeneratedUPackage(GeneratedStruct, Owner, *GeneratedPackageName);
 		}
 	}
 
@@ -2391,7 +2391,7 @@ UPackage* UCookOnTheFlyServer::TryPopulateGeneratedPackage(UE::Cook::FPopulatePa
 	using namespace UE::Cook;
 
 	FGeneratorPackage* GeneratorStruct = Context.GeneratorStruct;
-	const FGeneratorPackage::FGeneratedStruct& GeneratedStruct = *Context.GeneratedStruct;
+	FGeneratorPackage::FGeneratedStruct& GeneratedStruct = *Context.GeneratedStruct;
 	const UPackage* OwnerPackage = Context.OwnerPackage;
 	UE::Cook::FPackageData* GeneratedPackageData = Context.GeneratedStruct->PackageData;
 	const FString GeneratedPackageName = GeneratedPackageData->GetPackageName().ToString();
@@ -2424,20 +2424,25 @@ UPackage* UCookOnTheFlyServer::TryPopulateGeneratedPackage(UE::Cook::FPopulatePa
 	bool bPopulatedByPreSave = false;
 	if (GeneratedPackage)
 	{
-		if (!Splitter->UseDeferredPopulate())
+		if (!GeneratedStruct.bHasCreatedPackage)
 		{
-			UE_LOG(LogCook, Error, TEXT("PackageSplitter found an existing copy of a package it was trying to populate;")
-				TEXT("this is unexpected since we previously checked and did not find it. Splitter=%s, Generated=%s."),
-				*GeneratorStruct->GetSplitDataObjectName().ToString(), *GeneratedPackageName);
-			return nullptr;
-		}
-		else if (GeneratorStruct->HasGCOccurredAfterGenerate())
-		{
-			UE_LOG(LogCook, Error, TEXT("PackageSplitter found an existing copy of a package it was trying to populate;")
-				TEXT("this is unexpected since garbage has been collected and the package should have been unreferenced so it should have been collected.")
-				TEXT("Splitter=%s, Generated=%s."),
-				*GeneratorStruct->GetSplitDataObjectName().ToString(), *GeneratedPackageName);
-			FReferenceChainSearch RefChainSearch(GeneratedPackage, EReferenceChainSearchMode::Shortest | EReferenceChainSearchMode::PrintAllResults);
+			if (!Splitter->UseDeferredPopulate())
+			{
+				UE_LOG(LogCook, Error, TEXT("PackageSplitter found an existing copy of a package it was trying to populate;")
+					TEXT("this is unexpected since we previously checked and did not find it. Splitter=%s, Generated=%s."),
+					*GeneratorStruct->GetSplitDataObjectName().ToString(), *GeneratedPackageName);
+			}
+			else
+			{
+				UE_LOG(LogCook, Error, TEXT("PackageSplitter found an existing copy of a package it was trying to populate;")
+					TEXT("this is unexpected since garbage has been collected and the package should have been unreferenced so it should have been collected.")
+					TEXT("Splitter=%s, Generated=%s."),
+					*GeneratorStruct->GetSplitDataObjectName().ToString(), *GeneratedPackageName);
+				EReferenceChainSearchMode SearchMode = EReferenceChainSearchMode::Shortest
+					| EReferenceChainSearchMode::PrintAllResults
+					| EReferenceChainSearchMode::FullChain;
+				FReferenceChainSearch RefChainSearch(GeneratedPackage, SearchMode);
+			}
 			return nullptr;
 		}
 		// Otherwise this is the package that was created and passed to presave, and it is still valid because there has not been a GC since
@@ -2446,7 +2451,7 @@ UPackage* UCookOnTheFlyServer::TryPopulateGeneratedPackage(UE::Cook::FPopulatePa
 	}
 	else
 	{
-		GeneratedPackage = GeneratorStruct->CreateGeneratedUPackage(OwnerPackage, *GeneratedPackageName);
+		GeneratedPackage = GeneratorStruct->CreateGeneratedUPackage(GeneratedStruct, OwnerPackage, *GeneratedPackageName);
 	}
 
 	// Populate package using CookPackageSplitterInstance and pass GeneratedPackage's cooked name for it to
@@ -2456,7 +2461,9 @@ UPackage* UCookOnTheFlyServer::TryPopulateGeneratedPackage(UE::Cook::FPopulatePa
 	PopulateData.Package = GeneratedPackage;
 	PopulateData.bPopulatedByPreSave = bPopulatedByPreSave;
 	PopulateData.bCreatedAsMap = GeneratedStruct.bCreateAsMap;
-	if (!Splitter->TryPopulatePackage(OwnerPackage, Context.OwnerObject, PopulateData))
+	bool bWasOwnerReloaded = GeneratorStruct->GetWasOwnerReloaded();
+	GeneratorStruct->ClearWasOwnerReloaded();
+	if (!Splitter->TryPopulatePackage(OwnerPackage, Context.OwnerObject, PopulateData, bWasOwnerReloaded))
 	{
 		UE_LOG(LogCook, Error, TEXT("PackageSplitter returned false from TryPopulatePackage. Splitter=%s, Generated=%s."),
 			*GeneratorStruct->GetSplitDataObjectName().ToString(), *GeneratedPackageName);
@@ -2777,8 +2784,23 @@ void UCookOnTheFlyServer::ReleaseCookedPlatformData(UE::Cook::FPackageData& Pack
 			{
 				GeneratorPackage->GetCookPackageSplitterInstance()->PostSaveGeneratorPackage(PackageData.GetPackage(), SplitObject);
 			}
+
+			PackageData.ResetGenerationProgress();
+			if (IsCookByTheBookMode() && GeneratorPackage->IsComplete())
+			{
+				PackageData.DestroyGeneratorPackage();
+			}
+		}
+		else
+		{
+			PackageData.ResetGenerationProgress();
+		}
+
+		FGeneratorPackage* OwnerGeneratorPackage = PackageData.GetGeneratedOwner();
+		if (OwnerGeneratorPackage)
+		{
+			OwnerGeneratorPackage->SetGeneratedSaved(PackageData);
 		}
-		PackageData.ResetGenerationProgress();
 	}
 	PackageData.ClearCookedPlatformData();
 
diff --git a/Engine/Source/Editor/UnrealEd/Private/Cooker/CookPackageData.cpp b/Engine/Source/Editor/UnrealEd/Private/Cooker/CookPackageData.cpp
--- a/Engine/Source/Editor/UnrealEd/Private/Cooker/CookPackageData.cpp
+++ b/Engine/Source/Editor/UnrealEd/Private/Cooker/CookPackageData.cpp
@@ -15,6 +15,7 @@
 #include "ShaderCompiler.h"
 #include "UObject/Object.h"
 #include "UObject/Package.h"
+#include "UObject/ReferenceChainSearch.h"
 #include "UObject/UObjectGlobals.h"
 #include "UObject/UObjectHash.h"
 
@@ -1027,6 +1028,7 @@ namespace Cook
 			check(PackageData.GetGeneratedOwner() == nullptr); 
 			PackageData.SetGeneratedOwner(this);
 		}
+		RemainingToPopulate = GeneratorDatas.Num();
 		return true;
 	}
 
@@ -1057,22 +1059,95 @@ namespace Cook
 
 	void FGeneratorPackage::PostGarbageCollect()
 	{
-		if (Owner.HasCompletedGeneration())
+		if (!bGeneratedList)
 		{
-			bGCOccurredAfterGenerate = true;
+			return;
+		}
+		if (Owner.GetState() == EPackageState::Save)
+		{
+			// UCookOnTheFlyServer::PreCollectGarbage adds references for the Generator package and all its public objects, so it should still be loaded
+			if (!Owner.GetPackage() || !FindSplitDataObject())
+			{
+				UE_LOG(LogCook, Error, TEXT("PackageSplitter object was deleted by garbage collection while generation was still ongoing. This will break the generation.")
+					TEXT("\n\tSplitter=%s."), *GetSplitDataObjectName().ToString());
+			}
+		}
+		else
+		{
+			// After the Generator Package is saved, we drop its referenced and it can be garbage collected
+			// If we have any packages left to populate, our splitter contract requires that it be garbage collected;
+			// we promise that the package is not partially GC'd during calls to TryPopulateGeneratedPackage
+			// The splitter can opt-out of this contract and keep it referenced itself if it desires.
+			UPackage* OwnerPackage = FindObject<UPackage>(nullptr, *Owner.GetPackageName().ToString());
+			if (OwnerPackage)
+			{
+				if (RemainingToPopulate > 0 && !CookPackageSplitterInstance->UseInternalReferenceToAvoidGarbageCollect())
+				{
+					UE_LOG(LogCook, Error, TEXT("PackageSplitter found the Generator package still in memory after it should have been deleted by GC.")
+						TEXT("\n\tThis is unexpected since garbage has been collected and the package should have been unreferenced so it should have been collected, and will break population of Generated packages.")
+						TEXT("\n\tSplitter=%s"), *GetSplitDataObjectName().ToString());
+					EReferenceChainSearchMode SearchMode = EReferenceChainSearchMode::Shortest
+						| EReferenceChainSearchMode::PrintAllResults
+						| EReferenceChainSearchMode::FullChain;
+					FReferenceChainSearch RefChainSearch(OwnerPackage, SearchMode);
+				}
+			}
+			else
+			{
+				bWasOwnerReloaded = true;
+			}
+		}
+
+		bool bHasIssuedWarning = false;
+		for (FGeneratedStruct& GeneratedStruct : PackagesToGenerate)
+		{
+			GeneratedStruct.bHasCreatedPackage = false;
+			if (!GeneratedStruct.bHasSaved && !bHasIssuedWarning)
+			{
+				if (FindObject<UPackage>(nullptr, *GeneratedStruct.PackageData->GetPackageName().ToString()))
+				{
+					UE_LOG(LogCook, Warning, TEXT("PackageSplitter found a package it generated that was not removed from memory during garbage collection. This will cause errors later during population.")
+						TEXT("\n\tSplitter=%s, Generated=%s."), *GetSplitDataObjectName().ToString(), *GeneratedStruct.PackageData->GetPackageName().ToString());
+					bHasIssuedWarning = true; // Only issue the warning once per GC
+				}
+			}
 		}
 	}
 
-	UPackage* FGeneratorPackage::CreateGeneratedUPackage(const UPackage* OwnerPackage, const TCHAR* GeneratedPackageName)
+	UPackage* FGeneratorPackage::CreateGeneratedUPackage(FGeneratorPackage::FGeneratedStruct& GeneratedStruct,
+		const UPackage* OwnerPackage, const TCHAR* GeneratedPackageName)
 	{
 		UPackage* GeneratedPackage = CreatePackage(GeneratedPackageName);
 		PRAGMA_DISABLE_DEPRECATION_WARNINGS
 		GeneratedPackage->SetGuid(OwnerPackage->GetGuid());
 		PRAGMA_ENABLE_DEPRECATION_WARNINGS
 		GeneratedPackage->SetPersistentGuid(OwnerPackage->GetPersistentGuid());
+		GeneratedStruct.bHasCreatedPackage = true;
 		return GeneratedPackage;
 	}
 
+	void FGeneratorPackage::SetGeneratedSaved(FPackageData& PackageData)
+	{
+		FGeneratedStruct* GeneratedStruct = FindGeneratedStruct(&PackageData);
+		if (!GeneratedStruct)
+		{
+			UE_LOG(LogCook, Warning, TEXT("PackageSplitter called SetGeneratedSaved on a package that does not belong to the splitter.")
+				TEXT("\n\tSplitter=%s, Generated=%s."), *GetSplitDataObjectName().ToString(), *PackageData.GetPackageName().ToString());
+			return;
+		}
+		if (GeneratedStruct->bHasSaved)
+		{
+			return;
+		}
+		GeneratedStruct->bHasSaved = true;
+		--RemainingToPopulate;
+		check(RemainingToPopulate >= 0);
+	}
+
+	bool FGeneratorPackage::IsComplete() const
+	{
+		return bGeneratedList && RemainingToPopulate == 0;
+	}
 
 	void FGeneratorPackage::GetIntermediateMountPoint(FString& OutPackagePath, FString& OutLocalFilePath) const
 	{
diff --git a/Engine/Source/Editor/UnrealEd/Private/Cooker/CookPackageData.h b/Engine/Source/Editor/UnrealEd/Private/Cooker/CookPackageData.h
--- a/Engine/Source/Editor/UnrealEd/Private/Cooker/CookPackageData.h
+++ b/Engine/Source/Editor/UnrealEd/Private/Cooker/CookPackageData.h
@@ -420,6 +420,8 @@ namespace Cook
 			TArray<FName> Dependencies;
 			FPackageData* PackageData = nullptr;
 			bool bCreateAsMap = false;
+			bool bHasCreatedPackage = false;
+			bool bHasSaved = false;
 		};
 
 		/** Store the provided CookPackageSplitter and prepare the packages to generate. */
@@ -472,18 +474,19 @@ namespace Cook
 		void SetQueuedGeneratedPackages() { bQueuedGeneratedPackages = true; }
 		int32& GetNextPopulateIndex() { return NextPopulateIndex; }
 
-		/**
-		 * Records whether GC occurred after creating the placeholder packages.
-		 * If GC has occurred, we expect to no longer find the placeholder packages in memory,
-		 * and if we do find them its an error because we no longer know which objects are in the packages.
-		 */
-		bool HasGCOccurredAfterGenerate() const { return bGCOccurredAfterGenerate; }
-
 		/** Callback during garbage collection */
 		void PostGarbageCollect();
 
-		/** Call CreatePackage and set PackageData correctly for deterministic generated packages */
-		UPackage* CreateGeneratedUPackage(const UPackage* OwnerPackage, const TCHAR* GeneratedPackageName);
+		/** Get/Clear whether the Owner was garbage collected since the previous call to TryPopulatePackage */
+		bool GetWasOwnerReloaded() const { return bWasOwnerReloaded; }
+		void ClearWasOwnerReloaded() { bWasOwnerReloaded = false; }
+		/** Call CreatePackage and set PackageData correctly for deterministic generated packages, and update status. */
+		UPackage* CreateGeneratedUPackage(FGeneratedStruct& GeneratedStruct,
+			const UPackage* OwnerPackage, const TCHAR* GeneratedPackageName);
+		/** Mark that the generated package has been populated, to keep track of when this GeneratorPackage is no longer needed. */
+		void SetGeneratedSaved(FPackageData& PackageData);
+		/** Return whether list has been generated and all generated packages have been populated */
+		bool IsComplete() const;
 
 	private:
 		/** PackageData for the package that is being split */
@@ -495,16 +498,13 @@ namespace Cook
 		/** Recorded list of packages to generate from the splitter, and data we need about them */
 		TArray<FGeneratedStruct> PackagesToGenerate;
 
-		/** State variables for reentrant SplitPackage calls */
 		int32 NextPopulateIndex = 0;
+		int32 RemainingToPopulate = 0;
 		bool bGeneratedList = false;
 		bool bClearedOldPackages = false;
 		bool bClearedOldPackagesWithGC = false;
 		bool bQueuedGeneratedPackages = false;
-
-		/** Records whether GC occurred after creating the placeholder packages. */
-		bool bGCOccurredAfterGenerate = false;
-
+		bool bWasOwnerReloaded = false;
 	};
 
 
diff --git a/Engine/Source/Editor/UnrealEd/Public/CookPackageSplitter.h b/Engine/Source/Editor/UnrealEd/Public/CookPackageSplitter.h
--- a/Engine/Source/Editor/UnrealEd/Public/CookPackageSplitter.h
+++ b/Engine/Source/Editor/UnrealEd/Public/CookPackageSplitter.h
@@ -46,6 +46,12 @@ public:
 	 */
 	virtual bool UseDeferredPopulate() { return false; }
 
+	/**
+	 * If true, this splitter forces the Generator package objects it needs to remain referenced, and the cooker
+	 * should expect them to still be in memory after a garbage collect so long as the splitter is alive.
+	 */
+	virtual bool UseInternalReferenceToAvoidGarbageCollect() { return false; }
+
 	/** Return the list of packages to generate. */
 	virtual TArray<FGeneratedPackage> GetGenerateList(const UPackage* OwnerPackage, const UObject* OwnerObject) = 0;
 
@@ -74,9 +80,12 @@ public:
 	 * @param OwnerPackage				The parent package being split
 	 * @param OwnerObject				The SplitDataClass instance that this CookPackageSplitter instance was created for
 	 * @param GeneratedPackage			Pointer and information about the package to populate
+	 * @param bWasOwnerReloaded			True iff the OwnerPackage was garbage collected and reloaded since the previous call
+	 *                                  to GetGenerateList or TryPopulatePackage
 	 * @return							True if successfully populates,  false on error (this will cause a cook error).
 	 */
-	virtual bool TryPopulatePackage(const UPackage* OwnerPackage, const UObject* OwnerObject, const FGeneratedPackageForPopulate& GeneratedPackage) = 0;
+	virtual bool TryPopulatePackage(const UPackage* OwnerPackage, const UObject* OwnerObject,
+		const FGeneratedPackageForPopulate& GeneratedPackage, bool bWasOwnerReloaded) = 0;
 
 	struct FGeneratedPackageForPreSave
 	{
diff --git a/Engine/Source/Runtime/Engine/Private/WorldPartition/WorldPartitionActorDesc.cpp b/Engine/Source/Runtime/Engine/Private/WorldPartition/WorldPartitionActorDesc.cpp
--- a/Engine/Source/Runtime/Engine/Private/WorldPartition/WorldPartitionActorDesc.cpp
+++ b/Engine/Source/Runtime/Engine/Private/WorldPartition/WorldPartitionActorDesc.cpp
@@ -285,16 +285,20 @@ void FWorldPartitionActorDesc::Unload()
 {
 	if (AActor* Actor = GetActor())
 	{
-		check(Actor->IsPackageExternal());
-
-		ForEachObjectWithPackage(Actor->GetPackage(), [](UObject* Object)
+		// @todo_ow: FWorldPartitionCookPackageSplitter should mark each FWorldPartitionActorDesc as moved, which we would assert on here rather than asserting IsRunningCookCommandlet
+		// and the splitter should take responsbility for calling ClearFlags on every object in the package when it does the move
+		check(Actor->IsPackageExternal() || IsRunningCookCommandlet());
+		if (Actor->IsPackageExternal())
 		{
-			if (Object->HasAnyFlags(RF_Public | RF_Standalone))
+			ForEachObjectWithPackage(Actor->GetPackage(), [](UObject* Object)
 			{
-				CastChecked<UMetaData>(Object)->ClearFlags(RF_Public | RF_Standalone);
-			}
-			return true;
-		}, false);
+				if (Object->HasAnyFlags(RF_Public | RF_Standalone))
+				{
+					CastChecked<UMetaData>(Object)->ClearFlags(RF_Public | RF_Standalone);
+				}
+				return true;
+			}, false);
+		}
 
 		ActorPtr = nullptr;
 	}
diff --git a/Engine/Source/Runtime/Engine/Private/WorldPartition/WorldPartitionCookPackageSplitter.cpp b/Engine/Source/Runtime/Engine/Private/WorldPartition/WorldPartitionCookPackageSplitter.cpp
--- a/Engine/Source/Runtime/Engine/Private/WorldPartition/WorldPartitionCookPackageSplitter.cpp
+++ b/Engine/Source/Runtime/Engine/Private/WorldPartition/WorldPartitionCookPackageSplitter.cpp
@@ -34,6 +34,38 @@ bool FWorldPartitionCookPackageSplitter::UseDeferredPopulate()
 	return UseDeferredPopulate.Value;
 }
 
+FWorldPartitionCookPackageSplitter::FWorldPartitionCookPackageSplitter()
+{
+	FCoreUObjectDelegates::GetPreGarbageCollectDelegate().AddRaw(this, &FWorldPartitionCookPackageSplitter::PreGarbageCollect);
+}
+
+FWorldPartitionCookPackageSplitter::~FWorldPartitionCookPackageSplitter()
+{
+	FCoreUObjectDelegates::GetPreGarbageCollectDelegate().RemoveAll(this);
+	TeardownWorldPartition();
+}
+
+void FWorldPartitionCookPackageSplitter::PreGarbageCollect()
+{
+	TeardownWorldPartition();
+}
+
+void FWorldPartitionCookPackageSplitter::TeardownWorldPartition()
+{
+	if (bWorldPartitionNeedsTeardown)
+	{
+		if (UWorld* LocalWorld = ReferencedWorld.Get())
+		{
+			UWorldPartition* WorldPartition = LocalWorld->PersistentLevel->GetWorldPartition();
+			if (WorldPartition)
+			{
+				WorldPartition->Uninitialize();
+			}
+		}
+		bWorldPartitionNeedsTeardown = false;
+	}
+}
+
 UWorld* FWorldPartitionCookPackageSplitter::ValidateDataObject(UObject* SplitData)
 {
 	UWorld* PartitionedWorld = CastChecked<UWorld>(SplitData);
@@ -83,8 +115,18 @@ TArray<ICookPackageSplitter::FGeneratedPackage> FWorldPartitionCookPackageSplitt
 }
 
 bool FWorldPartitionCookPackageSplitter::TryPopulatePackage(const UPackage* OwnerPackage, const UObject* OwnerObject,
-	const FGeneratedPackageForPopulate& GeneratedPackage)
+	const FGeneratedPackageForPopulate& GeneratedPackage, bool bWasOwnerReloaded)
 {
+	if (bWasOwnerReloaded)
+	{
+		GetGenerateList(OwnerPackage, OwnerObject);
+		// When GetGenerateList is called by CookOnTheFlyServer, it is followed up with a call to 
+		// CleanupWorld when the Generator package is done saving, which calls WorldPartition->Uninitialize
+		// But when we call it here to initialize our data for the populate, we need to take responsibility
+		// for calling WorldPartition->Uninitialize before the next GarbageCollection
+		bWorldPartitionNeedsTeardown = true;
+	}
+
 	// TODO: Make PopulateGeneratedPackageForCook const so we can honor the constness of the OwnerObject in this API function
 	const UWorld* ConstPartitionedWorld = ValidateDataObject(OwnerObject);
 	UWorld* PartitionedWorld = const_cast<UWorld*>(ConstPartitionedWorld);
@@ -101,9 +143,4 @@ void FWorldPartitionCookPackageSplitter::PreSaveGeneratorPackage(UPackage* Owner
 	WorldPartition->FinalizeGeneratorPackageForCook(GeneratedPackages);
 }
 
-void FWorldPartitionCookPackageSplitter::AddReferencedObjects(FReferenceCollector& Collector)
-{
-	Collector.AddReferencedObject(ReferencedWorld);
-}
-
 #endif
\ No newline at end of file
diff --git a/Engine/Source/Runtime/Engine/Private/WorldPartition/WorldPartitionCookPackageSplitter.h b/Engine/Source/Runtime/Engine/Private/WorldPartition/WorldPartitionCookPackageSplitter.h
--- a/Engine/Source/Runtime/Engine/Private/WorldPartition/WorldPartitionCookPackageSplitter.h
+++ b/Engine/Source/Runtime/Engine/Private/WorldPartition/WorldPartitionCookPackageSplitter.h
@@ -6,35 +6,34 @@
 
 #include "CookPackageSplitter.h"
 #include "Engine/World.h"
-#include "UObject/GCObject.h"
+#include "UObject/WeakObjectPtrTemplates.h"
 
-class FWorldPartitionCookPackageSplitter : public ICookPackageSplitter, public FGCObject
+class FWorldPartitionCookPackageSplitter : public ICookPackageSplitter
 {
 public:
 	//~ Begin of ICookPackageSplitter
 	static bool ShouldSplit(UObject* SplitData);
+
+	FWorldPartitionCookPackageSplitter();
+	virtual ~FWorldPartitionCookPackageSplitter();
+
 	virtual bool UseDeferredPopulate() override;
-	virtual ~FWorldPartitionCookPackageSplitter() {}
+
 	virtual TArray<ICookPackageSplitter::FGeneratedPackage> GetGenerateList(const UPackage* OwnerPackage, const UObject* OwnerObject) override;
 	virtual bool TryPopulatePackage(const UPackage* OwnerPackage, const UObject* OwnerObject,
-		const ICookPackageSplitter::FGeneratedPackageForPopulate& GeneratedPackage) override;
+		const ICookPackageSplitter::FGeneratedPackageForPopulate& GeneratedPackage, bool bWasOwnerReloaded) override;
 	virtual void PreSaveGeneratorPackage(UPackage* OwnerPackage, UObject* OwnerObject,
 		const TArray<ICookPackageSplitter::FGeneratedPackageForPreSave>& GeneratedPackages) override;
 	//~ End of ICookPackageSplitter
 
-	//~ Begin of FGCObject
-	virtual void AddReferencedObjects(FReferenceCollector& Collector);
-	virtual FString GetReferencerName() const
-	{
-		return TEXT("FWorldPartitionCookPackageSplitter");
-	}
-	//~ End of FGCObject
-
 private:
 	const UWorld* ValidateDataObject(const UObject* SplitData);
 	UWorld* ValidateDataObject(UObject* SplitData);
+	void PreGarbageCollect();
+	void TeardownWorldPartition();
 
-	UWorld* ReferencedWorld = nullptr;
+	TWeakObjectPtr<UWorld> ReferencedWorld;
+	bool bWorldPartitionNeedsTeardown = false;
 };
 
 #endif
\ No newline at end of file
