-Have the remote executor force inputs to be present before building the merkle tree
-Place temporary limits on the sizes and quantities of build jobs that will be remote executed

diff --git a/Engine/Source/Editor/UnrealEd/Private/DerivedDataBuildRemoteExecutor.cpp b/Engine/Source/Editor/UnrealEd/Private/DerivedDataBuildRemoteExecutor.cpp
--- a/Engine/Source/Editor/UnrealEd/Private/DerivedDataBuildRemoteExecutor.cpp
+++ b/Engine/Source/Editor/UnrealEd/Private/DerivedDataBuildRemoteExecutor.cpp
@@ -2,6 +2,7 @@
 
 #include "DerivedDataBuildRemoteExecutor.h"
 
+#include "Algo/Find.h"
 #include "DerivedDataBuild.h"
 #include "DerivedDataBuildAction.h"
 #include "DerivedDataBuildInputs.h"
@@ -17,6 +18,7 @@
 #include "Modules/ModuleManager.h"
 #include "Serialization/CompactBinaryValidation.h"
 #include "Serialization/CompactBinaryWriter.h"
+#include <atomic>
 
 static uint32 GetTypeHash(const FDigest& Digest)
 {
@@ -132,7 +134,7 @@ public:
 		// Lookup tables for indexing in different scenarios
 		TMultiMap<FDigest, FVariantIndex> DigestFilesystemIndex;
 		TMap<FStringView, int32> PathToDirectoryIndex;
-		TMap<int32, FStringView> FileIndexToInputKey;
+		TMap<int32, FString> FileIndexToInputKey;
 
 		// Unique items in the tree
 		FCommand Command;
@@ -269,17 +271,16 @@ public:
 				}
 			}).Wait();
 
-		if (State.BuildInputs.IsValid())
-		{
-			State.BuildInputs.Get().IterateInputs([&State] (FStringView Key, const FCompressedBuffer& Buffer)
-				{
-					TStringBuilder<128> InputPath;
-					InputPath << TEXT("Inputs/") << FIoHash(Buffer.GetRawHash());
-					const FString& NewInputPath = State.InputPaths.Emplace_GetRef(InputPath);
-					int32 FileIndex = AddMerkleTreeFile(State, NewInputPath, FIoHash::HashBuffer(Buffer.GetCompressed()), Buffer.GetCompressedSize(), false, EFileType::Input, Buffer.GetCompressed());
-					State.FileIndexToInputKey.Add(FileIndex, Key);
-				});
-		}
+		State.BuildAction.IterateInputs([&State] (FStringView Key, const FIoHash& RawHash, uint64 RawSize)
+			{
+				TStringBuilder<128> InputPath;
+				InputPath << TEXT("Inputs/") << RawHash;
+				const FString& NewInputPath = State.InputPaths.Emplace_GetRef(InputPath);
+				const FCompressedBuffer& Buffer = State.BuildInputs.Get().FindInput(Key);
+				check(!Buffer.IsNull());
+				int32 FileIndex = AddMerkleTreeFile(State, NewInputPath, FIoHash::HashBuffer(Buffer.GetCompressed()), Buffer.GetCompressedSize(), false, EFileType::Input, Buffer.GetCompressed());
+				State.FileIndexToInputKey.Emplace(FileIndex, Key);
+			});
 
 		// This base directory must be created as worker executables (even those that don't exist in this directory) will attempt to change directories into it during startup.
 		int32 BaseDirectoryIndex = INDEX_NONE;
@@ -326,7 +327,6 @@ public:
 			});
 		ContentAddressableStorage->ToBlob(State.Command, State.CommandContentBytes, State.Action.CommandDigest);
 
-		State.Action.DoNotCache = true;
 		State.Action.Salt = Salt;
 		ContentAddressableStorage->ToBlob(State.Action, State.ActionContentBytes, State.ActionDigest);
 	}
@@ -460,6 +460,27 @@ public:
 		return ContentAddressableStorage->BatchUpdateBlobs(State.BatchUpdateBlobsRequest, State.BatchUpdateBlobsResponse);
 	}
 
+	bool ValidateUploadSuccess(FRemoteExecutionState& State)
+	{
+		bool bSuccess = true;
+		for (const FBatchUpdateBlobsRequest::FRequest& RequestedUpload : State.BatchUpdateBlobsRequest.Requests)
+		{
+			if (const FBatchUpdateBlobsResponse::FResponse* FoundResponse = Algo::FindByPredicate(State.BatchUpdateBlobsResponse.Responses, [&RequestedUpload] (const FBatchUpdateBlobsResponse::FResponse& Response)
+				{
+					return Response.Digest == RequestedUpload.Digest;
+				}))
+			{
+				if (!FoundResponse->Status.Ok())
+				{
+					FStringView ActionName = State.BuildAction.GetName();
+					UE_LOG(LogDerivedDataBuildRemoteExecutor, Warning, TEXT("Remote execution system error: data for action '%.*s' could not be uploaded (hash: %s, size: %u)"), ActionName.Len(), ActionName.GetData(), *LexToString(RequestedUpload.Digest.Hash), RequestedUpload.Digest.SizeBytes);
+					bSuccess = false;
+				}
+			}
+		}
+		return bSuccess;
+	}
+
 	bool ExecuteBuild(FRemoteExecutionState& State)
 	{
 		State.ExecuteRequest.InstanceName = InstanceName;
@@ -568,6 +589,30 @@ public:
 		return OutputBuilder.Build();
 	}
 
+	bool IsRemoteExecutionAttemptAllowed(const FBuildAction& Action)
+	{
+		uint64 TotalInputSize = 0;
+		Action.IterateInputs([&TotalInputSize] (FStringView Key, const FIoHash& RawHash, uint64 RawSize)
+			{
+				TotalInputSize += RawSize;
+			});
+
+		// Only bother sending large input data for remote execution right now.
+		if (TotalInputSize < 10*1024*1024)
+		{
+			return false;
+		}
+
+		static std::atomic<uint32> NumRemoteBuilds {0};
+		// Limit to 10 total remote builds for now until the remote instance can handle a high quantity better.
+		if (NumRemoteBuilds.fetch_add(1, std::memory_order_relaxed) > 8)
+		{
+			return false;
+		}
+
+		return true;
+	}
+
 	FRequest BuildAction(
 		const FBuildAction& Action,
 		const FOptionalBuildInputs& Inputs,
@@ -576,15 +621,45 @@ public:
 		EPriority Priority,
 		FOnBuildWorkerActionComplete&& OnComplete) final
 	{
+		if (!IsRemoteExecutionAttemptAllowed(Action))
+		{
+			OnComplete({Action.GetKey(), {}, {}, EStatus::Error});
+			return FRequest();
+		}
+
+		TArray<FString> MissingInputs;
+		TArray<FStringView> MissingInputViews;
+		{
+			// TODO: This block forces resolution of inputs before we attempt to determine which
+			//		 inputs need to be uploaded.  This is required because we can't refer to inputs
+			//		 in the Merkle tree by their RawHash/RawSize but instead must send their CompressedHash/
+			//		 CompressedSize.  Once the remote execution API allows us to represent inputs with RawHash/
+			//		 RawSize, this block can be removed and we can find missing CAS inputs without having resolved
+			//		 the inputs first.
+			Action.IterateInputs([&MissingInputs, &MissingInputViews, &Inputs] (FStringView Key, const FIoHash& RawHash, uint64 RawSize)
+				{
+					if (Inputs.IsNull() || Inputs.Get().FindInput(Key).IsNull())
+					{
+						MissingInputs.Emplace(Key);
+						MissingInputViews.Add(MissingInputs.Last());
+					}
+				});
+
+			if (!MissingInputViews.IsEmpty())
+			{
+				OnComplete({Action.GetKey(), {}, MissingInputViews, EStatus::Ok});
+				return FRequest();
+			}
+		}
+
 		FRemoteExecutionState State { Action, Inputs, Worker, Policy, Priority };
 
 		DetermineMissingBlobs(State);
 
-		TArray<FStringView> MissingInputs;
-		GatherMissingInputFileBlobs(State, MissingInputs);
-		if (!MissingInputs.IsEmpty())
+		GatherMissingInputFileBlobs(State, MissingInputViews);
+		if (!MissingInputViews.IsEmpty())
 		{
-			OnComplete({State.BuildAction.GetKey(), {}, MissingInputs, EStatus::Ok});
+			OnComplete({State.BuildAction.GetKey(), {}, MissingInputViews, EStatus::Ok});
 			return FRequest();
 		}
 
@@ -592,6 +667,12 @@ public:
 		if (!State.FindMissingBlobsResponse.MissingBlobDigests.IsEmpty())
 		{
 			UploadMissingBlobs(State);
+			UE_LOG(LogDerivedDataBuildRemoteExecutor, Display, TEXT("Uploaded %d data blobs for remote execution."), State.FindMissingBlobsResponse.MissingBlobDigests.Num());
+			if (!ValidateUploadSuccess(State))
+			{
+				OnComplete({Action.GetKey(), {}, {}, EStatus::Error});
+				return FRequest();
+			}
 		}
 
 		if (ExecuteBuild(State))
diff --git a/Engine/Source/Runtime/Engine/Private/TextureDerivedDataTask.h b/Engine/Source/Runtime/Engine/Private/TextureDerivedDataTask.h
--- a/Engine/Source/Runtime/Engine/Private/TextureDerivedDataTask.h
+++ b/Engine/Source/Runtime/Engine/Private/TextureDerivedDataTask.h
@@ -159,7 +159,7 @@ class FTextureCacheDerivedDataWorker : public FNonAbandonableTask
 			NewRef.Id = Source.GetId();
 			Source.OperateOnLoadedBulkData([&NewRef](const FSharedBuffer& BulkDataBuffer)
 			{
-				NewRef.Data = FCompressedBuffer::Compress(NAME_None, BulkDataBuffer);
+				NewRef.Data = FCompressedBuffer::Compress(NAME_Default, BulkDataBuffer);
 			});
 		}
 	}
