Fixed a crash when adding a point light in light editing mode in the level editor viewport.

# LevelViewportClickHandlers updated to use the shared AssetSelection::AddActor code instead of rolling its own version of adding an actor and finding the appropriate factory.

diff --git a/Engine/Source/Editor/UnrealEd/Private/LevelViewportClickHandlers.cpp b/Engine/Source/Editor/UnrealEd/Private/LevelViewportClickHandlers.cpp
--- a/Engine/Source/Editor/UnrealEd/Private/LevelViewportClickHandlers.cpp
+++ b/Engine/Source/Editor/UnrealEd/Private/LevelViewportClickHandlers.cpp
@@ -31,6 +31,7 @@
 #include "LightMap.h"
 #include "Subsystems/BrushEditingSubsystem.h"
 #include "Elements/Framework/TypedElementSelectionSet.h"
+#include "AssetSelection.h"
 
 #define LOCTEXT_NAMESPACE "ClickHandlers"
 
@@ -77,22 +78,7 @@ namespace LevelViewportClickHandlers
 	 */
 	static AActor* PrivateAddActor(UClass* ActorClass)
 	{
-		if ( ActorClass )
-		{
-			// use an actor factory if possible
-			UActorFactory* ActorFactory = GEditor->FindActorFactoryForActorClass( ActorClass );
-			if( ActorFactory )
-			{
-				return GEditor->UseActorFactory( ActorFactory, FAssetData(), nullptr );
-			}
-			// otherwise use AddActor so that we can return the newly created actor
-			else
-			{
-				const FTransform ActorTransform = FActorPositioning::GetCurrentViewportPlacementTransform(*ActorClass->GetDefaultObject<AActor>());
-				return GEditor->AddActor( GCurrentLevelEditingViewportClient->GetWorld()->GetCurrentLevel(), ActorClass, ActorTransform );
-			}
-		}
-		return NULL;
+		return FActorFactoryAssetProxy::AddActorForAsset(ActorClass);
 	}
 
 	/**
