Fix UEditorLevelUtils::CopyOrMoveActorsToLevel to remap pasted actors using an internal id value instead of by names.

diff --git a/Engine/Source/Editor/UnrealEd/Private/EditorExporters.cpp b/Engine/Source/Editor/UnrealEd/Private/EditorExporters.cpp
--- a/Engine/Source/Editor/UnrealEd/Private/EditorExporters.cpp
+++ b/Engine/Source/Editor/UnrealEd/Private/EditorExporters.cpp
@@ -514,11 +514,12 @@ bool ULevelExporterT3D::ExportText( const FExportObjectInnerContext* Context, UO
 
 				FString ParentActorString = ( ParentActor ? FString::Printf(TEXT(" ParentActor=%s"), *ParentActor->GetName() ) : TEXT(""));
 				FString SocketNameString = ( (ParentActor && SocketName != NAME_None) ? FString::Printf(TEXT(" SocketName=%s"), *SocketName.ToString() ) : TEXT(""));
-				FString GroupActor = (Actor->GroupActor? FString::Printf(TEXT(" GroupActor=%s"), *Actor->GroupActor->GetName() ) : TEXT(""));
-				FString GroupFolder = (Actor->GroupActor? FString::Printf(TEXT(" GroupFolder=%s"), *Actor->GroupActor->GetFolderPath().ToString() ) : TEXT(""));
-				Ar.Logf( TEXT("%sBegin Actor Class=%s Name=%s Archetype=%s'%s'%s%s%s%s") LINE_TERMINATOR, 
+				FString GroupActor = (Actor->GroupActor ? FString::Printf(TEXT(" GroupActor=%s"), *Actor->GroupActor->GetName() ) : TEXT(""));
+				FString GroupFolder = (Actor->GroupActor ? FString::Printf(TEXT(" GroupFolder=%s"), *Actor->GroupActor->GetFolderPath().ToString() ) : TEXT(""));
+				FString CopyPasteId = (Actor->CopyPasteId != INDEX_NONE) ? FString::Printf(TEXT(" CopyPasteId=%d"), Actor->CopyPasteId ) : TEXT("");
+				Ar.Logf( TEXT("%sBegin Actor Class=%s Name=%s Archetype=%s'%s'%s%s%s%s%s") LINE_TERMINATOR, 
 					FCString::Spc(TextIndent), *Actor->GetClass()->GetPathName(), *Actor->GetName(),
-					*Actor->GetArchetype()->GetClass()->GetPathName(), *Actor->GetArchetype()->GetPathName(), *ParentActorString, *SocketNameString, *GroupActor, *GroupFolder );
+					*Actor->GetArchetype()->GetClass()->GetPathName(), *Actor->GetArchetype()->GetPathName(), *ParentActorString, *SocketNameString, *GroupActor, *GroupFolder, *CopyPasteId );
 
 				ExportRootScope = Actor;
 				ExportObjectInner( Context, Actor, Ar, PortFlags | PPF_ExportsNotFullyQualified );
diff --git a/Engine/Source/Editor/UnrealEd/Private/EditorLevelUtils.cpp b/Engine/Source/Editor/UnrealEd/Private/EditorLevelUtils.cpp
--- a/Engine/Source/Editor/UnrealEd/Private/EditorLevelUtils.cpp
+++ b/Engine/Source/Editor/UnrealEd/Private/EditorLevelUtils.cpp
@@ -146,14 +146,15 @@ int32 UEditorLevelUtils::CopyOrMoveActorsToLevel(const TArray<AActor*>& ActorsTo
 
 		if (FinalMoveList.Num() > 0 && (!bMoveAllOrFail || FinalMoveList.Num() == ActorCount))
 		{
-			TMap<FSoftObjectPath, FSoftObjectPath> ActorPathMapping;
+			TArray<TTuple<FSoftObjectPath, FSoftObjectPath>> ActorPathMapping;
 			GEditor->SelectNone(false, true, false);
 
 			USelection* ActorSelection = GEditor->GetSelectedActors();
 			ActorSelection->BeginBatchSelectOperation();
 			for (AActor* Actor : FinalMoveList)
 			{
-				ActorPathMapping.Add(FSoftObjectPath(Actor), FSoftObjectPath());
+				check(Actor->CopyPasteId == INDEX_NONE);
+				Actor->CopyPasteId = ActorPathMapping.Add(TTuple<FSoftObjectPath, FSoftObjectPath>(FSoftObjectPath(Actor), FSoftObjectPath()));
 				GEditor->SelectActor(Actor, true, false);
 			}
 			ActorSelection->EndBatchSelectOperation(false);
@@ -194,45 +195,15 @@ int32 UEditorLevelUtils::CopyOrMoveActorsToLevel(const TArray<AActor*>& ActorsTo
 				for (FSelectionIterator It(GEditor->GetSelectedActorIterator()); It; ++It)
 				{
 					AActor* Actor = static_cast<AActor*>(*It);
-					FSoftObjectPath NewPath = FSoftObjectPath(Actor);
 
-					bool bFoundMatch = false;
-
-					// First try exact match
-					for (TPair<FSoftObjectPath, FSoftObjectPath>& Pair : ActorPathMapping)
+					if (ActorPathMapping.IsValidIndex(Actor->CopyPasteId))
 					{
-						if (Pair.Value.IsNull() && NewPath.GetSubPathString() == Pair.Key.GetSubPathString())
-						{
-							bFoundMatch = true;
-							Pair.Value = NewPath;
-							break;
-						}
-					}
+						TTuple<FSoftObjectPath, FSoftObjectPath>& Tuple = ActorPathMapping[Actor->CopyPasteId];
+						check(Tuple.Value.IsNull());
 
-					if (!bFoundMatch)
-					{
-						// Remove numbers from end as it may have had to add some to disambiguate
-						FString PartialPath = NewPath.GetSubPathString();
-						int32 IgnoreNumber;
-						FActorLabelUtilities::SplitActorLabel(PartialPath, IgnoreNumber);
-
-						for (TPair<FSoftObjectPath, FSoftObjectPath>& Pair : ActorPathMapping)
-						{
-							if (Pair.Value.IsNull())
-							{
-								FString KeyPartialPath = Pair.Key.GetSubPathString();
-								FActorLabelUtilities::SplitActorLabel(KeyPartialPath, IgnoreNumber);
-								if (PartialPath == KeyPartialPath)
-								{
-									bFoundMatch = true;
-									Pair.Value = NewPath;
-									break;
-								}
-							}
-						}
+						Tuple.Value = FSoftObjectPath(Actor);
 					}
-
-					if (!bFoundMatch)
+					else
 					{
 						UE_LOG(LogLevelTools, Error, TEXT("Cannot find remapping for moved actor ID %s, any soft references pointing to it will be broken!"), *Actor->GetPathName());
 					}
@@ -241,7 +212,7 @@ int32 UEditorLevelUtils::CopyOrMoveActorsToLevel(const TArray<AActor*>& ActorsTo
 				FAssetToolsModule& AssetToolsModule = FModuleManager::GetModuleChecked<FAssetToolsModule>(TEXT("AssetTools"));
 				TArray<FAssetRenameData> RenameData;
 
-				for (TPair<FSoftObjectPath, FSoftObjectPath>& Pair : ActorPathMapping)
+				for (TTuple<FSoftObjectPath, FSoftObjectPath>& Pair : ActorPathMapping)
 				{
 					if (Pair.Value.IsValid())
 					{
@@ -270,6 +241,12 @@ int32 UEditorLevelUtils::CopyOrMoveActorsToLevel(const TArray<AActor*>& ActorsTo
 
 			// The moved (pasted) actors will now be selected
 			NumMovedActors += FinalMoveList.Num();
+
+			for (AActor* Actor : FinalMoveList)
+			{
+				check(Actor->CopyPasteId != INDEX_NONE);
+				Actor->CopyPasteId = INDEX_NONE;
+			}
 		}
 
 		// Restore the original clipboard contents
diff --git a/Engine/Source/Editor/UnrealEd/Private/Factories/EditorFactories.cpp b/Engine/Source/Editor/UnrealEd/Private/Factories/EditorFactories.cpp
--- a/Engine/Source/Editor/UnrealEd/Private/Factories/EditorFactories.cpp
+++ b/Engine/Source/Editor/UnrealEd/Private/Factories/EditorFactories.cpp
@@ -984,6 +984,13 @@ UObject* ULevelFactory::FactoryCreateText
 									FParse::Value(Str, TEXT("GroupFolder="), GroupFolder);
 								}
 
+								uint32 CopyPasteId;
+								if (FParse::Value(Str, TEXT("CopyPasteId="), CopyPasteId))
+								{
+									check(CopyPasteId != INDEX_NONE);
+									NewActor->CopyPasteId = CopyPasteId;
+								}
+
 								// Does the group exist?
 								if(tmpNewGroup)
 								{
diff --git a/Engine/Source/Runtime/Engine/Classes/GameFramework/Actor.h b/Engine/Source/Runtime/Engine/Classes/GameFramework/Actor.h
--- a/Engine/Source/Runtime/Engine/Classes/GameFramework/Actor.h
+++ b/Engine/Source/Runtime/Engine/Classes/GameFramework/Actor.h
@@ -803,6 +803,9 @@ protected:
 	TArray<FActorDataLayer> DataLayers;
 
 public:
+	/** The copy/paste id used to remap actors during copy operations */
+	uint32 CopyPasteId;
+
 	/** The editor-only group this actor is a part of. */
 	UPROPERTY(Transient)
 	TObjectPtr<AActor> GroupActor;
diff --git a/Engine/Source/Runtime/Engine/Private/Actor.cpp b/Engine/Source/Runtime/Engine/Private/Actor.cpp
--- a/Engine/Source/Runtime/Engine/Private/Actor.cpp
+++ b/Engine/Source/Runtime/Engine/Private/Actor.cpp
@@ -176,6 +176,7 @@ void AActor::InitializeDefaults()
 
 #if WITH_EDITORONLY_DATA
 	GridPlacement = EActorGridPlacement::Location;
+	CopyPasteId = INDEX_NONE;
 #endif
 }
 
