Fix bad merge that remove the async validation before rendering a thumbnail.

diff --git a/Engine/Source/Editor/UnrealEd/Private/AssetThumbnail.cpp b/Engine/Source/Editor/UnrealEd/Private/AssetThumbnail.cpp
--- a/Engine/Source/Editor/UnrealEd/Private/AssetThumbnail.cpp
+++ b/Engine/Source/Editor/UnrealEd/Private/AssetThumbnail.cpp
@@ -1078,7 +1078,7 @@ void FAssetThumbnailPool::Tick( float DeltaTime )
 						}
 					}
 
-					bool bLoadedThumbnail = LoadThumbnail(InfoRef, CustomThumbnailAsset);
+					bool bLoadedThumbnail = LoadThumbnail(InfoRef, bIsAssetStillCompiling, CustomThumbnailAsset);
 
 					// If we failed to load a custom thumbnail, then load the custom thumbnail's asset and try again
 					if (!bLoadedThumbnail && CustomThumbnailAsset.IsValid())
@@ -1087,7 +1087,7 @@ void FAssetThumbnailPool::Tick( float DeltaTime )
 						//if (InfoRef->AssetData.IsAssetLoaded())
 						{
 							UObject* CustomThumbnail = CustomThumbnailAsset.GetAsset();
-							bLoadedThumbnail = LoadThumbnail(InfoRef, CustomThumbnailAsset);
+							bLoadedThumbnail = LoadThumbnail(InfoRef, bIsAssetStillCompiling, CustomThumbnailAsset);
 							CustomThumbnail->ClearFlags(RF_Standalone);
 						}
 					}
@@ -1114,43 +1114,51 @@ void FAssetThumbnailPool::Tick( float DeltaTime )
 	}
 }
 
-bool FAssetThumbnailPool::LoadThumbnail(TSharedRef<FThumbnailInfo> ThumbnailInfo, const FAssetData& CustomAssetToRender)
+bool FAssetThumbnailPool::LoadThumbnail(TSharedRef<FThumbnailInfo> ThumbnailInfo, bool &bIsAssetStillCompiling, const FAssetData& CustomAssetToRender)
 {
 	const FAssetData& AssetData = CustomAssetToRender.IsValid() ? CustomAssetToRender : ThumbnailInfo->AssetData;
 	UObject* Asset = AssetData.IsAssetLoaded() ? AssetData.GetAsset() : nullptr;
 
-	if (Asset)
+	const bool bAreShadersCompiling = (GShaderCompilingManager && GShaderCompilingManager->IsCompiling());
+	
+	if (Asset && !bAreShadersCompiling)
 	{
-		FThumbnailRenderingInfo* RenderInfo = GUnrealEd->GetThumbnailManager()->GetRenderingInfo(Asset);
-		if (RenderInfo != NULL && RenderInfo->Renderer != NULL)
+		//Avoid rendering the thumbnail of an asset that is currently edited asynchronously
+		const IInterface_AsyncCompilation* Interface_AsyncCompilation = Cast<IInterface_AsyncCompilation>(Asset);
+		bIsAssetStillCompiling = Interface_AsyncCompilation && Interface_AsyncCompilation->IsCompiling();
+		if (!bIsAssetStillCompiling)
 		{
-			FThumbnailInfo_RenderThread ThumbInfo = ThumbnailInfo.Get();
-			ENQUEUE_RENDER_COMMAND(SyncSlateTextureCommand)(
-				[ThumbInfo](FRHICommandListImmediate& RHICmdList)
+			FThumbnailRenderingInfo* RenderInfo = GUnrealEd->GetThumbnailManager()->GetRenderingInfo(Asset);
+			if (RenderInfo != nullptr && RenderInfo->Renderer != nullptr)
 			{
-				if (ThumbInfo.ThumbnailTexture->GetTypedResource() != ThumbInfo.ThumbnailRenderTarget->GetTextureRHI())
+				FThumbnailInfo_RenderThread ThumbInfo = ThumbnailInfo.Get();
+				ENQUEUE_RENDER_COMMAND(SyncSlateTextureCommand)(
+					[ThumbInfo](FRHICommandListImmediate& RHICmdList)
 				{
-					ThumbInfo.ThumbnailTexture->ClearTextureData();
-					ThumbInfo.ThumbnailTexture->ReleaseDynamicRHI();
-					ThumbInfo.ThumbnailTexture->SetRHIRef(ThumbInfo.ThumbnailRenderTarget->GetTextureRHI(), ThumbInfo.Width, ThumbInfo.Height);
+					if (ThumbInfo.ThumbnailTexture->GetTypedResource() != ThumbInfo.ThumbnailRenderTarget->GetTextureRHI())
+					{
+						ThumbInfo.ThumbnailTexture->ClearTextureData();
+						ThumbInfo.ThumbnailTexture->ReleaseDynamicRHI();
+						ThumbInfo.ThumbnailTexture->SetRHIRef(ThumbInfo.ThumbnailRenderTarget->GetTextureRHI(), ThumbInfo.Width, ThumbInfo.Height);
+					}
+				});
+
+				if (ThumbnailInfo->LastUpdateTime <= 0.0f || RenderInfo->Renderer->AllowsRealtimeThumbnails(Asset))
+				{
+					//@todo: this should be done on the GPU only but it is not supported by thumbnail tools yet
+					ThumbnailTools::RenderThumbnail(
+						Asset,
+						ThumbnailInfo->Width,
+						ThumbnailInfo->Height,
+						ThumbnailTools::EThumbnailTextureFlushMode::NeverFlush,
+						ThumbnailInfo->ThumbnailRenderTarget
+					);
 				}
-			});
 
-			if (ThumbnailInfo->LastUpdateTime <= 0.0f || RenderInfo->Renderer->AllowsRealtimeThumbnails(Asset))
-			{
-				//@todo: this should be done on the GPU only but it is not supported by thumbnail tools yet
-				ThumbnailTools::RenderThumbnail(
-					Asset,
-					ThumbnailInfo->Width,
-					ThumbnailInfo->Height,
-					ThumbnailTools::EThumbnailTextureFlushMode::NeverFlush,
-					ThumbnailInfo->ThumbnailRenderTarget
-				);
+				// Since this was rendered, add it to the list of thumbnails that can be rendered in real-time
+				RealTimeThumbnails.AddUnique(ThumbnailInfo);
+				return true;
 			}
-
-			// Since this was rendered, add it to the list of thumbnails that can be rendered in real-time
-			RealTimeThumbnails.AddUnique(ThumbnailInfo);
-			return true;
 		}
 	}
 
diff --git a/Engine/Source/Editor/UnrealEd/Public/AssetThumbnail.h b/Engine/Source/Editor/UnrealEd/Public/AssetThumbnail.h
--- a/Engine/Source/Editor/UnrealEd/Public/AssetThumbnail.h
+++ b/Engine/Source/Editor/UnrealEd/Public/AssetThumbnail.h
@@ -276,11 +276,12 @@ private:
 	 * Assign a thumbnail from its render target and re-render it if necessary.
 	 *
 	 * @param ThumbnailInfo The thumbnail info to assign a texture to
+	 * @param bIsAssetStillCompiling If the asset we want to load the thumbnail is compiling, this flag will be set to true, it wont be touch in other cases.
 	 * @param CustomAssetToRender The asset to render when generating the texture
 	 *
 	 * @return true if the thumbnail was assigned to a valid texture
 	 */
-	bool LoadThumbnail(TSharedRef<FThumbnailInfo> ThumbnailInfo, const FAssetData& CustomAssetToRender = FAssetData());
+	bool LoadThumbnail(TSharedRef<FThumbnailInfo> ThumbnailInfo, bool& bIsAssetStillCompiling, const FAssetData& CustomAssetToRender = FAssetData());
 
 	struct FThumbnailInfo_RenderThread
 	{
