Alllow tools with no toolkits to register. This is needed for our motion trail edit mode, which works with conjuction with the control rig edit mode, so that it can still recieve input and handle tracking.

diff --git a/Engine/Source/Editor/UnrealEd/Private/Tools/UEdMode.cpp b/Engine/Source/Editor/UnrealEd/Private/Tools/UEdMode.cpp
--- a/Engine/Source/Editor/UnrealEd/Private/Tools/UEdMode.cpp
+++ b/Engine/Source/Editor/UnrealEd/Private/Tools/UEdMode.cpp
@@ -86,23 +86,21 @@ void UEdMode::Enter()
 
 void UEdMode::RegisterTool(TSharedPtr<FUICommandInfo> UICommand, FString ToolIdentifier, UInteractiveToolBuilder* Builder)
 {
-	if (!Toolkit.IsValid())
-	{
-		return;
-	}
-
-	const TSharedRef<FUICommandList>& CommandList = Toolkit->GetToolkitCommands();
 	ToolsContext->ToolManager->RegisterToolType(ToolIdentifier, Builder);
-	CommandList->MapAction(UICommand,
-		FExecuteAction::CreateUObject(ToolsContext.Get(), &UEdModeInteractiveToolsContext::StartTool, ToolIdentifier),
-		FCanExecuteAction::CreateWeakLambda(ToolsContext.Get(), [this, ToolIdentifier]() {
-			return ShouldToolStartBeAllowed(ToolIdentifier) && 
-				 ToolsContext->ToolManager->CanActivateTool(EToolSide::Mouse, ToolIdentifier);
-			}),
-		FIsActionChecked::CreateUObject(ToolsContext.Get(), &UEdModeInteractiveToolsContext::IsToolActive, EToolSide::Mouse, ToolIdentifier),
-		EUIActionRepeatMode::RepeatDisabled);
 
-	RegisteredTools.Emplace(UICommand, ToolIdentifier);
+	if (Toolkit.IsValid())
+	{
+		const TSharedRef<FUICommandList>& CommandList = Toolkit->GetToolkitCommands();
+		CommandList->MapAction(UICommand,
+			FExecuteAction::CreateUObject(ToolsContext.Get(), &UEdModeInteractiveToolsContext::StartTool, ToolIdentifier),
+			FCanExecuteAction::CreateWeakLambda(ToolsContext.Get(), [this, ToolIdentifier]() {
+				return ShouldToolStartBeAllowed(ToolIdentifier) &&
+					ToolsContext->ToolManager->CanActivateTool(EToolSide::Mouse, ToolIdentifier);
+				}),
+			FIsActionChecked::CreateUObject(ToolsContext.Get(), &UEdModeInteractiveToolsContext::IsToolActive, EToolSide::Mouse, ToolIdentifier),
+					EUIActionRepeatMode::RepeatDisabled);
+		RegisteredTools.Emplace(UICommand, ToolIdentifier);
+	}
 }
 
 bool UEdMode::ShouldToolStartBeAllowed(const FString& ToolIdentifier) const
@@ -118,15 +116,19 @@ void UEdMode::Exit()
 		SettingsObject->SaveConfig();
 	}
 
+	//Tools can live without toolkit
+	for (auto& RegisteredTool : RegisteredTools)
+	{
+		ToolsContext->ToolManager->UnregisterToolType(RegisteredTool.Value);
+	}
+
 	if (Toolkit.IsValid())
 	{
 		const TSharedRef<FUICommandList>& CommandList = Toolkit->GetToolkitCommands();
 		for (auto& RegisteredTool : RegisteredTools)
 		{
 			CommandList->UnmapAction(RegisteredTool.Key);
-			ToolsContext->ToolManager->UnregisterToolType(RegisteredTool.Value);
 		}
-
 		Toolkit.Reset();
 	}
 	RegisteredTools.SetNum(0);
