Revert UEDMode change, require valid toolkit to register tool.

diff --git a/Engine/Source/Editor/UnrealEd/Private/Tools/UEdMode.cpp b/Engine/Source/Editor/UnrealEd/Private/Tools/UEdMode.cpp
--- a/Engine/Source/Editor/UnrealEd/Private/Tools/UEdMode.cpp
+++ b/Engine/Source/Editor/UnrealEd/Private/Tools/UEdMode.cpp
@@ -86,21 +86,23 @@ void UEdMode::Enter()
 
 void UEdMode::RegisterTool(TSharedPtr<FUICommandInfo> UICommand, FString ToolIdentifier, UInteractiveToolBuilder* Builder)
 {
-	ToolsContext->ToolManager->RegisterToolType(ToolIdentifier, Builder);
-
-	if (Toolkit.IsValid())
+	if (!Toolkit.IsValid())
 	{
-		const TSharedRef<FUICommandList>& CommandList = Toolkit->GetToolkitCommands();
-		CommandList->MapAction(UICommand,
-			FExecuteAction::CreateUObject(ToolsContext.Get(), &UEdModeInteractiveToolsContext::StartTool, ToolIdentifier),
-			FCanExecuteAction::CreateWeakLambda(ToolsContext.Get(), [this, ToolIdentifier]() {
-				return ShouldToolStartBeAllowed(ToolIdentifier) &&
-					ToolsContext->ToolManager->CanActivateTool(EToolSide::Mouse, ToolIdentifier);
-				}),
-			FIsActionChecked::CreateUObject(ToolsContext.Get(), &UEdModeInteractiveToolsContext::IsToolActive, EToolSide::Mouse, ToolIdentifier),
-					EUIActionRepeatMode::RepeatDisabled);
-		RegisteredTools.Emplace(UICommand, ToolIdentifier);
+		return;
 	}
+
+	const TSharedRef<FUICommandList>& CommandList = Toolkit->GetToolkitCommands();
+	ToolsContext->ToolManager->RegisterToolType(ToolIdentifier, Builder);
+	CommandList->MapAction(UICommand,
+		FExecuteAction::CreateUObject(ToolsContext.Get(), &UEdModeInteractiveToolsContext::StartTool, ToolIdentifier),
+		FCanExecuteAction::CreateWeakLambda(ToolsContext.Get(), [this, ToolIdentifier]() {
+		return ShouldToolStartBeAllowed(ToolIdentifier) &&
+			ToolsContext->ToolManager->CanActivateTool(EToolSide::Mouse, ToolIdentifier);
+	}),
+		FIsActionChecked::CreateUObject(ToolsContext.Get(), &UEdModeInteractiveToolsContext::IsToolActive, EToolSide::Mouse, ToolIdentifier),
+		EUIActionRepeatMode::RepeatDisabled);
+
+	RegisteredTools.Emplace(UICommand, ToolIdentifier);
 }
 
 bool UEdMode::ShouldToolStartBeAllowed(const FString& ToolIdentifier) const
