Cooker leak diagnostics: FReferenceChainSearch calls made from the cooker should use the new EReferenceChainSearchMode::FullChain flag, because RF_Standalone objects are GC'd during cooking invocations of Garbage Collection, so stopping at RF_Standalone does not explain why the RF_Standalone object is still referenced.

diff --git a/Engine/Source/Editor/UnrealEd/Private/EditorServer.cpp b/Engine/Source/Editor/UnrealEd/Private/EditorServer.cpp
--- a/Engine/Source/Editor/UnrealEd/Private/EditorServer.cpp
+++ b/Engine/Source/Editor/UnrealEd/Private/EditorServer.cpp
@@ -2011,7 +2011,8 @@ void UEditorEngine::CheckForWorldGCLeaks( UWorld* NewWorld, UPackage* WorldPacka
 		const bool bIsPersistantWorldType = (RemainingWorld->WorldType == EWorldType::Inactive) || (RemainingWorld->WorldType == EWorldType::EditorPreview) || (RemainingWorld->WorldType == EWorldType::GamePreview);
 		if(!bIsNewWorld && !bIsPersistantWorldType && !WorldHasValidContext(RemainingWorld))
 		{
-			FReferenceChainSearch RefChainSearch(RemainingWorld, EReferenceChainSearchMode::Shortest | EReferenceChainSearchMode::PrintResults);
+			EReferenceChainSearchMode SearchMode = EReferenceChainSearchMode::Shortest | EReferenceChainSearchMode::PrintResults | EReferenceChainSearchMode::FullChain;
+			FReferenceChainSearch RefChainSearch(RemainingWorld, SearchMode);
 			UE_LOG(LogEditorServer, Error, TEXT("Old world %s not cleaned up by garbage collection while loading new map! Referenced by:") LINE_TERMINATOR TEXT("%s"), *RemainingWorld->GetPathName(), *RefChainSearch.GetRootPath());
 			NumFailedToCleanup++;
 		}
@@ -2026,7 +2027,8 @@ void UEditorEngine::CheckForWorldGCLeaks( UWorld* NewWorld, UPackage* WorldPacka
 			const bool bIsNewWorldPackage = (NewWorldPackage && RemainingPackage == NewWorldPackage);
 			if(!bIsNewWorldPackage && RemainingPackage == WorldPackage)
 			{
-				FReferenceChainSearch RefChainSearch(RemainingPackage, EReferenceChainSearchMode::Shortest | EReferenceChainSearchMode::PrintResults);
+				EReferenceChainSearchMode SearchMode = EReferenceChainSearchMode::Shortest | EReferenceChainSearchMode::PrintResults | EReferenceChainSearchMode::FullChain;
+				FReferenceChainSearch RefChainSearch(RemainingPackage, SearchMode);
 				UE_LOG(LogEditorServer, Error, TEXT("Old level package %s not cleaned up by garbage collection while loading new map! Referenced by:") LINE_TERMINATOR TEXT("%s"), *RemainingPackage->GetPathName(), *RefChainSearch.GetRootPath());
 				NumFailedToCleanup++;
 			}
