Avoid calling UPackageTools::UnloadPackages in FPackageSourceControlHelper::Delete because it triggers garbage collection, which is illegal during some commandlets.
Instead, call ResetLoaders on the packages that needs to be deleted to make sure the process doesn't have an open handle.

diff --git a/Engine/Source/Editor/UnrealEd/Private/PackageSourceControlHelper.cpp b/Engine/Source/Editor/UnrealEd/Private/PackageSourceControlHelper.cpp
--- a/Engine/Source/Editor/UnrealEd/Private/PackageSourceControlHelper.cpp
+++ b/Engine/Source/Editor/UnrealEd/Private/PackageSourceControlHelper.cpp
@@ -183,30 +183,11 @@ bool FPackageSourceControlHelper::Delete(const TArray<UPackage*>& Packages) cons
 		return true;
 	}
 
-	// Store all packages names as we won't be able to retrieve them from the UPackages once they are unloaded
-	TArray<FString> PackagesNames;
-	PackagesNames.Reserve(Packages.Num());
-
-	for (UPackage* Package : Packages)
-	{
-		PackagesNames.Add(Package->GetName());
-
-		// Must clear dirty flag before unloading
-		Package->SetDirtyFlag(false);
-	}
-
-	// Unload packages so we can delete them
-	FText ErrorMessage;
-	if (!UPackageTools::UnloadPackages(Packages, ErrorMessage))
-	{
-		UE_LOG(LogCommandletPackageHelper, Error, TEXT("Error unloading package: %s"), *ErrorMessage.ToString());
-		return false;
-	}
-
 	bool bAllPackagesDeleted = true;
-	for (const FString& PackageName : PackagesNames)
+	for (UPackage* Package : Packages)
 	{
-		bAllPackagesDeleted &= Delete(PackageName);
+		ResetLoaders(Package);
+		bAllPackagesDeleted &= Delete(Package->GetName());
 	}
 
 	return bAllPackagesDeleted;
