AsyncIODelete: Silently fallback to _1, _2, ... suffixes if the requested TempRoot is undeletable.
#rnx

diff --git a/Engine/Source/Editor/UnrealEd/Private/CookOnTheFlyServer.cpp b/Engine/Source/Editor/UnrealEd/Private/CookOnTheFlyServer.cpp
--- a/Engine/Source/Editor/UnrealEd/Private/CookOnTheFlyServer.cpp
+++ b/Engine/Source/Editor/UnrealEd/Private/CookOnTheFlyServer.cpp
@@ -5587,11 +5587,6 @@ void UCookOnTheFlyServer::DeleteSandboxDirectory(const FString& PlatformName)
 	FAsyncIODelete& LocalAsyncIODelete = GetAsyncIODelete(PlatformName, &AsyncDeleteDirectory);
 	LocalAsyncIODelete.DeleteDirectory(SandboxDirectory);
 
-	// Part of Deleting the sandbox includes deleting the old AsyncDelete directory for the sandbox, in case a previous cooker crashed before cleaning it up.
-	// The AsyncDelete directory is associated with a sandbox but is necessarily outside of it since it is used to delete the sandbox.
-	// Note that for the Platform we used to create the AsyncIODelete, this Delete will fail because AsyncIODelete refuses to delete its own temproot; this is okay because it will delete the temproot on exit.
-	LocalAsyncIODelete.DeleteDirectory(AsyncDeleteDirectory);
-
 	// UE_DEPRECATED(4.25, "Delete the old location for AsyncDeleteDirectory until all users have cooked at least once")
 	LocalAsyncIODelete.DeleteDirectory(SandboxDirectory + TEXT("AsyncDelete"));
 }
diff --git a/Engine/Source/Editor/UnrealEd/Private/Cooker/AsyncIODelete.cpp b/Engine/Source/Editor/UnrealEd/Private/Cooker/AsyncIODelete.cpp
--- a/Engine/Source/Editor/UnrealEd/Private/Cooker/AsyncIODelete.cpp
+++ b/Engine/Source/Editor/UnrealEd/Private/Cooker/AsyncIODelete.cpp
@@ -12,6 +12,7 @@
 #include "Math/NumericLimits.h"
 #include "Misc/StringBuilder.h"
 #include "Misc/Paths.h"
+#include "Misc/PathViews.h"
 #include "Templates/UnrealTemplate.h"
 
 #if WITH_ASYNCIODELETE_DEBUG
@@ -33,18 +34,18 @@ void FAsyncIODelete::SetTempRoot(const FStringView& InOwnedTempRoot)
 	Teardown();
 
 #if WITH_ASYNCIODELETE_DEBUG
-	if (!TempRoot.IsEmpty())
+	if (!RequestedTempRoot.IsEmpty())
 	{
-		RemoveTempRoot(*TempRoot);
+		RemoveTempRoot(*RequestedTempRoot);
 	}
 #endif
 
-	TempRoot = InOwnedTempRoot;
+	RequestedTempRoot = InOwnedTempRoot;
 
 #if WITH_ASYNCIODELETE_DEBUG
-	if (!TempRoot.IsEmpty())
+	if (!RequestedTempRoot.IsEmpty())
 	{
-		AddTempRoot(*TempRoot);
+		AddTempRoot(*RequestedTempRoot);
 	}
 #endif
 }
@@ -79,7 +80,7 @@ bool FAsyncIODelete::Setup()
 		return true;
 	}
 
-	if (TempRoot.IsEmpty())
+	if (RequestedTempRoot.IsEmpty())
 	{
 		checkf(false, TEXT("DeleteDirectory called without having first set a TempRoot"));
 		return false;
@@ -88,17 +89,18 @@ bool FAsyncIODelete::Setup()
 	if (AsyncEnabled())
 	{
 		// Delete the TempRoot directory to clear the results from any previous process using the same TempRoot that did not shut down cleanly
-		uint32 ErrorCode;
-		if (!DeleteTempRootDirectory(ErrorCode))
+		if (!TryPurgeTempRootFamily(&TempRoot))
 		{
-			UE_LOG(LogCook, Error, TEXT("Could not clear asyncdelete root directory '%s'.  LastError: %i."), *TempRoot, ErrorCode);
+			// TryPurgeTempRootFamily logged the warning
 			return false;
 		}
 
 		// Create the empty directory to work in
 		if (!IFileManager::Get().MakeDirectory(*TempRoot, true))
 		{
-			UE_LOG(LogCook, Error, TEXT("Could not create asyncdelete root directory '%s'.  LastError: %i."), *TempRoot, FPlatformMisc::GetLastError());
+			UE_LOG(LogCook, Error, TEXT("Could not create asyncdelete root directory '%s'. LastError: %i. Falling back to synchronous delete."),
+				*TempRoot, FPlatformMisc::GetLastError());
+			TempRoot.Reset();
 			return false;
 		}
 
@@ -136,12 +138,8 @@ void FAsyncIODelete::Teardown()
 		TasksComplete = nullptr;
 
 		// Remove the temp directory from disk
-		uint32 ErrorCode;
-		if (!DeleteTempRootDirectory(ErrorCode))
-		{
-			// This will leave directories (and potentially files, if we were paused or if any of the asyncdeletes failed) on disk, so it is bad for users, but is not fatal for our operations.
-			UE_LOG(LogCook, Warning, TEXT("Could not delete asyncdelete root directory '%s'.  LastError: %i."), *TempRoot, ErrorCode);
-		}
+		TryPurgeTempRootFamily(nullptr);
+		TempRoot.Reset();
 
 		// Clear delete variables; we don't need to run the tasks for the remaining pauseddeletes because synchronously deleting the temp directory above did the work they were going to do
 		PausedDeletes.Empty();
@@ -300,43 +298,121 @@ bool FAsyncIODelete::SynchronousDelete(const TCHAR* InDeletePath, EPathType Path
 	return Result;
 }
 
-bool FAsyncIODelete::DeleteTempRootDirectory(uint32& OutErrorCode)
+bool FAsyncIODelete::TryPurgeTempRootFamily(FString* OutNewTempRoot)
 {
-	OutErrorCode = 0;
 	IFileManager& FileManager = IFileManager::Get();
-	if (!FileManager.DirectoryExists(*TempRoot))
-	{
-		return true;
-	}
+	FString RequestedLeaf = FPaths::GetPathLeaf(RequestedTempRoot);
+	FString ParentDir = FPaths::GetPath(RequestedTempRoot);
 
-	// Since we sometimes will be creating the directory again immediately, we need to take precautions against the delayed delete of directories that
-	// occurs on Windows platforms; creating a new file/directory in one that was just deleted can fail.  So we need to move-delete our TempRoot
-	// in addition to move-delete our clients' directories.  Since we don't have a TempRoot to move-delete into, we create a unique sibling directory name.
-	FString UniqueDirectory = FPaths::CreateTempFilename(*FPaths::GetPath(TempRoot), TEXT("DeleteTemp"), TEXT(""));
+	TArray<uint32, TInlineAllocator<2>> ExistingRoots;
+	FileManager.IterateDirectory(*ParentDir,
+		[&ExistingRoots, &RequestedLeaf](const TCHAR* FilenameOrDirectory, bool bIsDirectory)
+		{
+			// Compare by PathLeaf instead of full path because absolute vs relative paths and junctions
+			// may change the name of the parent directory
+			FStringView ExistingPath(FilenameOrDirectory);
+			FStringView ExistingLeaf = FPathViews::GetPathLeaf(ExistingPath);
+			if (bIsDirectory && ExistingLeaf.StartsWith(FStringView(RequestedLeaf), ESearchCase::IgnoreCase))
+			{
+				FStringView Suffix = ExistingLeaf.RightChop(RequestedLeaf.Len());
+				if (Suffix.Len() == 0)
+				{
+					ExistingRoots.Add(0);
+				}
+				else
+				{
+					// Suffix is null-terminated because it came from FilenameOrDirectory
+					const TCHAR* SuffixPtr = Suffix.GetData();
+					uint32 IntSuffix = 0;
+					LexFromString(IntSuffix, SuffixPtr + 1);
+					if (IntSuffix > 0)
+					{
+						ExistingRoots.Add(IntSuffix);
+					}
+				}
+			}
+			return true;
+		});
+
+
+	auto GetRootPath = [this](uint32 Suffix)
+	{
+		return Suffix == 0 ? RequestedTempRoot : FString::Printf(TEXT("%s_%d"), *RequestedTempRoot, Suffix);
+	};
 
-	const bool bReplace = false;
-	const bool bEvenIfReadOnly = true;
-	const TCHAR* DirectoryToDelete = *UniqueDirectory;
-	const bool bMoveSucceeded = FileManager.Move(DirectoryToDelete, *TempRoot, bReplace, bEvenIfReadOnly);
-	if (!bMoveSucceeded)
+	uint32 LastError = 0;
+	TArray<uint32> Undeletables;
+	for (uint32 Suffix : ExistingRoots)
 	{
-		// Move failed; fallback to inplace delete
-		DirectoryToDelete = *TempRoot;
+		FString ExistingRoot = GetRootPath(Suffix);
+
+		// Since we sometimes will be creating the directory again immediately, we need to take precautions against the delayed delete of directories that
+		// occurs on Windows platforms; creating a new file/directory in one that was just deleted can fail.  So we need to move-delete our TempRoot
+		// in addition to move-delete our clients' directories.  Since we don't have a TempRoot to move-delete into, we create a unique sibling directory name.
+		FString UniqueDirectory = FPaths::CreateTempFilename(*ParentDir, TEXT("DeleteTemp"), TEXT(""));
+
+		const bool bReplace = false;
+		const bool bEvenIfReadOnly = true;
+		const bool bMoveAttributes = false;
+		const bool bDoNotRetryOnError = true;
+		const TCHAR* DirectoryToDelete = *UniqueDirectory;
+		const bool bMoveSucceeded = FileManager.Move(DirectoryToDelete, *ExistingRoot, bReplace, bEvenIfReadOnly,
+			bMoveAttributes, bDoNotRetryOnError);
+		if (!bMoveSucceeded)
+		{
+			// Move failed; fallback to inplace delete
+			DirectoryToDelete = *ExistingRoot;
+		}
+
+		const bool bRequireExists = false;
+		const bool bTree = true;
+		const bool bDeleteSucceeded = FileManager.DeleteDirectory(DirectoryToDelete, bRequireExists, bTree);
+		if (!bDeleteSucceeded)
+		{
+			LastError = FPlatformMisc::GetLastError();
+			if (bMoveSucceeded && !bDeleteSucceeded)
+			{
+				// Try to move the directory back so that we can try again to delete it next time.
+				FileManager.Move(*ExistingRoot, DirectoryToDelete, bReplace, bEvenIfReadOnly,
+					bMoveAttributes, bDoNotRetryOnError);
+			}
+			Undeletables.Add(Suffix);
+		}
 	}
 
-	const bool bRequireExists = false;
-	const bool bTree = true;
-	const bool bDeleteSucceeded = FileManager.DeleteDirectory(DirectoryToDelete, bRequireExists, bTree);
-	if (!bDeleteSucceeded)
+	if (OutNewTempRoot)
 	{
-		OutErrorCode = FPlatformMisc::GetLastError();
-		if (bMoveSucceeded && !bDeleteSucceeded)
+		OutNewTempRoot->Reset();
+		uint32 NewSuffix = 0;
+		Undeletables.Sort();
+		if (Undeletables.Num() > 0 && Undeletables[0] == 0)
 		{
-			// Try to move the directory back so that we can try again to delete it next time.
-			FileManager.Move(*TempRoot, DirectoryToDelete, bReplace, bEvenIfReadOnly);
+			NewSuffix = 1;
+			for (uint32 Suffix : Undeletables)
+			{
+				if (Suffix == 0)
+				{
+					continue;
+				}
+				if (Suffix != NewSuffix)
+				{
+					break;
+				}
+				++NewSuffix;
+			}
 		}
+
+		const uint32 MaxHangingTempRoots = 20;
+		if (NewSuffix > MaxHangingTempRoots)
+		{
+			UE_LOG(LogCook, Error, TEXT("Could not clear %d old asyncdelete root directories '%s'_*.  LastError: %i.")
+				TEXT("\n\tFalling back to synchronous delete. Delete the directories manually to silence this message."),
+				Undeletables.Num(), *RequestedTempRoot, LastError);
+			return false;
+		}
+		*OutNewTempRoot = GetRootPath(NewSuffix);
 	}
-	return bDeleteSucceeded;
+	return true;
 }
 
 #if WITH_ASYNCIODELETE_DEBUG
diff --git a/Engine/Source/Editor/UnrealEd/Private/Cooker/AsyncIODelete.h b/Engine/Source/Editor/UnrealEd/Private/Cooker/AsyncIODelete.h
--- a/Engine/Source/Editor/UnrealEd/Private/Cooker/AsyncIODelete.h
+++ b/Engine/Source/Editor/UnrealEd/Private/Cooker/AsyncIODelete.h
@@ -46,9 +46,10 @@ public:
 	/**
 	 * Set the TempRoot directory used to hold the deleted directories from elsewhere.  The FAsyncIODelete takes ownership of this directory, and deletes it when destructed or when it is changed.
 	 * If the FAsyncIODelete has already been used to delete files, this Set will block until all deletes are finished.
+	 * If the given TempRoot is unavailable, falls back to _digit suffixes.
 	 */
 	void SetTempRoot(const FStringView& InOwnedTempRoot);
-	FStringView GetTempRoot() const { return TempRoot; }
+	FStringView GetTempRoot() const { return TempRoot.IsEmpty() ? RequestedTempRoot : TempRoot; }
 
 	/**
 	 * Set whether new background deletes are paused.  If paused, paths will be moved immediately but will not be deleted from the temporary location until unpaused (or at the FAsyncIODelete's destruction).
@@ -116,9 +117,10 @@ private:
 	/** Delete the given path synchronously; called from a task or in error fallback cases from the public thread */
 	bool SynchronousDelete(const TCHAR* InDeletePath, EPathType PathType);
 
-	bool DeleteTempRootDirectory(uint32& OutErrorCode);
+	bool TryPurgeTempRootFamily(FString* OutNextTempRoot);
 
-	FString	TempRoot; 
+	FString	RequestedTempRoot; 
+	FString TempRoot;
 	TArray<FString> PausedDeletes;
 	FCriticalSection CriticalSection; // We use a CriticalSection instead of a TAtomic ActiveTaskCount so that we can atomically { trigger TasksComplete if ActiveTaskCount == 0 }
 	FEvent* TasksComplete = nullptr;
