Matinee: Changes required to delete matinee editor module. Removes the Interp edit mode and UnrealEds dependency on the Matinee editor module.

diff --git a/Engine/Source/Developer/AssetTools/Private/AssetTypeActions/AssetTypeActions_CameraAnim.cpp b/Engine/Source/Developer/AssetTools/Private/AssetTypeActions/AssetTypeActions_CameraAnim.cpp
--- a/Engine/Source/Developer/AssetTools/Private/AssetTypeActions/AssetTypeActions_CameraAnim.cpp
+++ b/Engine/Source/Developer/AssetTools/Private/AssetTypeActions/AssetTypeActions_CameraAnim.cpp
@@ -180,7 +180,7 @@ void FAssetTypeActions_CameraAnim::OpenAssetEditor( const TArray<UObject*>& InOb
 		{
 			// construct a temporary matinee actor
 			CreateMatineeActorForCameraAnim(CameraAnim);
-
+PRAGMA_DISABLE_DEPRECATION_WARNINGS
 			if (GEditor->ShouldOpenMatinee(PreviewMatineeActor.Get()))
 			{
 				// changed the actor type, but don't want to lose any properties from previous
@@ -227,6 +227,7 @@ void FAssetTypeActions_CameraAnim::OpenAssetEditor( const TArray<UObject*>& InOb
 			{
 				GEditor->GetEditorWorldContext().World()->DestroyActor(PreviewMatineeActor.Get());
 			}
+PRAGMA_ENABLE_DEPRECATION_WARNINGS
 		}
 	}
 }
diff --git a/Engine/Source/Editor/LevelEditor/Private/LevelEditor.cpp b/Engine/Source/Editor/LevelEditor/Private/LevelEditor.cpp
--- a/Engine/Source/Editor/LevelEditor/Private/LevelEditor.cpp
+++ b/Engine/Source/Editor/LevelEditor/Private/LevelEditor.cpp
@@ -1132,11 +1132,6 @@ void FLevelEditorModule::BindGlobalLevelEditorCommands()
 		FExecuteAction::CreateStatic( &FLevelEditorActionCallbacks::ExecuteExecCommand, FString( TEXT("ACTOR SELECT MATCHINGMATERIAL") ) )
 		);
 
-	ActionList.MapAction(
-		Commands.SelectAllActorsControlledByMatinee,
-		FExecuteAction::CreateStatic( &FLevelEditorActionCallbacks::OnSelectAllActorsControlledByMatinee )
-		);
-
 	ActionList.MapAction(
 		Commands.SelectMatchingEmitter,
 		FExecuteAction::CreateStatic( &FLevelEditorActionCallbacks::ExecuteExecCommand, FString( TEXT("ACTOR SELECT MATCHINGEMITTER") ) )
diff --git a/Engine/Source/Editor/LevelEditor/Private/LevelEditorActions.cpp b/Engine/Source/Editor/LevelEditor/Private/LevelEditorActions.cpp
--- a/Engine/Source/Editor/LevelEditor/Private/LevelEditorActions.cpp
+++ b/Engine/Source/Editor/LevelEditor/Private/LevelEditorActions.cpp
@@ -37,7 +37,6 @@
 #include "EditorModes.h"
 #include "UnrealEdMisc.h"
 #include "FileHelpers.h"
-#include "EditorModeInterpolation.h"
 #include "UnrealEdGlobals.h"
 #include "WorldPartition/WorldPartitionSubsystem.h"
 #include "WorldPartition/WorldPartition.h"
@@ -50,7 +49,6 @@
 #include "Elements/Interfaces/TypedElementObjectInterface.h"
 
 #include "LevelEditor.h"
-#include "Matinee/MatineeActor.h"
 #include "Engine/LevelScriptBlueprint.h"
 #include "LightingBuildOptions.h"
 #include "EditorSupportDelegates.h"
@@ -1839,11 +1837,6 @@ bool FLevelEditorActionCallbacks::CanSelectComponentOwnerActor()
 	return GEditor->GetSelectedComponentCount() > 0;
 }
 
-void FLevelEditorActionCallbacks::OnSelectAllActorsControlledByMatinee()
-{
-	GEditor->SelectAllActorsControlledByMatinee();
-}
-
 void FLevelEditorActionCallbacks::OnSelectOwningHLODCluster()
 {
 	if (GEditor->GetSelectedActorCount() > 0)
@@ -1863,27 +1856,6 @@ void FLevelEditorActionCallbacks::OnSelectOwningHLODCluster()
 	}
 }
 
-void FLevelEditorActionCallbacks::OnSelectMatineeActor( AMatineeActor * ActorToSelect )
-{
-	GEditor->SelectNone( false, true );
-	GEditor->SelectActor(ActorToSelect, true, false, true);
-
-	GEditor->NoteSelectionChange();
-}
-
-void FLevelEditorActionCallbacks::OnSelectMatineeGroup( AActor* Actor )
-{
-	if( GLevelEditorModeTools().IsModeActive( FBuiltinEditorModes::EM_InterpEdit ) )
-	{
-		FEdModeInterpEdit* InterpEditMode = (FEdModeInterpEdit*)GLevelEditorModeTools().GetActiveMode( FBuiltinEditorModes::EM_InterpEdit );
-
-		if ( InterpEditMode && InterpEditMode->MatineeActor )
-		{
-			InterpEditMode->UpdateSelectedActor();
-		}
-	}
-}
-
 void FLevelEditorActionCallbacks::OnApplyMaterialToSurface()
 {
 	FEditorDelegates::LoadSelectedAssetsIfNeeded.Broadcast();
@@ -3462,7 +3434,6 @@ void FLevelEditorCommands::RegisterCommands()
 	UI_COMMAND( SelectStationaryLightsExceedingOverlap, "Select Stationary Lights exceeding overlap", "Selects all stationary lights exceeding the overlap limit", EUserInterfaceActionType::Button, FInputChord() ); 
 	UI_COMMAND( SelectAllAddditiveBrushes, "Select All Additive Brushes", "Selects all additive brushes", EUserInterfaceActionType::Button, FInputChord() ); 
 	UI_COMMAND( SelectAllSubtractiveBrushes, "Select All Subtractive Brushes", "Selects all subtractive brushes", EUserInterfaceActionType::Button, FInputChord() ); 
-	UI_COMMAND( SelectAllActorsControlledByMatinee, "Select Actors Used by This Matinee", "Selects all actors controlled by this Matinee", EUserInterfaceActionType::Button, FInputChord() ); 
 
 	UI_COMMAND( SelectAllSurfaces, "Select All Surfaces", "Selects all bsp surfaces", EUserInterfaceActionType::Button, FInputChord(EModifierKey::Shift, EKeys::S) ); 
 
diff --git a/Engine/Source/Editor/LevelEditor/Private/LevelEditorContextMenu.cpp b/Engine/Source/Editor/LevelEditor/Private/LevelEditorContextMenu.cpp
--- a/Engine/Source/Editor/LevelEditor/Private/LevelEditorContextMenu.cpp
+++ b/Engine/Source/Editor/LevelEditor/Private/LevelEditorContextMenu.cpp
@@ -30,9 +30,7 @@
 #include "Editor/GroupActor.h"
 #include "LevelEditorViewport.h"
 #include "EditorModes.h"
-#include "EditorModeInterpolation.h"
 #include "LevelEditor.h"
-#include "Matinee/MatineeActor.h"
 #include "Kismet2/KismetEditorUtilities.h"
 #include "AssetSelection.h"
 #include "LevelEditorActions.h"
@@ -150,13 +148,6 @@ public:
 	 */
 	static void FillSourceControlMenu(UToolMenu* Menu);
 
-private:
-	/**
-	 * Fills in menu options for the matinee selection menu
-	 *
-	 * @param MenuBuilder	The menu to add items to
-	 */
-	static void FillMatineeSelectActorMenu(UToolMenu* Menu);
 };
 
 FSelectedActorInfo FLevelEditorContextMenuImpl::SelectionInfo;
@@ -979,93 +970,6 @@ void FLevelEditorContextMenuImpl::FillSelectActorMenu(UToolMenu* Menu)
 			Section.AddMenuEntry(FLevelEditorCommands::Get().GeometryCollectionSelectInverseGeometry);
 		}
 	}
-
-	// build matinee related selection menu
-	FillMatineeSelectActorMenu( Menu );
-}
-
-void FLevelEditorContextMenuImpl::FillMatineeSelectActorMenu(UToolMenu* Menu)
-{
-	{
-		FToolMenuSection& Section = Menu->AddSection("SelectMatinee", LOCTEXT("SelectMatineeHeading", "Matinee"));
-		// show list of Matinee Actors that controls this actor
-		// this is ugly but we don't have good way of knowing which Matinee actor controls me
-		// in the future this can be cached to TMap somewhere and use that list
-		// for now we show only when 1 actor is selected
-		if ( SelectionInfo.SharedLevel && SelectionInfo.NumSelected == 1 )
-		{
-			TArray<AMatineeActor*> MatineeActors;	
-			// first collect all matinee actors
-			for ( AActor* Actor : SelectionInfo.SharedLevel->Actors )
-			{
-				AMatineeActor * CurActor = Cast<AMatineeActor>(Actor);
-				if ( CurActor )
-				{
-					MatineeActors.Add(CurActor);
-				}
-			}
-
-			if ( MatineeActors.Num() > 0 )
-			{
-				FSelectionIterator ActorIter( GEditor->GetSelectedActorIterator() );
-				AActor* SelectedActor = Cast<AActor>(*ActorIter);
-
-				// now delete the matinee actors that don't control currently selected actor
-				for (int32 MatineeActorIter=0; MatineeActorIter<MatineeActors.Num(); ++MatineeActorIter)
-				{
-					AMatineeActor * CurMatineeActor = MatineeActors[MatineeActorIter];
-					TArray<AActor *> CutMatineeControlledActors;
-					CurMatineeActor->GetControlledActors(CutMatineeControlledActors);
-					bool bIsMatineeControlled=false;
-					for ( AActor* ControlledActor : CutMatineeControlledActors )
-					{
-						if (ControlledActor == SelectedActor)
-						{
-							bIsMatineeControlled = true;
-						}
-					}
-
-					// if not, remove it
-					if (!bIsMatineeControlled)
-					{
-						MatineeActors.RemoveAt(MatineeActorIter);
-						--MatineeActorIter;
-					}
-				}
-
-				// if some matinee controls this, add to menu for direct selection
-				if ( MatineeActors.Num() > 0 )
-				{
-					for (int32 MatineeActorIter=0; MatineeActorIter<MatineeActors.Num(); ++MatineeActorIter)
-					{
-						AMatineeActor * CurMatineeActor = MatineeActors[MatineeActorIter];
-						const FText Text = FText::Format( LOCTEXT("SelectMatineeActor", "Select {0}"), FText::FromString( CurMatineeActor->GetName() ) );
-
-						FUIAction CurMatineeActorAction( FExecuteAction::CreateStatic( &FLevelEditorActionCallbacks::OnSelectMatineeActor, CurMatineeActor ) );
-						Section.AddMenuEntry(NAME_None, Text, Text, FSlateIcon(), CurMatineeActorAction);
-
-						// if matinee is opened, and if that is CurMatineeActor, show option to go to group
-						if( GLevelEditorModeTools().IsModeActive( FBuiltinEditorModes::EM_InterpEdit ) )
-						{
-							const FEdModeInterpEdit* InterpEditMode = (const FEdModeInterpEdit*)GLevelEditorModeTools().GetActiveMode( FBuiltinEditorModes::EM_InterpEdit );
-
-							if ( InterpEditMode && InterpEditMode->MatineeActor == CurMatineeActor )
-							{
-								FUIAction SelectedActorAction( FExecuteAction::CreateStatic( &FLevelEditorActionCallbacks::OnSelectMatineeGroup, SelectedActor ) );
-								Section.AddMenuEntry(NAME_None, LOCTEXT("SelectMatineeGroupForActorMenuTitle", "Select Matinee Group For This Actor"), LOCTEXT("SelectMatineeGroupForActorMenuTooltip", "Selects matinee group controlling this actor"), FSlateIcon(), SelectedActorAction);
-							}
-						}
-					}
-				}
-			}
-		}
-
-		// if this class is Matinee Actor, add option to allow select all controlled actors
-		if ( SelectionInfo.bHaveMatinee )
-		{
-			Section.AddMenuEntry(FLevelEditorCommands::Get().SelectAllActorsControlledByMatinee);	
-		}
-	}
 }
 
 void FLevelEditorContextMenuImpl::FillActorVisibilityMenu(UToolMenu* Menu)
diff --git a/Engine/Source/Editor/LevelEditor/Private/LevelEditorToolBar.cpp b/Engine/Source/Editor/LevelEditor/Private/LevelEditorToolBar.cpp
--- a/Engine/Source/Editor/LevelEditor/Private/LevelEditorToolBar.cpp
+++ b/Engine/Source/Editor/LevelEditor/Private/LevelEditorToolBar.cpp
@@ -33,7 +33,6 @@
 #include "SScalabilitySettings.h"
 #include "IContentBrowserSingleton.h"
 #include "ContentBrowserModule.h"
-#include "Matinee/MatineeActor.h"
 #include "LevelSequenceActor.h"
 #include "Engine/LevelScriptBlueprint.h"
 #include "ISettingsCategory.h"
@@ -63,11 +62,6 @@
 #include "Styling/ToolBarStyle.h"
 #include "PlatformInfo.h"
 
-static TAutoConsoleVariable<int32> CVarAllowMatineeActors(
-	TEXT("Matinee.AllowMatineeActors"),
-	0,
-	TEXT("Toggles whether matinee actors should appear in the cinematics menu so that they can be edited."));
-
 namespace LevelEditorActionHelpers
 {
 	/** Filters out any classes for the Class Picker when creating or selecting classes in the Blueprints dropdown */
@@ -1969,16 +1963,14 @@ void FLevelEditorToolBar::RegisterCinematicsMenu()
 			return;
 		}
 
-		const int32 bAllowMatineeActors = CVarAllowMatineeActors->GetInt();
-
 		UWorld* World = FoundContext->LevelEditor.IsValid() ? FoundContext->LevelEditor.Pin()->GetWorld() : nullptr;
-		const bool bHasAnyCinematicsActors = (bAllowMatineeActors && !!TActorIterator<AMatineeActor>(World)) || !!TActorIterator<ALevelSequenceActor>(World);
+		const bool bHasAnyCinematicsActors = !!TActorIterator<ALevelSequenceActor>(World);
 		if (!bHasAnyCinematicsActors)
 		{
 			return;
 		}
 
-		// We can't build a list of Matinees and LevelSequenceActors while the current World is a PIE world.
+		// We can't build a list of LevelSequenceActors while the current World is a PIE world.
 		FSceneOutlinerInitializationOptions InitOptions;
 		{
 			// We hide the header row to keep the UI compact.
@@ -1990,9 +1982,9 @@ void FLevelEditorToolBar::RegisterCinematicsMenu()
 			InitOptions.ColumnMap.Add(FSceneOutlinerBuiltInColumnTypes::Label(), FSceneOutlinerColumnInfo(ESceneOutlinerColumnVisibility::Visible, 0));
 			InitOptions.ColumnMap.Add(FSceneOutlinerBuiltInColumnTypes::ActorInfo(), FSceneOutlinerColumnInfo(ESceneOutlinerColumnVisibility::Visible, 10));
 
-			// Only display Matinee and MovieScene actors
+			// Only display MovieScene actors
 			auto ActorFilter = [&](const AActor* Actor) {
-				return (bAllowMatineeActors && Actor->IsA(AMatineeActor::StaticClass())) || Actor->IsA(ALevelSequenceActor::StaticClass());
+				return Actor->IsA(ALevelSequenceActor::StaticClass());
 			};
 			InitOptions.Filters->AddFilterPredicate<FActorTreeItem>(FActorTreeItem::FFilterPredicate::CreateLambda(ActorFilter));
 		}
@@ -2024,12 +2016,7 @@ void FLevelEditorToolBar::OnCinematicsActorPicked( AActor* Actor )
 	FSlateApplication::Get().DismissAllMenus();
 
 	// Make sure we dismiss the menus before we open this
-	if (AMatineeActor* MatineeActor = Cast<AMatineeActor>(Actor))
-	{
-		// Open Matinee for editing!
-		GEditor->OpenMatinee( MatineeActor );
-	}
-	else if (ALevelSequenceActor* LevelSequenceActor = Cast<ALevelSequenceActor>(Actor))
+	if (ALevelSequenceActor* LevelSequenceActor = Cast<ALevelSequenceActor>(Actor))
 	{
 		FScopedSlowTask SlowTask(1.f, NSLOCTEXT("LevelToolBarCinematicsMenu", "LoadSequenceSlowTask", "Loading Level Sequence..."));
 		SlowTask.MakeDialog();
diff --git a/Engine/Source/Editor/LevelEditor/Public/LevelEditorActions.h b/Engine/Source/Editor/LevelEditor/Public/LevelEditorActions.h
--- a/Engine/Source/Editor/LevelEditor/Public/LevelEditorActions.h
+++ b/Engine/Source/Editor/LevelEditor/Public/LevelEditorActions.h
@@ -17,7 +17,6 @@
 #include "Elements/Framework/TypedElementHandle.h"
 #include "Elements/Interfaces/TypedElementWorldInterface.h"
 
-class AMatineeActor;
 class FLightingBuildOptions;
 class SLevelEditor;
 class UActorFactory;
@@ -400,6 +399,7 @@ public:
 	TSharedPtr< FUICommandInfo > SelectAllWithSameMaterial;
 
 	/** Selects all actors used by currently selected matinee actor */
+	UE_DEPRECATED(5.0, "Matinee is no longer part of the editor.")
 	TSharedPtr< FUICommandInfo > SelectAllActorsControlledByMatinee;
 
 	/** Selects all emitters using the same particle system as the current selection */
@@ -1020,21 +1020,6 @@ public:
 	/** Selects stationary lights that are exceeding the overlap limit. */
 	static void OnSelectStationaryLightsExceedingOverlap();
 
-	/**
-	 * Selects the MatineeActor - used by Matinee Selection
-	 */
-	static void OnSelectMatineeActor( AMatineeActor * ActorToSelect );
-
-	/**
-	 * Selects the Matinee InterpGroup
-	 */
-	static void OnSelectMatineeGroup( AActor* Actor );
-
-	/**
-	 * Called when selecting all actors that's controlled by currently selected matinee actor
-	 */
-	static void OnSelectAllActorsControlledByMatinee();
-	
 	/**
 	* Called when selecting an Actor's (if available) owning HLOD cluster
 	*/
diff --git a/Engine/Source/Editor/UnrealEd/Classes/Editor/EditorEngine.h b/Engine/Source/Editor/UnrealEd/Classes/Editor/EditorEngine.h
--- a/Engine/Source/Editor/UnrealEd/Classes/Editor/EditorEngine.h
+++ b/Engine/Source/Editor/UnrealEd/Classes/Editor/EditorEngine.h
@@ -34,7 +34,6 @@
 
 #include "EditorEngine.generated.h"
 
-class AMatineeActor;
 class APlayerStart;
 class Error;
 class FEditorViewportClient;
@@ -1727,6 +1726,7 @@ public:
 
 	/**
 	 * Request to create a new PIE window and join the currently running PIE session.
+	 * Deferred until the next tick.
 	 */
 	void RequestLateJoin();
 
@@ -2149,7 +2149,8 @@ public:
 	/**
 	 * Selects all actors controlled by currently selected MatineeActor
 	 */
-	void SelectAllActorsControlledByMatinee();
+	UE_DEPRECATED(5.0, "Matinee is no longer part of the editor.")
+	void SelectAllActorsControlledByMatinee() {}
 
 	/**
 	 * Selects all actors with the same class as the current selection
@@ -2325,6 +2326,7 @@ public:
 	 *
 	 * @return The event delegate.
 	 */
+	UE_DEPRECATED(5.0, "Matinee is no longer part of the editor.")
 	DECLARE_DELEGATE_RetVal_OneParam(bool, FShouldOpenMatineeCallback, AMatineeActor*)
 	FShouldOpenMatineeCallback& OnShouldOpenMatinee() { return ShouldOpenMatineeCallback; }
 
@@ -2334,7 +2336,8 @@ public:
 	 * @param MatineeActor	The actor we wish to check (can be null)
 	 * @returns true if the user wishes to proceed
 	 */
-	bool ShouldOpenMatinee(AMatineeActor* MatineeActor) const;
+	UE_DEPRECATED(5.0, "Matinee is no longer part of the editor.")
+	bool ShouldOpenMatinee(AMatineeActor* MatineeActor) const { return false; }
 
 	/** 
 	 * Open the Matinee tool to edit the supplied MatineeActor. Will check that MatineeActor has an InterpData attached.
@@ -2342,7 +2345,8 @@ public:
 	 * @param MatineeActor	The actor we wish to edit
 	 * @param bWarnUser		If true, calls ShouldOpenMatinee as part of the open process
 	 */
-	void OpenMatinee(class AMatineeActor* MatineeActor, bool bWarnUser=true);
+	UE_DEPRECATED(5.0, "Matinee is no longer part of the editor.")
+	void OpenMatinee(class AMatineeActor* MatineeActor, bool bWarnUser = true) {}
 
 	/**
 	* Update any outstanding reflection captures
@@ -2830,7 +2834,7 @@ private:
 	FExecParticleInvoked ExecParticleInvokedEvent; 
 
 	/** Delegate to be called when a matinee is requested to be opened */
-	FShouldOpenMatineeCallback ShouldOpenMatineeCallback;
+	FShouldOpenMatineeCallback ShouldOpenMatineeCallback;	
 
 	/** Reference to owner of the current popup */
 	TWeakPtr<class SWindow> PopupWindow;
@@ -3093,6 +3097,7 @@ private:
 protected:
 
 	/** Called when Matinee is opened */
+	UE_DEPRECATED(5.0, "Matinee is no longer part of the editor.")
 	virtual void OnOpenMatinee(){};
 
 	/**
diff --git a/Engine/Source/Editor/UnrealEd/Classes/Editor/UnrealEdEngine.h b/Engine/Source/Editor/UnrealEd/Classes/Editor/UnrealEdEngine.h
--- a/Engine/Source/Editor/UnrealEd/Classes/Editor/UnrealEdEngine.h
+++ b/Engine/Source/Editor/UnrealEd/Classes/Editor/UnrealEdEngine.h
@@ -210,7 +210,8 @@ public:
 	virtual void GetPackageList( TArray<UPackage*>* InPackages, UClass* InClass ) override;
 	virtual bool ShouldAbortActorDeletion() const override final; // @note Final - Override ShouldAbortComponentDeletion or ShouldAbortActorDeletion (with parameters) instead.
 	virtual void CloseEditor() override;
-	virtual void OnOpenMatinee() override;
+
+	
 	virtual bool IsAutosaving() const override;
 	//~ End UEditorEngine Interface 
 	
diff --git a/Engine/Source/Editor/UnrealEd/Private/EditorEngine.cpp b/Engine/Source/Editor/UnrealEd/Private/EditorEngine.cpp
--- a/Engine/Source/Editor/UnrealEd/Private/EditorEngine.cpp
+++ b/Engine/Source/Editor/UnrealEd/Private/EditorEngine.cpp
@@ -76,11 +76,9 @@
 #include "UnrealEdMisc.h"
 #include "EditorDirectories.h"
 #include "FileHelpers.h"
-#include "EditorModeInterpolation.h"
 #include "Dialogs/Dialogs.h"
 #include "Dialogs/DialogsPrivate.h"
 #include "UnrealEdGlobals.h"
-#include "Matinee/MatineeActor.h"
 #include "InteractiveFoliageActor.h"
 #include "Engine/WorldComposition.h"
 #include "EditorSupportDelegates.h"
@@ -1478,7 +1476,7 @@ void UEditorEngine::Tick( float DeltaSeconds, bool bIdleMode )
 	{
 		AudioFocusViewportClient->SetAudioFocus();
 
-		// override realtime setting if viewport chooses (i.e. for matinee preview)
+		// override realtime setting if viewport chooses (i.e. for cinematic preview)
 		if (AudioFocusViewportClient->IsForcedRealtimeAudio())
 		{
 			bAudioIsRealtime = true;
@@ -3051,29 +3049,6 @@ void UEditorEngine::DeselectLevelInLevelBrowser()
 	LevelEditorModule.SummonWorldBrowserHierarchy();
 }
 
-void UEditorEngine::SelectAllActorsControlledByMatinee()
-{
-	TArray<AActor *> AllActors;
-	UWorld* IteratorWorld = GWorld;
-	for( FSelectedActorIterator Iter(IteratorWorld); Iter; ++Iter)
-	{
-		AMatineeActor * CurActor = Cast<AMatineeActor>(*Iter);
-		if ( CurActor )
-		{
-			TArray<AActor*> Actors;			
-			CurActor->GetControlledActors(Actors);
-			AllActors.Append(Actors);
-		}
-	}
-
-	SelectNone(false, true, false);
-	for(int32 i=0; i<AllActors.Num(); i++)
-	{
-		SelectActor( AllActors[i], true, false, true );
-	}
-	NoteSelectionChange();
-}
-
 void UEditorEngine::SelectAllActorsWithClass( bool bArchetype )
 {
 	if( !bArchetype )
@@ -3701,60 +3676,6 @@ void UEditorEngine::ConvertActorsFromClass( UClass* FromClass, UClass* ToClass )
 	}
 }
 
-bool UEditorEngine::ShouldOpenMatinee(AMatineeActor* MatineeActor) const
-{
-	if( PlayWorld )
-	{
-		FMessageDialog::Open( EAppMsgType::Ok, NSLOCTEXT("UnrealEd", "Error_MatineeCantOpenDuringPIE", "Matinee cannot be opened during Play in Editor.") );
-		return false;
-	}
-
-	if ( MatineeActor && !MatineeActor->MatineeData )
-	{
-		FMessageDialog::Open( EAppMsgType::Ok, NSLOCTEXT("UnrealEd", "Error_MatineeActionMustHaveData", "Matinee must have valid InterpData assigned before being edited.") );
-		return false;
-	}
-
-	// Make sure we can't open the same action twice in Matinee.
-	if( GLevelEditorModeTools().IsModeActive(FBuiltinEditorModes::EM_InterpEdit) )
-	{
-		FMessageDialog::Open( EAppMsgType::Ok, NSLOCTEXT("UnrealEd", "MatineeActionAlreadyOpen", "An Matinee sequence is currently open in an editor.  Please close it before proceeding.") );
-		return false;
-	}
-
-	// Don't let you open Matinee if a transaction is currently active.
-	if( IsTransactionActive() )
-	{
-		FMessageDialog::Open( EAppMsgType::Ok, NSLOCTEXT("UnrealEd", "TransactionIsActive", "Undo Transaction Is Active - Cannot Open Matinee.") );
-		return false;
-	}
-
-	return true;
-}
-
-void UEditorEngine::OpenMatinee(AMatineeActor* MatineeActor, bool bWarnUser)
-{
-	// Drop out if the user doesn't want to proceed to matinee atm
-	if( bWarnUser && ( (ShouldOpenMatineeCallback.IsBound() && !ShouldOpenMatineeCallback.Execute(MatineeActor)) || !ShouldOpenMatinee( MatineeActor ) ) )
-	{
-		return;
-	}
-
-	// If already in Matinee mode, exit out before going back in with new Interpolation.
-	if( GLevelEditorModeTools().IsModeActive( FBuiltinEditorModes::EM_InterpEdit ) )
-	{
-		GLevelEditorModeTools().DeactivateMode( FBuiltinEditorModes::EM_InterpEdit );
-	}
-
-	GLevelEditorModeTools().ActivateMode( FBuiltinEditorModes::EM_InterpEdit );
-
-	FEdModeInterpEdit* InterpEditMode = (FEdModeInterpEdit*)GLevelEditorModeTools().GetActiveMode( FBuiltinEditorModes::EM_InterpEdit );
-
-	InterpEditMode->InitInterpMode( MatineeActor );
-
-	OnOpenMatinee();
-}
-
 void UEditorEngine::BuildReflectionCaptures(UWorld* World)
 {
 	// Note: Lighting and reflection build operations should only dirty BuildData packages, not ULevel packages
diff --git a/Engine/Source/Editor/UnrealEd/Private/EditorModeInterpolation.cpp b/Engine/Source/Editor/UnrealEd/Private/EditorModeInterpolation.cpp
deleted file mode 100644
diff --git a/Engine/Source/Editor/UnrealEd/Private/EditorModeRegistry.cpp b/Engine/Source/Editor/UnrealEd/Private/EditorModeRegistry.cpp
--- a/Engine/Source/Editor/UnrealEd/Private/EditorModeRegistry.cpp
+++ b/Engine/Source/Editor/UnrealEd/Private/EditorModeRegistry.cpp
@@ -5,7 +5,6 @@
 #include "EditorStyleSet.h"
 #include "EdMode.h"
 #include "EditorModes.h"
-#include "EditorModeInterpolation.h"
 
 #include "Editor/PlacementMode/Public/IPlacementModeModule.h"
 #include "Editor/LandscapeEditor/Public/LandscapeEditorModule.h"
@@ -51,8 +50,6 @@ void FEditorModeRegistry::Initialize()
 	}
 
 	// Add default editor modes
-	RegisterMode<FEdModeInterpEdit>(FBuiltinEditorModes::EM_InterpEdit);
-
 	FModuleManager::LoadModuleChecked<FActorPickerModeModule>(TEXT("ActorPickerMode"));
 	FModuleManager::LoadModuleChecked<FSceneDepthPickerModeModule>(TEXT("SceneDepthPickerMode"));
 	FModuleManager::LoadModuleChecked<ILandscapeEditorModule>(TEXT("LandscapeEditor"));
diff --git a/Engine/Source/Editor/UnrealEd/Private/EditorViewportClient.cpp b/Engine/Source/Editor/UnrealEd/Private/EditorViewportClient.cpp
--- a/Engine/Source/Editor/UnrealEd/Private/EditorViewportClient.cpp
+++ b/Engine/Source/Editor/UnrealEd/Private/EditorViewportClient.cpp
@@ -25,14 +25,11 @@
 #include "EditorModes.h"
 #include "MouseDeltaTracker.h"
 #include "CameraController.h"
-#include "Editor/Matinee/Public/IMatinee.h"
-#include "Editor/Matinee/Public/MatineeConstants.h"
 #include "HighResScreenshot.h"
 #include "EditorDragTools.h"
 #include "EngineAnalytics.h"
 #include "AnalyticsEventAttribute.h"
 #include "Interfaces/IAnalyticsProvider.h"
-#include "Matinee/MatineeActor.h"
 #include "EngineModule.h"
 #include "Framework/Notifications/NotificationManager.h"
 #include "Widgets/Notifications/SNotificationList.h"
@@ -415,7 +412,6 @@ FEditorViewportClient::FEditorViewportClient(FEditorModeTools* InModeTools, FPre
 	, Widget(new FWidget)
 	, bShowWidget(true)
 	, MouseDeltaTracker(new FMouseDeltaTracker)
-	, RecordingInterpEd(NULL)
 	, bHasMouseMovedSinceClick(false)
 	, CameraController(new FEditorCameraController())
 	, CameraUserImpulseData(new FCameraControllerUserImpulseData())
@@ -1333,12 +1329,6 @@ void FEditorViewportClient::ReceivedFocus(FViewport* InViewport)
 	SetRequiredCursorOverride(false, EMouseCursor::Default);
 	FSlateApplication::Get().QueryCursor();
 
-	if( IsMatineeRecordingWindow() )
-	{
-		// Allow the joystick to be used for matinee capture
-		InViewport->SetUserFocus( true );
-	}
-
 	ModeTools->ReceivedFocus(this, Viewport);
 }
 
@@ -1537,64 +1527,7 @@ void FEditorViewportClient::UpdateCameraMovementFromJoystick(const bool bRelativ
 				{
 					CameraUserImpulseData->ZoomOutInImpulse -= InConfig.ZoomMultiplier;
 				}
-				else if (RecordingInterpEd)
-				{
-					bool bRepeatAllowed = RecordingInterpEd->IsRecordMenuChangeAllowedRepeat();
-					if ((Key == EKeys::Gamepad_DPad_Up) && bPressed)
-					{
-						const bool bNextMenuItem = false;
-						RecordingInterpEd->ChangeRecordingMenu(bNextMenuItem);
-						bRepeatAllowed = false;
-					}
-					else if ((Key == EKeys::Gamepad_DPad_Down) && bPressed)
-					{
-						const bool bNextMenuItem = true;
-						RecordingInterpEd->ChangeRecordingMenu(bNextMenuItem);
-						bRepeatAllowed = false;
-					}
-					else if ((Key == EKeys::Gamepad_DPad_Right) && (bPressed || (bRepeat && bRepeatAllowed)))
-					{
-						const bool bIncrease= true;
-						RecordingInterpEd->ChangeRecordingMenuValue(this, bIncrease);
-					}
-					else if ((Key == EKeys::Gamepad_DPad_Left) && (bPressed || (bRepeat && bRepeatAllowed)))
-					{
-						const bool bIncrease= false;
-						RecordingInterpEd->ChangeRecordingMenuValue(this, bIncrease);
-					}
-					else if ((Key == EKeys::Gamepad_RightThumbstick) && (bPressed))
-					{
-						const bool bIncrease= true;
-						RecordingInterpEd->ResetRecordingMenuValue(this);
-					}
-					else if ((Key == EKeys::Gamepad_LeftThumbstick) && (bPressed))
-					{
-						RecordingInterpEd->ToggleRecordMenuDisplay();
-					}
-					else if ((Key == EKeys::Gamepad_FaceButton_Bottom) && (bPressed))
-					{
-						RecordingInterpEd->ToggleRecordInterpValues();
-					}
-					else if ((Key == EKeys::Gamepad_FaceButton_Right) && (bPressed))
-					{
-						if (!RecordingInterpEd->GetMatineeActor()->bIsPlaying)
-						{
-							bool bLoop = true;
-							bool bForward = true;
-							RecordingInterpEd->StartPlaying(bLoop, bForward);
-						}
-						else
-						{
-							RecordingInterpEd->StopPlaying();
-						}
-					}
-
-					if (!bRepeatAllowed)
-					{
-						//only respond to this event ONCE
-						JoystickState->KeyEventValues.Remove(Key);
-					}
-				}
+				
 				if (bPressed)
 				{
 					//instantly set to repeat to stock rapid flickering until the time out
@@ -1895,29 +1828,13 @@ void FEditorViewportClient::UpdateCameraMovement( float DeltaTime )
 			CameraController->GetConfig().bForceRotationalPhysics = false;
 		}
 
-		bool bIgnoreJoystickControls = false;
-		//if we're playing back (without recording), stop input from being processed
-		if (RecordingInterpEd && RecordingInterpEd->GetMatineeActor())
-		{
-			if (RecordingInterpEd->GetMatineeActor()->bIsPlaying && !RecordingInterpEd->IsRecordingInterpValues())
-			{
-				bIgnoreJoystickControls = true;
-			}
-
-			CameraController->GetConfig().bPlanarCamera = (RecordingInterpEd->GetCameraMovementScheme() == MatineeConstants::ECameraScheme::CAMERA_SCHEME_PLANAR_CAM);
-		}
-
 		if( GetDefault<ULevelEditorViewportSettings>()->bLevelEditorJoystickControls )
 		{
 			//Now update for cached joystick info (relative movement first)
 			UpdateCameraMovementFromJoystick(true, CameraController->GetConfig());
 
-			//if we're not playing any cinematics right now
-			if (!bIgnoreJoystickControls)
-			{
-				//Now update for cached joystick info (absolute movement second)
-				UpdateCameraMovementFromJoystick(false, CameraController->GetConfig());
-			}
+			//Now update for cached joystick info (absolute movement second)
+			UpdateCameraMovementFromJoystick(false, CameraController->GetConfig());
 		}
 
 		FVector NewViewLocation = GetViewLocation();
@@ -1932,7 +1849,7 @@ void FEditorViewportClient::UpdateCameraMovement( float DeltaTime )
 		const float FinalCameraSpeedScale = FlightCameraSpeedScale * CameraSpeed * GetCameraSpeedScalar() * CameraBoost;
 
 		// Only allow FOV recoil if flight camera mode is currently inactive.
-		const bool bAllowRecoilIfNoImpulse = (!bUsingFlightInput) && (!IsMatineeRecordingWindow());
+		const bool bAllowRecoilIfNoImpulse = !bUsingFlightInput;
 
 		// Update the camera's position, rotation and FOV
 		float EditorMovementDeltaUpperBound = 1.0f;	// Never "teleport" the camera further than a reasonable amount after a large quantum
@@ -2610,17 +2527,6 @@ bool FEditorViewportClient::IsUsingAbsoluteTranslation(bool bAlsoCheckAbsoluteRo
 	return (!bCameraLockedToWidget && !bIsHotKeyAxisLocked && bAbsoluteMovementEnabled && (bCurrentWidgetSupportsAbsoluteMovement || bCurrentWidgetSupportsAbsoluteRotation) && bWidgetActivelyTrackingAbsoluteMovement && !IsOrtho() && bAnyMouseButtonsDown);
 }
 
-void FEditorViewportClient::SetMatineeRecordingWindow (IMatineeBase* InInterpEd)
-{
-	RecordingInterpEd = InInterpEd;
-	if (CameraController)
-	{
-		FCameraControllerConfig Config = CameraController->GetConfig();
-		RecordingInterpEd->LoadRecordingSettings(OUT Config);
-		CameraController->SetConfig(Config);
-	}
-}
-
 bool FEditorViewportClient::IsFlightCameraActive() const
 {
 	bool bIsFlightMovementKey = false;
@@ -3883,7 +3789,7 @@ void FEditorViewportClient::Draw(FViewport* InViewport, FCanvas* Canvas)
 		if( !UseEngineShowFlags.Game )
 		{
 			// in the editor, disable camera motion blur and other rendering features that rely on the former frame
-			// unless the view port is Matinee controlled
+			// unless the view port is cinematic controlled
 			ViewFamily.EngineShowFlags.CameraInterpolation = 0;
 		}
 
diff --git a/Engine/Source/Editor/UnrealEd/Private/LevelEditorViewport.cpp b/Engine/Source/Editor/UnrealEd/Private/LevelEditorViewport.cpp
--- a/Engine/Source/Editor/UnrealEd/Private/LevelEditorViewport.cpp
+++ b/Engine/Source/Editor/UnrealEd/Private/LevelEditorViewport.cpp
@@ -33,7 +33,6 @@
 #include "Editor.h"
 #include "EditorModeRegistry.h"
 #include "EditorModes.h"
-#include "EditorModeInterpolation.h"
 #include "PhysicsManipulationMode.h"
 #include "UnrealEdGlobals.h"
 #include "Materials/MaterialExpressionTextureSample.h"
@@ -98,6 +97,7 @@
 DEFINE_LOG_CATEGORY(LogEditorViewport);
 
 #define LOCTEXT_NAMESPACE "LevelEditorViewportClient"
+PRAGMA_DISABLE_OPTIMIZATION
 
 const FLevelViewportActorLock FLevelViewportActorLock::None(nullptr);
 
@@ -2034,13 +2034,6 @@ void FLevelEditorViewportClient::PerspectiveCameraMoved()
 		UpdateLockedActorViewports(GetActiveActorLock().Get(), false);
 	}
 
-	// Tell the editing mode that the camera moved, in case its interested.
-	FEdMode* Mode = ModeTools->GetActiveMode(FBuiltinEditorModes::EM_InterpEdit);
-	if( Mode )
-	{
-		((FEdModeInterpEdit*)Mode)->CamMoveNotify(this);
-	}
-
 	// Broadcast 'camera moved' delegate
 	FEditorDelegates::OnEditorCameraMoved.Broadcast(GetViewLocation(), GetViewRotation(), ViewportType, ViewIndex);
 }
@@ -2070,13 +2063,6 @@ void FLevelEditorViewportClient::ResetCamera()
 
 	ViewFOV = FOVAngle;
 
-	// If interp mode is active, tell it about the camera movement.
-	FEdMode* Mode = ModeTools->GetActiveMode(FBuiltinEditorModes::EM_InterpEdit);
-	if( Mode )
-	{
-		((FEdModeInterpEdit*)Mode)->CamMoveNotify(this);
-	}
-
 	SetIsCameraCut();
 
 	// Broadcast 'camera moved' delegate
@@ -2345,7 +2331,7 @@ void FLevelEditorViewportClient::Tick(float DeltaTime)
 
 void FLevelEditorViewportClient::UpdateViewForLockedActor(float DeltaTime)
 {
-	// We can't be locked to a matinee actor if this viewport doesn't allow matinee control
+	// We can't be locked to a cinematic actor if this viewport doesn't allow cinematic control
 	if (!bAllowCinematicControl && ActorLocks.CinematicActorLock.HasValidLockedActor())
 	{
 		ActorLocks.CinematicActorLock = FLevelViewportActorLock::None;
@@ -4975,3 +4961,4 @@ bool FLevelEditorViewportClient::GetPivotForOrbit(FVector& Pivot) const
 }
 
 #undef LOCTEXT_NAMESPACE
+PRAGMA_ENABLE_OPTIMIZATION
\ No newline at end of file
diff --git a/Engine/Source/Editor/UnrealEd/Private/UnrealEdEngine.cpp b/Engine/Source/Editor/UnrealEd/Private/UnrealEdEngine.cpp
--- a/Engine/Source/Editor/UnrealEd/Private/UnrealEdEngine.cpp
+++ b/Engine/Source/Editor/UnrealEd/Private/UnrealEdEngine.cpp
@@ -746,13 +746,6 @@ void UUnrealEdEngine::OnPostWindowsMessage(FViewport* Viewport, uint32 Message)
 	}		
 }
 
-
-void UUnrealEdEngine::OnOpenMatinee()
-{
-	// Register a delegate to pickup when Matinee is closed.
-	UpdateEdModeOnMatineeCloseDelegateHandle = GLevelEditorModeTools().OnEditorModeIDChanged().AddUObject( this, &UUnrealEdEngine::UpdateEdModeOnMatineeClose );
-}
-
 bool UUnrealEdEngine::IsAutosaving() const
 {
 	if (PackageAutoSaver)
diff --git a/Engine/Source/Editor/UnrealEd/Public/EditorModeInterpolation.h b/Engine/Source/Editor/UnrealEd/Public/EditorModeInterpolation.h
deleted file mode 100644
diff --git a/Engine/Source/Editor/UnrealEd/Public/EditorViewportClient.h b/Engine/Source/Editor/UnrealEd/Public/EditorViewportClient.h
--- a/Engine/Source/Editor/UnrealEd/Public/EditorViewportClient.h
+++ b/Engine/Source/Editor/UnrealEd/Public/EditorViewportClient.h
@@ -31,7 +31,6 @@ class FEditorViewportClient;
 class FEdMode;
 class FMouseDeltaTracker;
 class FPreviewScene;
-class IMatineeBase;
 class SEditorViewport;
 class UActorFactory;
 class UTypedElementViewportInteraction;
@@ -1068,15 +1067,14 @@ public:
 	/**
 	 * Allows custom disabling of camera recoil
 	 */
-	void SetMatineeRecordingWindow(IMatineeBase* InInterpEd);
+	UE_DEPRECATED(5.0, "Matinee is no longer part of the editor.")
+	void SetMatineeRecordingWindow(class IMatineeBase* InInterpEd) {}
 
 	/**
 	 * Returns true if camera recoil is currently allowed
 	 */
-	bool IsMatineeRecordingWindow() const
-	{
-		return (RecordingInterpEd != NULL);
-	}
+	UE_DEPRECATED(5.0, "Matinee is no longer part of the editor.")
+	bool IsMatineeRecordingWindow() const { return false; }
 	
 	EAxisList::Type GetCurrentWidgetAxis() const;
 
@@ -1649,9 +1647,6 @@ protected:
 	bool bShowWidget;
 
 	FMouseDeltaTracker*		MouseDeltaTracker;
-		
-	/**InterpEd, should only be set if used for matinee recording*/
-	IMatineeBase* RecordingInterpEd;
 
 	/** If true, the canvas has been been moved using bMoveCanvas Mode*/
 	bool bHasMouseMovedSinceClick;
diff --git a/Engine/Source/Editor/UnrealEd/UnrealEd.Build.cs b/Engine/Source/Editor/UnrealEd/UnrealEd.Build.cs
--- a/Engine/Source/Editor/UnrealEd/UnrealEd.Build.cs
+++ b/Engine/Source/Editor/UnrealEd/UnrealEd.Build.cs
@@ -214,7 +214,6 @@ public class UnrealEd : ModuleRules
 				"TextureEditor",
 				"Cascade",
 				"UMGEditor",
-				"Matinee",
 				"AssetTools",
 				"ClassViewer",
 				"StructViewer",
diff --git a/Engine/Source/Editor/WorldBrowser/Private/LevelCollectionModel.cpp b/Engine/Source/Editor/WorldBrowser/Private/LevelCollectionModel.cpp
--- a/Engine/Source/Editor/WorldBrowser/Private/LevelCollectionModel.cpp
+++ b/Engine/Source/Editor/WorldBrowser/Private/LevelCollectionModel.cpp
@@ -17,7 +17,6 @@
 #include "EditorModeManager.h"
 #include "EditorModes.h"
 #include "FileHelpers.h"
-#include "EditorModeInterpolation.h"
 #include "ScopedTransaction.h"
 #include "EditorLevelUtils.h"
 #include "LevelCollectionCommands.h"
@@ -26,7 +25,6 @@
 #include "IAssetTypeActions.h"
 #include "AssetToolsModule.h"
 #include "EditorSupportDelegates.h"
-#include "Matinee/MatineeActor.h"
 #include "GameFramework/WorldSettings.h"
 
 #include "ShaderCompiler.h"
@@ -697,19 +695,7 @@ void FLevelCollectionModel::UnloadLevels(const FLevelModelList& InLevelList)
 	UWorld* ThisWorld = GetWorld();
 	check(ThisWorld != nullptr);
 
-	// If matinee is opened, and if it belongs to the level being removed, close it
-	if (GLevelEditorModeTools().IsModeActive(FBuiltinEditorModes::EM_InterpEdit))
-	{
-		TArray<ULevel*> LevelsToRemove = GetLevelObjectList(InLevelList);
-		
-		const FEdModeInterpEdit* InterpEditMode = (const FEdModeInterpEdit*)GLevelEditorModeTools().GetActiveMode(FBuiltinEditorModes::EM_InterpEdit);
-
-		if (InterpEditMode && InterpEditMode->MatineeActor && LevelsToRemove.Contains(InterpEditMode->MatineeActor->GetLevel()))
-		{
-			GLevelEditorModeTools().ActivateDefaultMode();
-		}
-	}
-	else if(GLevelEditorModeTools().IsModeActive(FBuiltinEditorModes::EM_Landscape))
+	if(GLevelEditorModeTools().IsModeActive(FBuiltinEditorModes::EM_Landscape))
 	{
 		GLevelEditorModeTools().ActivateDefaultMode();
 	}
@@ -1760,34 +1746,6 @@ bool FLevelCollectionModel::IsValidFindInContentBrowser()
 
 void FLevelCollectionModel::MoveActorsToSelected_Executed()
 {
-	// If matinee is open, and if an actor being moved belongs to it, message the user
-	if (GLevelEditorModeTools().IsModeActive(FBuiltinEditorModes::EM_InterpEdit))
-	{
-		const FEdModeInterpEdit* InterpEditMode = (const FEdModeInterpEdit*)GLevelEditorModeTools().GetActiveMode(FBuiltinEditorModes::EM_InterpEdit);
-		if (InterpEditMode && InterpEditMode->MatineeActor)
-		{
-			TArray<AActor*> ControlledActors;
-			InterpEditMode->MatineeActor->GetControlledActors(ControlledActors);
-
-			// are any of the selected actors in the matinee
-			USelection* SelectedActors = GEditor->GetSelectedActors();
-			for (FSelectionIterator Iter(*SelectedActors); Iter; ++Iter)
-			{
-				AActor* Actor = CastChecked<AActor>(*Iter);
-				if (Actor != nullptr && (Actor == InterpEditMode->MatineeActor || ControlledActors.Contains(Actor)))
-				{
-					const bool ExitInterp = EAppReturnType::Yes == FMessageDialog::Open(EAppMsgType::YesNo, NSLOCTEXT("UnrealEd", "MatineeUnableToMove", "You must close Matinee before moving actors.\nDo you wish to do this now and continue?"));
-					if (!ExitInterp)
-					{
-						return;
-					}
-					GLevelEditorModeTools().DeactivateMode(FBuiltinEditorModes::EM_InterpEdit);
-					break;
-				}
-			}
-		}
-	}
-
 	MakeLevelCurrent_Executed();
 
 	const FScopedTransaction Transaction(LOCTEXT("MoveSelectedActorsToSelectedLevel", "Move Selected Actors to Level"));
