Undo //UE5/Main/Engine/Source/... changelist 16613917

diff --git a/Engine/Source/Developer/AssetTools/Private/AssetTypeActions/AssetTypeActions_CameraAnim.cpp b/Engine/Source/Developer/AssetTools/Private/AssetTypeActions/AssetTypeActions_CameraAnim.cpp
--- a/Engine/Source/Developer/AssetTools/Private/AssetTypeActions/AssetTypeActions_CameraAnim.cpp
+++ b/Engine/Source/Developer/AssetTools/Private/AssetTypeActions/AssetTypeActions_CameraAnim.cpp
@@ -180,7 +180,7 @@ void FAssetTypeActions_CameraAnim::OpenAssetEditor( const TArray<UObject*>& InOb
 		{
 			// construct a temporary matinee actor
 			CreateMatineeActorForCameraAnim(CameraAnim);
-PRAGMA_DISABLE_DEPRECATION_WARNINGS
+
 			if (GEditor->ShouldOpenMatinee(PreviewMatineeActor.Get()))
 			{
 				// changed the actor type, but don't want to lose any properties from previous
@@ -227,7 +227,6 @@ PRAGMA_DISABLE_DEPRECATION_WARNINGS
 			{
 				GEditor->GetEditorWorldContext().World()->DestroyActor(PreviewMatineeActor.Get());
 			}
-PRAGMA_ENABLE_DEPRECATION_WARNINGS
 		}
 	}
 }
diff --git a/Engine/Source/Editor/LevelEditor/Private/LevelEditor.cpp b/Engine/Source/Editor/LevelEditor/Private/LevelEditor.cpp
--- a/Engine/Source/Editor/LevelEditor/Private/LevelEditor.cpp
+++ b/Engine/Source/Editor/LevelEditor/Private/LevelEditor.cpp
@@ -1132,6 +1132,11 @@ void FLevelEditorModule::BindGlobalLevelEditorCommands()
 		FExecuteAction::CreateStatic( &FLevelEditorActionCallbacks::ExecuteExecCommand, FString( TEXT("ACTOR SELECT MATCHINGMATERIAL") ) )
 		);
 
+	ActionList.MapAction(
+		Commands.SelectAllActorsControlledByMatinee,
+		FExecuteAction::CreateStatic( &FLevelEditorActionCallbacks::OnSelectAllActorsControlledByMatinee )
+		);
+
 	ActionList.MapAction(
 		Commands.SelectMatchingEmitter,
 		FExecuteAction::CreateStatic( &FLevelEditorActionCallbacks::ExecuteExecCommand, FString( TEXT("ACTOR SELECT MATCHINGEMITTER") ) )
diff --git a/Engine/Source/Editor/LevelEditor/Private/LevelEditorActions.cpp b/Engine/Source/Editor/LevelEditor/Private/LevelEditorActions.cpp
--- a/Engine/Source/Editor/LevelEditor/Private/LevelEditorActions.cpp
+++ b/Engine/Source/Editor/LevelEditor/Private/LevelEditorActions.cpp
@@ -37,6 +37,7 @@
 #include "EditorModes.h"
 #include "UnrealEdMisc.h"
 #include "FileHelpers.h"
+#include "EditorModeInterpolation.h"
 #include "UnrealEdGlobals.h"
 #include "WorldPartition/WorldPartitionSubsystem.h"
 #include "WorldPartition/WorldPartition.h"
@@ -49,6 +50,7 @@
 #include "Elements/Interfaces/TypedElementObjectInterface.h"
 
 #include "LevelEditor.h"
+#include "Matinee/MatineeActor.h"
 #include "Engine/LevelScriptBlueprint.h"
 #include "LightingBuildOptions.h"
 #include "EditorSupportDelegates.h"
@@ -1837,6 +1839,11 @@ bool FLevelEditorActionCallbacks::CanSelectComponentOwnerActor()
 	return GEditor->GetSelectedComponentCount() > 0;
 }
 
+void FLevelEditorActionCallbacks::OnSelectAllActorsControlledByMatinee()
+{
+	GEditor->SelectAllActorsControlledByMatinee();
+}
+
 void FLevelEditorActionCallbacks::OnSelectOwningHLODCluster()
 {
 	if (GEditor->GetSelectedActorCount() > 0)
@@ -1856,6 +1863,27 @@ void FLevelEditorActionCallbacks::OnSelectOwningHLODCluster()
 	}
 }
 
+void FLevelEditorActionCallbacks::OnSelectMatineeActor( AMatineeActor * ActorToSelect )
+{
+	GEditor->SelectNone( false, true );
+	GEditor->SelectActor(ActorToSelect, true, false, true);
+
+	GEditor->NoteSelectionChange();
+}
+
+void FLevelEditorActionCallbacks::OnSelectMatineeGroup( AActor* Actor )
+{
+	if( GLevelEditorModeTools().IsModeActive( FBuiltinEditorModes::EM_InterpEdit ) )
+	{
+		FEdModeInterpEdit* InterpEditMode = (FEdModeInterpEdit*)GLevelEditorModeTools().GetActiveMode( FBuiltinEditorModes::EM_InterpEdit );
+
+		if ( InterpEditMode && InterpEditMode->MatineeActor )
+		{
+			InterpEditMode->UpdateSelectedActor();
+		}
+	}
+}
+
 void FLevelEditorActionCallbacks::OnApplyMaterialToSurface()
 {
 	FEditorDelegates::LoadSelectedAssetsIfNeeded.Broadcast();
@@ -3434,6 +3462,7 @@ void FLevelEditorCommands::RegisterCommands()
 	UI_COMMAND( SelectStationaryLightsExceedingOverlap, "Select Stationary Lights exceeding overlap", "Selects all stationary lights exceeding the overlap limit", EUserInterfaceActionType::Button, FInputChord() ); 
 	UI_COMMAND( SelectAllAddditiveBrushes, "Select All Additive Brushes", "Selects all additive brushes", EUserInterfaceActionType::Button, FInputChord() ); 
 	UI_COMMAND( SelectAllSubtractiveBrushes, "Select All Subtractive Brushes", "Selects all subtractive brushes", EUserInterfaceActionType::Button, FInputChord() ); 
+	UI_COMMAND( SelectAllActorsControlledByMatinee, "Select Actors Used by This Matinee", "Selects all actors controlled by this Matinee", EUserInterfaceActionType::Button, FInputChord() ); 
 
 	UI_COMMAND( SelectAllSurfaces, "Select All Surfaces", "Selects all bsp surfaces", EUserInterfaceActionType::Button, FInputChord(EModifierKey::Shift, EKeys::S) ); 
 
diff --git a/Engine/Source/Editor/LevelEditor/Private/LevelEditorContextMenu.cpp b/Engine/Source/Editor/LevelEditor/Private/LevelEditorContextMenu.cpp
--- a/Engine/Source/Editor/LevelEditor/Private/LevelEditorContextMenu.cpp
+++ b/Engine/Source/Editor/LevelEditor/Private/LevelEditorContextMenu.cpp
@@ -30,7 +30,9 @@
 #include "Editor/GroupActor.h"
 #include "LevelEditorViewport.h"
 #include "EditorModes.h"
+#include "EditorModeInterpolation.h"
 #include "LevelEditor.h"
+#include "Matinee/MatineeActor.h"
 #include "Kismet2/KismetEditorUtilities.h"
 #include "AssetSelection.h"
 #include "LevelEditorActions.h"
@@ -148,6 +150,13 @@ public:
 	 */
 	static void FillSourceControlMenu(UToolMenu* Menu);
 
+private:
+	/**
+	 * Fills in menu options for the matinee selection menu
+	 *
+	 * @param MenuBuilder	The menu to add items to
+	 */
+	static void FillMatineeSelectActorMenu(UToolMenu* Menu);
 };
 
 FSelectedActorInfo FLevelEditorContextMenuImpl::SelectionInfo;
@@ -970,6 +979,93 @@ void FLevelEditorContextMenuImpl::FillSelectActorMenu(UToolMenu* Menu)
 			Section.AddMenuEntry(FLevelEditorCommands::Get().GeometryCollectionSelectInverseGeometry);
 		}
 	}
+
+	// build matinee related selection menu
+	FillMatineeSelectActorMenu( Menu );
+}
+
+void FLevelEditorContextMenuImpl::FillMatineeSelectActorMenu(UToolMenu* Menu)
+{
+	{
+		FToolMenuSection& Section = Menu->AddSection("SelectMatinee", LOCTEXT("SelectMatineeHeading", "Matinee"));
+		// show list of Matinee Actors that controls this actor
+		// this is ugly but we don't have good way of knowing which Matinee actor controls me
+		// in the future this can be cached to TMap somewhere and use that list
+		// for now we show only when 1 actor is selected
+		if ( SelectionInfo.SharedLevel && SelectionInfo.NumSelected == 1 )
+		{
+			TArray<AMatineeActor*> MatineeActors;	
+			// first collect all matinee actors
+			for ( AActor* Actor : SelectionInfo.SharedLevel->Actors )
+			{
+				AMatineeActor * CurActor = Cast<AMatineeActor>(Actor);
+				if ( CurActor )
+				{
+					MatineeActors.Add(CurActor);
+				}
+			}
+
+			if ( MatineeActors.Num() > 0 )
+			{
+				FSelectionIterator ActorIter( GEditor->GetSelectedActorIterator() );
+				AActor* SelectedActor = Cast<AActor>(*ActorIter);
+
+				// now delete the matinee actors that don't control currently selected actor
+				for (int32 MatineeActorIter=0; MatineeActorIter<MatineeActors.Num(); ++MatineeActorIter)
+				{
+					AMatineeActor * CurMatineeActor = MatineeActors[MatineeActorIter];
+					TArray<AActor *> CutMatineeControlledActors;
+					CurMatineeActor->GetControlledActors(CutMatineeControlledActors);
+					bool bIsMatineeControlled=false;
+					for ( AActor* ControlledActor : CutMatineeControlledActors )
+					{
+						if (ControlledActor == SelectedActor)
+						{
+							bIsMatineeControlled = true;
+						}
+					}
+
+					// if not, remove it
+					if (!bIsMatineeControlled)
+					{
+						MatineeActors.RemoveAt(MatineeActorIter);
+						--MatineeActorIter;
+					}
+				}
+
+				// if some matinee controls this, add to menu for direct selection
+				if ( MatineeActors.Num() > 0 )
+				{
+					for (int32 MatineeActorIter=0; MatineeActorIter<MatineeActors.Num(); ++MatineeActorIter)
+					{
+						AMatineeActor * CurMatineeActor = MatineeActors[MatineeActorIter];
+						const FText Text = FText::Format( LOCTEXT("SelectMatineeActor", "Select {0}"), FText::FromString( CurMatineeActor->GetName() ) );
+
+						FUIAction CurMatineeActorAction( FExecuteAction::CreateStatic( &FLevelEditorActionCallbacks::OnSelectMatineeActor, CurMatineeActor ) );
+						Section.AddMenuEntry(NAME_None, Text, Text, FSlateIcon(), CurMatineeActorAction);
+
+						// if matinee is opened, and if that is CurMatineeActor, show option to go to group
+						if( GLevelEditorModeTools().IsModeActive( FBuiltinEditorModes::EM_InterpEdit ) )
+						{
+							const FEdModeInterpEdit* InterpEditMode = (const FEdModeInterpEdit*)GLevelEditorModeTools().GetActiveMode( FBuiltinEditorModes::EM_InterpEdit );
+
+							if ( InterpEditMode && InterpEditMode->MatineeActor == CurMatineeActor )
+							{
+								FUIAction SelectedActorAction( FExecuteAction::CreateStatic( &FLevelEditorActionCallbacks::OnSelectMatineeGroup, SelectedActor ) );
+								Section.AddMenuEntry(NAME_None, LOCTEXT("SelectMatineeGroupForActorMenuTitle", "Select Matinee Group For This Actor"), LOCTEXT("SelectMatineeGroupForActorMenuTooltip", "Selects matinee group controlling this actor"), FSlateIcon(), SelectedActorAction);
+							}
+						}
+					}
+				}
+			}
+		}
+
+		// if this class is Matinee Actor, add option to allow select all controlled actors
+		if ( SelectionInfo.bHaveMatinee )
+		{
+			Section.AddMenuEntry(FLevelEditorCommands::Get().SelectAllActorsControlledByMatinee);	
+		}
+	}
 }
 
 void FLevelEditorContextMenuImpl::FillActorVisibilityMenu(UToolMenu* Menu)
diff --git a/Engine/Source/Editor/LevelEditor/Private/LevelEditorToolBar.cpp b/Engine/Source/Editor/LevelEditor/Private/LevelEditorToolBar.cpp
--- a/Engine/Source/Editor/LevelEditor/Private/LevelEditorToolBar.cpp
+++ b/Engine/Source/Editor/LevelEditor/Private/LevelEditorToolBar.cpp
@@ -33,6 +33,7 @@
 #include "SScalabilitySettings.h"
 #include "IContentBrowserSingleton.h"
 #include "ContentBrowserModule.h"
+#include "Matinee/MatineeActor.h"
 #include "LevelSequenceActor.h"
 #include "Engine/LevelScriptBlueprint.h"
 #include "ISettingsCategory.h"
@@ -62,6 +63,11 @@
 #include "Styling/ToolBarStyle.h"
 #include "PlatformInfo.h"
 
+static TAutoConsoleVariable<int32> CVarAllowMatineeActors(
+	TEXT("Matinee.AllowMatineeActors"),
+	0,
+	TEXT("Toggles whether matinee actors should appear in the cinematics menu so that they can be edited."));
+
 namespace LevelEditorActionHelpers
 {
 	/** Filters out any classes for the Class Picker when creating or selecting classes in the Blueprints dropdown */
@@ -1963,14 +1969,16 @@ void FLevelEditorToolBar::RegisterCinematicsMenu()
 			return;
 		}
 
+		const int32 bAllowMatineeActors = CVarAllowMatineeActors->GetInt();
+
 		UWorld* World = FoundContext->LevelEditor.IsValid() ? FoundContext->LevelEditor.Pin()->GetWorld() : nullptr;
-		const bool bHasAnyCinematicsActors = !!TActorIterator<ALevelSequenceActor>(World);
+		const bool bHasAnyCinematicsActors = (bAllowMatineeActors && !!TActorIterator<AMatineeActor>(World)) || !!TActorIterator<ALevelSequenceActor>(World);
 		if (!bHasAnyCinematicsActors)
 		{
 			return;
 		}
 
-		// We can't build a list of LevelSequenceActors while the current World is a PIE world.
+		// We can't build a list of Matinees and LevelSequenceActors while the current World is a PIE world.
 		FSceneOutlinerInitializationOptions InitOptions;
 		{
 			// We hide the header row to keep the UI compact.
@@ -1982,9 +1990,9 @@ void FLevelEditorToolBar::RegisterCinematicsMenu()
 			InitOptions.ColumnMap.Add(FSceneOutlinerBuiltInColumnTypes::Label(), FSceneOutlinerColumnInfo(ESceneOutlinerColumnVisibility::Visible, 0));
 			InitOptions.ColumnMap.Add(FSceneOutlinerBuiltInColumnTypes::ActorInfo(), FSceneOutlinerColumnInfo(ESceneOutlinerColumnVisibility::Visible, 10));
 
-			// Only display MovieScene actors
+			// Only display Matinee and MovieScene actors
 			auto ActorFilter = [&](const AActor* Actor) {
-				return Actor->IsA(ALevelSequenceActor::StaticClass());
+				return (bAllowMatineeActors && Actor->IsA(AMatineeActor::StaticClass())) || Actor->IsA(ALevelSequenceActor::StaticClass());
 			};
 			InitOptions.Filters->AddFilterPredicate<FActorTreeItem>(FActorTreeItem::FFilterPredicate::CreateLambda(ActorFilter));
 		}
@@ -2016,7 +2024,12 @@ void FLevelEditorToolBar::OnCinematicsActorPicked( AActor* Actor )
 	FSlateApplication::Get().DismissAllMenus();
 
 	// Make sure we dismiss the menus before we open this
-	if (ALevelSequenceActor* LevelSequenceActor = Cast<ALevelSequenceActor>(Actor))
+	if (AMatineeActor* MatineeActor = Cast<AMatineeActor>(Actor))
+	{
+		// Open Matinee for editing!
+		GEditor->OpenMatinee( MatineeActor );
+	}
+	else if (ALevelSequenceActor* LevelSequenceActor = Cast<ALevelSequenceActor>(Actor))
 	{
 		FScopedSlowTask SlowTask(1.f, NSLOCTEXT("LevelToolBarCinematicsMenu", "LoadSequenceSlowTask", "Loading Level Sequence..."));
 		SlowTask.MakeDialog();
diff --git a/Engine/Source/Editor/LevelEditor/Public/LevelEditorActions.h b/Engine/Source/Editor/LevelEditor/Public/LevelEditorActions.h
--- a/Engine/Source/Editor/LevelEditor/Public/LevelEditorActions.h
+++ b/Engine/Source/Editor/LevelEditor/Public/LevelEditorActions.h
@@ -17,6 +17,7 @@
 #include "Elements/Framework/TypedElementHandle.h"
 #include "Elements/Interfaces/TypedElementWorldInterface.h"
 
+class AMatineeActor;
 class FLightingBuildOptions;
 class SLevelEditor;
 class UActorFactory;
@@ -399,7 +400,6 @@ public:
 	TSharedPtr< FUICommandInfo > SelectAllWithSameMaterial;
 
 	/** Selects all actors used by currently selected matinee actor */
-	UE_DEPRECATED(5.0, "Matinee is no longer part of the editor.")
 	TSharedPtr< FUICommandInfo > SelectAllActorsControlledByMatinee;
 
 	/** Selects all emitters using the same particle system as the current selection */
@@ -1020,6 +1020,21 @@ public:
 	/** Selects stationary lights that are exceeding the overlap limit. */
 	static void OnSelectStationaryLightsExceedingOverlap();
 
+	/**
+	 * Selects the MatineeActor - used by Matinee Selection
+	 */
+	static void OnSelectMatineeActor( AMatineeActor * ActorToSelect );
+
+	/**
+	 * Selects the Matinee InterpGroup
+	 */
+	static void OnSelectMatineeGroup( AActor* Actor );
+
+	/**
+	 * Called when selecting all actors that's controlled by currently selected matinee actor
+	 */
+	static void OnSelectAllActorsControlledByMatinee();
+	
 	/**
 	* Called when selecting an Actor's (if available) owning HLOD cluster
 	*/
diff --git a/Engine/Source/Editor/UnrealEd/Classes/Editor/EditorEngine.h b/Engine/Source/Editor/UnrealEd/Classes/Editor/EditorEngine.h
--- a/Engine/Source/Editor/UnrealEd/Classes/Editor/EditorEngine.h
+++ b/Engine/Source/Editor/UnrealEd/Classes/Editor/EditorEngine.h
@@ -34,6 +34,7 @@
 
 #include "EditorEngine.generated.h"
 
+class AMatineeActor;
 class APlayerStart;
 class Error;
 class FEditorViewportClient;
@@ -1726,7 +1727,6 @@ public:
 
 	/**
 	 * Request to create a new PIE window and join the currently running PIE session.
-	 * Deferred until the next tick.
 	 */
 	void RequestLateJoin();
 
@@ -2149,8 +2149,7 @@ public:
 	/**
 	 * Selects all actors controlled by currently selected MatineeActor
 	 */
-	UE_DEPRECATED(5.0, "Matinee is no longer part of the editor.")
-	void SelectAllActorsControlledByMatinee() {}
+	void SelectAllActorsControlledByMatinee();
 
 	/**
 	 * Selects all actors with the same class as the current selection
@@ -2326,7 +2325,6 @@ public:
 	 *
 	 * @return The event delegate.
 	 */
-	UE_DEPRECATED(5.0, "Matinee is no longer part of the editor.")
 	DECLARE_DELEGATE_RetVal_OneParam(bool, FShouldOpenMatineeCallback, AMatineeActor*)
 	FShouldOpenMatineeCallback& OnShouldOpenMatinee() { return ShouldOpenMatineeCallback; }
 
@@ -2336,8 +2334,7 @@ public:
 	 * @param MatineeActor	The actor we wish to check (can be null)
 	 * @returns true if the user wishes to proceed
 	 */
-	UE_DEPRECATED(5.0, "Matinee is no longer part of the editor.")
-	bool ShouldOpenMatinee(AMatineeActor* MatineeActor) const { return false; }
+	bool ShouldOpenMatinee(AMatineeActor* MatineeActor) const;
 
 	/** 
 	 * Open the Matinee tool to edit the supplied MatineeActor. Will check that MatineeActor has an InterpData attached.
@@ -2345,8 +2342,7 @@ public:
 	 * @param MatineeActor	The actor we wish to edit
 	 * @param bWarnUser		If true, calls ShouldOpenMatinee as part of the open process
 	 */
-	UE_DEPRECATED(5.0, "Matinee is no longer part of the editor.")
-	void OpenMatinee(class AMatineeActor* MatineeActor, bool bWarnUser = true) {}
+	void OpenMatinee(class AMatineeActor* MatineeActor, bool bWarnUser=true);
 
 	/**
 	* Update any outstanding reflection captures
@@ -2834,7 +2830,7 @@ private:
 	FExecParticleInvoked ExecParticleInvokedEvent; 
 
 	/** Delegate to be called when a matinee is requested to be opened */
-	FShouldOpenMatineeCallback ShouldOpenMatineeCallback;	
+	FShouldOpenMatineeCallback ShouldOpenMatineeCallback;
 
 	/** Reference to owner of the current popup */
 	TWeakPtr<class SWindow> PopupWindow;
@@ -3097,7 +3093,6 @@ private:
 protected:
 
 	/** Called when Matinee is opened */
-	UE_DEPRECATED(5.0, "Matinee is no longer part of the editor.")
 	virtual void OnOpenMatinee(){};
 
 	/**
diff --git a/Engine/Source/Editor/UnrealEd/Classes/Editor/UnrealEdEngine.h b/Engine/Source/Editor/UnrealEd/Classes/Editor/UnrealEdEngine.h
--- a/Engine/Source/Editor/UnrealEd/Classes/Editor/UnrealEdEngine.h
+++ b/Engine/Source/Editor/UnrealEd/Classes/Editor/UnrealEdEngine.h
@@ -210,8 +210,7 @@ public:
 	virtual void GetPackageList( TArray<UPackage*>* InPackages, UClass* InClass ) override;
 	virtual bool ShouldAbortActorDeletion() const override final; // @note Final - Override ShouldAbortComponentDeletion or ShouldAbortActorDeletion (with parameters) instead.
 	virtual void CloseEditor() override;
-
-	
+	virtual void OnOpenMatinee() override;
 	virtual bool IsAutosaving() const override;
 	//~ End UEditorEngine Interface 
 	
diff --git a/Engine/Source/Editor/UnrealEd/Private/EditorEngine.cpp b/Engine/Source/Editor/UnrealEd/Private/EditorEngine.cpp
--- a/Engine/Source/Editor/UnrealEd/Private/EditorEngine.cpp
+++ b/Engine/Source/Editor/UnrealEd/Private/EditorEngine.cpp
@@ -76,9 +76,11 @@
 #include "UnrealEdMisc.h"
 #include "EditorDirectories.h"
 #include "FileHelpers.h"
+#include "EditorModeInterpolation.h"
 #include "Dialogs/Dialogs.h"
 #include "Dialogs/DialogsPrivate.h"
 #include "UnrealEdGlobals.h"
+#include "Matinee/MatineeActor.h"
 #include "InteractiveFoliageActor.h"
 #include "Engine/WorldComposition.h"
 #include "EditorSupportDelegates.h"
@@ -1476,7 +1478,7 @@ void UEditorEngine::Tick( float DeltaSeconds, bool bIdleMode )
 	{
 		AudioFocusViewportClient->SetAudioFocus();
 
-		// override realtime setting if viewport chooses (i.e. for cinematic preview)
+		// override realtime setting if viewport chooses (i.e. for matinee preview)
 		if (AudioFocusViewportClient->IsForcedRealtimeAudio())
 		{
 			bAudioIsRealtime = true;
@@ -3049,6 +3051,29 @@ void UEditorEngine::DeselectLevelInLevelBrowser()
 	LevelEditorModule.SummonWorldBrowserHierarchy();
 }
 
+void UEditorEngine::SelectAllActorsControlledByMatinee()
+{
+	TArray<AActor *> AllActors;
+	UWorld* IteratorWorld = GWorld;
+	for( FSelectedActorIterator Iter(IteratorWorld); Iter; ++Iter)
+	{
+		AMatineeActor * CurActor = Cast<AMatineeActor>(*Iter);
+		if ( CurActor )
+		{
+			TArray<AActor*> Actors;			
+			CurActor->GetControlledActors(Actors);
+			AllActors.Append(Actors);
+		}
+	}
+
+	SelectNone(false, true, false);
+	for(int32 i=0; i<AllActors.Num(); i++)
+	{
+		SelectActor( AllActors[i], true, false, true );
+	}
+	NoteSelectionChange();
+}
+
 void UEditorEngine::SelectAllActorsWithClass( bool bArchetype )
 {
 	if( !bArchetype )
@@ -3676,6 +3701,60 @@ void UEditorEngine::ConvertActorsFromClass( UClass* FromClass, UClass* ToClass )
 	}
 }
 
+bool UEditorEngine::ShouldOpenMatinee(AMatineeActor* MatineeActor) const
+{
+	if( PlayWorld )
+	{
+		FMessageDialog::Open( EAppMsgType::Ok, NSLOCTEXT("UnrealEd", "Error_MatineeCantOpenDuringPIE", "Matinee cannot be opened during Play in Editor.") );
+		return false;
+	}
+
+	if ( MatineeActor && !MatineeActor->MatineeData )
+	{
+		FMessageDialog::Open( EAppMsgType::Ok, NSLOCTEXT("UnrealEd", "Error_MatineeActionMustHaveData", "Matinee must have valid InterpData assigned before being edited.") );
+		return false;
+	}
+
+	// Make sure we can't open the same action twice in Matinee.
+	if( GLevelEditorModeTools().IsModeActive(FBuiltinEditorModes::EM_InterpEdit) )
+	{
+		FMessageDialog::Open( EAppMsgType::Ok, NSLOCTEXT("UnrealEd", "MatineeActionAlreadyOpen", "An Matinee sequence is currently open in an editor.  Please close it before proceeding.") );
+		return false;
+	}
+
+	// Don't let you open Matinee if a transaction is currently active.
+	if( IsTransactionActive() )
+	{
+		FMessageDialog::Open( EAppMsgType::Ok, NSLOCTEXT("UnrealEd", "TransactionIsActive", "Undo Transaction Is Active - Cannot Open Matinee.") );
+		return false;
+	}
+
+	return true;
+}
+
+void UEditorEngine::OpenMatinee(AMatineeActor* MatineeActor, bool bWarnUser)
+{
+	// Drop out if the user doesn't want to proceed to matinee atm
+	if( bWarnUser && ( (ShouldOpenMatineeCallback.IsBound() && !ShouldOpenMatineeCallback.Execute(MatineeActor)) || !ShouldOpenMatinee( MatineeActor ) ) )
+	{
+		return;
+	}
+
+	// If already in Matinee mode, exit out before going back in with new Interpolation.
+	if( GLevelEditorModeTools().IsModeActive( FBuiltinEditorModes::EM_InterpEdit ) )
+	{
+		GLevelEditorModeTools().DeactivateMode( FBuiltinEditorModes::EM_InterpEdit );
+	}
+
+	GLevelEditorModeTools().ActivateMode( FBuiltinEditorModes::EM_InterpEdit );
+
+	FEdModeInterpEdit* InterpEditMode = (FEdModeInterpEdit*)GLevelEditorModeTools().GetActiveMode( FBuiltinEditorModes::EM_InterpEdit );
+
+	InterpEditMode->InitInterpMode( MatineeActor );
+
+	OnOpenMatinee();
+}
+
 void UEditorEngine::BuildReflectionCaptures(UWorld* World)
 {
 	// Note: Lighting and reflection build operations should only dirty BuildData packages, not ULevel packages
diff --git a/Engine/Source/Editor/UnrealEd/Private/EditorModeInterpolation.cpp b/Engine/Source/Editor/UnrealEd/Private/EditorModeInterpolation.cpp
new file mode 100644
--- /dev/null
+++ b/Engine/Source/Editor/UnrealEd/Private/EditorModeInterpolation.cpp
@@ -0,0 +1,394 @@
+// Copyright Epic Games, Inc. All Rights Reserved.
+
+/*=============================================================================
+	EditorModeInterpolation : Editor mode for setting up interpolation sequences.
+
+=============================================================================*/
+
+#include "EditorModeInterpolation.h"
+#include "EditorViewportClient.h"
+#include "Modules/ModuleManager.h"
+#include "Editor/GroupActor.h"
+#include "EditorModeManager.h"
+#include "EditorModes.h"
+#include "InterpolationHitProxy.h"
+
+#include "Editor/Matinee/Public/MatineeModule.h"
+#include "Editor/Matinee/Public/IMatinee.h"
+#include "Framework/Notifications/NotificationManager.h"
+#include "Widgets/Notifications/SNotificationList.h"
+#include "ActorGroupingUtils.h"
+
+static const float	CurveHandleScale = 0.5f;
+
+//////////////////////////////////////////////////////////////////////////
+// FEdModeInterpEdit
+//////////////////////////////////////////////////////////////////////////
+
+FEdModeInterpEdit::FEdModeInterpEdit()
+{
+	MatineeActor = NULL;
+	InterpEd = NULL;
+
+	Tools.Add( new FModeTool_InterpEdit() );
+	SetCurrentTool( MT_InterpEdit );
+
+	bLeavingMode = false;
+}
+
+FEdModeInterpEdit::~FEdModeInterpEdit()
+{
+}
+
+bool FEdModeInterpEdit::InputKey( FEditorViewportClient* ViewportClient, FViewport* Viewport, FKey Key, EInputEvent Event )
+{
+	// Enter key drops new key frames
+	if( Key == EKeys::Enter &&
+		( Event == IE_Pressed || Event == IE_Repeat ) &&
+		( !ViewportClient->IsShiftPressed() &&
+		  !ViewportClient->IsAltPressed() &&
+		  !ViewportClient->IsCtrlPressed() ) )
+	{
+		if( InterpEd != NULL )
+		{
+			// Add a new key!
+			InterpEd->AddKey();
+		}
+
+		return true;
+	}
+
+	return FEdMode::InputKey(ViewportClient, Viewport, Key, Event);
+}
+
+void FEdModeInterpEdit::Enter()
+{
+	FEdMode::Enter();
+
+	// Disable Grouping while in InterpEdit mode
+	bGroupingActiveSaved = UActorGroupingUtils::IsGroupingActive();
+	UActorGroupingUtils::SetGroupingActive(false);
+}
+
+void FEdModeInterpEdit::Exit()
+{
+	MatineeActor = NULL;
+
+	// If there is one, close the Interp Editor and clear pointers.
+	if(InterpEd != NULL)
+	{
+		bLeavingMode = true; // This is so the editor being closed doesn't try and change the mode again!
+
+		InterpEd->Close(true);
+
+		bLeavingMode = false;
+	}
+
+	InterpEd = NULL;
+
+	// Grouping is always disabled while in InterpEdit Mode, re-enable the saved value on exit
+	UActorGroupingUtils::SetGroupingActive(bGroupingActiveSaved);
+	AGroupActor::SelectGroupsInSelection();
+
+	FEdMode::Exit();
+}
+
+// We see if we have 
+void FEdModeInterpEdit::ActorMoveNotify()
+{
+	if(!InterpEd)
+		return;
+
+	InterpEd->ActorModified();
+}
+
+
+void FEdModeInterpEdit::CamMoveNotify(FEditorViewportClient* ViewportClient)
+{
+	if(!InterpEd)
+		return;
+
+	if( ViewportClient->AllowsCinematicControl() )
+	{
+		InterpEd->CamMoved( ViewportClient->GetViewLocation(), ViewportClient->GetViewRotation() );
+	}
+}
+
+void FEdModeInterpEdit::ActorPropChangeNotify()
+{
+	if(!InterpEd)
+		return;
+
+	InterpEd->ActorModified();
+}
+
+void FEdModeInterpEdit::UpdateSelectedActor()
+{
+	// would this be enough?
+	InterpEd->ActorSelectionChange();
+}
+
+// Set the currently edited MatineeActor and open timeline window. Should always be called after we change to FBuiltinEditorModes::EM_InterpEdit.
+void FEdModeInterpEdit::InitInterpMode(AMatineeActor* InMatineeActor)
+{
+	check(InMatineeActor);
+	check(!InterpEd);
+
+	MatineeActor = InMatineeActor;
+	
+	EToolkitMode::Type Mode = EToolkitMode::Standalone;
+	IMatineeModule* MatineeModule = &FModuleManager::LoadModuleChecked<IMatineeModule>( "Matinee" );
+	
+	InterpEd = &(MatineeModule->CreateMatinee(Mode, TSharedPtr<IToolkitHost>(), InMatineeActor).Get());
+}
+
+void FEdModeInterpEdit::Render(const FSceneView* View, FViewport* Viewport, FPrimitiveDrawInterface* PDI)
+{
+	FEdMode::Render(View, Viewport, PDI);
+
+	check( GLevelEditorModeTools().IsModeActive(FBuiltinEditorModes::EM_InterpEdit) );
+
+	if(InterpEd  && !InterpEd->Hide3DTrackView() && View->Family->EngineShowFlags.Splines)
+	{
+		InterpEd->DrawTracks3D(View, PDI);
+	}
+}
+
+void FEdModeInterpEdit::DrawHUD(FEditorViewportClient* ViewportClient,FViewport* Viewport,const FSceneView* View,FCanvas* Canvas)
+{
+	FEdMode::DrawHUD(ViewportClient,Viewport,View,Canvas);
+
+	if(InterpEd)
+	{
+		InterpEd->DrawModeHUD(ViewportClient,Viewport,View,Canvas);
+	}
+}
+
+bool FEdModeInterpEdit::AllowWidgetMove()
+{
+	FModeTool_InterpEdit* InterpTool = (FModeTool_InterpEdit*)FindTool(MT_InterpEdit);
+
+	if(InterpTool->bMovingHandle)
+	{
+		return false;
+	}
+	else
+	{
+		return true;
+	}
+}
+
+void FEdModeInterpEdit::ActorSelectionChangeNotify()
+{
+	FEdMode::ActorSelectionChangeNotify();
+
+	if ( InterpEd )
+	{
+		InterpEd->ActorSelectionChange();
+	}
+}
+
+bool FEdModeInterpEdit::IsCompatibleWith(FEditorModeID OtherModeID) const
+{
+	return OtherModeID == FBuiltinEditorModes::EM_MeshPaint;
+}
+
+//////////////////////////////////////////////////////////////////////////
+// FModeTool_InterpEdit
+//////////////////////////////////////////////////////////////////////////
+
+FModeTool_InterpEdit::FModeTool_InterpEdit()
+{
+	ID = MT_InterpEdit;
+
+	bMovingHandle = false;
+	DragGroup = NULL;
+	DragTrackIndex = false;
+	DragKeyIndex = false;
+	bDragArriving = false;
+}
+
+FModeTool_InterpEdit::~FModeTool_InterpEdit()
+{
+}
+
+bool FModeTool_InterpEdit::MouseMove(FEditorViewportClient* ViewportClient,FViewport* Viewport,int32 x, int32 y)
+{
+	check( GLevelEditorModeTools().IsModeActive(FBuiltinEditorModes::EM_InterpEdit) );
+
+	FEdModeInterpEdit* mode = (FEdModeInterpEdit*)GLevelEditorModeTools().GetActiveMode(FBuiltinEditorModes::EM_InterpEdit);
+
+	return 0;
+}
+
+bool FModeTool_InterpEdit::InputAxis(FEditorViewportClient* InViewportClient, FViewport* Viewport, int32 ControllerId, FKey Key, float Delta, float DeltaTime)
+{
+	if ( InViewportClient->GetCurrentWidgetAxis() == EAxisList::None && !InViewportClient->Viewport->KeyState(EKeys::MiddleMouseButton) )
+	{
+		// We need to set up a widget axis here to prevent our drag operation being co-opted by box/frustum selection
+		ELevelViewportType ViewportType = InViewportClient->ViewportType;
+		switch(ViewportType)
+		{
+		case LVT_OrthoXY:
+		case LVT_OrthoNegativeXY:
+			InViewportClient->SetCurrentWidgetAxis(EAxisList::XY);
+			break;
+		case LVT_OrthoXZ:
+		case LVT_OrthoNegativeXZ:
+			InViewportClient->SetCurrentWidgetAxis(EAxisList::XZ);
+			break;
+		case LVT_OrthoYZ:
+		case LVT_OrthoNegativeYZ:
+			InViewportClient->SetCurrentWidgetAxis(EAxisList::YZ);
+			break;
+		case LVT_OrthoFreelook:
+		case LVT_Perspective:
+			break;
+		}
+	}
+
+	return false;
+}
+
+bool FModeTool_InterpEdit::InputKey(FEditorViewportClient* ViewportClient, FViewport* Viewport, FKey Key, EInputEvent Event)
+{
+	check( GLevelEditorModeTools().IsModeActive(FBuiltinEditorModes::EM_InterpEdit) );
+
+	FEdModeInterpEdit* mode = (FEdModeInterpEdit*)GLevelEditorModeTools().GetActiveMode(FBuiltinEditorModes::EM_InterpEdit);
+	if( !mode->InterpEd )
+	{
+		// Abort cleanly on InerpEd not yet being assigned.  This can occasionally be the case when receiving
+		// modifier key release events when changing into interp edit mode.
+		return false;
+	}
+
+	bool bCtrlDown = Viewport->KeyState(EKeys::LeftControl) || Viewport->KeyState(EKeys::RightControl);
+	bool bAltDown = Viewport->KeyState(EKeys::LeftAlt) || Viewport->KeyState(EKeys::RightAlt);
+	bool bShiftDown = Viewport->KeyState(EKeys::LeftShift) || Viewport->KeyState(EKeys::RightShift);
+
+	if( Key == EKeys::LeftMouseButton )
+	{
+
+		if( Event == IE_Pressed)
+		{
+			int32 HitX = ViewportClient->Viewport->GetMouseX();
+			int32 HitY = ViewportClient->Viewport->GetMouseY();
+			HHitProxy*	HitResult = ViewportClient->Viewport->GetHitProxy(HitX, HitY);
+
+			if(HitResult)
+			{
+				if( HitResult->IsA(HInterpTrackKeypointProxy::StaticGetType()) )
+				{
+					HInterpTrackKeypointProxy* KeyProxy = (HInterpTrackKeypointProxy*)HitResult;
+					UInterpGroup* Group = KeyProxy->Group;
+					UInterpTrack* Track = KeyProxy->Track;
+					int32 KeyIndex = KeyProxy->KeyIndex;
+
+					// Using the CTRL modifier invokes multi-select keyframe selection.
+					if( bCtrlDown )
+					{
+						// If key is already selected, deselect the key.
+						if( mode->InterpEd->KeyIsInSelection(Group, Track, KeyIndex) )
+						{
+							mode->InterpEd->RemoveKeyFromSelection(Group, Track, KeyIndex);
+							mode->InterpEd->InvalidateTrackWindowViewports();
+						}
+						// Otherwise, select the key while preserving previous selection.
+						else
+						{
+							// This will invalidate the display - so we must not access the KeyProxy after this!
+							mode->InterpEd->SelectTrack( Group, Track, false);
+							mode->InterpEd->AddKeyToSelection(Group, Track, KeyIndex, !bShiftDown);
+						}
+					}
+					else
+					{
+						mode->InterpEd->SelectTrack( Group, Track );
+						// NOTE: Clear previously-selected tracks because ctrl is not down. 
+						mode->InterpEd->ClearKeySelection();
+						mode->InterpEd->AddKeyToSelection(Group, Track, KeyIndex, !bShiftDown);
+					}
+				}
+				else if( HitResult->IsA(HInterpTrackKeyHandleProxy::StaticGetType()) )
+				{
+					// If we clicked on a 3D track handle, remember which key.
+					HInterpTrackKeyHandleProxy* KeyProxy = (HInterpTrackKeyHandleProxy*)HitResult;
+					DragGroup = KeyProxy->Group;
+					DragTrackIndex = KeyProxy->TrackIndex;
+					DragKeyIndex = KeyProxy->KeyIndex;
+					bDragArriving = KeyProxy->bArriving;
+
+					bMovingHandle = true;
+
+					mode->InterpEd->BeginDrag3DHandle(DragGroup, DragTrackIndex);
+				}
+			}
+		}
+		else if( Event == IE_Released)
+		{
+			if(bMovingHandle)
+			{
+				mode->InterpEd->EndDrag3DHandle();
+				bMovingHandle = false;
+			}
+		}
+	}
+
+	// Handle keys
+	if( Event == IE_Pressed )
+	{
+		// Swallow 'Delete' key to avoid deleting stuff when trying to interpolate it!
+		if ( Key == EKeys::Platform_Delete )
+		{
+			// Can't delete actors if Matinee is open.
+			const FText ErrorMsg = NSLOCTEXT("UnrealEd", "Error_WrongModeForActorDeletion", "Cannot delete actor while Matinee is open");
+			FNotificationInfo Info(ErrorMsg);
+			FSlateNotificationManager::Get().AddNotification(Info);
+			return true;
+		}
+		else if( mode->InterpEd->ProcessKeyPress( Key, bCtrlDown, bAltDown ) )
+		{
+			return true;
+		}
+	}
+
+	return FModeTool::InputKey(ViewportClient, Viewport, Key, Event);
+}
+
+bool FModeTool_InterpEdit::InputDelta(FEditorViewportClient* InViewportClient, FViewport* InViewport, FVector& InDrag, FRotator& InRot, FVector& InScale)
+{
+	check( GLevelEditorModeTools().IsModeActive(FBuiltinEditorModes::EM_InterpEdit) );
+
+	FEdModeInterpEdit* mode = (FEdModeInterpEdit*)GLevelEditorModeTools().GetActiveMode(FBuiltinEditorModes::EM_InterpEdit);
+	check(mode->InterpEd);
+
+	bool bShiftDown = InViewport->KeyState(EKeys::LeftShift) || InViewport->KeyState(EKeys::RightShift);
+
+	FVector InputDeltaDrag( InDrag );
+
+	// If we are grabbing a 'handle' on the movement curve, pass that info to Matinee
+	if(bMovingHandle)
+	{
+		mode->InterpEd->Move3DHandle( DragGroup, DragTrackIndex, DragKeyIndex, bDragArriving, InputDeltaDrag * (1.f/CurveHandleScale) );
+
+		return 1;
+	}
+	// If shift is downOnly do 'move initial position' if dragging the widget
+	else if(bShiftDown && InViewportClient->GetCurrentWidgetAxis() != EAxisList::None)
+	{
+		mode->InterpEd->MoveInitialPosition( InputDeltaDrag, InRot );
+
+		return 1;
+	}
+
+	InViewportClient->Viewport->Invalidate();
+
+	return 0;
+}
+
+void FModeTool_InterpEdit::SelectNone()
+{
+	check( GLevelEditorModeTools().IsModeActive(FBuiltinEditorModes::EM_InterpEdit) );
+
+	FEdModeInterpEdit* mode = (FEdModeInterpEdit*)GLevelEditorModeTools().GetActiveMode(FBuiltinEditorModes::EM_InterpEdit);
+}
diff --git a/Engine/Source/Editor/UnrealEd/Private/EditorModeRegistry.cpp b/Engine/Source/Editor/UnrealEd/Private/EditorModeRegistry.cpp
--- a/Engine/Source/Editor/UnrealEd/Private/EditorModeRegistry.cpp
+++ b/Engine/Source/Editor/UnrealEd/Private/EditorModeRegistry.cpp
@@ -5,6 +5,7 @@
 #include "EditorStyleSet.h"
 #include "EdMode.h"
 #include "EditorModes.h"
+#include "EditorModeInterpolation.h"
 
 #include "Editor/PlacementMode/Public/IPlacementModeModule.h"
 #include "Editor/LandscapeEditor/Public/LandscapeEditorModule.h"
@@ -50,6 +51,8 @@ void FEditorModeRegistry::Initialize()
 	}
 
 	// Add default editor modes
+	RegisterMode<FEdModeInterpEdit>(FBuiltinEditorModes::EM_InterpEdit);
+
 	FModuleManager::LoadModuleChecked<FActorPickerModeModule>(TEXT("ActorPickerMode"));
 	FModuleManager::LoadModuleChecked<FSceneDepthPickerModeModule>(TEXT("SceneDepthPickerMode"));
 	FModuleManager::LoadModuleChecked<ILandscapeEditorModule>(TEXT("LandscapeEditor"));
diff --git a/Engine/Source/Editor/UnrealEd/Private/EditorViewportClient.cpp b/Engine/Source/Editor/UnrealEd/Private/EditorViewportClient.cpp
--- a/Engine/Source/Editor/UnrealEd/Private/EditorViewportClient.cpp
+++ b/Engine/Source/Editor/UnrealEd/Private/EditorViewportClient.cpp
@@ -25,11 +25,14 @@
 #include "EditorModes.h"
 #include "MouseDeltaTracker.h"
 #include "CameraController.h"
+#include "Editor/Matinee/Public/IMatinee.h"
+#include "Editor/Matinee/Public/MatineeConstants.h"
 #include "HighResScreenshot.h"
 #include "EditorDragTools.h"
 #include "EngineAnalytics.h"
 #include "AnalyticsEventAttribute.h"
 #include "Interfaces/IAnalyticsProvider.h"
+#include "Matinee/MatineeActor.h"
 #include "EngineModule.h"
 #include "Framework/Notifications/NotificationManager.h"
 #include "Widgets/Notifications/SNotificationList.h"
@@ -412,6 +415,7 @@ FEditorViewportClient::FEditorViewportClient(FEditorModeTools* InModeTools, FPre
 	, Widget(new FWidget)
 	, bShowWidget(true)
 	, MouseDeltaTracker(new FMouseDeltaTracker)
+	, RecordingInterpEd(NULL)
 	, bHasMouseMovedSinceClick(false)
 	, CameraController(new FEditorCameraController())
 	, CameraUserImpulseData(new FCameraControllerUserImpulseData())
@@ -1329,6 +1333,12 @@ void FEditorViewportClient::ReceivedFocus(FViewport* InViewport)
 	SetRequiredCursorOverride(false, EMouseCursor::Default);
 	FSlateApplication::Get().QueryCursor();
 
+	if( IsMatineeRecordingWindow() )
+	{
+		// Allow the joystick to be used for matinee capture
+		InViewport->SetUserFocus( true );
+	}
+
 	ModeTools->ReceivedFocus(this, Viewport);
 }
 
@@ -1527,7 +1537,64 @@ void FEditorViewportClient::UpdateCameraMovementFromJoystick(const bool bRelativ
 				{
 					CameraUserImpulseData->ZoomOutInImpulse -= InConfig.ZoomMultiplier;
 				}
-				
+				else if (RecordingInterpEd)
+				{
+					bool bRepeatAllowed = RecordingInterpEd->IsRecordMenuChangeAllowedRepeat();
+					if ((Key == EKeys::Gamepad_DPad_Up) && bPressed)
+					{
+						const bool bNextMenuItem = false;
+						RecordingInterpEd->ChangeRecordingMenu(bNextMenuItem);
+						bRepeatAllowed = false;
+					}
+					else if ((Key == EKeys::Gamepad_DPad_Down) && bPressed)
+					{
+						const bool bNextMenuItem = true;
+						RecordingInterpEd->ChangeRecordingMenu(bNextMenuItem);
+						bRepeatAllowed = false;
+					}
+					else if ((Key == EKeys::Gamepad_DPad_Right) && (bPressed || (bRepeat && bRepeatAllowed)))
+					{
+						const bool bIncrease= true;
+						RecordingInterpEd->ChangeRecordingMenuValue(this, bIncrease);
+					}
+					else if ((Key == EKeys::Gamepad_DPad_Left) && (bPressed || (bRepeat && bRepeatAllowed)))
+					{
+						const bool bIncrease= false;
+						RecordingInterpEd->ChangeRecordingMenuValue(this, bIncrease);
+					}
+					else if ((Key == EKeys::Gamepad_RightThumbstick) && (bPressed))
+					{
+						const bool bIncrease= true;
+						RecordingInterpEd->ResetRecordingMenuValue(this);
+					}
+					else if ((Key == EKeys::Gamepad_LeftThumbstick) && (bPressed))
+					{
+						RecordingInterpEd->ToggleRecordMenuDisplay();
+					}
+					else if ((Key == EKeys::Gamepad_FaceButton_Bottom) && (bPressed))
+					{
+						RecordingInterpEd->ToggleRecordInterpValues();
+					}
+					else if ((Key == EKeys::Gamepad_FaceButton_Right) && (bPressed))
+					{
+						if (!RecordingInterpEd->GetMatineeActor()->bIsPlaying)
+						{
+							bool bLoop = true;
+							bool bForward = true;
+							RecordingInterpEd->StartPlaying(bLoop, bForward);
+						}
+						else
+						{
+							RecordingInterpEd->StopPlaying();
+						}
+					}
+
+					if (!bRepeatAllowed)
+					{
+						//only respond to this event ONCE
+						JoystickState->KeyEventValues.Remove(Key);
+					}
+				}
 				if (bPressed)
 				{
 					//instantly set to repeat to stock rapid flickering until the time out
@@ -1828,13 +1895,29 @@ void FEditorViewportClient::UpdateCameraMovement( float DeltaTime )
 			CameraController->GetConfig().bForceRotationalPhysics = false;
 		}
 
+		bool bIgnoreJoystickControls = false;
+		//if we're playing back (without recording), stop input from being processed
+		if (RecordingInterpEd && RecordingInterpEd->GetMatineeActor())
+		{
+			if (RecordingInterpEd->GetMatineeActor()->bIsPlaying && !RecordingInterpEd->IsRecordingInterpValues())
+			{
+				bIgnoreJoystickControls = true;
+			}
+
+			CameraController->GetConfig().bPlanarCamera = (RecordingInterpEd->GetCameraMovementScheme() == MatineeConstants::ECameraScheme::CAMERA_SCHEME_PLANAR_CAM);
+		}
+
 		if( GetDefault<ULevelEditorViewportSettings>()->bLevelEditorJoystickControls )
 		{
 			//Now update for cached joystick info (relative movement first)
 			UpdateCameraMovementFromJoystick(true, CameraController->GetConfig());
 
-			//Now update for cached joystick info (absolute movement second)
-			UpdateCameraMovementFromJoystick(false, CameraController->GetConfig());
+			//if we're not playing any cinematics right now
+			if (!bIgnoreJoystickControls)
+			{
+				//Now update for cached joystick info (absolute movement second)
+				UpdateCameraMovementFromJoystick(false, CameraController->GetConfig());
+			}
 		}
 
 		FVector NewViewLocation = GetViewLocation();
@@ -1849,7 +1932,7 @@ void FEditorViewportClient::UpdateCameraMovement( float DeltaTime )
 		const float FinalCameraSpeedScale = FlightCameraSpeedScale * CameraSpeed * GetCameraSpeedScalar() * CameraBoost;
 
 		// Only allow FOV recoil if flight camera mode is currently inactive.
-		const bool bAllowRecoilIfNoImpulse = !bUsingFlightInput;
+		const bool bAllowRecoilIfNoImpulse = (!bUsingFlightInput) && (!IsMatineeRecordingWindow());
 
 		// Update the camera's position, rotation and FOV
 		float EditorMovementDeltaUpperBound = 1.0f;	// Never "teleport" the camera further than a reasonable amount after a large quantum
@@ -2527,6 +2610,17 @@ bool FEditorViewportClient::IsUsingAbsoluteTranslation(bool bAlsoCheckAbsoluteRo
 	return (!bCameraLockedToWidget && !bIsHotKeyAxisLocked && bAbsoluteMovementEnabled && (bCurrentWidgetSupportsAbsoluteMovement || bCurrentWidgetSupportsAbsoluteRotation) && bWidgetActivelyTrackingAbsoluteMovement && !IsOrtho() && bAnyMouseButtonsDown);
 }
 
+void FEditorViewportClient::SetMatineeRecordingWindow (IMatineeBase* InInterpEd)
+{
+	RecordingInterpEd = InInterpEd;
+	if (CameraController)
+	{
+		FCameraControllerConfig Config = CameraController->GetConfig();
+		RecordingInterpEd->LoadRecordingSettings(OUT Config);
+		CameraController->SetConfig(Config);
+	}
+}
+
 bool FEditorViewportClient::IsFlightCameraActive() const
 {
 	bool bIsFlightMovementKey = false;
@@ -3789,7 +3883,7 @@ void FEditorViewportClient::Draw(FViewport* InViewport, FCanvas* Canvas)
 		if( !UseEngineShowFlags.Game )
 		{
 			// in the editor, disable camera motion blur and other rendering features that rely on the former frame
-			// unless the view port is cinematic controlled
+			// unless the view port is Matinee controlled
 			ViewFamily.EngineShowFlags.CameraInterpolation = 0;
 		}
 
diff --git a/Engine/Source/Editor/UnrealEd/Private/LevelEditorViewport.cpp b/Engine/Source/Editor/UnrealEd/Private/LevelEditorViewport.cpp
--- a/Engine/Source/Editor/UnrealEd/Private/LevelEditorViewport.cpp
+++ b/Engine/Source/Editor/UnrealEd/Private/LevelEditorViewport.cpp
@@ -33,6 +33,7 @@
 #include "Editor.h"
 #include "EditorModeRegistry.h"
 #include "EditorModes.h"
+#include "EditorModeInterpolation.h"
 #include "PhysicsManipulationMode.h"
 #include "UnrealEdGlobals.h"
 #include "Materials/MaterialExpressionTextureSample.h"
@@ -97,7 +98,6 @@
 DEFINE_LOG_CATEGORY(LogEditorViewport);
 
 #define LOCTEXT_NAMESPACE "LevelEditorViewportClient"
-PRAGMA_DISABLE_OPTIMIZATION
 
 const FLevelViewportActorLock FLevelViewportActorLock::None(nullptr);
 
@@ -2034,6 +2034,13 @@ void FLevelEditorViewportClient::PerspectiveCameraMoved()
 		UpdateLockedActorViewports(GetActiveActorLock().Get(), false);
 	}
 
+	// Tell the editing mode that the camera moved, in case its interested.
+	FEdMode* Mode = ModeTools->GetActiveMode(FBuiltinEditorModes::EM_InterpEdit);
+	if( Mode )
+	{
+		((FEdModeInterpEdit*)Mode)->CamMoveNotify(this);
+	}
+
 	// Broadcast 'camera moved' delegate
 	FEditorDelegates::OnEditorCameraMoved.Broadcast(GetViewLocation(), GetViewRotation(), ViewportType, ViewIndex);
 }
@@ -2063,6 +2070,13 @@ void FLevelEditorViewportClient::ResetCamera()
 
 	ViewFOV = FOVAngle;
 
+	// If interp mode is active, tell it about the camera movement.
+	FEdMode* Mode = ModeTools->GetActiveMode(FBuiltinEditorModes::EM_InterpEdit);
+	if( Mode )
+	{
+		((FEdModeInterpEdit*)Mode)->CamMoveNotify(this);
+	}
+
 	SetIsCameraCut();
 
 	// Broadcast 'camera moved' delegate
@@ -2331,7 +2345,7 @@ void FLevelEditorViewportClient::Tick(float DeltaTime)
 
 void FLevelEditorViewportClient::UpdateViewForLockedActor(float DeltaTime)
 {
-	// We can't be locked to a cinematic actor if this viewport doesn't allow cinematic control
+	// We can't be locked to a matinee actor if this viewport doesn't allow matinee control
 	if (!bAllowCinematicControl && ActorLocks.CinematicActorLock.HasValidLockedActor())
 	{
 		ActorLocks.CinematicActorLock = FLevelViewportActorLock::None;
@@ -4961,4 +4975,3 @@ bool FLevelEditorViewportClient::GetPivotForOrbit(FVector& Pivot) const
 }
 
 #undef LOCTEXT_NAMESPACE
-PRAGMA_ENABLE_OPTIMIZATION
\ No newline at end of file
diff --git a/Engine/Source/Editor/UnrealEd/Private/UnrealEdEngine.cpp b/Engine/Source/Editor/UnrealEd/Private/UnrealEdEngine.cpp
--- a/Engine/Source/Editor/UnrealEd/Private/UnrealEdEngine.cpp
+++ b/Engine/Source/Editor/UnrealEd/Private/UnrealEdEngine.cpp
@@ -746,6 +746,13 @@ void UUnrealEdEngine::OnPostWindowsMessage(FViewport* Viewport, uint32 Message)
 	}		
 }
 
+
+void UUnrealEdEngine::OnOpenMatinee()
+{
+	// Register a delegate to pickup when Matinee is closed.
+	UpdateEdModeOnMatineeCloseDelegateHandle = GLevelEditorModeTools().OnEditorModeIDChanged().AddUObject( this, &UUnrealEdEngine::UpdateEdModeOnMatineeClose );
+}
+
 bool UUnrealEdEngine::IsAutosaving() const
 {
 	if (PackageAutoSaver)
diff --git a/Engine/Source/Editor/UnrealEd/Public/EditorModeInterpolation.h b/Engine/Source/Editor/UnrealEd/Public/EditorModeInterpolation.h
new file mode 100644
--- /dev/null
+++ b/Engine/Source/Editor/UnrealEd/Public/EditorModeInterpolation.h
@@ -0,0 +1,89 @@
+// Copyright Epic Games, Inc. All Rights Reserved.
+
+/*=============================================================================
+	EditorModeInterpolation : Editor mode for setting up interpolation sequences.
+
+=============================================================================*/
+
+#pragma once
+
+#include "CoreMinimal.h"
+#include "InputCoreTypes.h"
+#include "EdMode.h"
+#include "EditorModeTools.h"
+
+class FCanvas;
+class FEditorViewportClient;
+class FPrimitiveDrawInterface;
+class FSceneView;
+class FViewport;
+class UInterpGroup;
+
+//////////////////////////////////////////////////////////////////////////
+// FEdModeInterpEdit
+//////////////////////////////////////////////////////////////////////////
+
+class FEdModeInterpEdit : public FEdMode
+{
+public:
+	FEdModeInterpEdit();
+	~FEdModeInterpEdit();
+
+	virtual bool InputKey( FEditorViewportClient* ViewportClient, FViewport* Viewport, FKey Key, EInputEvent Event ) override;
+	virtual void Enter() override;
+	virtual void Exit() override;
+	virtual void ActorMoveNotify() override;
+	virtual void ActorPropChangeNotify() override;
+	virtual bool AllowWidgetMove() override;
+	virtual void ActorSelectionChangeNotify() override;
+	virtual bool IsCompatibleWith(FEditorModeID OtherModeID) const override;
+
+	virtual void Render(const FSceneView* View,FViewport* Viewport,FPrimitiveDrawInterface* PDI) override;
+	virtual void DrawHUD(FEditorViewportClient* ViewportClient,FViewport* Viewport,const FSceneView* View,FCanvas* Canvas) override;
+
+	void CamMoveNotify(FEditorViewportClient* ViewportClient);
+	void InitInterpMode(class AMatineeActor* InMatineeActor);
+	UNREALED_API void UpdateSelectedActor();
+
+	class AMatineeActor*	MatineeActor;
+	class IMatineeBase*		InterpEd;
+	bool					bLeavingMode;
+
+private:
+	// Grouping is always disabled while in InterpEdit Mode, re-enable the saved value on exit
+	bool					bGroupingActiveSaved;
+
+};
+
+//////////////////////////////////////////////////////////////////////////
+// FModeTool_InterpEdit
+//////////////////////////////////////////////////////////////////////////
+
+class FModeTool_InterpEdit : public FModeTool
+{
+public:
+	FModeTool_InterpEdit();
+	~FModeTool_InterpEdit();
+
+	virtual FString GetName() const override { return TEXT("Interp Edit"); }
+
+	/**
+	 * @return		true if the key was handled by this editor mode tool.
+	 */
+	virtual bool InputKey(FEditorViewportClient* ViewportClient, FViewport* Viewport, FKey Key, EInputEvent Event) override;
+	virtual bool InputAxis(FEditorViewportClient* InViewportClient, FViewport* Viewport, int32 ControllerId, FKey Key, float Delta, float DeltaTime) override;
+	virtual bool MouseMove(FEditorViewportClient* ViewportClient, FViewport* Viewport, int32 x, int32 y) override;
+
+	/**
+	 * @return		true if the delta was handled by this editor mode tool.
+	 */
+	virtual bool InputDelta(FEditorViewportClient* InViewportClient,FViewport* InViewport,FVector& InDrag,FRotator& InRot,FVector& InScale) override;
+	virtual void SelectNone() override;
+
+
+	bool bMovingHandle;
+	UInterpGroup* DragGroup;
+	int32 DragTrackIndex;
+	int32 DragKeyIndex;
+	bool bDragArriving;
+};
diff --git a/Engine/Source/Editor/UnrealEd/Public/EditorViewportClient.h b/Engine/Source/Editor/UnrealEd/Public/EditorViewportClient.h
--- a/Engine/Source/Editor/UnrealEd/Public/EditorViewportClient.h
+++ b/Engine/Source/Editor/UnrealEd/Public/EditorViewportClient.h
@@ -31,6 +31,7 @@ class FEditorViewportClient;
 class FEdMode;
 class FMouseDeltaTracker;
 class FPreviewScene;
+class IMatineeBase;
 class SEditorViewport;
 class UActorFactory;
 class UTypedElementViewportInteraction;
@@ -1067,14 +1068,15 @@ public:
 	/**
 	 * Allows custom disabling of camera recoil
 	 */
-	UE_DEPRECATED(5.0, "Matinee is no longer part of the editor.")
-	void SetMatineeRecordingWindow(class IMatineeBase* InInterpEd) {}
+	void SetMatineeRecordingWindow(IMatineeBase* InInterpEd);
 
 	/**
 	 * Returns true if camera recoil is currently allowed
 	 */
-	UE_DEPRECATED(5.0, "Matinee is no longer part of the editor.")
-	bool IsMatineeRecordingWindow() const { return false; }
+	bool IsMatineeRecordingWindow() const
+	{
+		return (RecordingInterpEd != NULL);
+	}
 	
 	EAxisList::Type GetCurrentWidgetAxis() const;
 
@@ -1647,6 +1649,9 @@ protected:
 	bool bShowWidget;
 
 	FMouseDeltaTracker*		MouseDeltaTracker;
+		
+	/**InterpEd, should only be set if used for matinee recording*/
+	IMatineeBase* RecordingInterpEd;
 
 	/** If true, the canvas has been been moved using bMoveCanvas Mode*/
 	bool bHasMouseMovedSinceClick;
diff --git a/Engine/Source/Editor/UnrealEd/UnrealEd.Build.cs b/Engine/Source/Editor/UnrealEd/UnrealEd.Build.cs
--- a/Engine/Source/Editor/UnrealEd/UnrealEd.Build.cs
+++ b/Engine/Source/Editor/UnrealEd/UnrealEd.Build.cs
@@ -214,6 +214,7 @@ public class UnrealEd : ModuleRules
 				"TextureEditor",
 				"Cascade",
 				"UMGEditor",
+				"Matinee",
 				"AssetTools",
 				"ClassViewer",
 				"StructViewer",
diff --git a/Engine/Source/Editor/WorldBrowser/Private/LevelCollectionModel.cpp b/Engine/Source/Editor/WorldBrowser/Private/LevelCollectionModel.cpp
--- a/Engine/Source/Editor/WorldBrowser/Private/LevelCollectionModel.cpp
+++ b/Engine/Source/Editor/WorldBrowser/Private/LevelCollectionModel.cpp
@@ -17,6 +17,7 @@
 #include "EditorModeManager.h"
 #include "EditorModes.h"
 #include "FileHelpers.h"
+#include "EditorModeInterpolation.h"
 #include "ScopedTransaction.h"
 #include "EditorLevelUtils.h"
 #include "LevelCollectionCommands.h"
@@ -25,6 +26,7 @@
 #include "IAssetTypeActions.h"
 #include "AssetToolsModule.h"
 #include "EditorSupportDelegates.h"
+#include "Matinee/MatineeActor.h"
 #include "GameFramework/WorldSettings.h"
 
 #include "ShaderCompiler.h"
@@ -695,7 +697,19 @@ void FLevelCollectionModel::UnloadLevels(const FLevelModelList& InLevelList)
 	UWorld* ThisWorld = GetWorld();
 	check(ThisWorld != nullptr);
 
-	if(GLevelEditorModeTools().IsModeActive(FBuiltinEditorModes::EM_Landscape))
+	// If matinee is opened, and if it belongs to the level being removed, close it
+	if (GLevelEditorModeTools().IsModeActive(FBuiltinEditorModes::EM_InterpEdit))
+	{
+		TArray<ULevel*> LevelsToRemove = GetLevelObjectList(InLevelList);
+		
+		const FEdModeInterpEdit* InterpEditMode = (const FEdModeInterpEdit*)GLevelEditorModeTools().GetActiveMode(FBuiltinEditorModes::EM_InterpEdit);
+
+		if (InterpEditMode && InterpEditMode->MatineeActor && LevelsToRemove.Contains(InterpEditMode->MatineeActor->GetLevel()))
+		{
+			GLevelEditorModeTools().ActivateDefaultMode();
+		}
+	}
+	else if(GLevelEditorModeTools().IsModeActive(FBuiltinEditorModes::EM_Landscape))
 	{
 		GLevelEditorModeTools().ActivateDefaultMode();
 	}
@@ -1746,6 +1760,34 @@ bool FLevelCollectionModel::IsValidFindInContentBrowser()
 
 void FLevelCollectionModel::MoveActorsToSelected_Executed()
 {
+	// If matinee is open, and if an actor being moved belongs to it, message the user
+	if (GLevelEditorModeTools().IsModeActive(FBuiltinEditorModes::EM_InterpEdit))
+	{
+		const FEdModeInterpEdit* InterpEditMode = (const FEdModeInterpEdit*)GLevelEditorModeTools().GetActiveMode(FBuiltinEditorModes::EM_InterpEdit);
+		if (InterpEditMode && InterpEditMode->MatineeActor)
+		{
+			TArray<AActor*> ControlledActors;
+			InterpEditMode->MatineeActor->GetControlledActors(ControlledActors);
+
+			// are any of the selected actors in the matinee
+			USelection* SelectedActors = GEditor->GetSelectedActors();
+			for (FSelectionIterator Iter(*SelectedActors); Iter; ++Iter)
+			{
+				AActor* Actor = CastChecked<AActor>(*Iter);
+				if (Actor != nullptr && (Actor == InterpEditMode->MatineeActor || ControlledActors.Contains(Actor)))
+				{
+					const bool ExitInterp = EAppReturnType::Yes == FMessageDialog::Open(EAppMsgType::YesNo, NSLOCTEXT("UnrealEd", "MatineeUnableToMove", "You must close Matinee before moving actors.\nDo you wish to do this now and continue?"));
+					if (!ExitInterp)
+					{
+						return;
+					}
+					GLevelEditorModeTools().DeactivateMode(FBuiltinEditorModes::EM_InterpEdit);
+					break;
+				}
+			}
+		}
+	}
+
 	MakeLevelCurrent_Executed();
 
 	const FScopedTransaction Transaction(LOCTEXT("MoveSelectedActorsToSelectedLevel", "Move Selected Actors to Level"));
