Allow newly created worlds to be re-used during map load

diff --git a/Engine/Source/Editor/UnrealEd/Private/EditorServer.cpp b/Engine/Source/Editor/UnrealEd/Private/EditorServer.cpp
--- a/Engine/Source/Editor/UnrealEd/Private/EditorServer.cpp
+++ b/Engine/Source/Editor/UnrealEd/Private/EditorServer.cpp
@@ -2553,23 +2553,29 @@ bool UEditorEngine::Map_Load(const TCHAR* Str, FOutputDevice& Ar)
 						FMessageLog("LoadErrors").NewPage( FText::Format( LOCTEXT("LoadMapLogPage", "Loading map: {MapFileName}"), Arguments ) );
 					}
 
+					// Only worlds that are uninitialized or reside within newly created packages (so not yet saved) may be considered valid for re-use.
+					// All other worlds need to be reloaded from disk, as a world is only initialized correctly as part of the level loading process.
+					auto IsWorldValidForReuse = [](UWorld* WorldToConsider)
+					{
+						return !WorldToConsider->bIsWorldInitialized || WorldToConsider->GetPackage()->HasAnyPackageFlags(PKG_NewlyCreated);
+					};
+
 					// If we are loading the same world again (reloading) then we must not specify that we want to keep this world in memory.
-					// Otherwise, try to keep the existing uninitialized world in memory since there is not reason to reload it.
+					// Otherwise, try to keep the existing world in memory since there is not reason to reload it.
 					UWorld* NewWorld = nullptr;
-					if (ExistingWorld && !ExistingWorld->bIsWorldInitialized && Context.World() != ExistingWorld && !bIsLoadingMapTemplate)
+					if (ExistingWorld && IsWorldValidForReuse(ExistingWorld) && Context.World() != ExistingWorld && !bIsLoadingMapTemplate)
 					{
 						NewWorld = ExistingWorld;
 					}
 					EditorDestroyWorld( Context, LocalizedLoadingMap, NewWorld );
 
-					// Unload all other map packages currently loaded and initialized, before opening a new map.
-					// The world is only initialized correctly as part of the level loading process, so ensure that every map package needs loading.
+					// Unload all other map packages before opening a new map.
 					{
 						TArray<UPackage*> WorldPackages;
 						for (TObjectIterator<UWorld> It; It; ++It)
 						{
 							UPackage* Package = Cast<UPackage>(It->GetOuter());
-							if (Package && Package != GetTransientPackage() && (It->bIsWorldInitialized || Package->GetPathName() != LongTempFname))
+							if (Package && Package != GetTransientPackage() && (!IsWorldValidForReuse(*It) || Package->GetPathName() != LongTempFname))
 							{
 								WorldPackages.AddUnique(Package);
 							}
@@ -2598,9 +2604,9 @@ bool UEditorEngine::Map_Load(const TCHAR* Str, FOutputDevice& Ar)
 						ExistingWorld = nullptr;
 					}
 
-					// If the existing world has already been initialized then the unload attempts above failed, and we need to 
+					// If the existing world is not valid for re-use then the unload attempts above failed, and we need to 
 					// fatally error and dump any lingering references (like we would when unloading the main editor world).
-					if ((ExistingWorld && ExistingWorld->bIsWorldInitialized) || (ExistingPackage && !ExistingWorld))
+					if ((ExistingWorld && !IsWorldValidForReuse(ExistingWorld)) || (ExistingPackage && !ExistingWorld))
 					{
 						int32 NumFailedToCleanup = 0;
 
