You can now avoid building UAT\UBT every time when Launching application from the Editor via Editor Preferences -> Miscalleneous -> Always Build UAT. This reduces iteration time when there has been minor or no changes done.

diff --git a/Engine/Source/Editor/UnrealEd/Classes/Editor/EditorEngine.h b/Engine/Source/Editor/UnrealEd/Classes/Editor/EditorEngine.h
--- a/Engine/Source/Editor/UnrealEd/Classes/Editor/EditorEngine.h
+++ b/Engine/Source/Editor/UnrealEd/Classes/Editor/EditorEngine.h
@@ -2945,6 +2945,9 @@ protected:
 	/** Start a Launcher session with the given parameters. Called by StartQueuedPlaySessionRequestImpl based on request settings. */
 	virtual void StartPlayUsingLauncherSession(FRequestPlaySessionParams& InRequestParams);
 
+	/** This flag is used to skip UAT\UBT compilation on every launch if it was successfully compiled once. */
+	bool bUATSuccessfullyCompiledOnce;
+
 	// This chunk is for Play in Editor
 public:
 	/** @return true if the editor is able to launch PIE with online platform support */
diff --git a/Engine/Source/Editor/UnrealEd/Classes/Editor/EditorPerProjectUserSettings.h b/Engine/Source/Editor/UnrealEd/Classes/Editor/EditorPerProjectUserSettings.h
--- a/Engine/Source/Editor/UnrealEd/Classes/Editor/EditorPerProjectUserSettings.h
+++ b/Engine/Source/Editor/UnrealEd/Classes/Editor/EditorPerProjectUserSettings.h
@@ -113,6 +113,10 @@ class UEditorPerProjectUserSettings : public UObject
 	UPROPERTY(EditAnywhere, config, Category = UnrealAutomationTool)
 	bool bGetAttentionOnUATCompletion;
 
+	/** Always build UAT\UBT before launching the game. It will decrease iteration times if disabled */
+	UPROPERTY(EditAnywhere, config, Category = UnrealAutomationTool)
+	bool bAlwaysBuildUAT = true;
+
 	/** How fast the SCS viewport camera moves */
 	UPROPERTY(config, meta=(UIMin = "1", UIMax = "8", ClampMin="1", ClampMax="8"))
 	int32 SCSViewportCameraSpeed;
diff --git a/Engine/Source/Editor/UnrealEd/Private/EditorEngine.cpp b/Engine/Source/Editor/UnrealEd/Private/EditorEngine.cpp
--- a/Engine/Source/Editor/UnrealEd/Private/EditorEngine.cpp
+++ b/Engine/Source/Editor/UnrealEd/Private/EditorEngine.cpp
@@ -393,6 +393,8 @@ UEditorEngine::UEditorEngine(const FObjectInitializer& ObjectInitializer)
 	EditorWorldExtensionsManager = nullptr;
 
 	ActorGroupingUtilsClassName = UActorGroupingUtils::StaticClass();
+
+	bUATSuccessfullyCompiledOnce = FApp::IsEngineInstalled() || FApp::GetEngineIsPromotedBuild();
 }
 
 
diff --git a/Engine/Source/Editor/UnrealEd/Private/PlayLevelViaLauncher.cpp b/Engine/Source/Editor/UnrealEd/Private/PlayLevelViaLauncher.cpp
--- a/Engine/Source/Editor/UnrealEd/Private/PlayLevelViaLauncher.cpp
+++ b/Engine/Source/Editor/UnrealEd/Private/PlayLevelViaLauncher.cpp
@@ -21,6 +21,7 @@
 #include "DesktopPlatformModule.h"
 #include "PlatformInfo.h"
 #include "Framework/Docking/TabManager.h"
+#include "Editor/EditorPerProjectUserSettings.h"
 
 #define LOCTEXT_NAMESPACE "PlayLevel"
 
@@ -268,6 +269,12 @@ void UEditorEngine::StartPlayUsingLauncherSession(FRequestPlaySessionParams& InR
 	LauncherProfile->SetIncrementalDeploying(bIncrimentalCooking);
 	LauncherProfile->SetEditorExe(FUnrealEdMisc::Get().GetExecutableForCommandlets());
 	LauncherProfile->SetShouldUpdateDeviceFlash(InRequestParams.LauncherTargetDevice->bUpdateDeviceFlash);
+	
+	if (LauncherProfile->IsBuildingUAT() && !GetDefault<UEditorPerProjectUserSettings>()->bAlwaysBuildUAT && bUATSuccessfullyCompiledOnce)
+	{
+		// UAT was built on a first launch and there's no need to rebuild it any more
+		LauncherProfile->SetBuildUAT(false);
+	}
 
 	const FString DummyIOSDeviceName(FString::Printf(TEXT("All_iOS_On_%s"), FPlatformProcess::ComputerName()));
 	const FString DummyTVOSDeviceName(FString::Printf(TEXT("All_tvOS_On_%s"), FPlatformProcess::ComputerName()));
@@ -636,6 +643,8 @@ void UEditorEngine::HandleLaunchCompleted(bool Succeeded, double TotalTime, int3
 		FEditorAnalytics::ReportEvent(TEXT( "Editor.LaunchOn.Completed" ), LastPlayUsingLauncherDeviceId.Left(LastPlayUsingLauncherDeviceId.Find(TEXT("@"))), bHasCode, ParamArray);
 
 		UE_LOG(LogPlayLevel, Log, TEXT("Launch On Completed. Time: %f"), TotalTime);
+
+		bUATSuccessfullyCompiledOnce = true;
 	}
 	else
 	{
