Revert the removal of deselecting an actor when ShouldExport returns false.

diff --git a/Engine/Source/Editor/UnrealEd/Private/EditorExporters.cpp b/Engine/Source/Editor/UnrealEd/Private/EditorExporters.cpp
--- a/Engine/Source/Editor/UnrealEd/Private/EditorExporters.cpp
+++ b/Engine/Source/Editor/UnrealEd/Private/EditorExporters.cpp
@@ -531,6 +531,10 @@ bool ULevelExporterT3D::ExportText( const FExportObjectInnerContext* Context, UO
 				// Restore dynamic delegate bindings.
 				UBlueprintGeneratedClass::BindDynamicDelegates(Actor->GetClass(), Actor);
 			}
+			else if (GEditor)
+			{
+				GEditor->GetSelectedActors()->Deselect(Actor);
+			}
 		}
 	}
 
