AsyncIODelete: Initialize TempRoot path even when async is disabled, to implement contract that it does not allowing temproot to be deleted.
#rnx

diff --git a/Engine/Source/Editor/UnrealEd/Private/Cooker/AsyncIODelete.cpp b/Engine/Source/Editor/UnrealEd/Private/Cooker/AsyncIODelete.cpp
--- a/Engine/Source/Editor/UnrealEd/Private/Cooker/AsyncIODelete.cpp
+++ b/Engine/Source/Editor/UnrealEd/Private/Cooker/AsyncIODelete.cpp
@@ -73,36 +73,39 @@ void FAsyncIODelete::SetDeletesPaused(bool bInPaused)
 	}
 }
 
-bool FAsyncIODelete::Setup()
+void FAsyncIODelete::Setup()
 {
 	if (bInitialized)
 	{
-		return true;
+		return;
 	}
 
+	bInitialized = true;
+	TempRoot = RequestedTempRoot;
 	if (RequestedTempRoot.IsEmpty())
 	{
 		checkf(false, TEXT("DeleteDirectory called without having first set a TempRoot"));
-		return false;
+		return;
 	}
 
 	if (AsyncEnabled())
 	{
 		// Delete the TempRoot directory to clear the results from any previous process using the same TempRoot that did not shut down cleanly
-		if (!TryPurgeTempRootFamily(&TempRoot))
+		FString FoundRoot;
+		if (!TryPurgeTempRootFamily(&FoundRoot))
 		{
 			// TryPurgeTempRootFamily logged the warning
-			return false;
+			return;
 		}
 
 		// Create the empty directory to work in
-		if (!IFileManager::Get().MakeDirectory(*TempRoot, true))
+		if (!IFileManager::Get().MakeDirectory(*FoundRoot, true))
 		{
 			UE_LOG(LogCook, Error, TEXT("Could not create asyncdelete root directory '%s'. LastError: %i. Falling back to synchronous delete."),
 				*TempRoot, FPlatformMisc::GetLastError());
-			TempRoot.Reset();
-			return false;
+			return;
 		}
+		TempRoot = FoundRoot;
 
 		// Allocate the task event
 		check(TasksComplete == nullptr);
@@ -114,12 +117,8 @@ bool FAsyncIODelete::Setup()
 		// TempRoot and bPaused are preserved across setup/teardown and may have any value
 		check(PausedDeletes.Num() == 0);
 		check(DeleteCounter == 0);
+		bAsyncInitialized = true;
 	}
-
-	// We are now setup and ready to create DeleteTasks
-	bInitialized = true;
-
-	return true;
 }
 
 void FAsyncIODelete::Teardown()
@@ -129,7 +128,7 @@ void FAsyncIODelete::Teardown()
 		return;
 	}
 
-	if (AsyncEnabled())
+	if (bAsyncInitialized)
 	{
 		// Clear task variables
 		WaitForAllTasks();
@@ -139,12 +138,13 @@ void FAsyncIODelete::Teardown()
 
 		// Remove the temp directory from disk
 		TryPurgeTempRootFamily(nullptr);
-		TempRoot.Reset();
 
 		// Clear delete variables; we don't need to run the tasks for the remaining pauseddeletes because synchronously deleting the temp directory above did the work they were going to do
 		PausedDeletes.Empty();
 		DeleteCounter = 0;
+		bAsyncInitialized = false;
 	}
+	TempRoot.Reset();
 
 	// We are now torn down and ready for a new setup
 	bInitialized = false;
@@ -152,38 +152,33 @@ void FAsyncIODelete::Teardown()
 
 bool FAsyncIODelete::WaitForAllTasks(float TimeLimitSeconds)
 {
-	if (AsyncEnabled())
+	if (!bAsyncInitialized)
 	{
-		if (!bInitialized)
-		{
-			return true;
-		}
+		return true;
+	}
 
-		if (TimeLimitSeconds <= 0.f)
-		{
-			TasksComplete->Wait();
-		}
-		else
+	if (TimeLimitSeconds <= 0.f)
+	{
+		TasksComplete->Wait();
+	}
+	else
+	{
+		if (!TasksComplete->Wait(FTimespan::FromSeconds(TimeLimitSeconds)))
 		{
-			if (!TasksComplete->Wait(FTimespan::FromSeconds(TimeLimitSeconds)))
-			{
-				return false;
-			}
+			return false;
 		}
-		check(ActiveTaskCount == 0);
 	}
+	check(ActiveTaskCount == 0);
 	return true;
 }
 
 bool FAsyncIODelete::Delete(const FStringView& PathToDelete, EPathType ExpectedType)
 {
 	IFileManager& FileManager = IFileManager::Get();
-	TStringBuilder<128> PathToDeleteBuffer;
-	PathToDeleteBuffer << PathToDelete;
-	const TCHAR* PathToDeleteSZ = PathToDeleteBuffer.ToString();
+	FString PathToDeleteStr(PathToDelete);
 
-	const bool IsDirectory = FileManager.DirectoryExists(PathToDeleteSZ);
-	const bool IsFile = !IsDirectory && FileManager.FileExists(PathToDeleteSZ);
+	const bool IsDirectory = FileManager.DirectoryExists(*PathToDeleteStr);
+	const bool IsFile = !IsDirectory && FileManager.FileExists(*PathToDeleteStr);
 	if (!IsDirectory && !IsFile)
 	{
 		return true;
@@ -198,28 +193,21 @@ bool FAsyncIODelete::Delete(const FStringView& PathToDelete, EPathType ExpectedT
 		checkf(false, TEXT("DeleteFile called on \"%.*s\" which is not a file."), PathToDelete.Len(), PathToDelete.GetData());
 		return false;
 	}
-	// Prevent the user from trying to delete our temproot or anything inside it
-	FString PathToDeleteStr(PathToDelete);
-	if (FPaths::IsUnderDirectory(PathToDeleteStr, TempRoot) || FPaths::IsUnderDirectory(TempRoot, PathToDeleteStr))
-	{
-		return false;
-	}
 
-	if (AsyncEnabled())
+	if (bAsyncInitialized)
 	{
 		if (DeleteCounter == UINT32_MAX)
 		{
 			Teardown();
 		}
 	}
-	if (!Setup())
+	Setup();
+	// Prevent the user from trying to delete our temproot or anything inside it
+	if (FPaths::IsUnderDirectory(PathToDeleteStr, TempRoot) || FPaths::IsUnderDirectory(TempRoot, PathToDeleteStr))
 	{
-		// Setup failed; we are not able to provide asynchronous deletes; fall back to synchronous
-		UE_LOG(LogCook, Warning, TEXT("Failed to setup an async delete, falling back to synchronous delete."));
-		return SynchronousDelete(PathToDeleteSZ, ExpectedType);
+		return false;
 	}
-
-	if (AsyncEnabled())
+	if (bAsyncInitialized)
 	{
 		const FString TempPath = FPaths::Combine(TempRoot, FString::Printf(TEXT("%u"), DeleteCounter));
 		DeleteCounter++;
@@ -228,11 +216,11 @@ bool FAsyncIODelete::Delete(const FStringView& PathToDelete, EPathType ExpectedT
 		const bool bEvenIfReadOnly = true;
 		const bool bMoveAttributes = false;
 		const bool bDoNotRetryOnError = true;
-		if (!IFileManager::Get().Move(*TempPath, PathToDeleteSZ, bReplace, bEvenIfReadOnly, bMoveAttributes, bDoNotRetryOnError)) // IFileManager::Move works on either files or directories
+		if (!IFileManager::Get().Move(*TempPath, *PathToDeleteStr, bReplace, bEvenIfReadOnly, bMoveAttributes, bDoNotRetryOnError)) // IFileManager::Move works on either files or directories
 		{
 			// The move failed; try a synchronous delete as backup
 			UE_LOG(LogCook, Warning, TEXT("Failed to move path '%.*s' for async delete (LastError == %i); falling back to synchronous delete."), PathToDelete.Len(), PathToDelete.GetData(), FPlatformMisc::GetLastError());
-			return SynchronousDelete(PathToDeleteSZ, ExpectedType);
+			return SynchronousDelete(*PathToDeleteStr, ExpectedType);
 		}
 
 		if (bPaused)
@@ -247,7 +235,7 @@ bool FAsyncIODelete::Delete(const FStringView& PathToDelete, EPathType ExpectedT
 	}
 	else
 	{
-		return SynchronousDelete(PathToDeleteSZ, ExpectedType);
+		return SynchronousDelete(*PathToDeleteStr, ExpectedType);
 	}
 }
 
diff --git a/Engine/Source/Editor/UnrealEd/Private/Cooker/AsyncIODelete.h b/Engine/Source/Editor/UnrealEd/Private/Cooker/AsyncIODelete.h
--- a/Engine/Source/Editor/UnrealEd/Private/Cooker/AsyncIODelete.h
+++ b/Engine/Source/Editor/UnrealEd/Private/Cooker/AsyncIODelete.h
@@ -61,9 +61,8 @@ public:
 	/**
 	 * Prepare the directory on disk.  Called OnDemand from the other interface functions.  Can also be called if desired by client code.
 	 * May take a long time to run if cleanup from a crashed previous process is required.
-	 * @return False if the TempRoot directory could not be constructed; calls to DeleteDirectory will fail in this case
 	 */
-	bool Setup();
+	void Setup();
 
 	/** Synchronously wait for all tasks to complete, and remove the temproot. */
 	void Teardown();
@@ -127,6 +126,7 @@ private:
 	uint32 ActiveTaskCount = 0;
 	uint32 DeleteCounter = 0;
 	bool bInitialized = false;
+	bool bAsyncInitialized = false;
 	bool bPaused = false;
 
 #if WITH_ASYNCIODELETE_DEBUG
