Clear CopyPasteId on new actors (Fix failed check)

diff --git a/Engine/Source/Editor/UnrealEd/Private/EditorLevelUtils.cpp b/Engine/Source/Editor/UnrealEd/Private/EditorLevelUtils.cpp
--- a/Engine/Source/Editor/UnrealEd/Private/EditorLevelUtils.cpp
+++ b/Engine/Source/Editor/UnrealEd/Private/EditorLevelUtils.cpp
@@ -207,6 +207,8 @@ int32 UEditorLevelUtils::CopyOrMoveActorsToLevel(const TArray<AActor*>& ActorsTo
 					{
 						UE_LOG(LogLevelTools, Error, TEXT("Cannot find remapping for moved actor ID %s, any soft references pointing to it will be broken!"), *Actor->GetPathName());
 					}
+					// Reset CopyPasteId on new actors
+					Actor->CopyPasteId = INDEX_NONE;
 				}
 
 				FAssetToolsModule& AssetToolsModule = FModuleManager::GetModuleChecked<FAssetToolsModule>(TEXT("AssetTools"));
@@ -245,6 +247,7 @@ int32 UEditorLevelUtils::CopyOrMoveActorsToLevel(const TArray<AActor*>& ActorsTo
 			for (AActor* Actor : FinalMoveList)
 			{
 				check(Actor->CopyPasteId != INDEX_NONE);
+				// Reset CopyPasteId on source actors 
 				Actor->CopyPasteId = INDEX_NONE;
 			}
 		}
