EditorActorSubsystem: Added functions to duplicate actors regardless of selection status

diff --git a/Engine/Source/Editor/UnrealEd/Private/Subsystems/EditorActorSubsystem.cpp b/Engine/Source/Editor/UnrealEd/Private/Subsystems/EditorActorSubsystem.cpp
--- a/Engine/Source/Editor/UnrealEd/Private/Subsystems/EditorActorSubsystem.cpp
+++ b/Engine/Source/Editor/UnrealEd/Private/Subsystems/EditorActorSubsystem.cpp
@@ -612,8 +612,53 @@ TArray<class AActor*> UEditorActorSubsystem::ConvertActors(const TArray<class AA
 		}
 	}
 
-	UE_LOG(LogUtils, Log, TEXT("ConvertActorWith. %d convertions occurred."), Result.Num());
+	UE_LOG(LogUtils, Log, TEXT("ConvertActorWith. %d conversions occurred."), Result.Num());
 	return Result;
 }
 
+AActor* UEditorActorSubsystem::DuplicateActor(AActor* ActorToDuplicate, UWorld* ToWorld/*= nullptr*/, FVector Offset/* = FVector::ZeroVector*/)
+{
+	TArray<AActor*> Duplicate = DuplicateActors({ ActorToDuplicate }, ToWorld, Offset);
+	return (Duplicate.Num() > 0) ? Duplicate[0] : nullptr;
+}
+
+TArray<AActor*> UEditorActorSubsystem::DuplicateActors(const TArray<AActor*>& ActorsToDuplicate, UWorld* InToWorld/*= nullptr*/, FVector Offset/* = FVector::ZeroVector*/)
+{
+	TGuardValue<bool> UnattendedScriptGuard(GIsRunningUnattendedScript, true);
+
+	FScopedTransaction Transaction(LOCTEXT("DuplicateActors", "Duplicate Actors"));
+
+	UUnrealEditorSubsystem* UnrealEditorSubsystem = GEditor->GetEditorSubsystem<UUnrealEditorSubsystem>();
+
+	if (!UnrealEditorSubsystem || !EditorScriptingHelpers::CheckIfInEditorAndPIE())
+	{
+		return TArray<AActor*>();
+	}
+
+	UWorld* ToWorld = InToWorld ? InToWorld : UnrealEditorSubsystem->GetEditorWorld();
+
+	if (!ToWorld)
+	{
+		return TArray<AActor*>();
+	}
+
+	FEditorDelegates::OnDuplicateActorsBegin.Broadcast();
+
+	TArray<AActor*> NewActors;
+	ABrush::SetSuppressBSPRegeneration(true);
+	GUnrealEd->DuplicateActors(ActorsToDuplicate, NewActors, ToWorld->GetCurrentLevel(), Offset);
+	ABrush::SetSuppressBSPRegeneration(false);
+
+	// Find out if any of the actors will change the BSP.
+	// and only then rebuild BSP as this is expensive. 
+	if (NewActors.FindItemByClass<ABrush>())
+	{
+		GEditor->RebuildAlteredBSP(); // Update the BSP of any levels containing a modified brush
+	}
+
+	FEditorDelegates::OnDuplicateActorsEnd.Broadcast();
+
+	return NewActors;
+}
+
 #undef LOCTEXT_NAMESPACE
diff --git a/Engine/Source/Editor/UnrealEd/Public/Subsystems/EditorActorSubsystem.h b/Engine/Source/Editor/UnrealEd/Public/Subsystems/EditorActorSubsystem.h
--- a/Engine/Source/Editor/UnrealEd/Public/Subsystems/EditorActorSubsystem.h
+++ b/Engine/Source/Editor/UnrealEd/Public/Subsystems/EditorActorSubsystem.h
@@ -74,6 +74,24 @@ public:
 	UPROPERTY(BlueprintAssignable, Category = "Editor Scripting | Level Utility")
 	FOnDeleteActorsEnd OnDeleteActorsEnd;
 
+	/**
+	 * Duplicate an actor from the world editor.
+	 * @param	ActorToDuplicate	Actor to duplicate.
+	 * @param	ToWorld				World to place the duplicated actor in.
+	 * @return	The duplicated actor, or none if it didn't succeed
+	 */
+	UFUNCTION(BlueprintCallable, Category = "Editor Scripting | Level Utility", meta = (AdvancedDisplay = 1))
+	AActor* DuplicateActor(AActor* ActorToDuplicate, UWorld* ToWorld = nullptr, FVector Offset = FVector::ZeroVector);
+
+	/**
+	 * Duplicate actors from the world editor.
+	 * @param	ActorsToDuplicate	Actors to duplicate.
+	 * @param	ToWorld				World to place the duplicated actors in.
+	 * @return	The duplicated actors, or empty if it didn't succeed
+	 */
+	UFUNCTION(BlueprintCallable, Category = "Editor Scripting | Level Utility", meta = (AdvancedDisplay = 1))
+	TArray<AActor*> DuplicateActors(const TArray<AActor*>& ActorsToDuplicate, UWorld* ToWorld = nullptr, FVector Offset = FVector::ZeroVector);
+
 	/**
 	 * Duplicate all the selected actors in the given world
 	 * @param	InWorld 	The world the actors are selected in.
