World Partition HLODs
* HLOD Builder - Submit step: Ensure all files modified by the source control helper are taken into account
* Make sure HLOD actors are always visible when running from a commandlet
* Skip creation of VT editor resource when running from cmdlet

diff --git a/Engine/Source/Editor/UnrealEd/Private/WorldPartition/WorldPartitionHLODsBuilder.cpp b/Engine/Source/Editor/UnrealEd/Private/WorldPartition/WorldPartitionHLODsBuilder.cpp
--- a/Engine/Source/Editor/UnrealEd/Private/WorldPartition/WorldPartitionHLODsBuilder.cpp
+++ b/Engine/Source/Editor/UnrealEd/Private/WorldPartition/WorldPartitionHLODsBuilder.cpp
@@ -536,7 +536,9 @@ bool UWorldPartitionHLODsBuilder::SubmitHLODActors()
 	bool bRet = true;
 
 	// Ensure all files modified by the source control helper are taken into account
-	TArray<FString> FilesToSubmit = SourceControlHelper->GetModifiedFiles().GetAllFiles();
+	ModifiedFiles.Append(SourceControlHelper->GetModifiedFiles());
+
+	TArray<FString> FilesToSubmit = ModifiedFiles.GetAllFiles();
 
 	// Check in all modified files
 	if (FilesToSubmit.Num() > 0)
diff --git a/Engine/Source/Runtime/Engine/Private/Texture2D.cpp b/Engine/Source/Runtime/Engine/Private/Texture2D.cpp
--- a/Engine/Source/Runtime/Engine/Private/Texture2D.cpp
+++ b/Engine/Source/Runtime/Engine/Private/Texture2D.cpp
@@ -1344,7 +1344,7 @@ void FVirtualTexture2DResource::InitializeEditorResources(IVirtualTexture* InVir
 {
 	// Create a texture resource from the lowest resolution VT page data
 	// this will then be used during asset tumbnails/hitproxies/...
-	if (GIsEditor)
+	if (GIsEditor && !IsRunningCommandlet())
 	{
 		struct FPageToProduce
 		{
diff --git a/Engine/Source/Runtime/Engine/Private/WorldPartition/HLOD/HLODActor.cpp b/Engine/Source/Runtime/Engine/Private/WorldPartition/HLOD/HLODActor.cpp
--- a/Engine/Source/Runtime/Engine/Private/WorldPartition/HLOD/HLODActor.cpp
+++ b/Engine/Source/Runtime/Engine/Private/WorldPartition/HLOD/HLODActor.cpp
@@ -73,7 +73,7 @@ void AWorldPartitionHLOD::RerunConstructionScripts()
 
 bool AWorldPartitionHLOD::IsHiddenEd() const
 {
-	return true;
+	return !IsRunningCommandlet();
 }
 
 EActorGridPlacement AWorldPartitionHLOD::GetGridPlacement() const
