GitHub #8141 : Fix tooltip/watch text for by-ref function inputs

diff --git a/Engine/Source/Editor/UnrealEd/Private/Kismet2/KismetDebugUtilities.cpp b/Engine/Source/Editor/UnrealEd/Private/Kismet2/KismetDebugUtilities.cpp
--- a/Engine/Source/Editor/UnrealEd/Private/Kismet2/KismetDebugUtilities.cpp
+++ b/Engine/Source/Editor/UnrealEd/Private/Kismet2/KismetDebugUtilities.cpp
@@ -1166,15 +1166,14 @@ FKismetDebugUtilities::EWatchTextResult FKismetDebugUtilities::GetWatchText(FStr
 	void* DataPtr = nullptr;
 	void* DeltaPtr = nullptr;
 	UObject* ParentObj = nullptr;
-	bool bShouldUseContainerOffset = false;
 	TArray<UObject*> SeenObjects;
-	FKismetDebugUtilities::EWatchTextResult Result = FindDebuggingData(Blueprint, ActiveObject, WatchPin, PropertyToDebug, DataPtr, DeltaPtr, ParentObj, SeenObjects, &bShouldUseContainerOffset);
+	bool bIsDirectPtr = false;
+	FKismetDebugUtilities::EWatchTextResult Result = FindDebuggingData(Blueprint, ActiveObject, WatchPin, PropertyToDebug, DataPtr, DeltaPtr, ParentObj, SeenObjects, &bIsDirectPtr);
 
 	if (Result == FKismetDebugUtilities::EWatchTextResult::EWTR_Valid)
 	{
-		// If this came from an array property we need to avoid using ExportText_InContainer in order to properly 
-		// calculate the internal offset
-		if(bShouldUseContainerOffset)
+		// If this came from an out parameter it isn't in a property container, so must be accessed directly
+		if (bIsDirectPtr)
 		{
 			PropertyToDebug->ExportText_Direct(/*inout*/ OutWatchText, DataPtr, DeltaPtr, ParentObj, PPF_PropertyWindow | PPF_BlueprintDebugView);
 		}
@@ -1194,17 +1193,26 @@ FKismetDebugUtilities::EWatchTextResult FKismetDebugUtilities::GetDebugInfo(FDeb
 	void* DeltaPtr = nullptr;
 	UObject* ParentObj = nullptr;
 	TArray<UObject*> SeenObjects;
-	FKismetDebugUtilities::EWatchTextResult Result = FindDebuggingData(Blueprint, ActiveObject, WatchPin, PropertyToDebug, DataPtr, DeltaPtr, ParentObj, SeenObjects);
+	bool bIsDirectPtr = false;
+	FKismetDebugUtilities::EWatchTextResult Result = FindDebuggingData(Blueprint, ActiveObject, WatchPin, PropertyToDebug, DataPtr, DeltaPtr, ParentObj, SeenObjects, &bIsDirectPtr);
 
 	if (Result == FKismetDebugUtilities::EWatchTextResult::EWTR_Valid)
 	{
-		GetDebugInfo_InContainer(0, OutDebugInfo, PropertyToDebug, DataPtr);
+		// If this came from an out parameter it isn't in a property container, so must be accessed directly
+		if (bIsDirectPtr)
+		{
+			GetDebugInfoInternal(OutDebugInfo, PropertyToDebug, DataPtr);
+		}
+		else
+		{
+			GetDebugInfo_InContainer(0, OutDebugInfo, PropertyToDebug, DataPtr);
+		}
 	}
 
 	return Result;
 }
 
-FKismetDebugUtilities::EWatchTextResult FKismetDebugUtilities::FindDebuggingData(UBlueprint* Blueprint, UObject* ActiveObject, const UEdGraphPin* WatchPin, FProperty*& OutProperty, void*& OutData, void*& OutDelta, UObject*& OutParent, TArray<UObject*>& SeenObjects, bool* OutbShouldUseContainerOffset /* = nullptr */)
+FKismetDebugUtilities::EWatchTextResult FKismetDebugUtilities::FindDebuggingData(UBlueprint* Blueprint, UObject* ActiveObject, const UEdGraphPin* WatchPin, FProperty*& OutProperty, void*& OutData, void*& OutDelta, UObject*& OutParent, TArray<UObject*>& SeenObjects, bool* bOutIsDirectPtr /* = nullptr */)
 {
 	FKismetDebugUtilitiesData& Data = FKismetDebugUtilitiesData::Get();
 
@@ -1254,7 +1262,7 @@ FKismetDebugUtilities::EWatchTextResult FKismetDebugUtilities::FindDebuggingData
 								// otherwise the output param won't show any data since the return node hasn't executed when we stop here
 								if (WatchPin->LinkedTo.Num() == 1)
 								{
-									return FindDebuggingData(Blueprint, ActiveObject, WatchPin->LinkedTo[0], OutProperty, OutData, OutDelta, OutParent, SeenObjects, OutbShouldUseContainerOffset);
+									return FindDebuggingData(Blueprint, ActiveObject, WatchPin->LinkedTo[0], OutProperty, OutData, OutDelta, OutParent, SeenObjects, bOutIsDirectPtr);
 								}
 								else if (!WatchPin->LinkedTo.Num())
 								{
@@ -1264,13 +1272,11 @@ FKismetDebugUtilities::EWatchTextResult FKismetDebugUtilities::FindDebuggingData
 								}
 							}
 
-							// If this is an out container property then a different offset must be used when exporting this property 
-							// to text. Only container properties are effected by this because ExportText_InContainer adds an extra 
-							// 8 byte offset, which  would point to the container's first element, not the container itself. 
-							const bool bIsContainer = OutParmRec->Property->IsA<FArrayProperty>() || OutParmRec->Property->IsA<FSetProperty>() || OutParmRec->Property->IsA<FMapProperty>();
-							if (PropertyBase == nullptr && OutbShouldUseContainerOffset && bIsContainer)
+							if (PropertyBase == nullptr && bOutIsDirectPtr)
 							{
-								*OutbShouldUseContainerOffset = true;
+								// Flag to caller that PropertyBase points directly at the data, not at the
+								// base of a property container (so don't apply property's Offset_Internal)
+								*bOutIsDirectPtr = true;
 								PropertyBase = OutParmRec->PropAddr;
 							}
 							break;
diff --git a/Engine/Source/Editor/UnrealEd/Public/Kismet2/KismetDebugUtilities.h b/Engine/Source/Editor/UnrealEd/Public/Kismet2/KismetDebugUtilities.h
--- a/Engine/Source/Editor/UnrealEd/Public/Kismet2/KismetDebugUtilities.h
+++ b/Engine/Source/Editor/UnrealEd/Public/Kismet2/KismetDebugUtilities.h
@@ -263,9 +263,10 @@ protected:
 	* @param 	OutDelta		Populated with the same thing as OutData
 	* @param 	OutParent		Populated with the active object
 	* @param 	SeenObjects		Used to track what objects have been traversed to find the OutProperty address
+	* @param 	bOutIsDirectPtr	True if OutData/OutDelta point directly at the property rather than its base address
 	* @return	EWTR_Valid if the debug data could be found, otherwise an appropriate error code
 	*/
-	static EWatchTextResult FindDebuggingData(UBlueprint* Blueprint, UObject* ActiveObject, const UEdGraphPin* WatchPin, FProperty*& OutProperty, void*& OutData, void*& OutDelta, UObject*& OutParent, TArray<UObject*>& SeenObjects, bool* OutbShouldUseContainerOffset = nullptr);
+	static EWatchTextResult FindDebuggingData(UBlueprint* Blueprint, UObject* ActiveObject, const UEdGraphPin* WatchPin, FProperty*& OutProperty, void*& OutData, void*& OutDelta, UObject*& OutParent, TArray<UObject*>& SeenObjects, bool* bOutIsDirectPtr = nullptr);
 
 private:
 	FKismetDebugUtilities() {}
