WorldPartition: Fix new saved maps not having any registered actor descs until they are reopened (also fixes a Level Instance crash related to newly saved wp maps)

diff --git a/Engine/Source/Editor/UnrealEd/Private/FileHelpers.cpp b/Engine/Source/Editor/UnrealEd/Private/FileHelpers.cpp
--- a/Engine/Source/Editor/UnrealEd/Private/FileHelpers.cpp
+++ b/Engine/Source/Editor/UnrealEd/Private/FileHelpers.cpp
@@ -736,6 +736,12 @@ static bool SaveWorld(UWorld* World,
 						World->Rename(*NewWorldAssetName, NULL, REN_NonTransactional | REN_DontCreateRedirectors | REN_ForceNoResetLoaders);
 					}
 
+					// SaveAs on package that didn't exist (!DuplicatedWorld). Needs to be rebased because Package was renamed.		
+					if (UWorldPartition* WorldPartition = World->GetWorldPartition())
+					{
+						WorldPartition->SetContainerPackage(Package->GetFName());
+					}
+
 					// We're changing the world path, add a path redirector so that soft object paths get fixed on save
 					FSoftObjectPath NewPath( World );
 					GRedirectCollector.AddAssetPathRedirection( *OldPath.GetAssetPathString(), *NewPath.GetAssetPathString() );
@@ -761,35 +767,32 @@ static bool SaveWorld(UWorld* World,
 
 		UWorld* SaveWorld = DuplicatedWorld ? DuplicatedWorld : World;
 
+		// Delete External Actors if we are saving over an existing map
 		if (bSuccess && bPackageNeedsRename)
 		{
-			// Delete External Actors if we are saving over an existing map
-			if (bPackageNeedsRename)
+			FPackageSourceControlHelper PackageHelper;
+			const FString ExternalActorsPath = ULevel::GetExternalActorsPath(NewPackageName);
+			FString ExternalActorsFilePath = FPackageName::LongPackageNameToFilename(ExternalActorsPath);
+			if (IFileManager::Get().DirectoryExists(*ExternalActorsFilePath))
 			{
-				FPackageSourceControlHelper PackageHelper;
-				const FString ExternalActorsPath = ULevel::GetExternalActorsPath(NewPackageName);
-				FString ExternalActorsFilePath = FPackageName::LongPackageNameToFilename(ExternalActorsPath);
-				if (IFileManager::Get().DirectoryExists(*ExternalActorsFilePath))
-				{
-					bSuccess = IFileManager::Get().IterateDirectoryRecursively(*ExternalActorsFilePath, [&PackageHelper](const TCHAR* FilenameOrDirectory, bool bIsDirectory)
+				bSuccess = IFileManager::Get().IterateDirectoryRecursively(*ExternalActorsFilePath, [&PackageHelper](const TCHAR* FilenameOrDirectory, bool bIsDirectory)
+					{
+						if (!bIsDirectory)
 						{
-							if (!bIsDirectory)
+							FString Filename(FilenameOrDirectory);
+							if (Filename.EndsWith(FPackageName::GetAssetPackageExtension()))
 							{
-								FString Filename(FilenameOrDirectory);
-								if (Filename.EndsWith(FPackageName::GetAssetPackageExtension()))
-								{
-									// If Delete fails it will return false and Directory Iteration will be stopped
-									return PackageHelper.Delete(Filename);
-								}
+								// If Delete fails it will return false and Directory Iteration will be stopped
+								return PackageHelper.Delete(Filename);
 							}
-							// Continue Directory Iteration
-							return true;
-						});
+						}
+						// Continue Directory Iteration
+						return true;
+					});
 
-					if (!bSuccess)
-					{
-						FMessageDialog::Open(EAppMsgType::Ok, NSLOCTEXT("UnrealEd", "FailedToDeleteExistingExternalActors", "Failed to delete existing map external actors."));
-					}
+				if (!bSuccess)
+				{
+					FMessageDialog::Open(EAppMsgType::Ok, NSLOCTEXT("UnrealEd", "FailedToDeleteExistingExternalActors", "Failed to delete existing map external actors."));
 				}
 			}
 		}
@@ -807,17 +810,6 @@ static bool SaveWorld(UWorld* World,
 			}
 		}
 
-		// Package was renamed if the world was partitioned the world partition needs to be reinitialized
-		if (bSuccess && bPackageNeedsRename)
-		{			
-			if (UWorldPartition* WorldPartition = SaveWorld->GetWorldPartition())
-			{
-				FTransform InstanceTransform = WorldPartition->GetInstanceTransform();
-				WorldPartition->Uninitialize();
-				WorldPartition->Initialize(SaveWorld, InstanceTransform);
-			}
-		}
-
 		SlowTask.EnterProgressFrame(25);
 
 		// If the package save was not successful. Trash the duplicated world or rename back if the duplicate failed.
diff --git a/Engine/Source/Runtime/Engine/Public/WorldPartition/ActorDescContainer.h b/Engine/Source/Runtime/Engine/Public/WorldPartition/ActorDescContainer.h
--- a/Engine/Source/Runtime/Engine/Public/WorldPartition/ActorDescContainer.h
+++ b/Engine/Source/Runtime/Engine/Public/WorldPartition/ActorDescContainer.h
@@ -29,6 +29,13 @@ public:
 	virtual void OnPackageDeleted(UPackage* Package);
 	virtual void OnObjectsReplaced(const TMap<UObject*, UObject*>& OldToNewObjectMap);
 
+	void SetContainerPackage(const FName& InContainerPackageName)
+	{
+		// Should only be called on empty container
+		check(GetActorDescCount() == 0);
+		ContainerPackageName = InContainerPackageName;
+	}
+
 	FName GetContainerPackage() const { return ContainerPackageName; }
 
 	/** Removes an actor desc without the need to load a package */
