Integrating //UE5/Dev-Cooker @ 16678003 to //UE5/Main (Zousar.Shaker-YEG-0943-Quaternary)

diff --git a/Engine/Source/Editor/UnrealEd/Private/CookOnTheFlyServer.cpp b/Engine/Source/Editor/UnrealEd/Private/CookOnTheFlyServer.cpp
--- a/Engine/Source/Editor/UnrealEd/Private/CookOnTheFlyServer.cpp
+++ b/Engine/Source/Editor/UnrealEd/Private/CookOnTheFlyServer.cpp
@@ -14,6 +14,9 @@
 #include "Cooker/CookRequests.h"
 #include "Cooker/CookTypes.h"
 #include "Cooker/PackageTracker.h"
+#include "Cooker/CookOnTheFlyServerInterface.h"
+#include "Cooker/NetworkFileCookOnTheFlyRequestManager.h"
+#include "Cooker/IoStoreCookOnTheFlyRequestManager.h"
 #include "CookPackageSplitter.h"
 #include "Containers/RingBuffer.h"
 #include "EditorDomain/EditorDomain.h"
@@ -122,7 +125,8 @@
 #include "Async/ParallelFor.h"
 
 #include "Commandlets/ShaderPipelineCacheToolsCommandlet.h"
-#include "PackageStoreWriter.h"
+#include "FilePackageStoreWriter.h"
+#include "ZenStoreWriter.h"
 
 #define LOCTEXT_NAMESPACE "Cooker"
 #define REMAPPED_PLUGINS TEXT("RemappedPlugins")
@@ -425,6 +429,7 @@ public:
 	bool							bSkipSoftReferences = false;
 	bool							bFullLoadAndSave = false;
 	bool							bIoStore = false;
+	bool							bZenStore = false;
 	bool							bCookAgainstFixedBase = false;
 	bool							bDlcLoadMainAssetRegistry = false;
 	TArray<FName>					StartupPackages;
@@ -433,6 +438,257 @@ public:
 	TMap<FName, TArray<FName>>		SourceToLocalizedPackageVariants;
 };
 
+//////////////////////////////////////////////////////////////////////////
+// Cook on the fly server interface adapter
+
+class UCookOnTheFlyServer::FCookOnTheFlyServerInterface final
+	: public UE::Cook::ICookOnTheFlyServer
+{
+public:
+	FCookOnTheFlyServerInterface(UCookOnTheFlyServer& InCooker)
+		: Cooker(InCooker)
+		, FlushCompletedEvent(FPlatformProcess::GetSynchEventFromPool(true))
+	{
+		FlushCompletedEvent->Trigger();
+	}
+
+	virtual ~FCookOnTheFlyServerInterface()
+	{
+		FPlatformProcess::ReturnSynchEventToPool(FlushCompletedEvent);
+	}
+
+	virtual FString GetSandboxDirectory() const override
+	{
+		return Cooker.SandboxFile->GetSandboxDirectory();
+	}
+
+	virtual bool AddPlatform(const FName& PlatformName) override
+	{
+		{
+			UE::Cook::FPlatformManager::FReadScopeLock PlatformScopeLock(Cooker.PlatformManager->ReadLockPlatforms());
+			if (!Cooker.AddCookOnTheFlyPlatform(PlatformName))
+			{
+				UE_LOG(LogCook, Warning, TEXT("Trying to add invalid platform '%s' on the fly"), *PlatformName.ToString());
+				return false;
+			}
+
+			Cooker.PlatformManager->AddRefCookOnTheFlyPlatform(PlatformName, Cooker);
+		}
+
+		Cooker.TrySetDefaultAsyncIODeletePlatform(PlatformName.ToString());
+
+		return true;
+	}
+	
+	virtual void RemovePlatform(const FName& PlatformName) override
+	{
+		UE::Cook::FPlatformManager::FReadScopeLock PlatformScopeLock(Cooker.PlatformManager->ReadLockPlatforms());
+		Cooker.PlatformManager->ReleaseCookOnTheFlyPlatform(PlatformName);
+	}
+
+	virtual void GetPrecookedFileList(const FName& PlatformName, UE::Cook::FPrecookedFileList& OutPrecookedFiles) override
+	{
+		ITargetPlatformManagerModule& TPM = GetTargetPlatformManagerRef();
+		const ITargetPlatform* TargetPlatform = TPM.FindTargetPlatform(PlatformName);
+		if (!TargetPlatform)
+		{
+			UE_LOG(LogCook, Warning, TEXT("Trying to get precooked file(s) on the fly for invalid platform '%s'"), *PlatformName.ToString());
+			return;
+		}
+
+		TArray<FName> CookedPlatformFiles;
+		Cooker.PackageDatas->GetCookedFileNamesForPlatform(TargetPlatform, CookedPlatformFiles, /* include failed */ true, /* include successful */ true);
+
+		for (const FName& CookedFile : CookedPlatformFiles)
+		{
+			const FString SandboxFilename = Cooker.ConvertToFullSandboxPath(CookedFile.ToString(), true, PlatformName.ToString());
+			if (IFileManager::Get().FileExists(*SandboxFilename))
+			{
+				continue;
+			}
+
+			OutPrecookedFiles.Add(CookedFile.ToString(),FDateTime::MinValue());
+		}
+	}
+
+	virtual void GetUnsolicitedFiles(const FName& PlatformName, const FString& Filename, const bool bIsCookable, TArray<FString>& OutUnsolicitedFiles) override
+	{
+		while (true)
+		{
+			{
+				UE::Cook::FPlatformManager::FReadScopeLock PlatformsScopeLock(Cooker.PlatformManager->ReadLockPlatforms());
+				const ITargetPlatform* TargetPlatform = Cooker.AddCookOnTheFlyPlatform(PlatformName);
+				if (!TargetPlatform)
+				{
+					break;
+				}
+				if (Cooker.PlatformManager->IsPlatformInitialized(TargetPlatform))
+				{
+					Cooker.GetCookOnTheFlyUnsolicitedFiles(TargetPlatform, PlatformName.ToString(), OutUnsolicitedFiles, Filename, bIsCookable);
+					break;
+				}
+			}
+			// Wait for the Platform to be added if this is the first time; it is not legal to call GetCookOnTheFlyUnsolicitedFiles until after the platform has been added
+			FPlatformProcess::Sleep(0.001f);
+		}
+	}
+
+	virtual bool EnqueueCookRequest(UE::Cook::FCookPackageRequest CookPackageRequest) override
+	{
+		UE::Cook::FPlatformManager::FReadScopeLock PlatformsScopeLock(Cooker.PlatformManager->ReadLockPlatforms());
+
+		if (const ITargetPlatform* TargetPlatform = Cooker.AddCookOnTheFlyPlatform(CookPackageRequest.PlatformName))
+		{
+			Cooker.PlatformManager->AddRefCookOnTheFlyPlatform(CookPackageRequest.PlatformName, Cooker);
+			FName StandardFileName(*FPaths::CreateStandardFilename(CookPackageRequest.Filename));
+			FRequest* Request = AllocRequest();
+			Request->CookPackageRequest = MoveTemp(CookPackageRequest);
+
+			UE::Cook::FCompletionCallback CookCompleted = [this, Request](UE::Cook::FPackageData* PackageData)
+			{
+				const UE::Cook::FPlatformData* PlatformData = Cooker.PlatformManager->GetPlatformDataByName(Request->CookPackageRequest.PlatformName);
+				check(PlatformData);
+				const UE::Cook::ECookResult CookResult = PackageData
+					? PackageData->GetCookResults(PlatformData->TargetPlatform)
+					: UE::Cook::ECookResult::Failed;
+				
+				UE_LOG(LogCook, Verbose, TEXT("Cook request completed, Filename='%s', Platform='%s', CookResult='%d'"),
+					*Request->CookPackageRequest.Filename, *Request->CookPackageRequest.PlatformName.ToString(), uint32(CookResult));
+				
+				if (Request->CookPackageRequest.CompletionCallback)
+				{
+					Request->CookPackageRequest.CompletionCallback(CookResult);
+				}
+				
+				Cooker.PlatformManager->ReleaseCookOnTheFlyPlatform(Request->CookPackageRequest.PlatformName);
+				FreeRequest(Request);
+			};
+
+			UE_LOG(LogCook, Verbose, TEXT("Enqueing cook request, Filename='%s', Platform='%s'"), *Request->CookPackageRequest.Filename, *Request->CookPackageRequest.PlatformName.ToString());
+			Cooker.ExternalRequests->EnqueueUnique(UE::Cook::FFilePlatformRequest(StandardFileName, TargetPlatform, MoveTemp(CookCompleted)), true);
+			if (Cooker.ExternalRequests->CookRequestEvent)
+			{
+				Cooker.ExternalRequests->CookRequestEvent->Trigger();
+			}
+
+			return true;
+		}
+		else
+		{
+			UE_LOG(LogCook, Warning, TEXT("Trying to cook package on the fly for invalid platform '%s'"), *CookPackageRequest.PlatformName.ToString());
+
+			return false;
+		}
+	};
+
+	virtual bool EnqueueRecompileShaderRequest(UE::Cook::FRecompileShaderRequest RecompileShaderRequest) override
+	{
+		UE_LOG(LogCook, Verbose, TEXT("Enqueing recompile shader(s) request, Platform='%s', ShaderPlatform='%d'"),
+			*RecompileShaderRequest.PlatformName.ToString(), RecompileShaderRequest.ShaderPlatformToCompile);
+
+		Cooker.PackageTracker->RecompileRequests.Enqueue(MoveTemp(RecompileShaderRequest));
+
+		return true;
+	}
+
+	virtual IPackageStoreWriter* GetPackageStoreWriter(const FName& PlatformName) override
+	{
+		if (Cooker.CookOnTheFlyOptions.bIoStore)
+		{
+			for (FSavePackageContext* Context : Cooker.SavePackageContexts)
+			{
+				if (FName(*Context->TargetPlatform->PlatformName()) == PlatformName)
+				{
+					return Context->PackageStoreWriter;
+				}
+			}
+		}
+
+		return nullptr;
+	}
+
+	DECLARE_DERIVED_EVENT(UCookOnTheFlyServer::FCookOnTheFlyServerInterface, UE::Cook::ICookOnTheFlyServer::FFlushEvent, FFlushEvent);
+	virtual FFlushEvent& OnFlush() override
+	{
+		return FlushEvent;
+	}
+
+	virtual double WaitForPendingFlush() override
+	{
+		const double StartTime = bIsFlushing ? FPlatformTime::Seconds() : 0.0;
+
+		FlushCompletedEvent->Wait();
+
+		return StartTime > 0.0 ? FPlatformTime::Seconds() - StartTime : 0.0;
+	}
+
+	void Flush()
+	{
+		bIsFlushing = true;
+		FlushCompletedEvent->Reset();
+
+		UE_LOG(LogCook, Log, TEXT("Flushing..."));
+
+		if (Cooker.CookOnTheFlyOptions.bIoStore)
+		{
+			for (FSavePackageContext* Context : Cooker.SavePackageContexts)
+			{
+				Context->PackageStoreWriter->Flush();
+			}
+		}
+
+		if (FlushEvent.IsBound())
+		{
+			FlushEvent.Broadcast();
+		}
+
+		UE_LOG(LogCook, Log, TEXT("Flush completed"));
+
+		bIsFlushing = false;
+		FlushCompletedEvent->Trigger();
+	}
+
+private:
+	struct FRequest
+	{
+		FRequest* NextFree = nullptr;
+		UE::Cook::FCookPackageRequest CookPackageRequest;
+	};
+
+	FRequest* AllocRequest()
+	{
+		FScopeLock _(&RequestsCriticalSection);
+		if (FRequest* Request = FirstFree)
+		{
+			FirstFree = Request->NextFree;
+			Request->NextFree = nullptr;
+			return Request;
+		}
+		else
+		{
+			const int32 RequestIndex = Requests.Add();
+			return &Requests[RequestIndex];
+		}
+	}
+
+	void FreeRequest(FRequest* Request)
+	{
+		FScopeLock _(&RequestsCriticalSection);
+		check(Request->NextFree == nullptr);
+		Request->NextFree = FirstFree;
+		FirstFree = Request;
+	}
+
+	UCookOnTheFlyServer& Cooker;
+	FFlushEvent FlushEvent;
+	FEvent* FlushCompletedEvent;
+	FCriticalSection FlushCriticalSection;
+	FCriticalSection RequestsCriticalSection;
+	TChunkedArray<FRequest> Requests;
+	FRequest* FirstFree = nullptr;
+	TAtomic<bool> bIsFlushing = {false};
+};
+
 /* UCookOnTheFlyServer functions
  *****************************************************************************/
 
@@ -519,7 +775,7 @@ TStatId UCookOnTheFlyServer::GetStatId() const
 	RETURN_QUICK_DECLARE_CYCLE_STAT(UCookServer, STATGROUP_Tickables);
 }
 
-bool UCookOnTheFlyServer::StartNetworkFileServer(const bool BindAnyPort, const TArray<ITargetPlatform*>& TargetPlatforms)
+bool UCookOnTheFlyServer::StartCookOnTheFly(FCookOnTheFlyOptions InCookOnTheFlyOptions)
 {
 	check(IsCookOnTheFlyMode());
 	//GetDerivedDataCacheRef().WaitForQuiescence(false);
@@ -528,6 +784,10 @@ bool UCookOnTheFlyServer::StartNetworkFileServer(const bool BindAnyPort, const T
 	NetworkRequestEvent = FPlatformProcess::GetSynchEventFromPool();
 #endif
 
+	CookOnTheFlyOptions = MoveTemp(InCookOnTheFlyOptions);
+	CookOnTheFlyServerInterface = MakeUnique<UCookOnTheFlyServer::FCookOnTheFlyServerInterface>(*this);
+	ExternalRequests->CookRequestEvent = FPlatformProcess::GetSynchEventFromPool();
+
 	// Precreate the map of all possible target platforms so we can access the collection of existing platforms in a threadsafe manner
 	// Each PlatformData in the map will be uninitialized until we call AddCookOnTheFlyPlatform for the platform
 	ITargetPlatformManagerModule& TPM = GetTargetPlatformManagerRef();
@@ -547,59 +807,38 @@ bool UCookOnTheFlyServer::StartNetworkFileServer(const bool BindAnyPort, const T
 
 	{
 		UE::Cook::FPlatformManager::FReadScopeLock PlatformScopeLock(PlatformManager->ReadLockPlatforms());
-		for (ITargetPlatform* TargetPlatform : TargetPlatforms)
+		for (ITargetPlatform* TargetPlatform : CookOnTheFlyOptions.TargetPlatforms)
 		{
-			AddCookOnTheFlyPlatform(TargetPlatform->PlatformName());
-		}
+			AddCookOnTheFlyPlatform(FName(*TargetPlatform->PlatformName()));
+	}
 	}
 
-	// start the listening thread
-	FNewConnectionDelegate NewConnectionDelegate(FNewConnectionDelegate::CreateUObject(this, &UCookOnTheFlyServer::HandleNetworkFileServerNewConnection));
-	FFileRequestDelegate FileRequestDelegate(FFileRequestDelegate::CreateUObject(this, &UCookOnTheFlyServer::HandleNetworkFileServerFileRequest));
-	FRecompileShadersDelegate RecompileShadersDelegate(FRecompileShadersDelegate::CreateUObject(this, &UCookOnTheFlyServer::HandleNetworkFileServerRecompileShaders));
-	FSandboxPathDelegate SandboxPathDelegate(FSandboxPathDelegate::CreateUObject(this, &UCookOnTheFlyServer::HandleNetworkGetSandboxPath));
-	FInitialPrecookedListDelegate InitialPrecookedListDelegate(FInitialPrecookedListDelegate::CreateUObject(this, &UCookOnTheFlyServer::HandleNetworkGetPrecookedList));
-
-	FNetworkFileServerOptions FileServerOptions;
-	FileServerOptions.Protocol = NFSP_Tcp;
-	FileServerOptions.Port = BindAnyPort ? 0 : -1;
-	FileServerOptions.Delegates.NewConnectionDelegate = NewConnectionDelegate;
-	FileServerOptions.Delegates.InitialPrecookedListDelegate = InitialPrecookedListDelegate;
-	FileServerOptions.Delegates.FileRequestDelegate = FileRequestDelegate;
-	FileServerOptions.Delegates.RecompileShadersDelegate = RecompileShadersDelegate;
-	FileServerOptions.Delegates.SandboxPathOverrideDelegate = SandboxPathDelegate;
-	FileServerOptions.Delegates.OnFileModifiedCallback = &FileModifiedDelegate;
-	FileServerOptions.bRestrictPackageAssetsToSandbox = true; // prevents sending uncooked packages if the package fails to cook
-
-	INetworkFileServer *TcpFileServer = FModuleManager::LoadModuleChecked<INetworkFileSystemModule>("NetworkFileSystem")
-		.CreateNetworkFileServer(MoveTemp(FileServerOptions), /* bLoadTargetPlatforms */ true);
-	if ( TcpFileServer )
+	if (CookOnTheFlyOptions.bIoStore)
 	{
-		NetworkFileServers.Add(TcpFileServer);
+		UE_LOG(LogCook, Display, TEXT("Initializing I/O store cook on the fly server"));
+		UE::Cook::FIoStoreCookOnTheFlyServerOptions ServerOptions;
+		ServerOptions.Port = -1; // Use default
+		CookOnTheFlyRequestManager = UE::Cook::MakeIoStoreCookOnTheFlyRequestManager(*CookOnTheFlyServerInterface, MoveTemp(ServerOptions));
 	}
-
-#if 0 // cookonthefly server via http
-	INetworkFileServer *HttpFileServer = FModuleManager::LoadModuleChecked<INetworkFileSystemModule>("NetworkFileSystem")
-		.CreateNetworkFileServer(true, BindAnyPort ? 0 : -1, NetworkFileDelegateContainer, ENetworkFileServerProtocol::NFSP_Http);
-	if ( HttpFileServer )
+	else
 	{
-		NetworkFileServers.Add( HttpFileServer );
+		UE_LOG(LogCook, Display, TEXT("Initializing Network File cook on the fly server"));
+		FNetworkFileServerOptions ServerOptions;
+		ServerOptions.Protocol = NFSP_Tcp;
+		ServerOptions.Port = CookOnTheFlyOptions.bBindAnyPort ? 0 : -1;
+		CookOnTheFlyRequestManager = UE::Cook::MakeNetworkFileCookOnTheFlyRequestManager(*CookOnTheFlyServerInterface, MoveTemp(ServerOptions));
 	}
-#endif
 
-	ExternalRequests->CookRequestEvent = FPlatformProcess::GetSynchEventFromPool();
-
-	// loop while waiting for requests
-	return true;
+	const bool bInitialized = CookOnTheFlyRequestManager->Initialize();
+	return bInitialized;
 }
 
-const ITargetPlatform* UCookOnTheFlyServer::AddCookOnTheFlyPlatform(const FString& PlatformNameString)
+const ITargetPlatform* UCookOnTheFlyServer::AddCookOnTheFlyPlatform(const FName& PlatformName)
 {
-	FName PlatformName(*PlatformNameString);
 	const UE::Cook::FPlatformData* PlatformData = PlatformManager->GetPlatformDataByName(PlatformName);
 	if (!PlatformData)
 	{
-		UE_LOG(LogCook, Warning, TEXT("Target platform %s wasn't found."), *PlatformNameString);
+		UE_LOG(LogCook, Warning, TEXT("Target platform %s wasn't found."), *PlatformName.ToString());
 		return nullptr;
 	}
 
@@ -615,21 +854,52 @@ const ITargetPlatform* UCookOnTheFlyServer::AddCookOnTheFlyPlatform(const FStrin
 	}
 	else
 	{
+		FEvent* PlatformInitializedEvent = FPlatformProcess::GetSynchEventFromPool();
+		
 		// Registering a new platform is not thread safe; queue the command for TickCookOnTheSide to execute
-		ExternalRequests->AddCallback([this, PlatformName]()
+		ExternalRequests->AddCallback([this, PlatformName, PlatformInitializedEvent]()
 			{
 				const UE::Cook::FPlatformData* PlatformData = PlatformManager->GetPlatformDataByName(PlatformName);
 				check(PlatformData);
 				AddCookOnTheFlyPlatformFromGameThread(PlatformData->TargetPlatform);
+			PlatformInitializedEvent->Trigger();
 			});
+
 		if (ExternalRequests->CookRequestEvent)
 		{
 			ExternalRequests->CookRequestEvent->Trigger();
 		}
+		
+		PlatformInitializedEvent->Wait();
+		FPlatformProcess::ReturnSynchEventToPool(PlatformInitializedEvent);
 	}
 	return PlatformData->TargetPlatform;
 }
 
+void UCookOnTheFlyServer::InitializeShadersForCookOnTheFly(const TArrayView<ITargetPlatform* const>& NewTargetPlatforms)
+{
+	for (const ITargetPlatform* TargetPlatform : NewTargetPlatforms)
+	{
+		UE_LOG(LogCook, Display, TEXT("Initializing shaders for cook-on-the-fly on platform '%s'"), *TargetPlatform->PlatformName());
+
+		const FString& PlatformName = TargetPlatform->PlatformName();
+		FString OutputDir = GetSandboxDirectory(PlatformName);
+
+		UE::Cook::FRecompileShaderRequest RecompileShaderRequest;
+		RecompileShaderRequest.bCompileChangedShaders = true;
+
+		RecompileShadersForRemote(
+			PlatformName, 
+			SP_NumPlatforms,
+			OutputDir, 
+			RecompileShaderRequest.MaterialsToLoad, 
+			RecompileShaderRequest.ShadersToRecompile,
+			RecompileShaderRequest.MeshMaterialMaps, 
+			RecompileShaderRequest.ModifiedFiles,
+			RecompileShaderRequest.bCompileChangedShaders);
+	}
+}
+
 void UCookOnTheFlyServer::AddCookOnTheFlyPlatformFromGameThread(ITargetPlatform* TargetPlatform)
 {
 	check(!!(CookFlags & ECookInitializationFlags::GeneratedAssetRegistry)); // GenerateAssetRegistry should have been called in StartNetworkFileServer
@@ -646,6 +916,7 @@ void UCookOnTheFlyServer::AddCookOnTheFlyPlatformFromGameThread(ITargetPlatform*
 	RefreshPlatformAssetRegistries(NewTargetPlatforms);
 	InitializeSandbox(NewTargetPlatforms);
 	InitializeTargetPlatforms(NewTargetPlatforms);
+	InitializeShadersForCookOnTheFly(NewTargetPlatforms);
 
 	// When cooking on the fly the full registry is saved at the beginning
 	// in cook by the book asset registry is saved after the cook is finished
@@ -654,6 +925,7 @@ void UCookOnTheFlyServer::AddCookOnTheFlyPlatformFromGameThread(ITargetPlatform*
 	{
 		Generator->SaveAssetRegistry(GetSandboxAssetRegistryFilename(), true);
 	}
+	InitializePackageStore(NewTargetPlatforms);
 	check(PlatformData->bIsSandboxInitialized); // This should have been set by InitializeSandbox, and it is what we use to determine whether a platform has been initialized
 }
 
@@ -1022,7 +1294,12 @@ bool UCookOnTheFlyServer::IsUsingShaderCodeLibrary() const
 
 bool UCookOnTheFlyServer::IsUsingIoStore() const
 {
-	return IsCookByTheBookMode() && CookByTheBookOptions->bIoStore;
+	return (IsCookByTheBookMode() && CookByTheBookOptions->bIoStore) || (IsCookOnTheFlyMode() && CookOnTheFlyOptions.bIoStore);
+}
+
+bool UCookOnTheFlyServer::IsUsingZenStore() const
+{
+	return (IsCookByTheBookMode() && CookByTheBookOptions->bZenStore) || (IsCookOnTheFlyMode() && CookOnTheFlyOptions.bIoStore);
 }
 
 bool UCookOnTheFlyServer::IsCookOnTheFlyMode() const
@@ -1250,6 +1527,11 @@ uint32 UCookOnTheFlyServer::TickCookOnTheSide(const float TimeSlice, uint32 &Coo
 		CookByTheBookFinished();
 	}
 
+	if (IsCookOnTheFlyMode() && bCookComplete)
+	{
+		CookOnTheFlyServerInterface->Flush();
+	}
+
 	CookedPackageCount += StackData.CookedPackageCount;
 	return StackData.ResultFlags;
 }
@@ -1616,7 +1898,7 @@ void UCookOnTheFlyServer::PumpExternalRequests(const UE::Cook::FCookerTimer& Coo
 				UE::Cook::FCompletionCallback CompletionCallback(MoveTemp(ToBuild.GetCompletionCallback()));
 				if (CompletionCallback)
 				{
-					CompletionCallback();
+					CompletionCallback(nullptr);
 				}
 				continue;
 			}
@@ -3705,22 +3987,15 @@ bool UCookOnTheFlyServer::IsInSession() const
 	return IsCookByTheBookRunning() || (IsCookOnTheFlyMode() && PlatformManager->HasSelectedSessionPlatforms());
 }
 
-void UCookOnTheFlyServer::EndNetworkFileServer()
+void UCookOnTheFlyServer::ShutdownCookOnTheFly()
 {
-	for ( int i = 0; i < NetworkFileServers.Num(); ++i )
-	{
-		INetworkFileServer *NetworkFileServer = NetworkFileServers[i];
-		// shutdown the server
-		NetworkFileServer->Shutdown();
-		delete NetworkFileServer;
-		NetworkFileServer = NULL;
-	}
-
-	if (NetworkFileServers.Num() > 0)
+	if (CookOnTheFlyRequestManager.IsValid())
 	{
+		UE_LOG(LogCook, Display, TEXT("Shutting down cook on the fly server"));
+		CookOnTheFlyRequestManager->Shutdown();
+		CookOnTheFlyRequestManager.Reset();
 		ClearHierarchyTimers();
 	}
-	NetworkFileServers.Empty();
 }
 
 uint32 UCookOnTheFlyServer::GetPackagesPerGC() const
@@ -3980,7 +4255,7 @@ void UCookOnTheFlyServer::PostGarbageCollect()
 
 void UCookOnTheFlyServer::BeginDestroy()
 {
-	EndNetworkFileServer();
+	ShutdownCookOnTheFly();
 
 	Super::BeginDestroy();
 }
@@ -3988,17 +4263,23 @@ void UCookOnTheFlyServer::BeginDestroy()
 void UCookOnTheFlyServer::TickRecompileShaderRequests()
 {
 	// try to pull off a request
-	UE::Cook::FRecompileRequest* Request = NULL;
-
-	PackageTracker->RecompileRequests.Dequeue(&Request);
-
-	// process it
-	if (Request)
+	UE::Cook::FRecompileShaderRequest RecompileShaderRequest;
+	if (PackageTracker->RecompileRequests.Dequeue(&RecompileShaderRequest))
 	{
-		HandleNetworkFileServerRecompileShaders(Request->RecompileData);
+		FString PlatformName = RecompileShaderRequest.PlatformName.ToString(); 
+		FString OutputDir = GetSandboxDirectory(PlatformName);
+
+		RecompileShadersForRemote(
+			PlatformName, 
+			RecompileShaderRequest.ShaderPlatformToCompile == -1 ? SP_NumPlatforms : (EShaderPlatform)RecompileShaderRequest.ShaderPlatformToCompile,
+			OutputDir, 
+			RecompileShaderRequest.MaterialsToLoad, 
+			RecompileShaderRequest.ShadersToRecompile,
+			RecompileShaderRequest.MeshMaterialMaps, 
+			RecompileShaderRequest.ModifiedFiles,
+			RecompileShaderRequest.bCompileChangedShaders);
 
-		// all done! other thread can unblock now
-		Request->bComplete = true;
+		RecompileShaderRequest.CompletionCallback();
 	}
 }
 
@@ -4342,7 +4623,7 @@ void UCookOnTheFlyServer::SaveCookedPackage(UE::Cook::FPackageData& PackageData,
 
 					if (DiffModeHelper.IsRunningCookDiff())
 					{
-						FSavePackageContext* const SavePackageContext = (IsCookByTheBookMode() && SavePackageContexts.Num() > 0) ? SavePackageContexts[PlatformIndex] : nullptr;
+						FSavePackageContext* const SavePackageContext = SavePackageContexts.Num() > 0 ? SavePackageContexts[PlatformIndex] : nullptr;
 
 						DiffModeHelper.ProcessPackage(Package);
 
@@ -4359,12 +4640,26 @@ void UCookOnTheFlyServer::SaveCookedPackage(UE::Cook::FPackageData& PackageData,
 					}
 					else
 					{
-						FSavePackageContext* const SavePackageContext = (IsCookByTheBookMode() && SavePackageContexts.Num() > 0) ? SavePackageContexts[PlatformIndex] : nullptr;
+						FSavePackageContext* const SavePackageContext = SavePackageContexts.Num() > 0 ? SavePackageContexts[PlatformIndex] : nullptr;
+
+						IPackageStoreWriter::FPackageBaseInfo Info;
+						Info.PackageName = Package->GetFName();
+
+						IPackageStoreWriter* PackageStoreWriter = SavePackageContext ? SavePackageContext->PackageStoreWriter : nullptr;
+						if (PackageStoreWriter)
+						{
+							PackageStoreWriter->BeginPackage(Info);
+						}
 
 						Result = GEditor->Save(	Package, World, FlagsToCook, *PlatFilename, 
 												GError, nullptr, bSwap, false, SaveFlags, Target,
 												FDateTime::MinValue(), false, /*DiffMap*/ nullptr, 
 												SavePackageContext);
+
+						if (PackageStoreWriter)
+						{
+							PackageStoreWriter->CommitPackage(Info);
+						}
 					}
 
 					// if running linker diff, resave the package again using the other save algorithm
@@ -6844,8 +7139,6 @@ void UCookOnTheFlyServer::CookByTheBookFinished()
 
 	UPackage::WaitForAsyncFileWrites();
 	
-	FinalizePackageStore();
-
 	GetDerivedDataCacheRef().WaitForQuiescence(true);
 	
 	UCookerSettings const* CookerSettings = GetDefault<UCookerSettings>();
@@ -7020,6 +7313,8 @@ void UCookOnTheFlyServer::CookByTheBookFinished()
 		}
 	}
 
+	FinalizePackageStore();
+
 	CookByTheBookOptions->BasedOnReleaseCookedPackages.Empty();
 	CookByTheBookOptions->bRunning = false;
 	CookByTheBookOptions->bFullLoadAndSave = false;
@@ -7389,6 +7684,7 @@ void UCookOnTheFlyServer::InitializePackageStore(const TArrayView<const ITargetP
 	const FString ProjectPathSandbox = ConvertToFullSandboxPath(*ProjectPath, true);
 
 	const bool bIsDiffOnly = FParse::Param(FCommandLine::Get(), TEXT("DIFFONLY"));
+	const bool bIsCleanBuild = !IsCookFlagSet(ECookInitializationFlags::Iterative);
 
 	SavePackageContexts.Reserve(TargetPlatforms.Num());
 
@@ -7399,10 +7695,27 @@ void UCookOnTheFlyServer::InitializePackageStore(const TArrayView<const ITargetP
 		const FString ResolvedRootPath = RootPathSandbox.Replace(TEXT("[Platform]"), *PlatformString);
 		const FString ResolvedProjectPath = ProjectPathSandbox.Replace(TEXT("[Platform]"), *PlatformString);
 
-		FPackageStoreWriter* PackageStoreWriter = nullptr;
+		IPackageStoreWriter* PackageStoreWriter = nullptr;
 		if (IsUsingIoStore())
 		{
-			PackageStoreWriter = new FPackageStoreWriter(ResolvedRootPath, TargetPlatform);
+			FString MetadataDirectoryPath = ProjectPathSandbox / TEXT("Metadata");
+			MetadataDirectoryPath.ReplaceInline(TEXT("[Platform]"), *PlatformString);
+
+			if (IsUsingZenStore())
+			{
+				PackageStoreWriter = new FZenStoreWriter(	ResolvedRootPath, 
+															MetadataDirectoryPath, 
+															TargetPlatform, 
+															bIsCleanBuild);
+			}
+			else
+			{
+				PackageStoreWriter = new FFilePackageStoreWriter(ResolvedRootPath, MetadataDirectoryPath, TargetPlatform);
+			}
+
+			IPackageStoreWriter::FCookInfo CookInfo;
+			CookInfo.CookMode = IsCookOnTheFlyMode() ? IPackageStoreWriter::FCookInfo::CookOnTheFlyMode : IPackageStoreWriter::FCookInfo::CookByTheBookMode;
+			PackageStoreWriter->BeginCook(CookInfo);
 		}
 
 		FConfigFile PlatformEngineIni;
@@ -7411,7 +7724,7 @@ void UCookOnTheFlyServer::InitializePackageStore(const TArrayView<const ITargetP
 		bool bLegacyBulkDataOffsets = false;
 		PlatformEngineIni.GetBool(TEXT("Core.System"), TEXT("LegacyBulkDataOffsets"), bLegacyBulkDataOffsets);
 
-		FSavePackageContext* SavePackageContext = new FSavePackageContext(PackageStoreWriter, bLegacyBulkDataOffsets);
+		FSavePackageContext* SavePackageContext = new FSavePackageContext(TargetPlatform, PackageStoreWriter, bLegacyBulkDataOffsets);
 		SavePackageContexts.Add(SavePackageContext);
 	}
 }
@@ -7427,7 +7740,7 @@ void UCookOnTheFlyServer::FinalizePackageStore()
 		{
 			if (PackageContext->PackageStoreWriter != nullptr)
 			{
-				PackageContext->PackageStoreWriter->Finalize();
+				PackageContext->PackageStoreWriter->EndCook();
 			}
 		}
 	}
@@ -7553,6 +7866,7 @@ void UCookOnTheFlyServer::StartCookByTheBook( const FCookByTheBookStartupOptions
 	CookByTheBookOptions->bSkipSoftReferences = !!(CookOptions & ECookByTheBookOptions::SkipSoftReferences);
 	CookByTheBookOptions->bFullLoadAndSave = !!(CookOptions & ECookByTheBookOptions::FullLoadAndSave);
 	CookByTheBookOptions->bIoStore = !!(CookOptions & ECookByTheBookOptions::IoStore);
+	CookByTheBookOptions->bZenStore = !!(CookOptions & ECookByTheBookOptions::ZenStore);
 	CookByTheBookOptions->bCookAgainstFixedBase = !!(CookOptions & ECookByTheBookOptions::CookAgainstFixedBase);
 	CookByTheBookOptions->bDlcLoadMainAssetRegistry = !!(CookOptions & ECookByTheBookOptions::DlcLoadMainAssetRegistry);
 	CookByTheBookOptions->bErrorOnEngineContentUse = CookByTheBookStartupOptions.bErrorOnEngineContentUse;
@@ -8093,35 +8407,6 @@ void UCookOnTheFlyServer::MaybeMarkPackageAsAlreadyLoaded(UPackage *Package)
 	}
 }
 
-
-bool UCookOnTheFlyServer::HandleNetworkFileServerNewConnection(const FString& VersionInfo, const FString& PlatformName)
-{
-	const uint32 CL = FEngineVersion::CompatibleWith().GetChangelist();
-	const FString Branch = FEngineVersion::CompatibleWith().GetBranch();
-
-	const FString LocalVersionInfo = FString::Printf(TEXT("%s %d"), *Branch, CL);
-
-	{
-		UE::Cook::FPlatformManager::FReadScopeLock PlatformScopeLock(PlatformManager->ReadLockPlatforms());
-		if (!AddCookOnTheFlyPlatform(PlatformName))
-		{
-			UE_LOG(LogCook, Warning, TEXT("Unrecognized PlatformName '%s', CookOnTheFly requests for this platform will fail."), *PlatformName);
-			return false;
-		}
-	}
-
-	TrySetDefaultAsyncIODeletePlatform(PlatformName);
-
-	UE_LOG(LogCook, Display, TEXT("Connection received of version %s local version %s"), *VersionInfo, *LocalVersionInfo);
-
-	if (LocalVersionInfo != VersionInfo)
-	{
-		UE_LOG(LogCook, Warning, TEXT("Connection tried to connect with incompatible version"));
-		// return false;
-	}
-	return true;
-}
-
 static void AppendExistingPackageSidecarFiles(const FString& PackageSandboxFilename, const FString& PackageStandardFilename, TArray<FString>& OutPackageSidecarFiles)
 {
 	const TCHAR* const PackageSidecarExtensions[] =
@@ -8175,205 +8460,6 @@ void UCookOnTheFlyServer::GetCookOnTheFlyUnsolicitedFiles(const ITargetPlatform*
 	}
 }
 
-void UCookOnTheFlyServer::HandleNetworkFileServerFileRequest(FString& Filename, const FString& PlatformNameString, TArray<FString>& UnsolicitedFiles)
-{
-	check(IsCookOnTheFlyMode());
-
-	FName PlatformName(*PlatformNameString);
-	const bool bIsCookable = FPackageName::IsPackageExtension(*FPaths::GetExtension(Filename, true));
-	if (!bIsCookable)
-	{
-		while (true)
-		{
-			{
-				UE::Cook::FPlatformManager::FReadScopeLock PlatformsScopeLock(PlatformManager->ReadLockPlatforms());
-				const ITargetPlatform* TargetPlatform = AddCookOnTheFlyPlatform(PlatformNameString);
-				if (!TargetPlatform)
-				{
-					break;
-				}
-				if (PlatformManager->IsPlatformInitialized(TargetPlatform))
-				{
-					GetCookOnTheFlyUnsolicitedFiles(TargetPlatform, PlatformNameString, UnsolicitedFiles, Filename, bIsCookable);
-					break;
-				}
-			}
-			// Wait for the Platform to be added if this is the first time; it is not legal to call GetCookOnTheFlyUnsolicitedFiles until after the platform has been added
-			FPlatformProcess::Sleep(0.001f);
-		}
-		return;
-	}
-
-	FString StandardFileName = Filename;
-	FPackageName::FindPackageFileWithoutExtension(FPaths::ChangeExtension(Filename, TEXT("")), Filename);
-	FPaths::MakeStandardFilename( StandardFileName );
-	FName StandardFileFname = FName(*StandardFileName);
-
-#if PROFILE_NETWORK
-	double StartTime = FPlatformTime::Seconds();
-	check(NetworkRequestEvent);
-	NetworkRequestEvent->Reset();
-#endif
-	
-	UE_LOG(LogCook, Display, TEXT("Requesting file from cooker %s"), *StandardFileName);
-	TAtomic<bool> bCookComplete(false);
-	UE::Cook::FCompletionCallback OnCookComplete = [&bCookComplete]()
-	{
-		bCookComplete = true;
-	};
-
-	{
-		// This lock guards us from having the TargetPlatform pointer invalidated as a key until after we have stored it in ExternalRequests
-		// Once it is in ExternalRequests we are safe because TargetPlatform pointers in ExternalRequests are updated whenever the TargetPlatform pointer changes.
-		// Note that we can not dereference the TargetPlatform pointer in this function as it could be invalidated at any time; we can only use it as a key value for identifying the platform in the external request.
-		UE::Cook::FPlatformManager::FReadScopeLock PlatformsScopeLock(PlatformManager->ReadLockPlatforms());
-
-		const ITargetPlatform* TargetPlatform = AddCookOnTheFlyPlatform(PlatformNameString);
-		if (!TargetPlatform)
-		{
-			UE_LOG(LogCook, Warning, TEXT("Unrecognized PlatformName '%s', CookOnTheFly FileServerRequest requests for this platform will fail."), *PlatformNameString);
-			return;
-		}
-		PlatformManager->AddRefCookOnTheFlyPlatform(PlatformName, *this);
-
-		UE::Cook::FFilePlatformRequest FileRequest(StandardFileFname, TargetPlatform, MoveTemp(OnCookComplete));
-		ExternalRequests->EnqueueUnique(MoveTemp(FileRequest), true);
-	}
-	
-	if (ExternalRequests->CookRequestEvent)
-	{
-		ExternalRequests->CookRequestEvent->Trigger();
-	}
-
-#if PROFILE_NETWORK
-	bool bFoundNetworkEventWait = true;
-	while (NetworkRequestEvent->Wait(1) == false)
-	{
-		// for some reason we missed the stat
-		if (bCookComplete)
-		{
-			double DeltaTimeTillRequestForfilled = FPlatformTime::Seconds() - StartTime;
-			TimeTillRequestForfilled += DeltaTimeTillRequestForfilled;
-			TimeTillRequestForfilledError += DeltaTimeTillRequestForfilled;
-			StartTime = FPlatformTime::Seconds();
-			bFoundNetworkEventWait = false;
-			break;
-		}
-	}
-
-	// wait for tick entry here
-	TimeTillRequestStarted += FPlatformTime::Seconds() - StartTime;
-	StartTime = FPlatformTime::Seconds();
-#endif
-
-	while (!bCookComplete)
-	{
-		FPlatformProcess::Sleep(0.001f);
-	}
-
-
-	{
-		UE::Cook::FPlatformManager::FReadScopeLock PlatformsScopeLock(PlatformManager->ReadLockPlatforms());
-		const ITargetPlatform* TargetPlatform = AddCookOnTheFlyPlatform(PlatformNameString);
-		PlatformManager->ReleaseCookOnTheFlyPlatform(PlatformName);
-		if (TargetPlatform)
-		{
-			GetCookOnTheFlyUnsolicitedFiles(TargetPlatform, PlatformNameString, UnsolicitedFiles, Filename, bIsCookable);
-		}
-	}
-
-
-#if PROFILE_NETWORK
-	if ( bFoundNetworkEventWait )
-	{
-		TimeTillRequestForfilled += FPlatformTime::Seconds() - StartTime;
-		StartTime = FPlatformTime::Seconds();
-	}
-	UE_LOG( LogCook, Display, TEXT("Cook complete %s"), *StandardFileFname.ToString());
-	WaitForAsyncFilesWrites += FPlatformTime::Seconds() - StartTime;
-	StartTime = FPlatformTime::Seconds();
-#endif
-#if DEBUG_COOKONTHEFLY
-	UE_LOG( LogCook, Display, TEXT("Processed file request %s"), *Filename );
-#endif
-
-}
-
-
-FString UCookOnTheFlyServer::HandleNetworkGetSandboxPath()
-{
-	return SandboxFile->GetSandboxDirectory();
-}
-
-void UCookOnTheFlyServer::HandleNetworkGetPrecookedList(const FString& PlatformName, TMap<FString, FDateTime>& PrecookedFileList)
-{
-	ITargetPlatformManagerModule& TPM = GetTargetPlatformManagerRef();
-	const ITargetPlatform* TargetPlatform = TPM.FindTargetPlatform(PlatformName);
-	if (!TargetPlatform)
-	{
-		UE_LOG(LogCook, Warning, TEXT("Unrecognized PlatformName '%s' in HandleNetworkGetPrrequests, returning 0 files."), *PlatformName);
-		return;
-	}
-
-	TArray<FName> CookedPlatformFiles;
-	PackageDatas->GetCookedFileNamesForPlatform(TargetPlatform, CookedPlatformFiles, /* include failed */ true, /* include successful */ true);
-
-
-	for ( const FName& CookedFile : CookedPlatformFiles)
-	{
-		const FString SandboxFilename = ConvertToFullSandboxPath(CookedFile.ToString(), true, PlatformName);
-		if (IFileManager::Get().FileExists(*SandboxFilename))
-		{
-			continue;
-		}
-
-		PrecookedFileList.Add(CookedFile.ToString(),FDateTime::MinValue());
-	}
-}
-
-void UCookOnTheFlyServer::HandleNetworkFileServerRecompileShaders(const FShaderRecompileData& RecompileData)
-{
-	// shouldn't receive network requests unless we are in cook on the fly mode
-	check( IsCookOnTheFlyMode() );
-	check( !IsCookingDLC() );
-	// if we aren't in the game thread, we need to push this over to the game thread and wait for it to finish
-	if (!IsInGameThread())
-	{
-		UE_LOG(LogCook, Verbose, TEXT("Got a recompile request on non-game thread"));
-
-		// make a new request
-		UE::Cook::FRecompileRequest* Request = new UE::Cook::FRecompileRequest;
-		Request->RecompileData = RecompileData;
-		Request->bComplete = false;
-
-		// push the request for the game thread to process
-		PackageTracker->RecompileRequests.Enqueue(Request);
-
-		// wait for it to complete (the game thread will pull it out of the TArray, but I will delete it)
-		while (!Request->bComplete)
-		{
-			FPlatformProcess::Sleep(0);
-		}
-		delete Request;
-		UE_LOG(LogCook, Verbose, TEXT("Completed recompile..."));
-
-		// at this point, we are done on the game thread, and ModifiedFiles will have been filled out
-		return;
-	}
-
-	FString OutputDir = GetSandboxDirectory(RecompileData.PlatformName);
-
-	RecompileShadersForRemote
-		(RecompileData.PlatformName, 
-		RecompileData.ShaderPlatform == -1 ? SP_NumPlatforms : (EShaderPlatform)RecompileData.ShaderPlatform,
-		OutputDir, 
-		RecompileData.MaterialsToLoad, 
-		RecompileData.ShadersToRecompile,
-		RecompileData.MeshMaterialMaps, 
-		RecompileData.ModifiedFiles,
-		RecompileData.bCompileChangedShaders);
-}
-
 bool UCookOnTheFlyServer::GetAllPackageFilenamesFromAssetRegistry( const FString& AssetRegistryPath, bool bVerifyPackagesExist, TArray<FName>& OutPackageFilenames ) const
 {
 	UE_SCOPED_COOKTIMER(GetAllPackageFilenamesFromAssetRegistry);
diff --git a/Engine/Source/Editor/UnrealEd/Private/Cooker/CookOnTheFlyServerInterface.h b/Engine/Source/Editor/UnrealEd/Private/Cooker/CookOnTheFlyServerInterface.h
new file mode 100644
--- /dev/null
+++ b/Engine/Source/Editor/UnrealEd/Private/Cooker/CookOnTheFlyServerInterface.h
@@ -0,0 +1,156 @@
+// Copyright Epic Games, Inc. All Rights Reserved.
+
+#pragma once
+
+#include "CoreTypes.h"
+#include "CookTypes.h"
+
+struct FODSCRequestPayload;
+class IPackageStoreWriter;
+
+namespace UE { namespace Cook
+{
+
+using FCookRequestCompletedCallback = TFunction<void(const ECookResult)>;
+
+using FRecompileShaderCompletedCallback = TFunction<void()>;
+
+using FPrecookedFileList = TMap<FString, FDateTime>;
+
+/**
+ * Cook package request.
+ */
+struct FCookPackageRequest
+{
+	/** The platform to cook for. */
+	FName PlatformName;
+	/* Asset filename to cook. */
+	FString Filename;
+	/** Completion callback. */
+	FCookRequestCompletedCallback CompletionCallback;
+};
+
+/**
+ * Recompile shader(s) request.
+ */
+struct FRecompileShaderRequest
+{
+	/** The platform name to compile for. */
+	FName PlatformName;
+	/** Shader platform, corresponds to EShaderPlatform, -1 indicates to compile for all target shader platforms. */
+	int32 ShaderPlatformToCompile = -1; // 
+	/** Materials to laod. */
+	TArray<FString> MaterialsToLoad;
+	/** On-deamnd shader compiler payload.  */
+	TArray<FODSCRequestPayload> ShadersToRecompile;
+	/** Mesh materials, returned to the caller.  */
+	TArray<uint8>* MeshMaterialMaps = nullptr;
+	/** All filenames that have been changed during the shader compilation. */
+	TArray<FString>* ModifiedFiles = nullptr;
+	/** Whether to compile changed shaders. */
+	bool bCompileChangedShaders = true;
+	/** Completion callback. */
+	FRecompileShaderCompletedCallback CompletionCallback;
+};
+
+/**
+ * Cook-on-the-fly server interface used by the request manager.
+ */
+class ICookOnTheFlyServer
+{
+public:
+	virtual ~ICookOnTheFlyServer() {}
+
+	/** Returns the cooker sandbox directory path. */
+	virtual FString GetSandboxDirectory() const = 0;
+	
+	/**
+	 * Add platform to the cook-on-the-fly session.
+	 *
+	 * @param PlatformName The platform name.
+	 */
+	virtual bool AddPlatform(const FName& PlatformName) = 0;
+
+	/**
+	 * Remove platform from the cook-on-the-fly session.
+	 *
+	 * @param PlatformName The platform name.
+	 */
+	virtual void RemovePlatform(const FName& PlatformName) = 0;
+
+	/**
+	 * Returns all asset file(s) that has been cooked.
+	 *
+	 * @param PlatformName The platform name.
+	 * @param OutPrecookedFiles Filename(s) returned to the caller.
+	 */
+	virtual void GetPrecookedFileList(const FName& PlatformName, FPrecookedFileList& OutPrecookedFiles) = 0;
+
+	/**
+	 * Returns all unsolicited files that has been produced has a result of a cook request.
+	 *
+	 * @param PlatformName The platform name.
+	 * @param Filename The filename.
+	 * @param Filename Whether the filename is a cookable asset.
+	 * @param OutUnsolicitedFiles All unsolicited filename(s) returned to the caller.
+	 */
+	virtual void GetUnsolicitedFiles(
+		const FName& PlatformName,
+		const FString& Filename,
+		const bool bIsCookable,
+		TArray<FString>& OutUnsolicitedFiles) = 0;
+
+	/**
+	 * Enqueue a new cook request.
+	 *
+	 * @param CookPackageRequest Cook package request parameters.
+	 */
+	virtual bool EnqueueCookRequest(FCookPackageRequest CookPackageRequest) = 0;
+
+	/**
+	 * Enqueue a new shader compile request.
+	 *
+	 * @param RecompileShaderRequest Recompile shader request parameters.
+	 */
+	virtual bool EnqueueRecompileShaderRequest(FRecompileShaderRequest RecompileShaderRequest) = 0;
+
+	/**
+	 * Returns the package store writer for the specified platform.
+	 *
+	 * @param PlatformName The platform name.
+	 */
+	virtual IPackageStoreWriter* GetPackageStoreWriter(const FName& PlatformName) = 0;
+
+	/**
+	 * A somewhat temporary event triggered when all pending packages has been cooked.
+	 * This is to ensure that all packages has been sent to all connected clients.
+	 */
+	DECLARE_EVENT(ICookOnTheFlyServer, FFlushEvent);
+	virtual FFlushEvent& OnFlush() = 0;
+
+	/**
+	 * Wait until any pending flush request is completed.
+	 *
+	 * @return Duration in seconds.
+	 */
+	virtual double WaitForPendingFlush() = 0;
+};
+
+/**
+ * The cook-on-the-fly request manager.
+ *
+ * Responsible for managing cook-on-the-fly requests from connected client(s).
+ */
+class ICookOnTheFlyRequestManager
+{
+public:
+	virtual ~ICookOnTheFlyRequestManager() {}
+
+	/** Initialze the request manager. */
+	virtual bool Initialize() = 0;
+
+	/** Shutdown the reques tmanager. */
+	virtual void Shutdown() = 0;
+};
+
+}} // namespace UE::Cook
diff --git a/Engine/Source/Editor/UnrealEd/Private/Cooker/CookPackageData.cpp b/Engine/Source/Editor/UnrealEd/Private/Cooker/CookPackageData.cpp
--- a/Engine/Source/Editor/UnrealEd/Private/Cooker/CookPackageData.cpp
+++ b/Engine/Source/Editor/UnrealEd/Private/Cooker/CookPackageData.cpp
@@ -630,7 +630,7 @@ namespace Cook
 		UE::Cook::FCompletionCallback LocalCompletionCallback(MoveTemp(GetCompletionCallback()));
 		if (LocalCompletionCallback)
 		{
-			LocalCompletionCallback();
+			LocalCompletionCallback(this);
 		}
 		ClearInProgressData();
 	}
diff --git a/Engine/Source/Editor/UnrealEd/Private/Cooker/CookTypes.h b/Engine/Source/Editor/UnrealEd/Private/Cooker/CookTypes.h
--- a/Engine/Source/Editor/UnrealEd/Private/Cooker/CookTypes.h
+++ b/Engine/Source/Editor/UnrealEd/Private/Cooker/CookTypes.h
@@ -54,8 +54,9 @@ namespace UE
 {
 namespace Cook
 {
+	struct FPackageData;
 	/** A function that is called when a requested package finishes cooking (when successful, failed, or skipped) */
-	typedef TUniqueFunction<void()> FCompletionCallback;
+	typedef TUniqueFunction<void(FPackageData*)> FCompletionCallback;
 
 	/** External Requests to the cooker can either by cook requests for a specific file, or arbitrary callbacks that need to execute within the Scheduler's lock. */
 	enum class EExternalRequestType
diff --git a/Engine/Source/Editor/UnrealEd/Private/Cooker/IoStoreCookOnTheFlyRequestManager.cpp b/Engine/Source/Editor/UnrealEd/Private/Cooker/IoStoreCookOnTheFlyRequestManager.cpp
new file mode 100644
--- /dev/null
+++ b/Engine/Source/Editor/UnrealEd/Private/Cooker/IoStoreCookOnTheFlyRequestManager.cpp
@@ -0,0 +1,751 @@
+// Copyright Epic Games, Inc. All Rights Reserved.
+
+#include "IoStoreCookOnTheFlyRequestManager.h"
+#include "CookOnTheFlyServerInterface.h"
+#include "CookOnTheFly.h"
+#include "Modules/ModuleManager.h"
+#include "Serialization/CookOnTheFlyPackageStore.h"
+#include "IO/PackageStoreWriter.h"
+#include "ShaderCompiler.h"
+#include "HAL/PlatformTime.h"
+#include "Misc/PackageName.h"
+#include "UObject/SavePackage.h"
+#include "ProfilingDebugging/CountersTrace.h"
+
+PRAGMA_DISABLE_OPTIMIZATION
+
+DEFINE_LOG_CATEGORY_STATIC(LogCookOnTheFlyTracker, Log, All);
+
+class FCookOnTheFlyPackageTracker
+{
+	enum class EPackageStatus
+	{
+		None,
+		Cooking,
+		WaitingForImports,
+		Cooked,
+		Failed
+	};
+
+	inline const TCHAR* LexToString(EPackageStatus Status)
+	{
+		switch (Status)
+		{
+			case EPackageStatus::None:
+				return TEXT("None");
+			case EPackageStatus::Cooking:
+				return TEXT("Cooking");
+			case EPackageStatus::WaitingForImports:
+				return TEXT("WaitingForImports");
+			case EPackageStatus::Cooked:
+				return TEXT("Cooked");
+			case EPackageStatus::Failed:
+				return TEXT("Failed");
+			default:
+				return TEXT("Unknown");
+		};
+	};
+
+	struct FPackage
+	{
+		FPackageId PackageId;
+		FName PackageName;
+		EPackageStatus Status = EPackageStatus::None;
+		int32 EntryIndex = INDEX_NONE;
+		TSet<FPackageId> WaitingPackages;
+		TSet<FPackageId> PendingImports;
+		bool bAnyImportFailed = false;
+	};
+
+public:
+
+	struct FCompletedPackages
+	{
+		TArray<int32> CookedEntryIndices;
+		TArray<FPackageId> FailedPackages;
+	};
+
+	TArrayView<const int32> GetCookedEntryIndices() const
+	{
+		return CookedEntries;
+	}
+
+	TArrayView<const FName> GetFailedPackages() const
+	{
+		return FailedPackages; 
+	}
+
+	int32 TotalTracked() const
+	{
+		return Packages.Num();
+	}
+
+	EPackageStoreEntryStatus GetStatus(const FName& PackageName)
+	{
+		FPackage& Package = GetPackage(FPackageId::FromName(PackageName));
+
+		switch(Package.Status)
+		{
+			case EPackageStatus::None:
+				return EPackageStoreEntryStatus::None;
+			case EPackageStatus::Cooking:
+			case EPackageStatus::WaitingForImports:
+				return EPackageStoreEntryStatus::Pending;
+			case EPackageStatus::Cooked:
+				return EPackageStoreEntryStatus::Ok;
+			default:
+				return EPackageStoreEntryStatus::Missing;
+		}
+	}
+
+	bool IsCooked(const FName& PackageName)
+	{
+		FPackage& Package = GetPackage(FPackageId::FromName(PackageName));
+		return Package.Status != EPackageStatus::None;
+	}
+
+	void MarkAsRequested(const FName& PackageName)
+	{
+		FPackage& Package = GetPackage(FPackageId::FromName(PackageName));
+		check(Package.Status == EPackageStatus::None);
+		Package.Status = EPackageStatus::Cooking;
+	}
+
+	FCompletedPackages MarkAsFailed(const FName& PackageName)
+	{
+		return MarkAsCompleted(PackageName, INDEX_NONE, TArrayView<const FPackageStoreEntryResource>());
+	}
+
+	FCompletedPackages MarkAsCompleted(const FName& PackageName, const int32 EntryIndex, TArrayView<const FPackageStoreEntryResource> Entries)
+	{
+		const int32 NumCookedPackages = CookedEntries.Num();
+		const int32 NumFailedPackages = FailedPackages.Num();
+
+		const FPackageId PackageId = FPackageId::FromName(PackageName);
+
+		FPackage& Package	= GetPackage(PackageId);
+		check(Package.Status == EPackageStatus::None || Package.Status == EPackageStatus::Cooking);
+
+		Package.PackageId	= PackageId;
+		Package.PackageName = PackageName;
+		Package.Status		= EntryIndex != INDEX_NONE ? EPackageStatus::Cooked : EPackageStatus::Failed;
+		Package.EntryIndex	= EntryIndex;
+
+		if (Package.Status == EPackageStatus::Cooked)
+		{
+			const FPackageStoreEntryResource& Entry = Entries[EntryIndex];
+
+			for (FPackageId ImportedPackageId : Entry.ImportedPackageIds)
+			{
+				FPackage& ImportedPackage = GetPackage(ImportedPackageId);
+				bool bWaitForImport = ImportedPackage.Status == EPackageStatus::None || ImportedPackage.Status == EPackageStatus::Cooking;
+
+				if (ImportedPackage.Status == EPackageStatus::WaitingForImports)
+				{
+					// Check if this import or any of it's waiting imports is waiting for this cooked package
+					const bool bAnyWaitingForPackage = IsAnyWaitingForPackage(Package, ImportedPackageId);
+					bWaitForImport = !bAnyWaitingForPackage;
+				}
+
+				if (bWaitForImport)
+				{
+					Package.Status = EPackageStatus::WaitingForImports;
+
+					bool bIsAlreadyInSet = false;
+					ImportedPackage.WaitingPackages.Add(PackageId, &bIsAlreadyInSet);
+					check(!bIsAlreadyInSet);
+
+					Package.PendingImports.Add(ImportedPackageId, &bIsAlreadyInSet);
+					check(!bIsAlreadyInSet);
+				}
+			}
+		}
+
+		if (Package.Status == EPackageStatus::Cooked || Package.Status == EPackageStatus::Failed)
+		{
+			AddCompletedPackage(Package);
+			ResolveWaitingPackages(Package);
+		}
+		else
+		{
+			AddWaitingPackage(Package);
+		}
+
+		FCompletedPackages CompletedPackages;
+
+		if (CookedEntries.Num() > NumCookedPackages || FailedPackages.Num() > NumFailedPackages)
+		{
+			for (int32 Idx = NumCookedPackages; Idx < CookedEntries.Num(); ++Idx)
+			{
+				CompletedPackages.CookedEntryIndices.Add(CookedEntries[Idx]);
+			}
+
+			for (int32 Idx = NumFailedPackages; Idx < FailedPackages.Num(); ++Idx)
+			{
+				CompletedPackages.FailedPackages.Add(FPackageId::FromName(FailedPackages[Idx]));
+			}
+		}
+
+		return MoveTemp(CompletedPackages);
+	}
+
+	FCompletedPackages Flush()
+	{
+		UE_LOG(LogCookOnTheFlyTracker, Log, TEXT("Flushing '%d' waiting package(s)"), WaitingPackages.Num());
+
+		FCompletedPackages CompletedPackages;
+
+		for (auto& KeyValue : Packages)
+		{
+			TUniquePtr<FPackage>& Package = KeyValue.Value;
+			
+			if (Package->Status == EPackageStatus::WaitingForImports)
+			{
+				check(Package->EntryIndex != INDEX_NONE);
+				
+				UE_LOG(LogCookOnTheFlyTracker, Warning, TEXT("Package '%s' is still waiting for:"), *Package->PackageName.ToString());
+				for (FPackageId PendingImportId : Package->PendingImports)
+				{
+					const TUniquePtr<FPackage>* ImportedPackage = Packages.Find(PendingImportId);
+					UE_LOG(LogCookOnTheFlyTracker, Warning, TEXT( " ==> '%s'"),
+						ImportedPackage
+							? *FString::Printf(TEXT("%s, (Status=%s)"), *(*ImportedPackage)->PackageName.ToString(), LexToString((*ImportedPackage)->Status))
+							: *FString::Printf(TEXT("<Untracked PackageId='0x%llX'>"), PendingImportId.Value()));
+				}
+				
+				Package->WaitingPackages.Empty();
+				Package->PendingImports.Empty();
+
+				CookedEntries.Add(Package->EntryIndex);
+				CompletedPackages.CookedEntryIndices.Add(Package->EntryIndex);
+			}
+		}
+
+		for (auto& KeyValue : Packages)
+		{
+			TUniquePtr<FPackage>& Package = KeyValue.Value;
+			
+			if (Package->Status == EPackageStatus::WaitingForImports)
+			{
+				Package->Status = EPackageStatus::Cooked;
+			}
+		}
+
+		WaitingPackages.Empty();
+
+		return MoveTemp(CompletedPackages);
+	}
+
+private:
+
+	void ResolveWaitingPackages(FPackage& Package)
+	{
+		check(Package.Status == EPackageStatus::Cooked || Package.Status == EPackageStatus::Failed);
+
+		TSet<FPackageId> ResolvedPackageIds;
+
+		for (FPackageId WaitingPackageId : Package.WaitingPackages)
+		{
+			FPackage& WaitingPackage = GetPackage(WaitingPackageId);
+			check(WaitingPackage.PendingImports.Contains(Package.PackageId));
+			check(WaitingPackage.Status == EPackageStatus::WaitingForImports);
+			WaitingPackage.PendingImports.Remove(Package.PackageId);
+			WaitingPackage.bAnyImportFailed |= Package.Status == EPackageStatus::Failed;
+
+			if (!WaitingPackage.PendingImports.Num())
+			{
+				WaitingPackage.Status = WaitingPackage.bAnyImportFailed ? EPackageStatus::Failed : Package.Status;
+				ResolvedPackageIds.Add(WaitingPackage.PackageId);
+				WaitingPackages.Remove(WaitingPackage.PackageName);
+				AddCompletedPackage(WaitingPackage);
+			}
+		}
+
+		Package.WaitingPackages.Empty();
+
+		for (FPackageId ResolvedPackageId : ResolvedPackageIds)
+		{
+			FPackage& ResolvedPackage = GetPackage(ResolvedPackageId);
+			ResolveWaitingPackages(ResolvedPackage);
+		}
+	}
+
+	bool IsAnyWaitingForPackageRecursive(const FPackage& Package, FPackageId TargetPackageId, TSet<FPackageId>& Visited)
+	{
+		Visited.Add(TargetPackageId);
+
+		FPackage& TargetPackage = GetPackage(TargetPackageId);
+		
+		if (TargetPackage.Status == EPackageStatus::WaitingForImports)
+		{
+			if (Package.WaitingPackages.Contains(TargetPackageId))
+			{
+				return true;
+			}
+
+			for (FPackageId PendingImportId : TargetPackage.PendingImports)
+			{
+				if (!Visited.Contains(PendingImportId))
+				{
+					const bool bIsWaitingForPackage = IsAnyWaitingForPackageRecursive(Package, PendingImportId, Visited);
+					if (bIsWaitingForPackage)
+					{
+						return true;
+					}
+				}
+			}
+		}
+	
+		return false;
+	}
+
+	bool IsAnyWaitingForPackage(const FPackage& Package, FPackageId TargetPackageId)
+	{
+		TSet<FPackageId> Visited;
+		return IsAnyWaitingForPackageRecursive(Package, TargetPackageId, Visited);
+	}
+
+	FPackage& GetPackage(FPackageId PackageId)
+	{
+		TUniquePtr<FPackage>& Package = Packages.FindOrAdd(PackageId, MakeUnique<FPackage>());
+		check(Package.IsValid());
+		return *Package;
+	}
+	
+	void AddCompletedPackage(FPackage& Package)
+	{
+		check(Package.Status == EPackageStatus::Cooked || Package.Status == EPackageStatus::Failed);
+
+		if (Package.Status == EPackageStatus::Cooked)
+		{
+			CookedEntries.Add(Package.EntryIndex);
+		}
+		else
+		{
+			FailedPackages.Add(Package.PackageName);
+		}
+
+		UE_LOG(LogCookOnTheFlyTracker, Log, TEXT("'%s' (0x%llX) [%s], Cooked/Fully/Waiting/Failed='%d/%d/%d/%d'"),
+			*Package.PackageName.ToString(), Package.PackageId.Value(), LexToString(Package.Status),
+			CookedEntries.Num() + WaitingPackages.Num(), CookedEntries.Num(), WaitingPackages.Num(), FailedPackages.Num());
+	}
+
+	void AddWaitingPackage(FPackage& Package)
+	{
+		check(Package.Status == EPackageStatus::WaitingForImports);
+
+		bool bIsAlreadyInSet = false;
+		WaitingPackages.Add(Package.PackageName, &bIsAlreadyInSet);
+		check(!bIsAlreadyInSet);
+
+		UE_LOG(LogCookOnTheFlyTracker, Log, TEXT("'%s' (0x%llX) [%s], Cooked/Fully/Waiting/Failed='%d/%d/%d/%d'"),
+			*Package.PackageName.ToString(), Package.PackageId.Value(), LexToString(Package.Status),
+			CookedEntries.Num() + WaitingPackages.Num(), CookedEntries.Num(), WaitingPackages.Num(), FailedPackages.Num());
+
+		for (const FPackageId ImportedPackageId : Package.PendingImports)
+		{
+			FPackage& ImportedPackage = GetPackage(ImportedPackageId);
+
+			UE_LOG(LogCookOnTheFlyTracker, Verbose, TEXT(" ==> '%s', Status='%s'"),
+				ImportedPackage.PackageName.IsNone()
+					? *FString::Printf(TEXT("0x%llX"), ImportedPackageId.Value())
+					: *ImportedPackage.PackageName.ToString(),
+				LexToString(ImportedPackage.Status));
+		}
+	}
+
+	TMap<FPackageId, TUniquePtr<FPackage>> Packages;
+	TArray<int32> CookedEntries;
+	TArray<FName> FailedPackages;
+	TSet<FName> WaitingPackages;
+};
+
+class FIoStoreCookOnTheFlyRequestManager final
+	: public UE::Cook::ICookOnTheFlyRequestManager
+{
+public:
+	FIoStoreCookOnTheFlyRequestManager(UE::Cook::ICookOnTheFlyServer& InCookOnTheFlyServer, UE::Cook::FIoStoreCookOnTheFlyServerOptions InOptions)
+		: CookOnTheFlyServer(InCookOnTheFlyServer)
+		, Options(MoveTemp(InOptions))
+	{
+		CookOnTheFlyServer.OnFlush().AddRaw(this, &FIoStoreCookOnTheFlyRequestManager::OnFlush);
+	}
+
+private:
+	struct FPlatformContext
+	{
+		FCriticalSection CriticalSection;
+		FName PlatformName;
+		FCookOnTheFlyPackageTracker PackageTracker;
+		IPackageStoreWriter* PackageStoreWriter = nullptr;
+	};
+
+	virtual bool Initialize() override
+	{
+		using namespace UE::Cook;
+
+		ICookOnTheFlyModule& CookOnTheFlyModule = FModuleManager::LoadModuleChecked<ICookOnTheFlyModule>(TEXT("CookOnTheFly"));
+
+		const int32 Port = Options.Port > 0 ? Options.Port : UE::Cook::DefaultCookOnTheFlyServingPort;
+
+		ConnectionServer = CookOnTheFlyModule.CreateConnectionServer(FCookOnTheFlyServerOptions
+		{
+			Port,
+			[this](FCookOnTheFlyClient Client, ECookOnTheFlyConnectionStatus ConnectionStatus)
+			{
+				return HandleClientConnection(Client, ConnectionStatus);
+			},
+			[this](FCookOnTheFlyClient Client, const FCookOnTheFlyRequest& Request, FCookOnTheFlyResponse& Response)
+			{ 
+				return HandleClientRequest(Client, Request, Response);
+			}
+		});
+
+		const bool bServerStarted = ConnectionServer->StartServer();
+
+		return bServerStarted;
+	}
+
+	virtual void Shutdown() override
+	{
+		ConnectionServer->StopServer();
+		ConnectionServer.Reset();
+	}
+
+private:
+	bool HandleClientConnection(UE::Cook::FCookOnTheFlyClient Client, UE::Cook::ECookOnTheFlyConnectionStatus ConnectionStatus)
+	{
+		FScopeLock _(&CriticalSection);
+
+		if (ConnectionStatus == UE::Cook::ECookOnTheFlyConnectionStatus::Connected)
+		{
+			if (CookOnTheFlyServer.AddPlatform(Client.PlatformName))
+			{
+				if (!PlatformContexts.Contains(Client.PlatformName))
+				{
+					TUniquePtr<FPlatformContext>& Context = PlatformContexts.Add(Client.PlatformName, MakeUnique<FPlatformContext>());
+					Context->PlatformName = Client.PlatformName;
+					Context->PackageStoreWriter = CookOnTheFlyServer.GetPackageStoreWriter(Client.PlatformName);
+					Context->PackageStoreWriter->OnCommit().AddRaw(this, &FIoStoreCookOnTheFlyRequestManager::OnPackageCooked); 
+				}
+
+				return true;
+			}
+
+			return false;
+		}
+		else
+		{
+			CookOnTheFlyServer.RemovePlatform(Client.PlatformName);
+			return true;
+		}
+	}
+
+	bool HandleClientRequest(UE::Cook::FCookOnTheFlyClient Client, const UE::Cook::FCookOnTheFlyRequest& Request, UE::Cook::FCookOnTheFlyResponse& Response)
+	{
+		bool bRequestOk = false;
+
+		const double StartTime = FPlatformTime::Seconds();
+
+		UE_LOG(LogCookOnTheFly, Verbose, TEXT("New request, Type='%s', Client='%s (%d)'"), LexToString(Request.GetHeader().MessageType), *Client.PlatformName.ToString(), Client.ClientId);
+
+		switch (Request.GetHeader().MessageType)
+		{
+			case UE::Cook::ECookOnTheFlyMessage::CookPackage:
+				bRequestOk = HandleCookPackageRequest(Client, Request, Response);
+				break;
+			case UE::Cook::ECookOnTheFlyMessage::GetCookedPackages:
+				bRequestOk = HandleGetCookedPackagesRequest(Client, Request, Response);
+				break;
+			case UE::Cook::ECookOnTheFlyMessage::RecompileShaders:
+				bRequestOk = HandleRecompileShadersRequest(Client, Request, Response);
+				break;
+			default:
+				UE_LOG(LogCookOnTheFly, Fatal, TEXT("Unknown request, Type='%s', Client='%s (%d)'"), LexToString(Request.GetHeader().MessageType), *Client.PlatformName.ToString(), Client.ClientId);
+				break;
+		}
+
+		const double Duration = FPlatformTime::Seconds() - StartTime;
+
+		UE_LOG(LogCookOnTheFly, Verbose, TEXT("Request handled, Type='%s', Client='%s (%u)', Status='%s', Duration='%.6lfs'"),
+			LexToString(Request.GetHeader().MessageType),
+			*Client.PlatformName.ToString(), Client.ClientId,
+			bRequestOk ? TEXT("Ok") : TEXT("Failed"),
+			Duration);
+
+		return bRequestOk;
+	}
+
+	bool HandleGetCookedPackagesRequest(UE::Cook::FCookOnTheFlyClient Client, const UE::Cook::FCookOnTheFlyRequest& Request, UE::Cook::FCookOnTheFlyResponse& Response)
+	{
+		using namespace UE::Cook;
+		using namespace UE::PackageStore::Messaging;
+
+		FPlatformContext& Context = GetContext(Client.PlatformName);
+		FScopeLock _(&Context.CriticalSection);
+		
+		const TArrayView<const int32> CookedEntryIndices = Context.PackageTracker.GetCookedEntryIndices();
+		const TArrayView<const FName> FailedPackages = Context.PackageTracker.GetFailedPackages();
+
+		UE::PackageStore::Messaging::FPackageStoreData PackageStoreData;
+
+		PackageStoreData.TotalCookedPackages	= CookedEntryIndices.Num();
+		PackageStoreData.TotalFailedPackages	= FailedPackages.Num();
+		//PackageStoreData.FailedPackages			= FailedPackages;
+
+		Context.PackageStoreWriter->GetEntries([&CookedEntryIndices, &PackageStoreData](TArrayView<const FPackageStoreEntryResource> Entries)
+		{
+			for (int32 EntryIndex : CookedEntryIndices)
+			{
+				check(EntryIndex != INDEX_NONE);
+				PackageStoreData.CookedPackages.Add(Entries[EntryIndex]);
+			}
+		});
+
+		for (const FName& FailedPackageName : FailedPackages)
+		{
+			PackageStoreData.FailedPackages.Add(FPackageId::FromName(FailedPackageName));
+		}
+
+		Response.SetBodyTo(FGetCookedPackagesResponse { MoveTemp(PackageStoreData) });
+		Response.SetStatus(ECookOnTheFlyMessageStatus::Ok);
+
+		return true;
+	}
+
+	bool HandleCookPackageRequest(UE::Cook::FCookOnTheFlyClient Client, const UE::Cook::FCookOnTheFlyRequest& Request, UE::Cook::FCookOnTheFlyResponse& Response)
+	{
+		using namespace UE::PackageStore::Messaging;
+
+		TRACE_CPUPROFILER_EVENT_SCOPE(CookOnTheFly::HandleCookPackageRequest);
+
+		const double FlushDurationInSeconds = CookOnTheFlyServer.WaitForPendingFlush();
+		UE_CLOG(FlushDurationInSeconds > 0.0, LogCookOnTheFly, Log, TEXT("Waited '%.2llf's for cooker to flush pending package(s), Client='%s (%u)'"),
+			FlushDurationInSeconds, *Client.PlatformName.ToString(), Client.ClientId);
+
+		FPlatformContext& Context = GetContext(Client.PlatformName);
+		FScopeLock _(&Context.CriticalSection);
+
+		FCookPackageRequest CookRequest = Request.GetBodyAs<FCookPackageRequest>();
+		
+		UE_LOG(LogCookOnTheFly, Log, TEXT("Received cook request, PackageName='%s'"), *CookRequest.PackageName.ToString());
+
+		const bool bIsCooked = Context.PackageTracker.IsCooked(CookRequest.PackageName);
+
+		if (!bIsCooked)
+		{
+			Context.PackageTracker.MarkAsRequested(CookRequest.PackageName);
+
+			FString Filename;
+			if (FPackageName::TryConvertLongPackageNameToFilename(CookRequest.PackageName.ToString(), Filename))
+			{
+				const bool bEnqueued = CookOnTheFlyServer.EnqueueCookRequest(UE::Cook::FCookPackageRequest { Client.PlatformName, Filename });
+				check(bEnqueued);
+			}
+			else
+			{
+				UE_LOG(LogCookOnTheFly, Warning, TEXT("Failed to cook package '%s' (File not found)"), *CookRequest.PackageName.ToString());
+				Context.PackageTracker.MarkAsFailed(CookRequest.PackageName);
+			}
+		}
+
+		Response.SetBodyTo(FCookPackageResponse { Context.PackageTracker.GetStatus(CookRequest.PackageName) });
+		Response.SetStatus(UE::Cook::ECookOnTheFlyMessageStatus::Ok);
+
+		return true;
+	}
+
+	void OnPackageCooked(const IPackageStoreWriter::FCommitEventArgs& EventArgs)
+	{
+		using namespace UE::Cook;
+		using namespace UE::PackageStore::Messaging;
+
+		TRACE_CPUPROFILER_EVENT_SCOPE(CookOnTheFly::OnPackageCooked);
+
+		FPlatformContext& Context = GetContext(EventArgs.PlatformName);
+		FScopeLock _(&Context.CriticalSection);
+
+		if (EventArgs.AdditionalFiles.Num())
+		{
+			TArray<FString> Filenames;
+			TArray<FIoChunkId> ChunkIds;
+
+			for (const IPackageStoreWriter::FAdditionalFileInfo& FileInfo : EventArgs.AdditionalFiles)
+			{
+				UE_LOG(LogCookOnTheFly, Verbose, TEXT("Sending additional cooked file '%s'"), *FileInfo.Filename);
+				Filenames.Add(FileInfo.Filename);
+				ChunkIds.Add(FileInfo.ChunkId);
+			}
+
+			FCookOnTheFlyMessage Message(ECookOnTheFlyMessage::FilesAdded);
+			{
+				TUniquePtr<FArchive> Ar = Message.WriteBody();
+				*Ar << Filenames;
+				*Ar << ChunkIds;
+			}
+
+			ConnectionServer->BroadcastMessage(Message, Context.PlatformName);
+		}
+
+		FCookOnTheFlyPackageTracker::FCompletedPackages NewCompletedPackages;
+		{
+			TRACE_CPUPROFILER_EVENT_SCOPE(CookOnTheFly::MarkAsCooked);
+			NewCompletedPackages = Context.PackageTracker.MarkAsCompleted(EventArgs.PackageName, EventArgs.EntryIndex, EventArgs.Entries);
+		}
+
+		if (NewCompletedPackages.CookedEntryIndices.Num() || NewCompletedPackages.FailedPackages.Num())
+		{
+			TRACE_CPUPROFILER_EVENT_SCOPE(CookOnTheFly::SendCookedPackagesMessage);
+
+			FPackageStoreData PackageStoreData;
+			PackageStoreData.TotalCookedPackages = Context.PackageTracker.GetCookedEntryIndices().Num();
+			PackageStoreData.TotalFailedPackages = Context.PackageTracker.GetFailedPackages().Num();
+			PackageStoreData.FailedPackages = MoveTemp(NewCompletedPackages.FailedPackages);
+
+			for (int32 EntryIndex : NewCompletedPackages.CookedEntryIndices)
+			{
+				check(EntryIndex != INDEX_NONE);
+				PackageStoreData.CookedPackages.Add(EventArgs.Entries[EntryIndex]);
+			}
+
+			UE_LOG(LogCookOnTheFly, Verbose, TEXT("Sending '%s' message, Cooked='%d', Failed='%d', Server total='%d/%d/%d' (Tracked/Cooked/Failed)"),
+				LexToString(ECookOnTheFlyMessage::PackagesCooked),
+				PackageStoreData.CookedPackages.Num(),
+				PackageStoreData.FailedPackages.Num(),
+				Context.PackageTracker.TotalTracked(),
+				PackageStoreData.TotalCookedPackages,
+				PackageStoreData.TotalFailedPackages);
+
+			FCookOnTheFlyMessage Message(ECookOnTheFlyMessage::PackagesCooked);
+			Message.SetBodyTo<FPackagesCookedMessage>(FPackagesCookedMessage { MoveTemp(PackageStoreData) });
+			ConnectionServer->BroadcastMessage(Message, Context.PlatformName);
+		}
+	}
+
+	bool HandleRecompileShadersRequest(UE::Cook::FCookOnTheFlyClient Client, const UE::Cook::FCookOnTheFlyRequest& Request, UE::Cook::FCookOnTheFlyResponse& Response)
+	{
+		TRACE_CPUPROFILER_EVENT_SCOPE(CookOnTheFly::HandleRecompileShadersRequest);
+
+		TArray<FString> RecompileModifiedFiles;
+		TArray<uint8> MeshMaterialMaps;
+
+		FShaderRecompileData RecompileData;
+		RecompileData.PlatformName = Client.PlatformName.ToString();
+		RecompileData.ModifiedFiles = &RecompileModifiedFiles;
+		RecompileData.MeshMaterialMaps = &MeshMaterialMaps;
+
+		{
+			TUniquePtr<FArchive> Ar = Request.ReadBody();
+			*Ar << RecompileData.MaterialsToLoad;
+			*Ar << RecompileData.ShaderPlatform;
+			*Ar << RecompileData.bCompileChangedShaders;
+			*Ar << RecompileData.ShadersToRecompile;
+		}
+
+		RecompileShaders(RecompileData);
+
+		{
+			TUniquePtr<FArchive> Ar = Response.WriteBody();
+			*Ar << MeshMaterialMaps;
+		}
+
+		Response.SetStatus(UE::Cook::ECookOnTheFlyMessageStatus::Ok);
+
+		return true;
+	}
+
+	void RecompileShaders(FShaderRecompileData& RecompileData)
+	{
+		FEvent* RecompileCompletedEvent = FPlatformProcess::GetSynchEventFromPool();
+		UE::Cook::FRecompileShaderCompletedCallback RecompileCompleted = [this, RecompileCompletedEvent]()
+		{
+			RecompileCompletedEvent->Trigger();
+		};
+
+		const bool bEnqueued = CookOnTheFlyServer.EnqueueRecompileShaderRequest(UE::Cook::FRecompileShaderRequest
+		{
+			FName(*RecompileData.PlatformName),
+			RecompileData.ShaderPlatform,
+			RecompileData.MaterialsToLoad,
+			RecompileData.ShadersToRecompile,
+			RecompileData.MeshMaterialMaps,
+			RecompileData.ModifiedFiles,
+			RecompileData.bCompileChangedShaders,
+			MoveTemp(RecompileCompleted)
+		});
+		check(bEnqueued);
+
+		RecompileCompletedEvent->Wait();
+		FPlatformProcess::ReturnSynchEventToPool(RecompileCompletedEvent);
+	}
+
+	void OnFlush()
+	{
+		using namespace UE::Cook;
+		using namespace UE::PackageStore::Messaging;
+
+		FScopeLock _(&CriticalSection);
+
+		for (auto& KeyValue : PlatformContexts)
+		{
+			const FName PlatformName = KeyValue.Key;
+			FPlatformContext& Context = *KeyValue.Value;
+			FScopeLock PlatformLock(&Context.CriticalSection);
+
+			FCookOnTheFlyPackageTracker::FCompletedPackages CompletedPackages = Context.PackageTracker.Flush();
+
+			if (CompletedPackages.CookedEntryIndices.Num() || CompletedPackages.FailedPackages.Num())
+			{
+				UE_LOG(LogCookOnTheFly, Warning, TEXT("Flushing '%d' pending package entry(s) that should have been completed"),
+					CompletedPackages.CookedEntryIndices.Num() || CompletedPackages.FailedPackages.Num());
+
+				FPackageStoreData PackageStoreData;
+				PackageStoreData.TotalCookedPackages = Context.PackageTracker.GetCookedEntryIndices().Num();
+				PackageStoreData.TotalFailedPackages = Context.PackageTracker.GetFailedPackages().Num();
+				PackageStoreData.FailedPackages = MoveTemp(CompletedPackages.FailedPackages);
+
+				Context.PackageStoreWriter->GetEntries([&CompletedPackages, &PackageStoreData](TArrayView<const FPackageStoreEntryResource> Entries)
+				{
+					for (int32 EntryIndex : CompletedPackages.CookedEntryIndices)
+					{
+						check(EntryIndex != INDEX_NONE);
+						PackageStoreData.CookedPackages.Add(Entries[EntryIndex]);
+					}
+				});
+
+				UE_LOG(LogCookOnTheFly, Log, TEXT("Sending cooked package(s) message, Cooked = '%d', Failed = '%d', Total cooked = '%d', Total failed = '%d'"),
+					PackageStoreData.CookedPackages.Num(), PackageStoreData.FailedPackages.Num(), PackageStoreData.TotalCookedPackages, PackageStoreData.TotalFailedPackages);
+
+				FCookOnTheFlyMessage Message(ECookOnTheFlyMessage::PackagesCooked);
+				Message.SetBodyTo<FPackagesCookedMessage>(FPackagesCookedMessage { MoveTemp(PackageStoreData) });
+				ConnectionServer->BroadcastMessage(Message, Context.PlatformName);
+			}
+		}
+	}
+
+	FPlatformContext& GetContext(const FName& PlatformName)
+	{
+		FScopeLock _(&CriticalSection);
+		TUniquePtr<FPlatformContext>& Ctx = PlatformContexts.FindChecked(PlatformName);
+		check(Ctx.IsValid());
+		return *Ctx;
+	}
+
+	UE::Cook::ICookOnTheFlyServer& CookOnTheFlyServer;
+	UE::Cook::FIoStoreCookOnTheFlyServerOptions Options;
+	TUniquePtr<UE::Cook::ICookOnTheFlyConnectionServer> ConnectionServer;
+	FCriticalSection CriticalSection;
+	TMap<FName, TUniquePtr<FPlatformContext>> PlatformContexts;
+};
+
+namespace UE { namespace Cook
+{
+
+TUniquePtr<ICookOnTheFlyRequestManager> MakeIoStoreCookOnTheFlyRequestManager(ICookOnTheFlyServer& CookOnTheFlyServer, FIoStoreCookOnTheFlyServerOptions Options)
+{
+	return MakeUnique<FIoStoreCookOnTheFlyRequestManager>(CookOnTheFlyServer, Options);
+}
+
+}}
+
+PRAGMA_ENABLE_OPTIMIZATION
diff --git a/Engine/Source/Editor/UnrealEd/Private/Cooker/IoStoreCookOnTheFlyRequestManager.h b/Engine/Source/Editor/UnrealEd/Private/Cooker/IoStoreCookOnTheFlyRequestManager.h
new file mode 100644
--- /dev/null
+++ b/Engine/Source/Editor/UnrealEd/Private/Cooker/IoStoreCookOnTheFlyRequestManager.h
@@ -0,0 +1,20 @@
+// Copyright Epic Games, Inc. All Rights Reserved.
+
+#pragma once
+
+#include "CoreMinimal.h"
+
+namespace UE { namespace Cook
+{
+
+class ICookOnTheFlyServer;
+class ICookOnTheFlyRequestManager;
+
+struct FIoStoreCookOnTheFlyServerOptions
+{
+	int32 Port = -1; // -1 indicates the default COTF serving port
+};
+
+TUniquePtr<ICookOnTheFlyRequestManager> MakeIoStoreCookOnTheFlyRequestManager(ICookOnTheFlyServer& CookOnTheFlyServer, FIoStoreCookOnTheFlyServerOptions Options);
+
+}} // namespace UE::Cook
diff --git a/Engine/Source/Editor/UnrealEd/Private/Cooker/NetworkFileCookOnTheFlyRequestManager.cpp b/Engine/Source/Editor/UnrealEd/Private/Cooker/NetworkFileCookOnTheFlyRequestManager.cpp
new file mode 100644
--- /dev/null
+++ b/Engine/Source/Editor/UnrealEd/Private/Cooker/NetworkFileCookOnTheFlyRequestManager.cpp
@@ -0,0 +1,134 @@
+// Copyright Epic Games, Inc. All Rights Reserved.
+
+#include "Cooker/NetworkFileCookOnTheFlyRequestManager.h"
+#include "CookOnTheFlyServerInterface.h"
+#include "INetworkFileServer.h"
+#include "INetworkFileSystemModule.h"
+#include "Modules/ModuleManager.h"
+
+class FNetworkFileCookOnTheFlyRequestManager final
+	: public UE::Cook::ICookOnTheFlyRequestManager
+{
+public:
+	FNetworkFileCookOnTheFlyRequestManager(UE::Cook::ICookOnTheFlyServer& InCookOnTheFlyServer, FNetworkFileServerOptions FileServerOptions)
+		: CookOnTheFlyServer(InCookOnTheFlyServer)
+	{
+		FileServerOptions.Delegates.NewConnectionDelegate			= FNewConnectionDelegate::CreateRaw(this, &FNetworkFileCookOnTheFlyRequestManager::OnNewConnection);
+		FileServerOptions.Delegates.InitialPrecookedListDelegate	= FInitialPrecookedListDelegate::CreateRaw(this, &FNetworkFileCookOnTheFlyRequestManager::OnGetPrecookedList);
+		FileServerOptions.Delegates.FileRequestDelegate				= FFileRequestDelegate::CreateRaw(this, &FNetworkFileCookOnTheFlyRequestManager::OnFileRequest);
+		FileServerOptions.Delegates.RecompileShadersDelegate		= FRecompileShadersDelegate::CreateRaw(this, &FNetworkFileCookOnTheFlyRequestManager::OnRecompileShaders);
+		FileServerOptions.Delegates.SandboxPathOverrideDelegate		= FSandboxPathDelegate::CreateRaw(this, &FNetworkFileCookOnTheFlyRequestManager::OnGetSandboxPath);
+		FileServerOptions.Delegates.OnFileModifiedCallback			= &FileModifiedDelegate;
+		FileServerOptions.bRestrictPackageAssetsToSandbox			= true; // prevents sending uncooked packages if the package fails to cook
+
+		NetworkFileServer.Reset(FModuleManager::LoadModuleChecked<INetworkFileSystemModule>("NetworkFileSystem")
+			.CreateNetworkFileServer(MoveTemp(FileServerOptions), /* bLoadTargetPlatforms */ true));
+	}
+
+	virtual ~FNetworkFileCookOnTheFlyRequestManager()
+	{
+		Shutdown();
+	}
+
+	virtual bool Initialize() override
+	{
+		return NetworkFileServer.IsValid();
+	}
+
+	virtual void Shutdown() override
+	{
+		if (NetworkFileServer.IsValid())
+		{
+			NetworkFileServer->Shutdown();
+			NetworkFileServer.Reset();
+		}
+	}
+
+private:
+	bool OnNewConnection(const FString& VersionInfo, const FString& PlatformName)
+	{
+		return CookOnTheFlyServer.AddPlatform(FName(*PlatformName));
+	}
+
+	void OnGetPrecookedList(const FString& PlatformName, TMap<FString, FDateTime>& PrecookedFileList)
+	{
+		CookOnTheFlyServer.GetPrecookedFileList(FName(*PlatformName), PrecookedFileList);
+	}
+
+	void OnFileRequest(FString& Filename, const FString& PlatformNameStr, TArray<FString>& UnsolicitedFiles)
+	{
+		using namespace UE::Cook;
+
+		FName PlatformName(*PlatformNameStr);
+		const bool bIsCookable = FPackageName::IsPackageExtension(*FPaths::GetExtension(Filename, true));
+
+		if (bIsCookable)
+		{
+			FEvent* CookCompletedEvent = FPlatformProcess::GetSynchEventFromPool();
+			FCookRequestCompletedCallback CookRequestCompleted = [this, CookCompletedEvent](const UE::Cook::ECookResult)
+			{
+				CookCompletedEvent->Trigger();
+			};
+
+			const bool bEnqueued = CookOnTheFlyServer.EnqueueCookRequest(UE::Cook::FCookPackageRequest
+			{ 
+				PlatformName,
+				Filename,
+				MoveTemp(CookRequestCompleted)
+			});
+			check(bEnqueued);
+
+			CookCompletedEvent->Wait();
+			FPlatformProcess::ReturnSynchEventToPool(CookCompletedEvent);
+			CookOnTheFlyServer.GetUnsolicitedFiles(PlatformName, Filename, bIsCookable, UnsolicitedFiles);
+		}
+		else
+		{
+			CookOnTheFlyServer.GetUnsolicitedFiles(PlatformName, Filename, bIsCookable, UnsolicitedFiles);
+		}
+	}
+
+	void OnRecompileShaders(const FShaderRecompileData& RecompileData)
+	{
+		FEvent* RecompileCompletedEvent = FPlatformProcess::GetSynchEventFromPool();
+		UE::Cook::FRecompileShaderCompletedCallback RecompileCompleted = [this, RecompileCompletedEvent]()
+		{
+			RecompileCompletedEvent->Trigger();
+		};
+
+		const bool bEnqueued = CookOnTheFlyServer.EnqueueRecompileShaderRequest(UE::Cook::FRecompileShaderRequest
+		{
+			FName(*RecompileData.PlatformName),
+			RecompileData.ShaderPlatform,
+			RecompileData.MaterialsToLoad,
+			RecompileData.ShadersToRecompile,
+			RecompileData.MeshMaterialMaps,
+			RecompileData.ModifiedFiles,
+			RecompileData.bCompileChangedShaders,
+			MoveTemp(RecompileCompleted)
+		});
+		check(bEnqueued);
+
+		RecompileCompletedEvent->Wait();
+		FPlatformProcess::ReturnSynchEventToPool(RecompileCompletedEvent);
+	}
+
+	FString OnGetSandboxPath()
+	{
+		return CookOnTheFlyServer.GetSandboxDirectory();
+	}
+
+	UE::Cook::ICookOnTheFlyServer& CookOnTheFlyServer;
+	TUniquePtr<INetworkFileServer> NetworkFileServer;
+	FOnFileModifiedDelegate FileModifiedDelegate;
+};
+
+namespace UE { namespace Cook
+{
+
+TUniquePtr<ICookOnTheFlyRequestManager> MakeNetworkFileCookOnTheFlyRequestManager(ICookOnTheFlyServer& CookOnTheFlyServer, const FNetworkFileServerOptions& FileServerOptions)
+{
+	return MakeUnique<FNetworkFileCookOnTheFlyRequestManager>(CookOnTheFlyServer, FileServerOptions);
+}
+
+}} // namespace UE::Cook
diff --git a/Engine/Source/Editor/UnrealEd/Private/Cooker/NetworkFileCookOnTheFlyRequestManager.h b/Engine/Source/Editor/UnrealEd/Private/Cooker/NetworkFileCookOnTheFlyRequestManager.h
new file mode 100644
--- /dev/null
+++ b/Engine/Source/Editor/UnrealEd/Private/Cooker/NetworkFileCookOnTheFlyRequestManager.h
@@ -0,0 +1,17 @@
+// Copyright Epic Games, Inc. All Rights Reserved.
+
+#pragma once
+
+#include "CoreMinimal.h"
+
+struct FNetworkFileServerOptions;
+
+namespace UE { namespace Cook
+{
+
+class ICookOnTheFlyServer;
+class ICookOnTheFlyRequestManager;
+
+TUniquePtr<ICookOnTheFlyRequestManager> MakeNetworkFileCookOnTheFlyRequestManager(ICookOnTheFlyServer& CookOnTheFlyServer, const FNetworkFileServerOptions& FileServerOptions);
+
+}} // namespace UE::Cook
diff --git a/Engine/Source/Editor/UnrealEd/Private/Cooker/PackageTracker.h b/Engine/Source/Editor/UnrealEd/Private/Cooker/PackageTracker.h
--- a/Engine/Source/Editor/UnrealEd/Private/Cooker/PackageTracker.h
+++ b/Engine/Source/Editor/UnrealEd/Private/Cooker/PackageTracker.h
@@ -7,8 +7,8 @@
 #include "Containers/Set.h"
 #include "CookRequests.h"
 #include "CookTypes.h"
+#include "CookOnTheFlyServerInterface.h"
 #include "HAL/CriticalSection.h"
-#include "INetworkFileSystemModule.h"
 #include "Misc/ScopeLock.h"
 #include "Templates/UnrealTemplate.h"
 #include "UObject/NameTypes.h"
@@ -21,16 +21,7 @@ namespace UE
 {
 namespace Cook
 {
-
 	struct FPackageDatas;
-	struct FRecompileRequest;
-
-	/** Helper to pass a recompile request to game thread */
-	struct FRecompileRequest
-	{
-		struct FShaderRecompileData RecompileData;
-		volatile bool bComplete = false;
-	};
 
 	template<typename Type>
 	struct FThreadSafeQueue
@@ -44,6 +35,12 @@ namespace Cook
 			FScopeLock ScopeLock(&SynchronizationObject);
 			Items.Add(Item);
 		}
+		
+		void Enqueue(Type&& Item)
+		{
+			FScopeLock ScopeLock(&SynchronizationObject);
+			Items.Add(MoveTempIfPossible(Item));
+		}
 
 		void EnqueueUnique(const Type& Item)
 		{
@@ -201,7 +198,7 @@ namespace Cook
 		FPackageDatas& PackageDatas;
 
 		FThreadSafeUnsolicitedPackagesList UnsolicitedCookedPackages;
-		FThreadSafeQueue<FRecompileRequest*> RecompileRequests;
+		FThreadSafeQueue<FRecompileShaderRequest> RecompileRequests;
 
 		FThreadSafeSet<FName> NeverCookPackageList;
 		FThreadSafeSet<FName> UncookedEditorOnlyPackages; // set of packages that have been rejected due to being referenced by editor-only properties
diff --git a/Engine/Source/Editor/UnrealEd/Private/EditorEngine.cpp b/Engine/Source/Editor/UnrealEd/Private/EditorEngine.cpp
--- a/Engine/Source/Editor/UnrealEd/Private/EditorEngine.cpp
+++ b/Engine/Source/Editor/UnrealEd/Private/EditorEngine.cpp
@@ -14,6 +14,7 @@
 #include "UObject/MetaData.h"
 #include "UObject/ObjectSaveContext.h"
 #include "UObject/ConstructorHelpers.h"
+#include "UObject/SavePackage.h"
 #include "Application/ThrottleManager.h"
 #include "Framework/Application/SlateApplication.h"
 #include "Framework/MultiBox/MultiBoxDefs.h"
diff --git a/Engine/Source/Editor/UnrealEd/Private/PlayLevelViaLauncher.cpp b/Engine/Source/Editor/UnrealEd/Private/PlayLevelViaLauncher.cpp
--- a/Engine/Source/Editor/UnrealEd/Private/PlayLevelViaLauncher.cpp
+++ b/Engine/Source/Editor/UnrealEd/Private/PlayLevelViaLauncher.cpp
@@ -288,6 +288,7 @@ void UEditorEngine::StartPlayUsingLauncherSession(FRequestPlaySessionParams& InR
 	if (LauncherProfile->GetCookMode() == ELauncherProfileCookModes::OnTheFlyInEditor || LauncherProfile->GetCookMode() == ELauncherProfileCookModes::OnTheFly)
 	{
 		LauncherProfile->SetDeploymentMode(ELauncherProfileDeploymentModes::FileServer);
+		LauncherProfile->SetUseIoStore(EditorPlaySettings->EnableIoStoreCookOnTheFlyMode);
 	}
 
 	switch(EditorPlaySettings->PackFilesForLaunch)
diff --git a/Engine/Source/Editor/UnrealEd/Private/UnrealEdEngine.cpp b/Engine/Source/Editor/UnrealEd/Private/UnrealEdEngine.cpp
--- a/Engine/Source/Editor/UnrealEd/Private/UnrealEdEngine.cpp
+++ b/Engine/Source/Editor/UnrealEd/Private/UnrealEdEngine.cpp
@@ -203,7 +203,10 @@ void UUnrealEdEngine::Init(IEngineLoop* InEngineLoop)
 
 			CookServer = NewObject<UCookOnTheFlyServer>();
 			CookServer->Initialize(ECookMode::CookOnTheFlyFromTheEditor, BaseCookingFlags);
-			CookServer->StartNetworkFileServer(false);
+
+			UCookOnTheFlyServer::FCookOnTheFlyOptions CookOnTheFlyStartupOptions;
+			CookOnTheFlyStartupOptions.bBindAnyPort = false;
+			CookServer->StartCookOnTheFly(CookOnTheFlyStartupOptions);
 		}
 		else if (!ExperimentalSettings->bDisableCookInEditor)
 		{
@@ -1579,4 +1582,4 @@ void UUnrealEdEngine::ValidateFreeDiskSpace() const
 
 		FMessageDialog::Open(EAppMsgType::Ok, Message, &Title);
 	}
-}
\ No newline at end of file
+}
diff --git a/Engine/Source/Editor/UnrealEd/UnrealEd.Build.cs b/Engine/Source/Editor/UnrealEd/UnrealEd.Build.cs
--- a/Engine/Source/Editor/UnrealEd/UnrealEd.Build.cs
+++ b/Engine/Source/Editor/UnrealEd/UnrealEd.Build.cs
@@ -202,6 +202,7 @@ public class UnrealEd : ModuleRules
 				"AnimationModifiers",
 				"AnimationBlueprintLibrary",
 				"MaterialBaking",
+				"CookOnTheFly",
 				"BazelExecutor",
 			}
 		);
diff --git a/Engine/Source/Programs/AutomationTool/AutomationUtils/ProjectParams.cs b/Engine/Source/Programs/AutomationTool/AutomationUtils/ProjectParams.cs
--- a/Engine/Source/Programs/AutomationTool/AutomationUtils/ProjectParams.cs
+++ b/Engine/Source/Programs/AutomationTool/AutomationUtils/ProjectParams.cs
@@ -286,6 +286,7 @@ namespace AutomationTool
 			this.Pak = InParams.Pak;
 			this.IgnorePaksFromDifferentCookSource = InParams.IgnorePaksFromDifferentCookSource;
 			this.IoStore = InParams.IoStore;
+			this.Cook4IoStore = InParams.Cook4IoStore;
 			this.GenerateOptimizationData = InParams.GenerateOptimizationData;
 			this.SignPak = InParams.SignPak;
 			this.SignedPak = InParams.SignedPak;
@@ -371,6 +372,7 @@ namespace AutomationTool
 			this.bTreatNonShippingBinariesAsDebugFiles = InParams.bTreatNonShippingBinariesAsDebugFiles;
 			this.bUseExtraFlavor = InParams.bUseExtraFlavor;
 			this.AdditionalPackageOptions = InParams.AdditionalPackageOptions;
+			this.StorageServerHost = InParams.StorageServerHost;
 		}
 
 		/// <summary>
@@ -472,6 +474,7 @@ namespace AutomationTool
 			bool? Pak = null,
 			bool? IgnorePaksFromDifferentCookSource = null,
 			bool? IoStore = null,
+			bool? Cook4IoStore = null,
 			bool? SkipIoStore = null,
 			bool? GenerateOptimizationData = null,
 			bool? Prereqs = null,
@@ -521,7 +524,8 @@ namespace AutomationTool
 		    bool AutomaticSigning = false,
 			ParamList<string> InMapsToRebuildLightMaps = null,
             ParamList<string> InMapsToRebuildHLOD = null,
-            ParamList<string> TitleID = null
+            ParamList<string> TitleID = null,
+			string StorageServerhost = null
 			)
 		{
 			//
@@ -659,6 +663,12 @@ namespace AutomationTool
 			this.IgnorePaksFromDifferentCookSource = GetParamValueIfNotSpecified(Command, IgnorePaksFromDifferentCookSource, this.IgnorePaksFromDifferentCookSource, "IgnorePaksFromDifferentCookSource");
 			this.IoStore = GetParamValueIfNotSpecified(Command, IoStore, this.IoStore, "iostore");
 			this.SkipIoStore = GetParamValueIfNotSpecified(Command, SkipIoStore, this.SkipIoStore, "skipiostore");
+			this.Cook4IoStore = GetParamValueIfNotSpecified(Command, Cook4IoStore, this.Cook4IoStore, "cook4iostore");
+			if (this.Cook4IoStore)
+			{
+				this.IoStore = true;
+				this.AdditionalCookerOptions += " -IoStore";
+			}
 			this.GenerateOptimizationData = GetParamValueIfNotSpecified(Command, GenerateOptimizationData, this.GenerateOptimizationData, "makebinaryconfig");
 			
 			this.SkipPak = GetParamValueIfNotSpecified(Command, SkipPak, this.SkipPak, "skippak");
@@ -867,6 +877,7 @@ namespace AutomationTool
 			this.SpecifiedArchitecture = ParseParamValueIfNotSpecified(Command, SpecifiedArchitecture, "specifiedarchitecture", String.Empty);
 			this.UbtArgs = ParseParamValueIfNotSpecified(Command, UbtArgs, "ubtargs", String.Empty);
 			this.AdditionalPackageOptions = ParseParamValueIfNotSpecified(Command, AdditionalPackageOptions, "AdditionalPackageOptions", String.Empty);
+			this.StorageServerHost = ParseParamValueIfNotSpecified(Command, StorageServerHost, "StorageServerHost", "localhost");
 
 			if (ClientConfigsToBuild == null)
 			{
@@ -1231,6 +1242,12 @@ namespace AutomationTool
 		[Help("iostore", "generate I/O store container file(s)")]
 		public bool IoStore { private set; get; }
 
+		/// <summary>
+		/// Shared: True if the cooker should write directly to container file(s)
+		/// </summary>
+		[Help("cook4iostore", "generate I/O store container file(s)")]
+		public bool Cook4IoStore { private set; get; }
+
 		/// <summary>
 		/// Shared: True if optimization data is generated during staging that can improve loadtimes
 		/// </summary>
@@ -2057,6 +2074,9 @@ namespace AutomationTool
         [Help("IgnoreLightMapErrors", "Whether Light Map errors should be treated as critical")]
 		public bool IgnoreLightMapErrors { get; set; }
 
+		[Help("StorageServerHost", "Storage server host address")]
+		public string StorageServerHost { get; set; }
+
 		private List<SingleTargetProperties> DetectedTargets;
 		private Dictionary<UnrealTargetPlatform, ConfigHierarchy> LoadedEngineConfigs;
 		private Dictionary<UnrealTargetPlatform, ConfigHierarchy> LoadedGameConfigs;
@@ -2936,6 +2956,7 @@ namespace AutomationTool
 				CommandUtils.LogLog("IgnorePaksFromDifferentCookSource={0}", IgnorePaksFromDifferentCookSource);
 				CommandUtils.LogLog("IoStore={0}", IoStore);
 				CommandUtils.LogLog("SkipIoStore={0}", SkipIoStore);
+				CommandUtils.LogLog("Cook4IoStore={0}", Cook4IoStore);
 				CommandUtils.LogLog("SkipEncryption={0}", SkipEncryption);
 				CommandUtils.LogLog("GenerateOptimizationData={0}", GenerateOptimizationData);
 				CommandUtils.LogLog("SkipPackage={0}", SkipPackage);
diff --git a/Engine/Source/Programs/AutomationTool/Scripts/CookCommand.Automation.cs b/Engine/Source/Programs/AutomationTool/Scripts/CookCommand.Automation.cs
--- a/Engine/Source/Programs/AutomationTool/Scripts/CookCommand.Automation.cs
+++ b/Engine/Source/Programs/AutomationTool/Scripts/CookCommand.Automation.cs
@@ -65,6 +65,11 @@ public partial class Project : CommandUtils
 					COTFCommandLine += " -ddc=" + Params.DDCGraph;
 				}
 
+				if (Params.IoStore)
+				{
+					COTFCommandLine += " -iostore";
+				}
+
 				var ServerLogFile = CombinePaths(LogFolderOutsideOfSandbox, "Server.log");
 				Platform ClientPlatformInst = Params.ClientTargetPlatformInstances[0];
 				string TargetCook = ClientPlatformInst.GetCookPlatform(false, false); // cook on he fly doesn't support server cook platform... 
diff --git a/Engine/Source/Programs/AutomationTool/Scripts/CopyBuildToStagingDirectory.Automation.cs b/Engine/Source/Programs/AutomationTool/Scripts/CopyBuildToStagingDirectory.Automation.cs
--- a/Engine/Source/Programs/AutomationTool/Scripts/CopyBuildToStagingDirectory.Automation.cs
+++ b/Engine/Source/Programs/AutomationTool/Scripts/CopyBuildToStagingDirectory.Automation.cs
@@ -16,6 +16,7 @@ using EpicGames.Core;
 using System.Diagnostics;
 using System.Collections.Concurrent;
 using UnrealBuildBase;
+using System.Text.Json;
 
 /// <summary>
 /// Helper command used for cooking.
@@ -601,6 +602,30 @@ public partial class Project : CommandUtils
 		return Plugins.ToList();
 	}
 
+	private class PackageStoreManifest
+	{
+		public static IList<FileReference> GetCookedFiles(FileReference PackageStoreManifestFile)
+		{
+			List<FileReference> CookedFiles = new List<FileReference>();
+			if (FileReference.Exists(PackageStoreManifestFile))
+			{
+				String JsonString = FileReference.ReadAllText(PackageStoreManifestFile);
+				PackageStoreManifest PackageStoreManifest = JsonSerializer.Deserialize<PackageStoreManifest>(JsonString);
+				foreach (PackageStoreManifest.FileInfo File in PackageStoreManifest.Files)
+				{
+					CookedFiles.Add(new FileReference(File.Path));
+				}
+			}
+			return CookedFiles;
+		}
+		public IList<FileInfo> Files { get; set; }
+
+		public class FileInfo
+		{
+			public String Path { get; set; }
+		}
+	}
+
 	public static void CreateStagingManifest(ProjectParams Params, DeploymentContext SC)
 	{
 		if (!Params.Stage)
@@ -696,6 +721,13 @@ public partial class Project : CommandUtils
 
 			// Stage all the cooked data, this is the same rule as normal stage except we may skip Engine
 			List<FileReference> CookedFiles = DirectoryReference.EnumerateFiles(SC.PlatformCookDir, "*", SearchOption.AllDirectories).ToList();
+
+			// When cooking directly to I/O store contaier files look for the package store manifest
+			if (Params.IoStore)
+			{
+				CookedFiles.AddRange(PackageStoreManifest.GetCookedFiles(FileReference.Combine(SC.MetadataDir, "packagestore.manifest")));
+			}
+
 			foreach (FileReference CookedFile in CookedFiles)
 			{
 				// Skip metadata directory
@@ -718,7 +750,7 @@ public partial class Project : CommandUtils
 
 				// json files have never been staged
 				// metallib files cannot *currently* be staged as UFS as the Metal API needs to mmap them from files on disk in order to function efficiently
-				if (!CookedFile.HasExtension(".json") && !CookedFile.HasExtension(".metallib"))
+				if (!CookedFile.HasExtension(".json") && !CookedFile.HasExtension(".metallib") && !CookedFile.HasExtension(".utoc") && !CookedFile.HasExtension(".ucas"))
 				{
 					SC.StageFile(StagedFileType.UFS, CookedFile, new StagedFileReference(CookedFile.MakeRelativeTo(SC.PlatformCookDir)));
 				}
@@ -1020,6 +1052,13 @@ public partial class Project : CommandUtils
 				if (DirectoryReference.Exists(CookOutputDir))
 				{
 					List<FileReference> CookedFiles = DirectoryReference.EnumerateFiles(CookOutputDir, "*", SearchOption.AllDirectories).ToList();
+					
+					// When cooking directly to I/O store contaier files look for the package store manifest
+					if (Params.IoStore)
+					{
+						CookedFiles.AddRange(PackageStoreManifest.GetCookedFiles(FileReference.Combine(SC.MetadataDir, "packagestore.manifest")));
+					}
+
 					foreach (FileReference CookedFile in CookedFiles)
 					{
 						// Skip metadata directory
@@ -1036,7 +1075,7 @@ public partial class Project : CommandUtils
 
 						// json files have never been staged
 						// metallib files cannot *currently* be staged as UFS as the Metal API needs to mmap them from files on disk in order to function efficiently
-						if (!CookedFile.HasExtension(".json") && !CookedFile.HasExtension(".metallib"))
+						if (!CookedFile.HasExtension(".json") && !CookedFile.HasExtension(".metallib") && !CookedFile.HasExtension(".utoc") && !CookedFile.HasExtension(".ucas"))
 						{
 							SC.StageFile(StagedFileType.UFS, CookedFile, new StagedFileReference(CookedFile.MakeRelativeTo(CookOutputDir)));
 						}
@@ -3857,7 +3896,7 @@ public partial class Project : CommandUtils
 			string FileHostParams = " ";
 			if (Params.CookOnTheFly || Params.FileServer)
 			{
-				FileHostParams += "-filehostip=";
+				FileHostParams += Params.IoStore ? "-cookontheflyhost=" : "-filehostip=";
 				// add localhost first for platforms using redirection
 				const string LocalHost = "127.0.0.1";
 				if (!IsNullOrEmpty(Params.Port))
@@ -3973,7 +4012,15 @@ public partial class Project : CommandUtils
 						}
 					}
 				}
-				FileHostParams += " ";
+
+				if (Params.IoStore)
+				{
+					FileHostParams += string.Format(" -storageserverhost={0} ", Params.StorageServerHost);
+				}
+				else
+				{
+					FileHostParams += " ";
+				}
 			}
 
 			String ProjectFile = String.Format("{0} ", SC.ProjectArgForCommandLines);
diff --git a/Engine/Source/Programs/AutomationTool/Scripts/RunProjectCommand.Automation.cs b/Engine/Source/Programs/AutomationTool/Scripts/RunProjectCommand.Automation.cs
--- a/Engine/Source/Programs/AutomationTool/Scripts/RunProjectCommand.Automation.cs
+++ b/Engine/Source/Programs/AutomationTool/Scripts/RunProjectCommand.Automation.cs
@@ -528,7 +528,7 @@ public partial class Project : CommandUtils
 				if (ClientProcess == null && !String.IsNullOrEmpty(Output))
 				{
 					AllClientOutput += Output;
-					if (AllClientOutput.Contains("Game Engine Initialized") || AllClientOutput.Contains("Unreal Network File Server is ready"))
+					if (AllClientOutput.Contains("Game Engine Initialized") || AllClientOutput.Contains("Unreal Network File Server is ready") || AllClientOutput.Contains("COTF server is ready"))
 					{
 						LogInformation("Starting Client....");
 						var SC = DeployContextList[0];
@@ -629,7 +629,15 @@ public partial class Project : CommandUtils
 
 			if (Params.CookOnTheFly || Params.FileServer)
 			{
-				TempCmdLine += "-filehostip=";
+				if (Params.IoStore)
+				{
+					TempCmdLine += "-cookontheflyhost=";
+				}
+				else
+				{
+					TempCmdLine += "-filehostip=";
+				}
+
 				// add localhost first for platforms using redirection
 				const string LocalHost = "127.0.0.1";
 				if (!IsNullOrEmpty(Params.Port))
@@ -744,7 +752,15 @@ public partial class Project : CommandUtils
 						}
 					}
 				}
-				TempCmdLine += " ";
+
+				if (Params.IoStore)
+				{
+					TempCmdLine += string.Format(" -storageserverhost={0} ", Params.StorageServerHost);
+				}
+				else
+				{
+					TempCmdLine += " ";
+				}
 
 				if (Params.CookOnTheFlyStreaming)
 				{
diff --git a/Engine/Source/Programs/UnrealHeaderTool/Private/StringUtils.cpp b/Engine/Source/Programs/UnrealHeaderTool/Private/StringUtils.cpp
--- a/Engine/Source/Programs/UnrealHeaderTool/Private/StringUtils.cpp
+++ b/Engine/Source/Programs/UnrealHeaderTool/Private/StringUtils.cpp
@@ -79,12 +79,23 @@ FString GetClassPrefix(const FString& InClassName, bool& bIsLabeledDeprecated )
 // Returns zero only for '\r' and '\0' 
 FORCEINLINE uint64 FindCrOrNulHelper(TCHAR C)
 {
+	if (C == '\r')
+	{
+		return 0;
+	}
+	return C;
+
+	// NOTE: This is the previous implementation. It might have been faster
+	// once but current compilers do a much better job and compile into
+	// conditional moves which are faster than the variable-count bit shift below
+#if 0
 	uint64 LoPassMask = ~((uint64(1) << '\0') | (uint64(1) << '\r'));
 	uint64 LoBit = uint64(1) << (C & 63);
 	uint64 LoPass = LoBit & LoPassMask;
 	uint64 HiPass = C & ~63;
 
 	return LoPass | HiPass;
+#endif
 }
 
 // Finds next CR or null term with a single branch 
diff --git a/Engine/Source/Runtime/CookOnTheFly/Private/CookOnTheFly.cpp b/Engine/Source/Runtime/CookOnTheFly/Private/CookOnTheFly.cpp
new file mode 100644
--- /dev/null
+++ b/Engine/Source/Runtime/CookOnTheFly/Private/CookOnTheFly.cpp
@@ -0,0 +1,142 @@
+// Copyright Epic Games, Inc. All Rights Reserved.
+
+#include "CookOnTheFly.h"
+#include "Modules/ModuleManager.h"
+#include "Misc/CommandLine.h"
+#include "Misc/DateTime.h"
+#include "HAL/FileManager.h"
+
+namespace UE { namespace Cook
+{
+
+FString FCookOnTheFlyMessageHeader::ToString() const
+{
+	return FString::Printf(TEXT("Message='%s', Status='%s', Sender='%u', CorrelationId='%u'"),
+		LexToString(MessageType),
+		LexToString(MessageStatus),
+		SenderId,
+		CorrelationId);
+}
+
+FArchive& operator<<(FArchive& Ar, FCookOnTheFlyMessageHeader& Header)
+{
+	uint32 MessageType = static_cast<uint32>(Header.MessageType);
+	uint32 MessageStatus = static_cast<uint32>(Header.MessageStatus);
+	
+	Ar << MessageType;
+	Ar << MessageStatus;
+	Ar << Header.SenderId;
+	Ar << Header.CorrelationId;
+	Ar << Header.Timestamp;
+
+	if (Ar.IsLoading())
+	{
+		Header.MessageType = static_cast<ECookOnTheFlyMessage>(MessageType);
+		Header.MessageStatus = static_cast<ECookOnTheFlyMessageStatus>(MessageStatus);
+	}
+
+	return Ar;
+}
+
+FArchive& operator<<(FArchive& Ar, FCookOnTheFlyMessage& Message)
+{
+	Ar << Message.Header;
+	Ar << Message.Body;
+
+	return Ar;
+}
+
+void FCookOnTheFlyMessage::SetBody(TArray<uint8> InBody)
+{
+	Body = MoveTemp(InBody);
+}
+
+TUniquePtr<FArchive> FCookOnTheFlyMessage::ReadBody() const
+{
+	return MakeUnique<FMemoryReader>(Body);
+}
+
+TUniquePtr<FArchive> FCookOnTheFlyMessage::WriteBody()
+{
+	return MakeUnique<FMemoryWriter>(Body);
+}
+
+bool GetCookOnTheFlyHost(UE::Cook::FCookOnTheFlyHostOptions& OutHostOptions)
+{
+	bool bOptionsOk = false;
+
+	FString Host;
+	if (FParse::Value(FCommandLine::Get(), TEXT("-CookOnTheFlyHost="), Host))
+	{
+		if (!Host.ParseIntoArray(OutHostOptions.Hosts, TEXT("+"), true))
+		{
+			OutHostOptions.Hosts.Add(Host);
+		}
+
+		bOptionsOk = true;
+	}
+
+	double ServerWaitTimeInSeconds;
+	if (FParse::Value(FCommandLine::Get(), TEXT("-CookOnTheFlyServerWaitTime="), ServerWaitTimeInSeconds))
+	{
+		OutHostOptions.ServerStartupWaitTime = FTimespan::FromSeconds(ServerWaitTimeInSeconds);
+	}
+
+	return bOptionsOk;
+}
+
+bool SendCookOnTheFlyRequest(ECookOnTheFlyMessage RequestType, TFunction<void(FArchive&)>&& FillRequest, TFunction<bool(FArchive&)>&& ProcessResponse)
+{
+	static FString FileHostIP;
+	static bool bUseFileServer = FParse::Value(FCommandLine::Get(), TEXT("filehostip"), FileHostIP);
+
+	if (bUseFileServer)
+	{
+		struct FFileServerMessageHandler
+			: public IPlatformFile::IFileServerMessageHandler
+		{
+			FFileServerMessageHandler(TFunction<void(FArchive&)>&& InFillRequest, TFunction<bool(FArchive&)>&& InProcessResponse)
+				: FillRequestHandler(InFillRequest)
+				, ProcessResponseHandler(InProcessResponse)
+			{ }
+
+			virtual void FillPayload(FArchive& Payload) override
+			{
+				FillRequestHandler(Payload);
+			}
+
+			virtual void ProcessResponse(FArchive& Response) override
+			{
+				ProcessResponseHandler(Response);
+			}
+
+			TFunction<void(FArchive&)> FillRequestHandler;
+			TFunction<bool(FArchive&)> ProcessResponseHandler;
+		};
+
+		FFileServerMessageHandler MessageHandler(MoveTemp(FillRequest), MoveTemp(ProcessResponse));
+		return IFileManager::Get().SendMessageToServer(LexToString(RequestType), &MessageHandler);
+	}
+	else
+	{
+		static ICookOnTheFlyModule& CookOnTheFlyModule = FModuleManager::LoadModuleChecked<ICookOnTheFlyModule>(TEXT("CookOnTheFly"));
+		ICookOnTheFlyServerConnection& ServerConnection = CookOnTheFlyModule.GetServerConnection();
+
+		FCookOnTheFlyRequest Request(RequestType);
+		{
+			TUniquePtr<FArchive> Ar = Request.WriteBody();
+			FillRequest(*Ar);
+		}
+
+		FCookOnTheFlyResponse Response = ServerConnection.SendRequest(Request).Get();
+		if (Response.IsOk())
+		{
+			TUniquePtr<FArchive> Ar = Response.ReadBody();
+			ProcessResponse(*Ar);
+		}
+
+		return Response.IsOk();
+	}
+}
+
+}} // namesapce UE::Cook
