Do not include default map packages in the soft object reference collector, when loading the game config.  Doing this invalidates the cookers abillity to strip these packages from a build.

preflight mostly succeeded, but timed out.
https://horde.devtools.epicgames.com/job/60c89fe343ec380001f35313

diff --git a/Engine/Source/Editor/UnrealEd/Classes/CookOnTheSide/CookOnTheFlyServer.h b/Engine/Source/Editor/UnrealEd/Classes/CookOnTheSide/CookOnTheFlyServer.h
--- a/Engine/Source/Editor/UnrealEd/Classes/CookOnTheSide/CookOnTheFlyServer.h
+++ b/Engine/Source/Editor/UnrealEd/Classes/CookOnTheSide/CookOnTheFlyServer.h
@@ -671,7 +671,14 @@ private:
 	void CollectFilesToCook(TArray<FName>& FilesInPath, 
 		const TArray<FString>& CookMaps, const TArray<FString>& CookDirectories, 
 		const TArray<FString>& IniMapSections, ECookByTheBookOptions FilesToCookFlags,
-		const TArrayView<const ITargetPlatform* const>& TargetPlatforms);
+		const TArrayView<const ITargetPlatform* const>& TargetPlatforms,
+		const TMap<FName, TArray<FName>>& GameDefaultObjects);
+
+	/**
+	* Gets all game default objects for all platforms.
+	*/
+	static void GetGameDefaultObjects(const TArray<ITargetPlatform*>& TargetPlatforms, TMap<FName, TArray<FName>>& GameDefaultObjectsOut);
+
 	/* Collect filespackages that should not be cooked from ini settings and commandline. Does not include checking UAssetManager, which has to be queried later */
 	TArray<FName> GetNeverCookPackageFileNames(TArrayView<const FString> ExtraNeverCookDirectories = TArrayView<const FString>());
 
diff --git a/Engine/Source/Editor/UnrealEd/Private/CookOnTheFlyServer.cpp b/Engine/Source/Editor/UnrealEd/Private/CookOnTheFlyServer.cpp
--- a/Engine/Source/Editor/UnrealEd/Private/CookOnTheFlyServer.cpp
+++ b/Engine/Source/Editor/UnrealEd/Private/CookOnTheFlyServer.cpp
@@ -122,6 +122,7 @@
 #include "Misc/NetworkVersion.h"
 
 #include "Algo/Find.h"
+#include "Algo/AnyOf.h"
 #include "Async/ParallelFor.h"
 
 #include "Commandlets/ShaderPipelineCacheToolsCommandlet.h"
@@ -6394,7 +6395,8 @@ void UCookOnTheFlyServer::AddFileToCook( TArray<FName>& InOutFilesToCook, const
 }
 
 void UCookOnTheFlyServer::CollectFilesToCook(TArray<FName>& FilesInPath, const TArray<FString>& CookMaps, const TArray<FString>& InCookDirectories,
-	const TArray<FString> &IniMapSections, ECookByTheBookOptions FilesToCookFlags, const TArrayView<const ITargetPlatform* const>& TargetPlatforms)
+	const TArray<FString> &IniMapSections, ECookByTheBookOptions FilesToCookFlags, const TArrayView<const ITargetPlatform* const>& TargetPlatforms,
+	const TMap<FName, TArray<FName>>& GameDefaultObjects)
 {
 	UE_SCOPED_HIERARCHICAL_COOKTIMER(CollectFilesToCook);
 
@@ -6640,60 +6642,28 @@ void UCookOnTheFlyServer::CollectFilesToCook(TArray<FName>& FilesInPath, const T
 
 	if (!(FilesToCookFlags & ECookByTheBookOptions::NoDefaultMaps))
 	{
-		// make sure we cook the default maps
-		// Collect the default maps for all requested platforms.  Our additions are potentially wasteful if different platforms in the requested list have different default maps.
-		// In that case we will wastefully cook maps for platforms that don't require them.
-		for (const ITargetPlatform* TargetPlatform : TargetPlatforms)
+		for (const auto& GameDefaultSet : GameDefaultObjects)
 		{
-			// load the platform specific ini to get its DefaultMap
-			FConfigFile PlatformEngineIni;
-			FConfigCacheIni::LoadLocalIniFile(PlatformEngineIni, TEXT("Engine"), true, *TargetPlatform->IniPlatformName());
-
-			// get the server and game default maps/modes and cook them
-			TArray<FString, TInlineAllocator<5>> GameDefaultObjs; 
-			FString ObjPath;
-
-			if (PlatformEngineIni.GetString(TEXT("/Script/EngineSettings.GameMapsSettings"), TEXT("GameDefaultMap"), ObjPath))
+			if (GameDefaultSet.Key == FName("ServerDefaultMap") && !IsCookFlagSet(ECookInitializationFlags::IncludeServerMaps))
 			{
-				GameDefaultObjs.Add(MoveTemp(ObjPath));
+				continue;
 			}
 
-			if (IsCookFlagSet(ECookInitializationFlags::IncludeServerMaps))
+			for (FName PackagePath : GameDefaultSet.Value)
 			{
-				if (PlatformEngineIni.GetString(TEXT("/Script/EngineSettings.GameMapsSettings"), TEXT("ServerDefaultMap"), ObjPath))
+				TArray<FAssetData> Assets;
+				if (!AssetRegistry->GetAssetsByPackageName(PackagePath, Assets))
 				{
-					GameDefaultObjs.Add(MoveTemp(ObjPath));
+					const FText ErrorMessage = FText::Format(LOCTEXT("GameMapSettingsMissing", "GameMapsSettings contains a path to a missing asset '{0}'. The intended asset will fail to load in a packaged build. Select the intended asset again in Project Settings to fix this issue."), FText::FromName(PackagePath));
+					LogCookerMessage(ErrorMessage.ToString(), EMessageSeverity::Error);
 				}
-			}
-
-			if (PlatformEngineIni.GetString(TEXT("/Script/EngineSettings.GameMapsSettings"), TEXT("GlobalDefaultGameMode"), ObjPath))
-			{
-				GameDefaultObjs.Add(MoveTemp(ObjPath));
-			}
-
-			if (PlatformEngineIni.GetString(TEXT("/Script/EngineSettings.GameMapsSettings"), TEXT("GlobalDefaultServerGameMode"), ObjPath))
-			{
-				GameDefaultObjs.Add(MoveTemp(ObjPath));
-			}
-
-			if (PlatformEngineIni.GetString(TEXT("/Script/EngineSettings.GameMapsSettings"), TEXT("GameInstanceClass"), ObjPath))
-			{
-				GameDefaultObjs.Add(MoveTemp(ObjPath));
-			}
-
-			for (const FString& Obj : GameDefaultObjs)
-			{
-				if (Obj != FName(NAME_None).ToString())
+				else if (Algo::AnyOf(Assets, [](const FAssetData& Asset) { return Asset.IsRedirector(); }))
 				{
-					FAssetData DestinationData = AssetRegistry->GetAssetByObjectPath(FName(*Obj), true);
-					if (DestinationData.IsRedirector())
-					{
-						const FText ErrorMessage = FText::Format(LOCTEXT("GameMapSettingsRedirectorDetected", "GameMapsSettings contains a redirected reference '{0}'. The intended asset will fail to load in a packaged build. Select the intended asset again in Project Settings to fix this issue."), FText::FromString(Obj));
-						LogCookerMessage(ErrorMessage.ToString(), EMessageSeverity::Error);
-					}
-
-					AddFileToCook(FilesInPath, Obj);
+					const FText ErrorMessage = FText::Format(LOCTEXT("GameMapSettingsRedirectorDetected", "GameMapsSettings contains a redirected reference '{0}'. The intended asset will fail to load in a packaged build. Select the intended asset again in Project Settings to fix this issue."), FText::FromName(PackagePath));
+					LogCookerMessage(ErrorMessage.ToString(), EMessageSeverity::Error);
 				}
+
+				AddFileToCook(FilesInPath, PackagePath.ToString());
 			}
 		}
 	}
@@ -6751,6 +6721,43 @@ void UCookOnTheFlyServer::CollectFilesToCook(TArray<FName>& FilesInPath, const T
 	}
 }
 
+void UCookOnTheFlyServer::GetGameDefaultObjects(const TArray<ITargetPlatform*>& TargetPlatforms, TMap<FName, TArray<FName>>& GameDefaultObjectsOut)
+{
+	// Collect all default objects from all cooked platforms engine configurations.
+	for (const ITargetPlatform* TargetPlatform : TargetPlatforms)
+	{
+		// load the platform specific ini to get its DefaultMap
+		FConfigFile PlatformEngineIni;
+		FConfigCacheIni::LoadLocalIniFile(PlatformEngineIni, TEXT("Engine"), true, *TargetPlatform->IniPlatformName());
+
+		FConfigSection* MapSettingsSection = PlatformEngineIni.Find(TEXT("/Script/EngineSettings.GameMapsSettings"));
+
+		if (MapSettingsSection == nullptr)
+		{
+			continue;
+		}
+
+		auto AddDefaultObject = [&GameDefaultObjectsOut, &PlatformEngineIni, MapSettingsSection](FName PropertyName)
+		{
+			const FConfigValue* PairString = MapSettingsSection->Find(PropertyName);
+			if (PairString == nullptr)
+			{
+				return;
+			}
+
+			FSoftObjectPath Path(PairString->GetValue());
+			GameDefaultObjectsOut.FindOrAdd(PropertyName).AddUnique(Path.GetLongPackageFName());
+		};
+
+		// get the server and game default maps/modes and cook them
+		AddDefaultObject(FName("GameDefaultMap"));
+		AddDefaultObject(FName("ServerDefaultMap"));
+		AddDefaultObject(FName("GlobalDefaultGameMode"));
+		AddDefaultObject(FName("GlobalDefaultServerGameMode"));
+		AddDefaultObject(FName("GameInstanceClass"));
+	}
+}
+
 bool UCookOnTheFlyServer::IsCookByTheBookRunning() const
 {
 	return CookByTheBookOptions && CookByTheBookOptions->bRunning;
@@ -8185,7 +8192,19 @@ void UCookOnTheFlyServer::StartCookByTheBook( const FCookByTheBookStartupOptions
 		}
 	}
 
-	CollectFilesToCook(FilesInPath, CookMaps, CookDirectories, IniMapSections, CookOptions, TargetPlatforms);
+	TMap<FName, TArray<FName>> GameDefaultObjects;
+	GetGameDefaultObjects(TargetPlatforms, GameDefaultObjects);
+
+	// Strip out the default maps from SoftObjectPaths collected from startup packages. They will be added to the cook if necessary by CollectFilesToCook.
+	for (const auto& GameDefaultSet : GameDefaultObjects)
+	{
+		for (FName AssetName : GameDefaultSet.Value)
+		{
+			StartupSoftObjectPackages.Remove(AssetName);
+		}
+	}
+	
+	CollectFilesToCook(FilesInPath, CookMaps, CookDirectories, IniMapSections, CookOptions, TargetPlatforms, GameDefaultObjects);
 
 	// Add string asset packages after collecting files, to avoid accidentally activating the behavior to cook all maps if none are specified
 	for (FName SoftObjectPackage : StartupSoftObjectPackages)
