Fix spurious cook error GameMapsSettings contains a path to a missing asset 'None'.  Projects can set the GameMapSettings fields to the empty string:
    [/Script/EngineSettings.GameMapsSettings]
    ServerDefaultMap=
We need to silently ignore those settings.
#rnx

diff --git a/Engine/Source/Editor/UnrealEd/Private/CookOnTheFlyServer.cpp b/Engine/Source/Editor/UnrealEd/Private/CookOnTheFlyServer.cpp
--- a/Engine/Source/Editor/UnrealEd/Private/CookOnTheFlyServer.cpp
+++ b/Engine/Source/Editor/UnrealEd/Private/CookOnTheFlyServer.cpp
@@ -6644,7 +6644,7 @@ void UCookOnTheFlyServer::CollectFilesToCook(TArray<FName>& FilesInPath, const T
 	{
 		for (const auto& GameDefaultSet : GameDefaultObjects)
 		{
-			if (GameDefaultSet.Key == FName("ServerDefaultMap") && !IsCookFlagSet(ECookInitializationFlags::IncludeServerMaps))
+			if (GameDefaultSet.Key == FName(TEXT("ServerDefaultMap")) && !IsCookFlagSet(ECookInitializationFlags::IncludeServerMaps))
 			{
 				continue;
 			}
@@ -6654,12 +6654,14 @@ void UCookOnTheFlyServer::CollectFilesToCook(TArray<FName>& FilesInPath, const T
 				TArray<FAssetData> Assets;
 				if (!AssetRegistry->GetAssetsByPackageName(PackagePath, Assets))
 				{
-					const FText ErrorMessage = FText::Format(LOCTEXT("GameMapSettingsMissing", "GameMapsSettings contains a path to a missing asset '{0}'. The intended asset will fail to load in a packaged build. Select the intended asset again in Project Settings to fix this issue."), FText::FromName(PackagePath));
+					const FText ErrorMessage = FText::Format(LOCTEXT("GameMapSettingsMissing", "{0} contains a path to a missing asset '{1}'. The intended asset will fail to load in a packaged build. Select the intended asset again in Project Settings to fix this issue."),
+						FText::FromName(GameDefaultSet.Key), FText::FromName(PackagePath));
 					LogCookerMessage(ErrorMessage.ToString(), EMessageSeverity::Error);
 				}
 				else if (Algo::AnyOf(Assets, [](const FAssetData& Asset) { return Asset.IsRedirector(); }))
 				{
-					const FText ErrorMessage = FText::Format(LOCTEXT("GameMapSettingsRedirectorDetected", "GameMapsSettings contains a redirected reference '{0}'. The intended asset will fail to load in a packaged build. Select the intended asset again in Project Settings to fix this issue."), FText::FromName(PackagePath));
+					const FText ErrorMessage = FText::Format(LOCTEXT("GameMapSettingsRedirectorDetected", "{0} contains a redirected reference '{1}'. The intended asset will fail to load in a packaged build. Select the intended asset again in Project Settings to fix this issue."),
+						FText::FromName(GameDefaultSet.Key), FText::FromName(PackagePath));
 					LogCookerMessage(ErrorMessage.ToString(), EMessageSeverity::Error);
 				}
 
@@ -6721,7 +6723,7 @@ void UCookOnTheFlyServer::CollectFilesToCook(TArray<FName>& FilesInPath, const T
 	}
 }
 
-void UCookOnTheFlyServer::GetGameDefaultObjects(const TArray<ITargetPlatform*>& TargetPlatforms, TMap<FName, TArray<FName>>& GameDefaultObjectsOut)
+void UCookOnTheFlyServer::GetGameDefaultObjects(const TArray<ITargetPlatform*>& TargetPlatforms, TMap<FName, TArray<FName>>& OutGameDefaultObjects)
 {
 	// Collect all default objects from all cooked platforms engine configurations.
 	for (const ITargetPlatform* TargetPlatform : TargetPlatforms)
@@ -6737,24 +6739,34 @@ void UCookOnTheFlyServer::GetGameDefaultObjects(const TArray<ITargetPlatform*>&
 			continue;
 		}
 
-		auto AddDefaultObject = [&GameDefaultObjectsOut, &PlatformEngineIni, MapSettingsSection](FName PropertyName)
+		auto AddDefaultObject = [&OutGameDefaultObjects, &PlatformEngineIni, MapSettingsSection](FName PropertyName)
 		{
 			const FConfigValue* PairString = MapSettingsSection->Find(PropertyName);
 			if (PairString == nullptr)
 			{
 				return;
 			}
+			FString ObjectPath = PairString->GetValue();
+			if (ObjectPath.IsEmpty())
+			{
+				return;
+			}
 
-			FSoftObjectPath Path(PairString->GetValue());
-			GameDefaultObjectsOut.FindOrAdd(PropertyName).AddUnique(Path.GetLongPackageFName());
+			FSoftObjectPath Path(ObjectPath);
+			FName PackageName = Path.GetLongPackageFName();
+			if (PackageName.IsNone())
+			{
+				return;
+			}
+			OutGameDefaultObjects.FindOrAdd(PropertyName).AddUnique(PackageName);
 		};
 
 		// get the server and game default maps/modes and cook them
-		AddDefaultObject(FName("GameDefaultMap"));
-		AddDefaultObject(FName("ServerDefaultMap"));
-		AddDefaultObject(FName("GlobalDefaultGameMode"));
-		AddDefaultObject(FName("GlobalDefaultServerGameMode"));
-		AddDefaultObject(FName("GameInstanceClass"));
+		AddDefaultObject(FName(TEXT("GameDefaultMap")));
+		AddDefaultObject(FName(TEXT("ServerDefaultMap")));
+		AddDefaultObject(FName(TEXT("GlobalDefaultGameMode")));
+		AddDefaultObject(FName(TEXT("GlobalDefaultServerGameMode")));
+		AddDefaultObject(FName(TEXT("GameInstanceClass")));
 	}
 }
 
