Matinee/Editor: Removed CameraAnim, InterpData and MatineeActor factories as new Matinee assets should not be created.

diff --git a/Engine/Plugins/Tests/EditorTests/Source/EditorTests/Private/UnrealEd/EditorAssetAutomationTests.cpp b/Engine/Plugins/Tests/EditorTests/Source/EditorTests/Private/UnrealEd/EditorAssetAutomationTests.cpp
--- a/Engine/Plugins/Tests/EditorTests/Source/EditorTests/Private/UnrealEd/EditorAssetAutomationTests.cpp
+++ b/Engine/Plugins/Tests/EditorTests/Source/EditorTests/Private/UnrealEd/EditorAssetAutomationTests.cpp
@@ -61,13 +61,11 @@
 #include "Factories/BlueprintFactory.h"
 #include "Factories/BlueprintFunctionLibraryFactory.h"
 #include "Factories/BlueprintMacroFactory.h"
-#include "Factories/CameraAnimFactory.h"
 #include "Factories/DialogueVoiceFactory.h"
 #include "Factories/DialogueWaveFactory.h"
 #include "Factories/EnumFactory.h"
 #include "Factories/ForceFeedbackAttenuationFactory.h"
 #include "Factories/ForceFeedbackEffectFactory.h"
-#include "Factories/InterpDataFactoryNew.h"
 #include "Factories/MaterialFactoryNew.h"
 #include "Factories/MaterialFunctionFactoryNew.h"
 #include "Factories/MaterialInstanceConstantFactoryNew.h"
@@ -634,7 +632,6 @@ bool FAssetEditorTest::RunTest(const FString& Parameters)
 	ASSET_TEST_CREATE_BY_NAME(AIModule.BehaviorTree, BehaviorTreeEditor.BehaviorTreeFactory, BT, )
 	ASSET_TEST_CREATE(UBlueprint, UBlueprintFunctionLibraryFactory, BFL, )
 	ASSET_TEST_CREATE(UBlueprint, UBlueprintMacroFactory, MPL, FactoryInst->ParentClass = AActor::StaticClass();)
-	ASSET_TEST_CREATE(UCameraAnim, UCameraAnimFactory, CA, )
 	ASSET_TEST_CREATE(UCurveBase, UCurveFactory, C, FactoryInst->CurveClass = UCurveFloat::StaticClass();)
 	UClass* GameplayAbilityClass = StaticLoadClass(UObject::StaticClass(), NULL, TEXT("GameplayAbilities.GameplayAbilitySet"), NULL, LOAD_None, NULL);
 	if (GameplayAbilityClass)
diff --git a/Engine/Source/Editor/UnrealEd/Classes/ActorFactories/ActorFactoryMatineeActor.h b/Engine/Source/Editor/UnrealEd/Classes/ActorFactories/ActorFactoryMatineeActor.h
deleted file mode 100644
diff --git a/Engine/Source/Editor/UnrealEd/Classes/Factories/CameraAnimFactory.h b/Engine/Source/Editor/UnrealEd/Classes/Factories/CameraAnimFactory.h
deleted file mode 100644
diff --git a/Engine/Source/Editor/UnrealEd/Classes/Factories/InterpDataFactoryNew.h b/Engine/Source/Editor/UnrealEd/Classes/Factories/InterpDataFactoryNew.h
deleted file mode 100644
diff --git a/Engine/Source/Editor/UnrealEd/Private/Factories/ActorFactory.cpp b/Engine/Source/Editor/UnrealEd/Private/Factories/ActorFactory.cpp
--- a/Engine/Source/Editor/UnrealEd/Private/Factories/ActorFactory.cpp
+++ b/Engine/Source/Editor/UnrealEd/Private/Factories/ActorFactory.cpp
@@ -30,7 +30,6 @@ ActorFactory.cpp:
 #include "ActorFactories/ActorFactoryEmptyActor.h"
 #include "ActorFactories/ActorFactoryPawn.h"
 #include "ActorFactories/ActorFactoryExponentialHeightFog.h"
-#include "ActorFactories/ActorFactoryMatineeActor.h"
 #include "ActorFactories/ActorFactoryNote.h"
 #include "ActorFactories/ActorFactoryPhysicsAsset.h"
 #include "ActorFactories/ActorFactoryPlaneReflectionCapture.h"
@@ -102,8 +101,6 @@ ActorFactory.cpp:
 #include "Kismet2/BlueprintEditorUtils.h"
 #include "Kismet2/KismetEditorUtilities.h"
 #include "BSPOps.h"
-#include "Matinee/MatineeActor.h"
-#include "Matinee/InterpData.h"
 #include "InteractiveFoliageActor.h"
 #include "VT/RuntimeVirtualTextureVolume.h"
 
@@ -1653,69 +1650,6 @@ bool UActorFactoryBlueprint::PreSpawnActor( UObject* Asset, FTransform& InOutLoc
 	return true;
 }
 
-
-
-/*-----------------------------------------------------------------------------
-UActorFactoryMatineeActor
------------------------------------------------------------------------------*/
-UActorFactoryMatineeActor::UActorFactoryMatineeActor(const FObjectInitializer& ObjectInitializer)
-	: Super(ObjectInitializer)
-{
-	DisplayName = LOCTEXT("MatineeDisplayName", "Matinee");
-	NewActorClass = AMatineeActor::StaticClass();
-}
-
-bool UActorFactoryMatineeActor::CanCreateActorFrom( const FAssetData& AssetData, FText& OutErrorMsg )
-{
-	//We allow creating AMatineeActors without an existing asset
-	if ( UActorFactory::CanCreateActorFrom( AssetData, OutErrorMsg ) )
-	{
-		return true;
-	}
-
-	if ( AssetData.IsValid() && !AssetData.GetClass()->IsChildOf( UInterpData::StaticClass() ) )
-	{
-		OutErrorMsg = NSLOCTEXT("CanCreateActor", "NoInterpData", "A valid InterpData must be specified.");
-		return false;
-	}
-
-	return true;
-}
-
-void UActorFactoryMatineeActor::PostSpawnActor( UObject* Asset, AActor* NewActor )
-{
-	Super::PostSpawnActor(Asset, NewActor);
-
-	UInterpData* MatineeData = Cast<UInterpData>( Asset );
-	AMatineeActor* MatineeActor = CastChecked<AMatineeActor>( NewActor );
-
-	if( MatineeData )
-	{
-		MatineeActor->MatineeData = MatineeData;
-	}
-	else
-	{
-		// if MatineeData isn't set yet, create default one
-		UInterpData* NewMatineeData = NewObject<UInterpData>(NewActor);
-		MatineeActor->MatineeData = NewMatineeData;
-	}
-}
-
-void UActorFactoryMatineeActor::PostCreateBlueprint( UObject* Asset, AActor* CDO )
-{
-	if (Asset != nullptr && CDO != nullptr)
-	{
-		UInterpData* MatineeData = Cast<UInterpData>(Asset);
-		AMatineeActor* MatineeActor = CastChecked<AMatineeActor>(CDO);
-
-		// @todo sequencer: Don't ever need or want InterpData for Sequencer.  We will probably get rid of this after old Matinee goes away.
-		// also note the PostSpawnActor() code above creates an UInterpData and puts it in the actor's outermost package.  Definitely do not
-		// want that for Sequencer.
-		MatineeActor->MatineeData = MatineeData;
-	}
-}
-
-
 /*-----------------------------------------------------------------------------
 UActorFactoryDirectionalLight
 -----------------------------------------------------------------------------*/
diff --git a/Engine/Source/Editor/UnrealEd/Private/Factories/EditorFactories.cpp b/Engine/Source/Editor/UnrealEd/Private/Factories/EditorFactories.cpp
--- a/Engine/Source/Editor/UnrealEd/Private/Factories/EditorFactories.cpp
+++ b/Engine/Source/Editor/UnrealEd/Private/Factories/EditorFactories.cpp
@@ -69,7 +69,6 @@
 #include "Factories/BlueprintFunctionLibraryFactory.h"
 #include "Factories/BlueprintMacroFactory.h"
 #include "Factories/BlueprintInterfaceFactory.h"
-#include "Factories/CameraAnimFactory.h"
 #include "Factories/CurveFactory.h"
 #include "Factories/CurveImportFactory.h"
 #include "Factories/DataAssetFactory.h"
@@ -85,7 +84,6 @@
 #include "Factories/HapticFeedbackEffectCurveFactory.h"
 #include "Factories/HapticFeedbackEffectBufferFactory.h"
 #include "Factories/HapticFeedbackEffectSoundWaveFactory.h"
-#include "Factories/InterpDataFactoryNew.h"
 #include "Factories/LevelFactory.h"
 #include "Factories/MaterialFactoryNew.h"
 #include "Factories/MaterialFunctionFactoryNew.h"
@@ -179,8 +177,6 @@
 #include "Engine/UserDefinedStruct.h"
 #include "Internationalization/StringTable.h"
 #include "Editor.h"
-#include "Matinee/InterpData.h"
-#include "Matinee/InterpGroupCamera.h"
 #include "Materials/MaterialExpressionTextureSample.h"
 #include "Sound/SoundNodeWavePlayer.h"
 #include "Sound/SoundNodeAttenuation.h"
@@ -2115,23 +2111,6 @@ int32 UPhysicalMaterialMaskFactory::GetPriority() const
 	return ImportPriority;
 }
 
-/*------------------------------------------------------------------------------
-	UInterpDataFactoryNew.
-------------------------------------------------------------------------------*/
-UInterpDataFactoryNew::UInterpDataFactoryNew(const FObjectInitializer& ObjectInitializer)
-	: Super(ObjectInitializer)
-{
-
-	SupportedClass = UInterpData::StaticClass();
-	bCreateNew = true;
-	bEditAfterNew = true;
-}
-
-UObject* UInterpDataFactoryNew::FactoryCreateNew(UClass* Class,UObject* InParent,FName Name,EObjectFlags Flags,UObject* Context,FFeedbackContext* Warn)
-{
-	return NewObject<UObject>(InParent, Class, Name, Flags);
-}
-
 /*-----------------------------------------------------------------------------
 	UTextureRenderTargetFactoryNew
 -----------------------------------------------------------------------------*/
@@ -7650,40 +7629,6 @@ UObject* UTouchInterfaceFactory::FactoryCreateNew( UClass* InClass, UObject* InP
 	return NewObject<UTouchInterface>(InParent, InName, Flags);
 }
 
-/*------------------------------------------------------------------------------
-	UCameraAnimFactory implementation.
-------------------------------------------------------------------------------*/
-
-UCameraAnimFactory::UCameraAnimFactory(const FObjectInitializer& ObjectInitializer)
-	: Super(ObjectInitializer)
-{
-	SupportedClass = UCameraAnim::StaticClass();
-	bCreateNew = true;
-}
-
-FText UCameraAnimFactory::GetDisplayName() const
-{
-	return LOCTEXT("CameraAnimFactoryDescription", "Camera Anim");
-}
-
-FName UCameraAnimFactory::GetNewAssetThumbnailOverride() const
-{
-	return TEXT("ClassThumbnail.CameraAnim");
-}
-
-uint32 UCameraAnimFactory::GetMenuCategories() const
-{
-	return EAssetTypeCategories::Misc;
-}
-
-UObject* UCameraAnimFactory::FactoryCreateNew(UClass* Class,UObject* InParent,FName Name,EObjectFlags Flags,UObject* Context,FFeedbackContext* Warn)
-{
-	UCameraAnim* NewCamAnim = NewObject<UCameraAnim>(InParent, Class, Name, Flags);
-	NewCamAnim->CameraInterpGroup = NewObject<UInterpGroupCamera>(NewCamAnim);
-	NewCamAnim->CameraInterpGroup->GroupName = Name;
-	return NewCamAnim;
-}
-
 /*------------------------------------------------------------------------------
 UStringTableFactory implementation.
 ------------------------------------------------------------------------------*/
diff --git a/Engine/Source/Editor/UnrealEd/Public/UnrealEd.h b/Engine/Source/Editor/UnrealEd/Public/UnrealEd.h
--- a/Engine/Source/Editor/UnrealEd/Public/UnrealEd.h
+++ b/Engine/Source/Editor/UnrealEd/Public/UnrealEd.h
@@ -34,7 +34,6 @@ MONOLITHIC_HEADER_BOILERPLATE()
 #include "ActorFactories/ActorFactoryEmptyActor.h"
 #include "ActorFactories/ActorFactoryPawn.h"
 #include "ActorFactories/ActorFactoryExponentialHeightFog.h"
-#include "ActorFactories/ActorFactoryMatineeActor.h"
 #include "ActorFactories/ActorFactoryNote.h"
 #include "ActorFactories/ActorFactoryPhysicsAsset.h"
 #include "ActorFactories/ActorFactoryPlaneReflectionCapture.h"
@@ -131,7 +130,6 @@ MONOLITHIC_HEADER_BOILERPLATE()
 #include "Factories/BlueprintFunctionLibraryFactory.h"
 #include "Factories/BlueprintMacroFactory.h"
 #include "Factories/BlueprintInterfaceFactory.h"
-#include "Factories/CameraAnimFactory.h"
 #include "Factories/EnumFactory.h"
 #include "Factories/FbxFactory.h"
 #include "Factories/FbxSceneImportFactory.h"
@@ -145,7 +143,6 @@ MONOLITHIC_HEADER_BOILERPLATE()
 #include "Factories/HapticFeedbackEffectCurveFactory.h"
 #include "Factories/HapticFeedbackEffectBufferFactory.h"
 #include "Factories/HapticFeedbackEffectSoundWaveFactory.h"
-#include "Factories/InterpDataFactoryNew.h"
 #include "Factories/LevelFactory.h"
 #include "Factories/MaterialFactoryNew.h"
 #include "Factories/MaterialFunctionFactoryNew.h"
