Add EditorActorSubsystem functions to delete actors without relying on selection state

diff --git a/Engine/Source/Editor/UnrealEd/Private/Layers/LayersSubsystem.cpp b/Engine/Source/Editor/UnrealEd/Private/Layers/LayersSubsystem.cpp
--- a/Engine/Source/Editor/UnrealEd/Private/Layers/LayersSubsystem.cpp
+++ b/Engine/Source/Editor/UnrealEd/Private/Layers/LayersSubsystem.cpp
@@ -221,23 +221,10 @@ bool ULayersSubsystem::InitializeNewActorLayers(AActor* Actor)
 
 bool ULayersSubsystem::DisassociateActorFromLayers(AActor* Actor)
 {
-	if(	!IsActorValidForLayer( Actor ) )
-	{
-		return false;
-	}
-
-	bool bChangeOccurred = false;
-	for( auto LayerNameIt = Actor->Layers.CreateConstIterator(); LayerNameIt; ++LayerNameIt )
-	{
-		const FName LayerName = *LayerNameIt;
-		ULayer* Layer = EnsureLayerExists( LayerName );
-
-		Layer->Modify();
-		RemoveActorFromStats( Layer, Actor);
-		bChangeOccurred = true;
-	}
+	TArray< AActor* > Actors;
+	Actors.Add(Actor);
 
-	return bChangeOccurred;
+	return DisassociateActorsFromLayers(Actors);
 }
 
 
@@ -357,6 +344,29 @@ bool ULayersSubsystem::AddActorsToLayers(const TArray< TWeakObjectPtr< AActor >
 	return AddActorsToLayers(ActorsRawPtr, LayerNames);
 }
 
+bool ULayersSubsystem::DisassociateActorsFromLayers(const TArray<AActor*>& Actors)
+{
+	bool bChangesOccurred = false;
+	
+	for(AActor* Actor : Actors)
+	{
+		if (!IsActorValidForLayer(Actor))
+		{
+			continue;
+		}
+
+		for(const FName& LayerName : Actor->Layers)
+		{
+			ULayer* Layer = EnsureLayerExists(LayerName);
+
+			Layer->Modify();
+			RemoveActorFromStats(Layer, Actor);
+			bChangesOccurred = true;
+		}
+	}
+
+	return bChangesOccurred;
+}
 
 bool ULayersSubsystem::RemoveActorFromLayer(AActor* Actor, const FName& LayerName, const bool bUpdateStats)
 {
diff --git a/Engine/Source/Editor/UnrealEd/Private/Subsystems/EditorActorSubsystem.cpp b/Engine/Source/Editor/UnrealEd/Private/Subsystems/EditorActorSubsystem.cpp
--- a/Engine/Source/Editor/UnrealEd/Private/Subsystems/EditorActorSubsystem.cpp
+++ b/Engine/Source/Editor/UnrealEd/Private/Subsystems/EditorActorSubsystem.cpp
@@ -482,25 +482,36 @@ AActor* UEditorActorSubsystem::SpawnActorFromClass(TSubclassOf<class AActor> Act
 	return InternalActorUtilitiesSubsystemLibrary::SpawnActor(TEXT("SpawnActorFromClass"), ActorClass.Get(), Location, Rotation, bTransient);
 }
 
-bool UEditorActorSubsystem::DestroyActor(class AActor* ToDestroyActor)
+bool UEditorActorSubsystem::DestroyActor(AActor* ActorToDestroy)
+{
+	TArray<AActor*> ActorsToDestroy;
+	ActorsToDestroy.Add(ActorToDestroy);
+	return DestroyActors(ActorsToDestroy);
+}
+
+bool UEditorActorSubsystem::DestroyActors(const TArray<AActor*>& ActorsToDestroy)
 {
 	TGuardValue<bool> UnattendedScriptGuard(GIsRunningUnattendedScript, true);
 
+	const FScopedTransaction Transaction(LOCTEXT("DeleteActors", "Delete Actors"));
+	
 	if (!EditorScriptingHelpers::CheckIfInEditorAndPIE())
 	{
 		return false;
 	}
 
-	if (!ToDestroyActor)
+	for (AActor* ActorToDestroy : ActorsToDestroy)
 	{
-		UE_LOG(LogUtils, Error, TEXT("DestroyActor. ToDestroyActor is invalid."));
-		return false;
-	}
-
-	if (!InternalActorUtilitiesSubsystemLibrary::IsEditorLevelActor(ToDestroyActor))
-	{  
-		UE_LOG(LogUtils, Error, TEXT("DestroyActor. The Actor is not part of the world editor."));
-		return false;
+		if (!ActorToDestroy)
+		{
+			UE_LOG(LogUtils, Error, TEXT("DestroyActors. An actor to destroy is invalid."));
+			return false;
+		}
+		if (!InternalActorUtilitiesSubsystemLibrary::IsEditorLevelActor(ActorToDestroy))
+		{  
+			UE_LOG(LogUtils, Error, TEXT("DestroyActors. An actor to destroy is not part of the world editor."));
+			return false;
+		}
 	}
 
 	UUnrealEditorSubsystem* UnrealEditorSubsystem = GEditor->GetEditorSubsystem<UUnrealEditorSubsystem>();
@@ -513,19 +524,40 @@ bool UEditorActorSubsystem::DestroyActor(class AActor* ToDestroyActor)
 	UWorld* World = UnrealEditorSubsystem->GetEditorWorld();
 	if (!World)
 	{
-		UE_LOG(LogUtils, Error, TEXT("DestroyActor. Can't destroy the actor because there is no world."));
+		UE_LOG(LogUtils, Error, TEXT("DestroyActors. Can't destroy actors because there is no world."));
 		return false;
 	}
 
-	//To avoid dangling gizmo after actor has been destroyed
-	FTypedElementHandle ActorHandle = UEngineElementsLibrary::AcquireEditorActorElementHandle(ToDestroyActor);
+	FEditorDelegates::OnDeleteActorsBegin.Broadcast();
+	
+	TArray<FTypedElementHandle> ActorHandles;
+	ActorHandles.Reserve(ActorsToDestroy.Num());
+	for (AActor* ActorToDestroy : ActorsToDestroy)
+	{
+		ActorHandles.Add(UEngineElementsLibrary::AcquireEditorActorElementHandle(ActorToDestroy));
+	}
+
+	//To avoid dangling gizmo after actors have been destroyed
 	USelection* ActorSelection = GEditor->GetSelectedActors();
 	UTypedElementSelectionSet* SelectionSet = ActorSelection->GetElementSelectionSet();
-	SelectionSet->DeselectElement(ActorHandle, FTypedElementSelectionOptions());
+	SelectionSet->DeselectElements(ActorHandles, FTypedElementSelectionOptions());
 
 	ULayersSubsystem* Layers = GEditor->GetEditorSubsystem<ULayersSubsystem>();
-	Layers->DisassociateActorFromLayers(ToDestroyActor);
-	return World->EditorDestroyActor(ToDestroyActor, true);
+	Layers->DisassociateActorsFromLayers(ActorsToDestroy);
+
+	bool SuccessfullyDestroyedAll = true;
+	
+	for (AActor* ActorToDestroy : ActorsToDestroy)
+	{
+		if (!World->EditorDestroyActor(ActorToDestroy, true))
+		{
+			SuccessfullyDestroyedAll = false;
+		}
+	}
+
+	FEditorDelegates::OnDeleteActorsEnd.Broadcast();
+	
+	return SuccessfullyDestroyedAll;
 }
 
 TArray<class AActor*> UEditorActorSubsystem::ConvertActors(const TArray<class AActor*>& Actors, TSubclassOf<class AActor> ActorClass, const FString& StaticMeshPackagePath)
diff --git a/Engine/Source/Editor/UnrealEd/Public/Layers/LayersSubsystem.h b/Engine/Source/Editor/UnrealEd/Public/Layers/LayersSubsystem.h
--- a/Engine/Source/Editor/UnrealEd/Public/Layers/LayersSubsystem.h
+++ b/Engine/Source/Editor/UnrealEd/Public/Layers/LayersSubsystem.h
@@ -191,6 +191,14 @@ public:
 	 */
 	virtual bool AddActorsToLayers(const TArray< TWeakObjectPtr< AActor > >& Actors, const TArray< FName >& LayerNames) final;
 
+	/**
+	*	Disassociates actors from the layer system, generally used before deleting the Actors
+	*
+	*	@param	Actors	The actors to disassociate from the layer system
+	*/
+	UFUNCTION(BlueprintCallable, Category = Layers)
+	virtual bool DisassociateActorsFromLayers(const TArray< AActor* >& Actors) final;
+	
 	/**
 	 * Removes the actors from the specified layer.
 	 *
diff --git a/Engine/Source/Editor/UnrealEd/Public/Subsystems/EditorActorSubsystem.h b/Engine/Source/Editor/UnrealEd/Public/Subsystems/EditorActorSubsystem.h
--- a/Engine/Source/Editor/UnrealEd/Public/Subsystems/EditorActorSubsystem.h
+++ b/Engine/Source/Editor/UnrealEd/Public/Subsystems/EditorActorSubsystem.h
@@ -78,6 +78,7 @@ public:
 	 * Duplicate an actor from the world editor.
 	 * @param	ActorToDuplicate	Actor to duplicate.
 	 * @param	ToWorld				World to place the duplicated actor in.
+	 * @param	Offset				Translation to offset duplicated actor by.
 	 * @return	The duplicated actor, or none if it didn't succeed
 	 */
 	UFUNCTION(BlueprintCallable, Category = "Editor Scripting | Level Utility", meta = (AdvancedDisplay = 1))
@@ -87,6 +88,7 @@ public:
 	 * Duplicate actors from the world editor.
 	 * @param	ActorsToDuplicate	Actors to duplicate.
 	 * @param	ToWorld				World to place the duplicated actors in.
+	 * * @param	Offset				Translation to offset duplicated actors by.
 	 * @return	The duplicated actors, or empty if it didn't succeed
 	 */
 	UFUNCTION(BlueprintCallable, Category = "Editor Scripting | Level Utility", meta = (AdvancedDisplay = 1))
@@ -199,11 +201,19 @@ public:
 
 	/**
 	 * Destroy the actor from the world editor. Notify the Editor that the actor got destroyed.
-	 * @param	ToDestroyActor	Actor to destroy.
+	 * @param	ActorToDestroy	Actor to destroy.
 	 * @return	True if the operation succeeds.
 	 */
 	UFUNCTION(BlueprintCallable, Category = "Editor Scripting | Level Utility")
-	bool DestroyActor(class AActor* ActorToDestroy);
+	bool DestroyActor(AActor* ActorToDestroy);
+
+	/**
+	* Destroy the actors from the world editor. Notify the Editor that the actor got destroyed.
+	* @param	ActorsToDestroy		Actors to destroy.
+	* @return	True if the operation succeeds.
+	*/
+	UFUNCTION(BlueprintCallable, Category = "Editor Scripting | Level Utility")
+	bool DestroyActors(const TArray<AActor*>& ActorsToDestroy);
 
 	/**
 	 * Replace in the level all Actors provided with a new actor of type ActorClass. Destroy all Actors provided.
