Fix an issue where hit proxies were not being updated after placement mode painted assets.

Start/EndTracking was not being called for the mode manager after input reshuffles; additionally if we are circumventing the normal viewport start/end tracking, we need to make sure that hit proxies redraw for items that may have been placed via an ITF tool by invalidating the viewport.

diff --git a/Engine/Source/Editor/Experimental/EditorInteractiveToolsFramework/Private/EdModeInteractiveToolsContext.cpp b/Engine/Source/Editor/Experimental/EditorInteractiveToolsFramework/Private/EdModeInteractiveToolsContext.cpp
--- a/Engine/Source/Editor/Experimental/EditorInteractiveToolsFramework/Private/EdModeInteractiveToolsContext.cpp
+++ b/Engine/Source/Editor/Experimental/EditorInteractiveToolsFramework/Private/EdModeInteractiveToolsContext.cpp
@@ -1120,6 +1120,10 @@ bool UEdModeInteractiveToolsContext::EndTracking(FEditorViewportClient* InViewpo
 {
 	if (bIsTrackingMouse)
 	{
+		// If the input router captured the mouse input, we need to invalidate the viewport client here, since the mouse delta tracker's end tracking will not be called.
+		constexpr bool bForceChildViewportRedraw = true;
+		constexpr bool bInvalidateHitProxies = true;
+		InViewportClient->Invalidate(bForceChildViewportRedraw, bInvalidateHitProxies);
 		bIsTrackingMouse = false;
 		return true;
 	}
diff --git a/Engine/Source/Editor/UnrealEd/Private/EditorModeManager.cpp b/Engine/Source/Editor/UnrealEd/Private/EditorModeManager.cpp
--- a/Engine/Source/Editor/UnrealEd/Private/EditorModeManager.cpp
+++ b/Engine/Source/Editor/UnrealEd/Private/EditorModeManager.cpp
@@ -1452,10 +1452,19 @@ bool FEditorModeTools::ProcessCapturedMouseMoves( FEditorViewportClient* InViewp
 /** Notifies all active modes of keyboard input via a viewport client */
 bool FEditorModeTools::InputKey(FEditorViewportClient* InViewportClient, FViewport* Viewport, FKey Key, EInputEvent Event, bool bRouteToToolsContext)
 {
-	bool bHandled = bRouteToToolsContext && InteractiveToolsContext->InputKey(InViewportClient, Viewport, Key, Event);
+	const bool bWasHandledByToolsContext = bRouteToToolsContext && InteractiveToolsContext->InputKey(InViewportClient, Viewport, Key, Event);
+	if (bWasHandledByToolsContext && !bIsTracking && GetInteractiveToolsContext()->InputRouter->HasActiveMouseCapture())
+	{
+		StartTracking(InViewportClient, Viewport);
+	}
+	else if (bRouteToToolsContext && bIsTracking && !GetInteractiveToolsContext()->InputRouter->HasActiveMouseCapture())
+	{
+		EndTracking(InViewportClient, Viewport);
+	}
 
 	// If the toolkit should process the command, it should not have been handled by ITF, or be tracked elsewhere.
-	const bool bPassToToolkitCommands = bRouteToToolsContext && !bHandled;
+	const bool bPassToToolkitCommands = bRouteToToolsContext && !bWasHandledByToolsContext;
+	bool bHandled = bWasHandledByToolsContext;
 	ForEachEdMode([&bHandled, bPassToToolkitCommands, Event, Key, InViewportClient, Viewport](UEdMode* Mode)
 	{
 		// First, always give the legacy viewport interface a chance to process they key press. This is to support any of the FModeTools that may still exist.
