Fixed crash when converting a map to world partition : a viewport temporarily has a null world then.

diff --git a/Engine/Source/Editor/UnrealEd/Private/LevelEditorViewport.cpp b/Engine/Source/Editor/UnrealEd/Private/LevelEditorViewport.cpp
--- a/Engine/Source/Editor/UnrealEd/Private/LevelEditorViewport.cpp
+++ b/Engine/Source/Editor/UnrealEd/Private/LevelEditorViewport.cpp
@@ -3840,21 +3840,26 @@ EMouseCursor::Type FLevelEditorViewportClient::GetCursor(FViewport* InViewport,i
 {
 	EMouseCursor::Type CursorType = FEditorViewportClient::GetCursor(InViewport,X,Y);
 
-	// Allow the viewport interaction to override any previously set mouse cursor
-	UViewportWorldInteraction* WorldInteraction = Cast<UViewportWorldInteraction>(GEditor->GetEditorWorldExtensionsManager()->GetEditorWorldExtensions(GetWorld())->FindExtension(UViewportWorldInteraction::StaticClass()));
-	if (WorldInteraction != nullptr)
+	UWorld* ViewportWorld = GetWorld();
+
+	if (ViewportWorld != nullptr)
 	{
-		if (WorldInteraction->ShouldForceCursor())
-		{
-			CursorType = EMouseCursor::Crosshairs;
-			SetRequiredCursor(false, true);
-			UpdateRequiredCursorVisibility();
-		}
-		else if (WorldInteraction->ShouldSuppressExistingCursor())
+		// Allow the viewport interaction to override any previously set mouse cursor
+		UViewportWorldInteraction* WorldInteraction = Cast<UViewportWorldInteraction>(GEditor->GetEditorWorldExtensionsManager()->GetEditorWorldExtensions(GetWorld())->FindExtension(UViewportWorldInteraction::StaticClass()));
+		if (WorldInteraction != nullptr)
 		{
-			CursorType = EMouseCursor::None;
-			SetRequiredCursor(false, false);
-			UpdateRequiredCursorVisibility();
+			if (WorldInteraction->ShouldForceCursor())
+			{
+				CursorType = EMouseCursor::Crosshairs;
+				SetRequiredCursor(false, true);
+				UpdateRequiredCursorVisibility();
+			}
+			else if (WorldInteraction->ShouldSuppressExistingCursor())
+			{
+				CursorType = EMouseCursor::None;
+				SetRequiredCursor(false, false);
+				UpdateRequiredCursorVisibility();
+			}
 		}
 	}
 
diff --git a/Engine/Source/Runtime/Engine/Private/SceneView.cpp b/Engine/Source/Runtime/Engine/Private/SceneView.cpp
--- a/Engine/Source/Runtime/Engine/Private/SceneView.cpp
+++ b/Engine/Source/Runtime/Engine/Private/SceneView.cpp
@@ -2691,20 +2691,25 @@ const FSceneView& FSceneViewFamily::GetStereoEyeView(const EStereoscopicPass Eye
 
 bool FSceneViewFamily::SupportsScreenPercentage() const
 {
-	EShadingPath ShadingPath = Scene->GetShadingPath();
-
-	// The deferred shading renderer supports screen percentage when used normally
-	if (Scene->GetShadingPath() == EShadingPath::Deferred)
+	if (Scene != nullptr)
 	{
+		EShadingPath ShadingPath = Scene->GetShadingPath();
+
+		// The deferred shading renderer supports screen percentage when used normally
+		if (Scene->GetShadingPath() == EShadingPath::Deferred)
+		{
+			return true;
+		}
+
+		// Mobile renderer does not support screen percentage with LDR.
+		if ((GetFeatureLevel() <= ERHIFeatureLevel::ES3_1 && !IsMobileHDR()))
+		{
+			return false;
+		}
 		return true;
 	}
 
-	// Mobile renderer does not support screen percentage with LDR.
-	if ((GetFeatureLevel() <= ERHIFeatureLevel::ES3_1 && !IsMobileHDR()))
-	{
-		return false;
-	}
-	return true;
+	return false;
 }
 
 FSceneViewFamilyContext::~FSceneViewFamilyContext()
diff --git a/Engine/Source/Runtime/Engine/Private/UnrealClient.cpp b/Engine/Source/Runtime/Engine/Private/UnrealClient.cpp
--- a/Engine/Source/Runtime/Engine/Private/UnrealClient.cpp
+++ b/Engine/Source/Runtime/Engine/Private/UnrealClient.cpp
@@ -1651,8 +1651,9 @@ const TArray<FColor>& FViewport::GetRawHitProxyData(FIntRect InRect)
 	const bool bIsRenderingStereo = GEngine->IsStereoscopic3D( this ) && this->IsStereoRenderingAllowed();
 
 	bool bFetchHitProxyBytes = !bIsRenderingStereo && ( !bHitProxiesCached || (SizeY*SizeX) != CachedHitProxyData.Num() );
+	UWorld* World = ViewportClient->GetWorld();
 
-	if( bIsRenderingStereo )
+	if( bIsRenderingStereo || (World == nullptr))
 	{
 		// Stereo viewports don't support hit proxies, and we don't want to update them because it will adversely
 		// affect performance.
@@ -1677,8 +1678,8 @@ const TArray<FColor>& FViewport::GetRawHitProxyData(FIntRect InRect)
 			});
 
 		// Let the viewport client draw its hit proxies.
-		UWorld* World = ViewportClient->GetWorld();
-		FCanvas Canvas(&HitProxyMap, &HitProxyMap, World, World ? World->FeatureLevel.GetValue() : GMaxRHIFeatureLevel, FCanvas::CDM_DeferDrawing, ViewportClient->ShouldDPIScaleSceneCanvas() ? ViewportClient->GetDPIScale() : 1.0f);
+		check(World != nullptr);
+		FCanvas Canvas(&HitProxyMap, &HitProxyMap, World, World->FeatureLevel.GetValue(), FCanvas::CDM_DeferDrawing, ViewportClient->ShouldDPIScaleSceneCanvas() ? ViewportClient->GetDPIScale() : 1.0f);
 		{
 			ViewportClient->Draw(this, &Canvas);
 		}
