Anim blueprint encapsulation improvements

Adds 'template' anim BP concept. These anim BPs have no TargetSkeleton and as such cannot have direct references to animations placed inside of their anim graphs
Adds function call support from anim nodes. All anim graph nodes can now call functions when initialized, updated, evaluated or when they first become relevant.
Relevancy is established using a new FAnimSubsystemInstance_NodeRelevancy which tracks nodes in a local map, per UAnimInstance, if nodes require it.
Functions are displayed on the node if bound, and in the details panel.
Added a new override point to FAnimSubsystemInstance to allow for WT init.
Moved FMemberReference customization into a public header so it can be used on anim node functions.
Wrapped functions up into FAnimNodeFunctionRef structure so they can be re-used more effectively. Converted CallFunction node to use them.
Added a couple of simple BP function libraries to demonstrate the use of the new FAnimNodeContext (this is intended to be a generic way of exposing FAnimNode_Base types to script without the node types themselves having to be known to script directly).
Added the ability to set exposed properties as 'always dynamic' so they appear in mutable data rather than being merged into constants. This allows for the anim node data API to be changed to be less strict (and more script friendly). Now values can be set in mutable data (via GET_MUTABLE_ANIM_NODE_DATA_PTR) and attempted sets to constant data can be ignored and reported to client code.
A few minor crash fixes with edge cases (inc when trying to open an asset editor fails because of a missing skeleton/cancellation).
Also fixes an issue where literal linked anim graph/control rig/custom poroperty node inputs didnt get copied

diff --git a/Engine/Plugins/Experimental/Animation/PoseSearch/Source/Runtime/Private/AnimNode_MotionMatching.cpp b/Engine/Plugins/Experimental/Animation/PoseSearch/Source/Runtime/Private/AnimNode_MotionMatching.cpp
--- a/Engine/Plugins/Experimental/Animation/PoseSearch/Source/Runtime/Private/AnimNode_MotionMatching.cpp
+++ b/Engine/Plugins/Experimental/Animation/PoseSearch/Source/Runtime/Private/AnimNode_MotionMatching.cpp
@@ -295,11 +295,19 @@ bool FAnimNode_MotionMatching::GetIgnoreForRelevancyTest() const
 	return GET_ANIM_NODE_DATA(bool, bIgnoreForRelevancyTest);
 }
 
-void FAnimNode_MotionMatching::SetIgnoreForRelevancyTest(bool bInIgnoreForRelevancyTest)
+bool FAnimNode_MotionMatching::SetIgnoreForRelevancyTest(bool bInIgnoreForRelevancyTest)
 {
 #if WITH_EDITORONLY_DATA
 	bIgnoreForRelevancyTest = bInIgnoreForRelevancyTest;
 #endif
+
+	if(bool* bIgnoreForRelevancyTestPtr = GET_INSTANCE_ANIM_NODE_DATA_PTR(bool, bIgnoreForRelevancyTest))
+	{
+		*bIgnoreForRelevancyTestPtr = bInIgnoreForRelevancyTest;
+		return true;
+	}
+
+	return false;
 }
 
 #undef LOCTEXT_NAMESPACE
\ No newline at end of file
diff --git a/Engine/Plugins/Experimental/Animation/PoseSearch/Source/Runtime/Public/PoseSearch/AnimNode_MotionMatching.h b/Engine/Plugins/Experimental/Animation/PoseSearch/Source/Runtime/Public/PoseSearch/AnimNode_MotionMatching.h
--- a/Engine/Plugins/Experimental/Animation/PoseSearch/Source/Runtime/Public/PoseSearch/AnimNode_MotionMatching.h
+++ b/Engine/Plugins/Experimental/Animation/PoseSearch/Source/Runtime/Public/PoseSearch/AnimNode_MotionMatching.h
@@ -99,7 +99,7 @@ protected:
 	virtual float GetCurrentAssetTime() const override;
 	virtual float GetCurrentAssetTimePlayRateAdjusted() const override;
 	virtual bool GetIgnoreForRelevancyTest() const override;
-	virtual void SetIgnoreForRelevancyTest(bool bInIgnoreForRelevancyTest) override;
+	virtual bool SetIgnoreForRelevancyTest(bool bInIgnoreForRelevancyTest) override;
 	// End of FAnimNode_AssetPlayerBase interface
 
 private:
diff --git a/Engine/Source/Developer/AssetTools/Private/AssetTypeActions/AssetTypeActions_AnimBlueprint.cpp b/Engine/Source/Developer/AssetTools/Private/AssetTypeActions/AssetTypeActions_AnimBlueprint.cpp
--- a/Engine/Source/Developer/AssetTools/Private/AssetTypeActions/AssetTypeActions_AnimBlueprint.cpp
+++ b/Engine/Source/Developer/AssetTools/Private/AssetTypeActions/AssetTypeActions_AnimBlueprint.cpp
@@ -134,7 +134,7 @@ void FAssetTypeActions_AnimBlueprint::OpenAssetEditor( const TArray<UObject*>& I
 		auto AnimBlueprint = Cast<UAnimBlueprint>(*ObjIt);
 		if (AnimBlueprint != NULL && AnimBlueprint->SkeletonGeneratedClass && AnimBlueprint->GeneratedClass)
 		{
-			if(AnimBlueprint->BlueprintType != BPTYPE_Interface && !AnimBlueprint->TargetSkeleton)
+			if(AnimBlueprint->BlueprintType != BPTYPE_Interface && !AnimBlueprint->TargetSkeleton && !AnimBlueprint->bIsTemplate)
 			{
 				FText ShouldRetargetMessage = LOCTEXT("ShouldRetarget_Message", "Could not find the skeleton for Anim Blueprint '{BlueprintName}' Would you like to choose a new one?");
 				
diff --git a/Engine/Source/Editor/AnimGraph/Private/AnimBlueprintCompilationContext.h b/Engine/Source/Editor/AnimGraph/Private/AnimBlueprintCompilationContext.h
--- a/Engine/Source/Editor/AnimGraph/Private/AnimBlueprintCompilationContext.h
+++ b/Engine/Source/Editor/AnimGraph/Private/AnimBlueprintCompilationContext.h
@@ -84,18 +84,21 @@ class FAnimBlueprintNodeCopyTermDefaultsContext : public IAnimBlueprintNodeCopyT
 private:
 	friend class FAnimBlueprintCompilerContext;
 
-	FAnimBlueprintNodeCopyTermDefaultsContext(const FProperty* InTargetProperty, uint8* InDestinationPtr, const uint8* InSourcePtr, int32 InNodePropertyIndex)
-		: TargetProperty(InTargetProperty)
+	FAnimBlueprintNodeCopyTermDefaultsContext(UObject* InCDO, const FProperty* InTargetProperty, uint8* InDestinationPtr, const uint8* InSourcePtr, int32 InNodePropertyIndex)
+		: CDO(InCDO)
+		, TargetProperty(InTargetProperty)
 		, DestinationPtr(InDestinationPtr)
 		, SourcePtr(InSourcePtr)
 		, NodePropertyIndex(InNodePropertyIndex)
 	{}
 
+	virtual UObject* GetClassDefaultObjectImpl() const override { return CDO; }
 	virtual const FProperty* GetTargetPropertyImpl() const override { return TargetProperty; }
 	virtual uint8* GetDestinationPtrImpl() const override { return DestinationPtr; }
 	virtual const uint8* GetSourcePtrImpl() const override { return SourcePtr; }
 	virtual int32 GetNodePropertyIndexImpl() const override { return NodePropertyIndex; }
 
+	UObject* CDO;
 	const FProperty* TargetProperty;
 	uint8* DestinationPtr;
 	const uint8* SourcePtr;
@@ -107,18 +110,21 @@ class FAnimBlueprintExtensionCopyTermDefaultsContext : public IAnimBlueprintExte
 private:
 	friend class FAnimBlueprintCompilerContext;
 
-	FAnimBlueprintExtensionCopyTermDefaultsContext(const FProperty* InTargetProperty, uint8* InDestinationPtr, const uint8* InSourcePtr, int32 InNodePropertyIndex)
-        : TargetProperty(InTargetProperty)
-        , DestinationPtr(InDestinationPtr)
-        , SourcePtr(InSourcePtr)
-        , NodePropertyIndex(InNodePropertyIndex)
+	FAnimBlueprintExtensionCopyTermDefaultsContext(UObject* InCDO, const FProperty* InTargetProperty, uint8* InDestinationPtr, const uint8* InSourcePtr, int32 InNodePropertyIndex)
+		: CDO(InCDO)
+		, TargetProperty(InTargetProperty)
+		, DestinationPtr(InDestinationPtr)
+		, SourcePtr(InSourcePtr)
+		, NodePropertyIndex(InNodePropertyIndex)
 	{}
 
+	virtual UObject* GetClassDefaultObjectImpl() const override { return CDO; }
 	virtual const FProperty* GetTargetPropertyImpl() const override { return TargetProperty; }
 	virtual uint8* GetDestinationPtrImpl() const override { return DestinationPtr; }
 	virtual const uint8* GetSourcePtrImpl() const override { return SourcePtr; }
 	virtual int32 GetNodePropertyIndexImpl() const override { return NodePropertyIndex; }
 
+	UObject* CDO;
 	const FProperty* TargetProperty;
 	uint8* DestinationPtr;
 	const uint8* SourcePtr;
diff --git a/Engine/Source/Editor/AnimGraph/Private/AnimBlueprintCompiler.cpp b/Engine/Source/Editor/AnimGraph/Private/AnimBlueprintCompiler.cpp
--- a/Engine/Source/Editor/AnimGraph/Private/AnimBlueprintCompiler.cpp
+++ b/Engine/Source/Editor/AnimGraph/Private/AnimBlueprintCompiler.cpp
@@ -366,6 +366,7 @@ void FAnimBlueprintCompilerContext::GatherFoldRecordsForAnimationNode(const UScr
 		{
 			bool bAllPinsExposed = true;
 			bool bAllPinsDisconnected = true;
+			const bool bAlwaysDynamic = InVisualAnimNode->AlwaysDynamicProperties.Contains(SubProperty->GetFName());
 
 			// If a value is exposed on a pin but disconnected, push the value to the (intermediate) node here to simplify later logic
 			if(SubProperty->HasAnyPropertyFlags(CPF_Edit|CPF_BlueprintVisible))
@@ -417,7 +418,7 @@ void FAnimBlueprintCompilerContext::GatherFoldRecordsForAnimationNode(const UScr
 			if(SubProperty->HasMetaData(NAME_FoldProperty))
 			{
 				// Add folding candidate
-				AddFoldedPropertyRecord(InVisualAnimNode, InNodeProperty, SubProperty, bAllPinsExposed, !bAllPinsDisconnected);
+				AddFoldedPropertyRecord(InVisualAnimNode, InNodeProperty, SubProperty, bAllPinsExposed, !bAllPinsDisconnected, bAlwaysDynamic);
 			}
 		}
 		else
@@ -426,6 +427,7 @@ void FAnimBlueprintCompilerContext::GatherFoldRecordsForAnimationNode(const UScr
 			const bool bExposedOnPin = Pin != nullptr;
 			const bool bPinConnected = InVisualAnimNode->IsPinExposedAndLinked(SubProperty->GetName(), EGPD_Input);
 			const bool bPinBound = InVisualAnimNode->IsPinExposedAndBound(SubProperty->GetName(), EGPD_Input);
+			const bool bAlwaysDynamic = InVisualAnimNode->AlwaysDynamicProperties.Contains(SubProperty->GetFName());
 
 			// If a value is exposed on a pin but disconnected, push the value to the (intermediate) node here to simplify later logic
 			if(SubProperty->HasAnyPropertyFlags(CPF_Edit|CPF_BlueprintVisible))
@@ -454,7 +456,7 @@ void FAnimBlueprintCompilerContext::GatherFoldRecordsForAnimationNode(const UScr
 			if(SubProperty->HasMetaData(NAME_FoldProperty))
 			{
 				// Add folding candidate
-				AddFoldedPropertyRecord(InVisualAnimNode, InNodeProperty, SubProperty, bExposedOnPin, bPinConnected || bPinBound);
+				AddFoldedPropertyRecord(InVisualAnimNode, InNodeProperty, SubProperty, bExposedOnPin, bPinConnected || bPinBound, bAlwaysDynamic);
 			}
 		}
 	}
@@ -627,7 +629,7 @@ void FAnimBlueprintCompilerContext::ProcessAllAnimationNodes()
 	BP_SCOPED_COMPILER_EVENT_STAT(EAnimBlueprintCompilerStats_ProcessAllAnimationNodes);
 	
 	// Validate that we have a skeleton
-	if ((AnimBlueprint->TargetSkeleton == nullptr) && !AnimBlueprint->bIsNewlyCreated)
+	if ((AnimBlueprint->TargetSkeleton == nullptr) && !AnimBlueprint->bIsNewlyCreated && !AnimBlueprint->bIsTemplate)
 	{
 		MessageLog.Error(*LOCTEXT("NoSkeleton", "@@ - The skeleton asset for this animation Blueprint is missing, so it cannot be compiled!").ToString(), AnimBlueprint);
 		return;
@@ -818,7 +820,7 @@ void FAnimBlueprintCompilerContext::CopyTermDefaultsToDefaultObject(UObject* Def
 					SourcePtr = SourceExtensionProperty->ContainerPtrToValuePtr<uint8>(Extension);
 				}
 				
-				FAnimBlueprintExtensionCopyTermDefaultsContext ExtensionContext(TargetProperty, DestinationPtr, SourcePtr, LinkIndexCount);
+				FAnimBlueprintExtensionCopyTermDefaultsContext ExtensionContext(DefaultObject, TargetProperty, DestinationPtr, SourcePtr, LinkIndexCount);
 				Extension->CopyTermDefaultsToDefaultObject(DefaultAnimInstance, CompilerContext, ExtensionContext);
 			}
 		}
@@ -837,7 +839,7 @@ void FAnimBlueprintCompilerContext::CopyTermDefaultsToDefaultObject(UObject* Def
 				uint8* DestinationPtr = TargetProperty->ContainerPtrToValuePtr<uint8>(DefaultAnimInstance);
 				const uint8* SourcePtr = SourceNodeProperty->ContainerPtrToValuePtr<uint8>(VisualAnimNode);
 
-				FAnimBlueprintNodeCopyTermDefaultsContext NodeContext(TargetProperty, DestinationPtr, SourcePtr, LinkIndexCount);
+				FAnimBlueprintNodeCopyTermDefaultsContext NodeContext(DefaultObject, TargetProperty, DestinationPtr, SourcePtr, LinkIndexCount);
 				VisualAnimNode->CopyTermDefaultsToDefaultObject(CompilerContext, NodeContext, CompiledData);
 
 				LinkIndexMap.Add(VisualAnimNode, LinkIndexCount);
@@ -899,7 +901,7 @@ void FAnimBlueprintCompilerContext::CopyTermDefaultsToDefaultObject(UObject* Def
 						uint8* DestinationPtr = TargetProperty->ContainerPtrToValuePtr<uint8>(Constants);
 						const uint8* SourcePtr = SourceExtensionProperty->ContainerPtrToValuePtr<uint8>(Extension);
 					
-						FAnimBlueprintExtensionCopyTermDefaultsContext NodeContext(TargetProperty, DestinationPtr, SourcePtr, LinkIndexCount);
+						FAnimBlueprintExtensionCopyTermDefaultsContext NodeContext(DefaultObject, TargetProperty, DestinationPtr, SourcePtr, LinkIndexCount);
 						Extension->CopyTermDefaultsToSparseClassData(CompilerContext, NodeContext);
 					}
 				}
@@ -1414,13 +1416,15 @@ void FAnimBlueprintCompilerContext::PostCompile()
 	if(UAnimInstance* DefaultAnimInstance = Cast<UAnimInstance>(AnimBlueprintGeneratedClass->GetDefaultObject()))
 	{
 		// iterate all anim node and call PostCompile
-		const USkeleton* CurrentSkeleton = AnimBlueprint->TargetSkeleton;
-		for (FStructProperty* Property : TFieldRange<FStructProperty>(AnimBlueprintGeneratedClass, EFieldIteratorFlags::IncludeSuper))
+		if(const USkeleton* CurrentSkeleton = AnimBlueprint->TargetSkeleton)
 		{
-			if (Property->Struct->IsChildOf(FAnimNode_Base::StaticStruct()))
+			for (FStructProperty* Property : TFieldRange<FStructProperty>(AnimBlueprintGeneratedClass, EFieldIteratorFlags::IncludeSuper))
 			{
-				FAnimNode_Base* AnimNode = Property->ContainerPtrToValuePtr<FAnimNode_Base>(DefaultAnimInstance);
-				AnimNode->PostCompile(CurrentSkeleton);
+				if (Property->Struct->IsChildOf(FAnimNode_Base::StaticStruct()))
+				{
+					FAnimNode_Base* AnimNode = Property->ContainerPtrToValuePtr<FAnimNode_Base>(DefaultAnimInstance);
+					AnimNode->PostCompile(CurrentSkeleton);
+				}
 			}
 		}
 	}
@@ -1760,9 +1764,9 @@ FProperty* FAnimBlueprintCompilerContext::CreateStructVariable(UScriptStruct* In
 	return NewProperty;
 }
 
-void FAnimBlueprintCompilerContext::AddFoldedPropertyRecord(UAnimGraphNode_Base* InAnimGraphNode, FStructProperty* InAnimNodeProperty, FProperty* InProperty, bool bInExposedOnPin, bool bInPinConnected)
+void FAnimBlueprintCompilerContext::AddFoldedPropertyRecord(UAnimGraphNode_Base* InAnimGraphNode, FStructProperty* InAnimNodeProperty, FProperty* InProperty, bool bInExposedOnPin, bool bInPinConnected, bool bInAlwaysDynamic)
 {
-	const bool bConstant = !bInExposedOnPin || (bInExposedOnPin && !bInPinConnected);
+	const bool bConstant = !bInAlwaysDynamic && (!bInExposedOnPin || (bInExposedOnPin && !bInPinConnected));
 
 	if(!InProperty->HasAnyPropertyFlags(CPF_EditorOnly))
 	{
diff --git a/Engine/Source/Editor/AnimGraph/Private/AnimBlueprintCompiler.h b/Engine/Source/Editor/AnimGraph/Private/AnimBlueprintCompiler.h
--- a/Engine/Source/Editor/AnimGraph/Private/AnimBlueprintCompiler.h
+++ b/Engine/Source/Editor/AnimGraph/Private/AnimBlueprintCompiler.h
@@ -194,7 +194,7 @@ private:
 	FProperty* CreateStructVariable(UScriptStruct* InStruct, const FName VarName, const FEdGraphPinType& VarType);
 
 	/** Adds a record for a potentially-folded anim node property */
-	void AddFoldedPropertyRecord(UAnimGraphNode_Base* InAnimGraphNode, FStructProperty* InAnimNodeProperty, FProperty* InProperty, bool bInExposedOnPin, bool bInPinConnected);
+	void AddFoldedPropertyRecord(UAnimGraphNode_Base* InAnimGraphNode, FStructProperty* InAnimNodeProperty, FProperty* InProperty, bool bInExposedOnPin, bool bInPinConnected, bool bInAlwaysDynamic);
 
 	/** Process any anim node properties that are 'foldable' */
 	void ProcessFoldedPropertyRecords();
diff --git a/Engine/Source/Editor/AnimGraph/Private/AnimBlueprintDetails.cpp b/Engine/Source/Editor/AnimGraph/Private/AnimBlueprintDetails.cpp
new file mode 100644
--- /dev/null
+++ b/Engine/Source/Editor/AnimGraph/Private/AnimBlueprintDetails.cpp
@@ -0,0 +1,22 @@
+// Copyright Epic Games, Inc. All Rights Reserved.
+
+#include "AnimBlueprintDetails.h"
+#include "Animation/AnimBlueprint.h"
+#include "DetailLayoutBuilder.h"
+
+#define LOCTEXT_NAMESPACE "AnimBlueprintDetails"
+
+void FAnimBlueprintDetails::CustomizeDetails(IDetailLayoutBuilder& DetailBuilder)
+{
+	TArray<TWeakObjectPtr<UObject>> Objects;
+	DetailBuilder.GetObjectsBeingCustomized(Objects);
+	UAnimBlueprint* Asset = CastChecked<UAnimBlueprint>(Objects[0].Get());
+
+	if(Asset->bIsTemplate)
+	{
+		TSharedRef<IPropertyHandle> TargetSkeletonHandle = DetailBuilder.GetProperty(GET_MEMBER_NAME_CHECKED(UAnimBlueprint, TargetSkeleton));
+		TargetSkeletonHandle->MarkHiddenByCustomization();
+	}
+}
+
+#undef LOCTEXT_NAMESPACE
\ No newline at end of file
diff --git a/Engine/Source/Editor/AnimGraph/Private/AnimBlueprintDetails.h b/Engine/Source/Editor/AnimGraph/Private/AnimBlueprintDetails.h
new file mode 100644
--- /dev/null
+++ b/Engine/Source/Editor/AnimGraph/Private/AnimBlueprintDetails.h
@@ -0,0 +1,17 @@
+// Copyright Epic Games, Inc. All Rights Reserved.
+
+#pragma once
+
+#include "IDetailCustomization.h"
+
+class FAnimBlueprintDetails : public IDetailCustomization
+{
+public:
+	static TSharedRef<IDetailCustomization> MakeInstance()
+	{
+		return MakeShareable(new FAnimBlueprintDetails());
+	}
+
+	// IDetailCustomization interface
+	virtual void CustomizeDetails(IDetailLayoutBuilder& DetailBuilder) override;
+};
diff --git a/Engine/Source/Editor/AnimGraph/Private/AnimBlueprintExtension.cpp b/Engine/Source/Editor/AnimGraph/Private/AnimBlueprintExtension.cpp
--- a/Engine/Source/Editor/AnimGraph/Private/AnimBlueprintExtension.cpp
+++ b/Engine/Source/Editor/AnimGraph/Private/AnimBlueprintExtension.cpp
@@ -10,6 +10,7 @@
 #include "AnimBlueprintExtension_Attributes.h"
 #include "AnimBlueprintExtension_PropertyAccess.h"
 #include "Kismet2/BlueprintEditorUtils.h"
+#include "AnimBlueprintExtension_NodeRelevancy.h"
 
 // Set used to refresh extensions. Checks that an extension has a reference from an anim node for each refresh.
 static TSet<TSubclassOf<UAnimBlueprintExtension>> RefreshSet;
@@ -76,7 +77,12 @@ void UAnimBlueprintExtension::RequestExtensionsForNode(UAnimGraphNode_Base* InAn
             UAnimBlueprintExtension_Attributes::StaticClass(),
             UAnimBlueprintExtension_PropertyAccess::StaticClass()
         };
-	
+
+		if(InAnimGraphNode->BecomeRelevantFunction.GetMemberGuid().IsValid())
+		{
+			ExtensionClasses.Add(UAnimBlueprintExtension_NodeRelevancy::StaticClass());
+		}
+
 		InAnimGraphNode->GetRequiredExtensions(ExtensionClasses);
 	
 		for(const TSubclassOf<UAnimBlueprintExtension>& ExtensionClass : ExtensionClasses)
diff --git a/Engine/Source/Editor/AnimGraph/Private/AnimBlueprintExtension_NodeRelevancy.h b/Engine/Source/Editor/AnimGraph/Private/AnimBlueprintExtension_NodeRelevancy.h
new file mode 100644
--- /dev/null
+++ b/Engine/Source/Editor/AnimGraph/Private/AnimBlueprintExtension_NodeRelevancy.h
@@ -0,0 +1,18 @@
+// Copyright Epic Games, Inc. All Rights Reserved.
+
+#pragma once
+
+#include "CoreMinimal.h"
+#include "AnimBlueprintExtension.h"
+#include "Animation/AnimSubsystem_NodeRelevancy.h"
+#include "AnimBlueprintExtension_NodeRelevancy.generated.h"
+
+UCLASS(MinimalAPI)
+class UAnimBlueprintExtension_NodeRelevancy : public UAnimBlueprintExtension
+{
+	GENERATED_BODY()
+
+private:
+	UPROPERTY()
+	FAnimSubsystemInstance_NodeRelevancy Subsystem;
+};
\ No newline at end of file
diff --git a/Engine/Source/Editor/AnimGraph/Private/AnimGraphModule.cpp b/Engine/Source/Editor/AnimGraph/Private/AnimGraphModule.cpp
--- a/Engine/Source/Editor/AnimGraph/Private/AnimGraphModule.cpp
+++ b/Engine/Source/Editor/AnimGraph/Private/AnimGraphModule.cpp
@@ -31,6 +31,7 @@
 #include "Animation/AnimRootMotionProvider.h"
 #include "Features/IModularFeatures.h"
 #include "PropertyAccessAnimBlueprintBinding.h"
+#include "AnimBlueprintDetails.h"
 
 class FAnimGraphModule : public IModuleInterface
 {
@@ -70,6 +71,7 @@ void FAnimGraphModule::StartupModule()
 
 	// Register details customization
 	FPropertyEditorModule& PropertyModule = FModuleManager::LoadModuleChecked<FPropertyEditorModule>("PropertyEditor");
+	PropertyModule.RegisterCustomClassLayout(UAnimBlueprint::StaticClass()->GetFName(), FOnGetDetailCustomizationInstance::CreateStatic(&FAnimBlueprintDetails::MakeInstance));
 	PropertyModule.RegisterCustomClassLayout(UAnimGraphNode_PoseDriver::StaticClass()->GetFName(), FOnGetDetailCustomizationInstance::CreateStatic(&FPoseDriverDetails::MakeInstance));
 
 	PropertyModule.RegisterCustomPropertyTypeLayout("AnimBlueprintFunctionPinInfo", FOnGetPropertyTypeCustomizationInstance::CreateStatic(&FAnimBlueprintFunctionPinInfoDetails::MakeInstance));
diff --git a/Engine/Source/Editor/AnimGraph/Private/AnimGraphNode_AssetPlayerBase.cpp b/Engine/Source/Editor/AnimGraph/Private/AnimGraphNode_AssetPlayerBase.cpp
--- a/Engine/Source/Editor/AnimGraph/Private/AnimGraphNode_AssetPlayerBase.cpp
+++ b/Engine/Source/Editor/AnimGraph/Private/AnimGraphNode_AssetPlayerBase.cpp
@@ -451,4 +451,54 @@ FText UAnimGraphNode_AssetPlayerBase::GetNodeTitleForAsset(ENodeTitleType::Type
 	}
 }
 
+void UAnimGraphNode_AssetPlayerBase::ValidateAnimNodeDuringCompilationHelper(USkeleton* ForSkeleton, FCompilerResultsLog& MessageLog, UAnimationAsset* InAsset, TSubclassOf<UAnimationAsset> InAssetType, UEdGraphPin* InExposedPin)
+{
+	UAnimationAsset* AssetToCheck = InAsset;
+	if(InExposedPin != nullptr && AssetToCheck == nullptr)
+	{
+		AssetToCheck = Cast<UAnimationAsset>(InExposedPin->DefaultObject);
+	}
+
+	if(AssetToCheck == nullptr)
+	{
+		if(!GetAnimBlueprint()->bIsTemplate)
+		{
+			// Check for bindings
+			bool bHasBinding = false;
+			bool bAlwaysDynamic = false;
+			if(InExposedPin != nullptr)
+			{
+				if(PropertyBindings.Contains(InExposedPin->GetFName()))
+				{
+					bHasBinding = true;
+				}
+
+				if(AlwaysDynamicProperties.Contains(InExposedPin->GetFName()))
+				{
+					bAlwaysDynamic = true;
+				}
+			}
+
+			// we may have a connected node or binding
+			if(!bAlwaysDynamic)
+			{
+				if (InExposedPin == nullptr || (InExposedPin->LinkedTo.Num() == 0 && !bHasBinding))
+				{
+					MessageLog.Error(*FText::Format(LOCTEXT("MissingAssetFormat", "@@ references an unknown {0}"), InAssetType->GetDisplayNameText()).ToString(), this);
+				}
+			}
+		}
+	}
+	else 
+	{
+		USkeleton* AssetSkeleton = AssetToCheck->GetSkeleton();
+		
+		// if asset doesn't have a skeleton, it might be due to the asset no being not loaded yet
+		if(AssetSkeleton && !ForSkeleton->IsCompatible(AssetSkeleton))
+		{
+			MessageLog.Error(*FText::Format(LOCTEXT("IncompatibleSkeletonFormat", "@@ references {0} that uses an incompatible skeleton @@"), InAssetType->GetDisplayNameText()).ToString(), this, AssetSkeleton);
+		}
+	}
+}
+
 #undef LOCTEXT_NAMESPACE
\ No newline at end of file
diff --git a/Engine/Source/Editor/AnimGraph/Private/AnimGraphNode_Base.cpp b/Engine/Source/Editor/AnimGraph/Private/AnimGraphNode_Base.cpp
--- a/Engine/Source/Editor/AnimGraph/Private/AnimGraphNode_Base.cpp
+++ b/Engine/Source/Editor/AnimGraph/Private/AnimGraphNode_Base.cpp
@@ -85,7 +85,28 @@ void UAnimGraphNode_Base::PostEditChangeProperty(FPropertyChangedEvent& Property
 		FOptionalPinManager::EvaluateOldShownPins(ShowPinForProperties, OldShownPins, this);
 		GetSchema()->ReconstructNode(*this);
 	}
-
+	else if(PropertyName == GET_MEMBER_NAME_CHECKED(UAnimGraphNode_Base, InitializeFunction))
+	{
+		GetFNode()->InitializeFunction.SetFromFunction(InitializeFunction.ResolveMember<UFunction>(GetAnimBlueprint()->SkeletonGeneratedClass));
+		GetSchema()->ReconstructNode(*this);
+	}
+	else if(PropertyName == GET_MEMBER_NAME_CHECKED(UAnimGraphNode_Base, BecomeRelevantFunction))
+	{
+		GetFNode()->BecomeRelevantFunction.SetFromFunction(BecomeRelevantFunction.ResolveMember<UFunction>(GetAnimBlueprint()->SkeletonGeneratedClass));
+		GetAnimBlueprint()->RequestRefreshExtensions();
+		GetSchema()->ReconstructNode(*this);
+	}
+	else if(PropertyName == GET_MEMBER_NAME_CHECKED(UAnimGraphNode_Base, UpdateFunction))
+	{
+		GetFNode()->UpdateFunction.SetFromFunction(UpdateFunction.ResolveMember<UFunction>(GetAnimBlueprint()->SkeletonGeneratedClass));
+		GetSchema()->ReconstructNode(*this);
+	}
+	else if(PropertyName == GET_MEMBER_NAME_CHECKED(UAnimGraphNode_Base, EvaluateFunction))
+	{
+		GetFNode()->EvaluateFunction.SetFromFunction(EvaluateFunction.ResolveMember<UFunction>(GetAnimBlueprint()->SkeletonGeneratedClass));
+		GetSchema()->ReconstructNode(*this);
+	}
+	
 	Super::PostEditChangeProperty(PropertyChangedEvent);
 
 	PropertyChangeEvent.Broadcast(PropertyChangedEvent);
@@ -185,6 +206,23 @@ void UAnimGraphNode_Base::ValidateAnimNodeDuringCompilation(USkeleton* ForSkelet
 			}
 		}
 	}
+
+	auto ValidateFunctionRef = [this, &MessageLog](const FMemberReference& InRef, const FText& InError)
+	{
+		if(InRef.GetMemberName() != NAME_None)
+		{
+			UFunction* Function = InRef.ResolveMember<UFunction>(GetAnimBlueprint()->SkeletonGeneratedClass);
+			if(Function == nullptr)
+			{
+				MessageLog.Error(*InError.ToString());
+			}
+		}
+	};
+	
+	ValidateFunctionRef(InitializeFunction, LOCTEXT("MissingInitializeFunction", "Could not resolve initialize function"));
+	ValidateFunctionRef(BecomeRelevantFunction, LOCTEXT("MissingBecomeRelevantFunction", "Could not resolve become relevant function"));
+	ValidateFunctionRef(UpdateFunction, LOCTEXT("MissingUpdateFunction", "Could not resolve update function"));
+	ValidateFunctionRef(EvaluateFunction, LOCTEXT("MissingEvaluateFunction", "Could not resolve evaluate function"));
 }
 
 void UAnimGraphNode_Base::CopyTermDefaultsToDefaultObject(IAnimBlueprintCopyTermDefaultsContext& InCompilationContext, IAnimBlueprintNodeCopyTermDefaultsContext& InPerNodeContext, IAnimBlueprintGeneratedClassCompiledData& OutCompiledData)
@@ -467,6 +505,12 @@ void UAnimGraphNode_Base::ProcessDuringCompilation(IAnimBlueprintCompilationCont
 	// Record pose pins for later patchup and gather pins that have an associated evaluation handler
 	Extension->ProcessNodePins(this, InCompilationContext, OutCompiledData);
 
+	// Resolve functions
+	GetFNode()->InitializeFunction.SetFromFunction(InitializeFunction.ResolveMember<UFunction>(GetAnimBlueprint()->SkeletonGeneratedClass));  
+	GetFNode()->BecomeRelevantFunction.SetFromFunction(BecomeRelevantFunction.ResolveMember<UFunction>(GetAnimBlueprint()->SkeletonGeneratedClass));
+	GetFNode()->UpdateFunction.SetFromFunction(UpdateFunction.ResolveMember<UFunction>(GetAnimBlueprint()->SkeletonGeneratedClass));
+	GetFNode()->EvaluateFunction.SetFromFunction(EvaluateFunction.ResolveMember<UFunction>(GetAnimBlueprint()->SkeletonGeneratedClass)); 
+	
 	// Call the override point
 	OnProcessDuringCompilation(InCompilationContext, OutCompiledData);
 }
@@ -1054,81 +1098,155 @@ TSharedRef<SWidget> UAnimGraphNode_Base::MakePropertyBindingWidget(const FAnimPr
 		{
 			InMenuBuilder.BeginSection("Pins", LOCTEXT("Pin", "Pin"));
 			{
-				InMenuBuilder.AddMenuEntry(
-					LOCTEXT("ExposeAsPin", "Expose As Pin"),
-					LOCTEXT("ExposeAsPinTooltip", "Show/hide this property as a pin on the node"),
-					FSlateIcon("EditorStyle", "GraphEditor.PinIcon"),
-					FUIAction(
-						FExecuteAction::CreateLambda([InArgs, Blueprint]()
-						{
-							bool bHasBinding = false;
+				auto ExposeAsPin = [InArgs, Blueprint]()
+				{
+					bool bHasBinding = false;
 
-							// Comparison without name index to deal with arrays
-							const FName ComparisonName = FName(InArgs.PinName, 0);
+					// Comparison without name index to deal with arrays
+					const FName ComparisonName = FName(InArgs.PinName, 0);
 							
-							for(UAnimGraphNode_Base* AnimGraphNode : InArgs.Nodes)
+					for(UAnimGraphNode_Base* AnimGraphNode : InArgs.Nodes)
+					{
+						for(const TPair<FName, FAnimGraphNodePropertyBinding>& BindingPair : AnimGraphNode->PropertyBindings)
+						{
+							if(ComparisonName == FName(BindingPair.Key, 0))
 							{
-								for(const TPair<FName, FAnimGraphNodePropertyBinding>& BindingPair : AnimGraphNode->PropertyBindings)
-								{
-									if(ComparisonName == FName(BindingPair.Key, 0))
-									{
-										bHasBinding = true;
-										break;
-									}
-								}
+								bHasBinding = true;
+								break;
 							}
+						}
+					}
 
-							{
-								FScopedTransaction Transaction(LOCTEXT("PinExposure", "Pin Exposure"));
+					{
+						FScopedTransaction Transaction(LOCTEXT("PinExposure", "Pin Exposure"));
 
-								// Switching from non-pin to pin, remove any bindings
-								for(UAnimGraphNode_Base* AnimGraphNode : InArgs.Nodes)
+						// Switching from non-pin to pin, remove any bindings
+						for(UAnimGraphNode_Base* AnimGraphNode : InArgs.Nodes)
+						{
+							AnimGraphNode->Modify();
+
+							bool bVisible = AnimGraphNode->ShowPinForProperties[InArgs.OptionalPinIndex].bShowPin;
+							AnimGraphNode->SetPinVisibility(!bVisible || bHasBinding, InArgs.OptionalPinIndex);
+
+							// Remove all bindings that match the property, array or array elements
+							for(auto It = AnimGraphNode->PropertyBindings.CreateIterator(); It; ++It)
+							{
+								if(ComparisonName == FName(It.Key(), 0))
 								{
-									AnimGraphNode->Modify();
-
-									bool bVisible = AnimGraphNode->ShowPinForProperties[InArgs.OptionalPinIndex].bShowPin;
-									AnimGraphNode->SetPinVisibility(!bVisible || bHasBinding, InArgs.OptionalPinIndex);
-
-									// Remove all bindings that match the property, array or array elements
-									for(auto It = AnimGraphNode->PropertyBindings.CreateIterator(); It; ++It)
-									{
-										if(ComparisonName == FName(It.Key(), 0))
-										{
-											It.RemoveCurrent();
-										}
-									}
+									It.RemoveCurrent();
 								}
+							}
+						}
 
-								FBlueprintEditorUtils::MarkBlueprintAsModified(Blueprint);
+						FBlueprintEditorUtils::MarkBlueprintAsModified(Blueprint);
+					}
+				}; 
+
+				auto GetExposedAsPinCheckState = [InArgs]()
+				{
+					bool bPinShown = false;
+					bool bHasBinding = false;
+
+					// Comparison without name index to deal with arrays
+					const FName ComparisonName = FName(InArgs.PinName, 0);
+							
+					for(UAnimGraphNode_Base* AnimGraphNode : InArgs.Nodes)
+					{
+						for(const TPair<FName, FAnimGraphNodePropertyBinding>& BindingPair : AnimGraphNode->PropertyBindings)
+						{
+							if(ComparisonName == FName(BindingPair.Key, 0))
+							{
+								bHasBinding = true;
+								break;
 							}
-						}),
+						}
+								
+						bPinShown |= AnimGraphNode->ShowPinForProperties[InArgs.OptionalPinIndex].bShowPin;
+					}
+
+					// Pins are exposed if we have a binding or not, so treat as unchecked only if we have
+					// no binding
+					return bPinShown && !bHasBinding ? ECheckBoxState::Checked : ECheckBoxState::Unchecked;
+				};
+				
+				InMenuBuilder.AddMenuEntry(
+					LOCTEXT("ExposeAsPin", "Expose As Pin"),
+					LOCTEXT("ExposeAsPinTooltip", "Show/hide this property as a pin on the node"),
+					FSlateIcon("EditorStyle", "GraphEditor.PinIcon"),
+					FUIAction(
+						FExecuteAction::CreateLambda(ExposeAsPin),
 						FCanExecuteAction(),
-						FGetActionCheckState::CreateLambda([InArgs]()
+						FGetActionCheckState::CreateLambda(GetExposedAsPinCheckState)
+					),
+					NAME_None,
+					EUserInterfaceActionType::Check
+				);
+
+				auto MakeDynamic = [InArgs, Blueprint]()
+				{
+					// Comparison without name index to deal with arrays
+					const FName ComparisonName = FName(InArgs.PinName, 0);
+
+					bool bIsAlwaysDynamic = false;
+					for(UAnimGraphNode_Base* AnimGraphNode : InArgs.Nodes)
+					{
+						if(AnimGraphNode->AlwaysDynamicProperties.Contains(ComparisonName))
 						{
-							bool bPinShown = false;
-							bool bHasBinding = false;
+							bIsAlwaysDynamic = true;
+							break;
+						}
+					}
 
-							// Comparison without name index to deal with arrays
-							const FName ComparisonName = FName(InArgs.PinName, 0);
-							
-							for(UAnimGraphNode_Base* AnimGraphNode : InArgs.Nodes)
+					{
+						FScopedTransaction Transaction(LOCTEXT("AlwaysDynamic", "Always Dynamic"));
+
+						for(UAnimGraphNode_Base* AnimGraphNode : InArgs.Nodes)
+						{
+							AnimGraphNode->Modify();
+
+							if(bIsAlwaysDynamic)
 							{
-								for(const TPair<FName, FAnimGraphNodePropertyBinding>& BindingPair : AnimGraphNode->PropertyBindings)
-								{
-									if(ComparisonName == FName(BindingPair.Key, 0))
-									{
-										bHasBinding = true;
-										break;
-									}
-								}
-								
-								bPinShown |= AnimGraphNode->ShowPinForProperties[InArgs.OptionalPinIndex].bShowPin;
+								AnimGraphNode->AlwaysDynamicProperties.Remove(ComparisonName);
 							}
+							else
+							{
+								AnimGraphNode->AlwaysDynamicProperties.Add(ComparisonName);
+							}
+						}
 
-							// Pins are exposed if we have a binding or not, so treat as unchecked only if we have
-							// no binding
-							return bPinShown && !bHasBinding ? ECheckBoxState::Checked : ECheckBoxState::Unchecked;
-						})
+						FBlueprintEditorUtils::MarkBlueprintAsModified(Blueprint);
+					}
+				};
+
+				auto GetDynamicCheckState = [InArgs]()
+				{
+					// Comparison without name index to deal with arrays
+					const FName ComparisonName = FName(InArgs.PinName, 0);
+
+					bool bIsAlwaysDynamic = false;
+					for(UAnimGraphNode_Base* AnimGraphNode : InArgs.Nodes)
+					{
+						for(const FName& AlwaysDynamicPropertyName : AnimGraphNode->AlwaysDynamicProperties)
+						{
+							if(ComparisonName == FName(AlwaysDynamicPropertyName, 0))
+							{
+								bIsAlwaysDynamic = true;
+								break;
+							}
+						}
+					}
+					
+					return bIsAlwaysDynamic ? ECheckBoxState::Checked : ECheckBoxState::Unchecked;
+				};
+				
+				InMenuBuilder.AddMenuEntry(
+					LOCTEXT("DynamicValue", "Dynamic Value"),
+					LOCTEXT("DynamicValueTooltip", "Flag this value as dynamic. This way it can be set from functions even when not exposed as a pin."),
+					FSlateIcon(),
+					FUIAction(
+						FExecuteAction::CreateLambda(MakeDynamic),
+						FCanExecuteAction(),
+						FGetActionCheckState::CreateLambda(GetDynamicCheckState)
 					),
 					NAME_None,
 					EUserInterfaceActionType::Check
diff --git a/Engine/Source/Editor/AnimGraph/Private/AnimGraphNode_BlendSpaceEvaluator.cpp b/Engine/Source/Editor/AnimGraph/Private/AnimGraphNode_BlendSpaceEvaluator.cpp
--- a/Engine/Source/Editor/AnimGraph/Private/AnimGraphNode_BlendSpaceEvaluator.cpp
+++ b/Engine/Source/Editor/AnimGraph/Private/AnimGraphNode_BlendSpaceEvaluator.cpp
@@ -65,40 +65,7 @@ void UAnimGraphNode_BlendSpaceEvaluator::ValidateAnimNodeDuringCompilation(class
 {
 	Super::ValidateAnimNodeDuringCompilation(ForSkeleton, MessageLog);
 
-	UBlendSpace* BlendSpaceToCheck = Node.GetBlendSpace();
-	UEdGraphPin* BlendSpacePin = FindPin(GET_MEMBER_NAME_STRING_CHECKED(FAnimNode_BlendSpaceEvaluator, BlendSpace));
-	if (BlendSpacePin != nullptr && BlendSpaceToCheck == nullptr)
-	{
-		BlendSpaceToCheck = Cast<UBlendSpace>(BlendSpacePin->DefaultObject);
-	}
-
-	if (BlendSpaceToCheck == nullptr)
-	{
-		// Check for bindings
-		bool bHasBinding = false;
-		if(BlendSpacePin != nullptr)
-		{
-			if (FAnimGraphNodePropertyBinding* BindingPtr = PropertyBindings.Find(BlendSpacePin->GetFName()))
-			{
-				bHasBinding = true;
-			}
-		}
-
-		// we may have a connected node or binding
-		if (BlendSpacePin == nullptr || (BlendSpacePin->LinkedTo.Num() == 0 && !bHasBinding))
-		{
-			MessageLog.Error(TEXT("@@ references an unknown blend space"), this);
-		}
-	}
-	else 
-	{
-		USkeleton* BlendSpaceSkeleton = BlendSpaceToCheck->GetSkeleton();
-		if (BlendSpaceSkeleton&& // if blend space doesn't have skeleton, it might be due to blend space not loaded yet, @todo: wait with anim blueprint compilation until all assets are loaded?
-			!ForSkeleton->IsCompatible(BlendSpaceSkeleton))
-		{
-			MessageLog.Error(TEXT("@@ references blendspace that uses an incompatible skeleton @@"), this, BlendSpaceSkeleton);
-		}
-	}
+	ValidateAnimNodeDuringCompilationHelper(ForSkeleton, MessageLog, Node.GetBlendSpace(), UBlendSpace::StaticClass(), FindPin(GET_MEMBER_NAME_STRING_CHECKED(FAnimNode_BlendSpaceEvaluator, BlendSpace)));
 }
 
 void UAnimGraphNode_BlendSpaceEvaluator::BakeDataDuringCompilation(class FCompilerResultsLog& MessageLog)
diff --git a/Engine/Source/Editor/AnimGraph/Private/AnimGraphNode_BlendSpacePlayer.cpp b/Engine/Source/Editor/AnimGraph/Private/AnimGraphNode_BlendSpacePlayer.cpp
--- a/Engine/Source/Editor/AnimGraph/Private/AnimGraphNode_BlendSpacePlayer.cpp
+++ b/Engine/Source/Editor/AnimGraph/Private/AnimGraphNode_BlendSpacePlayer.cpp
@@ -35,40 +35,7 @@ void UAnimGraphNode_BlendSpacePlayer::ValidateAnimNodeDuringCompilation(class US
 {
 	Super::ValidateAnimNodeDuringCompilation(ForSkeleton, MessageLog);
 
-	UBlendSpace* BlendSpaceToCheck = Node.GetBlendSpace();
-	UEdGraphPin* BlendSpacePin = FindPin(GET_MEMBER_NAME_STRING_CHECKED(FAnimNode_BlendSpacePlayer, BlendSpace));
-	if (BlendSpacePin != nullptr && BlendSpaceToCheck == nullptr)
-	{
-		BlendSpaceToCheck = Cast<UBlendSpace>(BlendSpacePin->DefaultObject);
-	}
-
-	if (BlendSpaceToCheck == nullptr)
-	{
-		// Check for bindings
-		bool bHasBinding = false;
-		if(BlendSpacePin != nullptr)
-		{
-			if (FAnimGraphNodePropertyBinding* BindingPtr = PropertyBindings.Find(BlendSpacePin->GetFName()))
-			{
-				bHasBinding = true;
-			}
-		}
-
-		// we may have a connected node or binding
-		if (BlendSpacePin == nullptr || (BlendSpacePin->LinkedTo.Num() == 0 && !bHasBinding))
-		{
-			MessageLog.Error(TEXT("@@ references an unknown blend space"), this);
-		}		
-	}
-	else
-	{
-		USkeleton* BlendSpaceSkeleton = BlendSpaceToCheck->GetSkeleton();
-		if (BlendSpaceSkeleton && // if blend space doesn't have skeleton, it might be due to blend space not loaded yet, @todo: wait with anim blueprint compilation until all assets are loaded?
-			!ForSkeleton->IsCompatible(BlendSpaceSkeleton))
-		{
-			MessageLog.Error(TEXT("@@ references blendspace that uses an incompatible skeleton @@"), this, BlendSpaceSkeleton);
-		}
-	}
+	ValidateAnimNodeDuringCompilationHelper(ForSkeleton, MessageLog, Node.GetBlendSpace(), UBlendSpace::StaticClass(), FindPin(GET_MEMBER_NAME_STRING_CHECKED(FAnimNode_BlendSpacePlayer, BlendSpace)));
 }
 
 void UAnimGraphNode_BlendSpacePlayer::BakeDataDuringCompilation(class FCompilerResultsLog& MessageLog)
diff --git a/Engine/Source/Editor/AnimGraph/Private/AnimGraphNode_CallFunction.cpp b/Engine/Source/Editor/AnimGraph/Private/AnimGraphNode_CallFunction.cpp
--- a/Engine/Source/Editor/AnimGraph/Private/AnimGraphNode_CallFunction.cpp
+++ b/Engine/Source/Editor/AnimGraph/Private/AnimGraphNode_CallFunction.cpp
@@ -127,9 +127,7 @@ void UAnimGraphNode_CallFunction::ExpandNode(FKismetCompilerContext& CompilerCon
 		UAnimBlueprintExtension_CallFunction* Extension = UAnimBlueprintExtension::GetExtension<UAnimBlueprintExtension_CallFunction>(GetAnimBlueprint());
 		
 		const FName EventName = Extension->AddCustomEventName(this);
-
-		// @TODO: move this name copy to the CDO baking step
-		Node.FunctionName = EventName;
+		Node.Function.SetFromFunctionName(EventName);
 
 		UK2Node_CustomEvent* CustomEventNode = CompilerContext.SpawnIntermediateEventNode<UK2Node_CustomEvent>(this, nullptr, CompilerContext.ConsolidatedEventGraph);
 		CustomEventNode->bInternalEvent = true;
diff --git a/Engine/Source/Editor/AnimGraph/Private/AnimGraphNode_CallFunction.h b/Engine/Source/Editor/AnimGraph/Private/AnimGraphNode_CallFunction.h
--- a/Engine/Source/Editor/AnimGraph/Private/AnimGraphNode_CallFunction.h
+++ b/Engine/Source/Editor/AnimGraph/Private/AnimGraphNode_CallFunction.h
@@ -12,7 +12,7 @@
 class UEdGraph;
 class UK2Node_CallFunction;
 
-UCLASS(MinimalAPI, meta=(Keywords = "Event"))
+UCLASS(MinimalAPI, Experimental, meta=(Keywords = "Event"))
 class UAnimGraphNode_CallFunction : public UAnimGraphNode_Base
 {
 	GENERATED_BODY()
diff --git a/Engine/Source/Editor/AnimGraph/Private/AnimGraphNode_CustomProperty.cpp b/Engine/Source/Editor/AnimGraph/Private/AnimGraphNode_CustomProperty.cpp
--- a/Engine/Source/Editor/AnimGraph/Private/AnimGraphNode_CustomProperty.cpp
+++ b/Engine/Source/Editor/AnimGraph/Private/AnimGraphNode_CustomProperty.cpp
@@ -10,6 +10,7 @@
 #include "PropertyCustomizationHelpers.h"
 #include "KismetCompiler.h"
 #include "IAnimBlueprintCompilationContext.h"
+#include "IAnimBlueprintCopyTermDefaultsContext.h"
 
 #define LOCTEXT_NAMESPACE "CustomPropNode"
 
@@ -70,6 +71,26 @@ void UAnimGraphNode_CustomProperty::OnProcessDuringCompilation(IAnimBlueprintCom
 	}
 }
 
+void UAnimGraphNode_CustomProperty::OnCopyTermDefaultsToDefaultObject(IAnimBlueprintCopyTermDefaultsContext& InCompilationContext, IAnimBlueprintNodeCopyTermDefaultsContext& InPerNodeContext, IAnimBlueprintGeneratedClassCompiledData& OutCompiledData)
+{
+	// Copy pin default values to generated properties
+	for (UEdGraphPin* Pin : Pins)
+	{
+		if (!Pin->bOrphanedPin && !UAnimationGraphSchema::IsPosePin(Pin->PinType))
+		{
+			FString PrefixedName = GetPinTargetVariableName(Pin);
+
+			if (FProperty* GeneratedProperty = FindFProperty<FProperty>(InPerNodeContext.GetClassDefaultObject()->GetClass(), *PrefixedName))
+			{
+				if(!FBlueprintEditorUtils::PropertyValueFromString(GeneratedProperty, Pin->GetDefaultAsString(), (uint8*)InPerNodeContext.GetClassDefaultObject(), InPerNodeContext.GetClassDefaultObject()))
+				{
+					InCompilationContext.GetMessageLog().Warning(TEXT("Unable to push default value for pin @@ on @@"), Pin, this);
+				}
+			}
+		}
+	}
+}
+
 void UAnimGraphNode_CustomProperty::ValidateAnimNodeDuringCompilation(USkeleton* ForSkeleton, FCompilerResultsLog& MessageLog)
 {
 	Super::ValidateAnimNodeDuringCompilation(ForSkeleton, MessageLog);
diff --git a/Engine/Source/Editor/AnimGraph/Private/AnimGraphNode_LinkedAnimGraph.cpp b/Engine/Source/Editor/AnimGraph/Private/AnimGraphNode_LinkedAnimGraph.cpp
--- a/Engine/Source/Editor/AnimGraph/Private/AnimGraphNode_LinkedAnimGraph.cpp
+++ b/Engine/Source/Editor/AnimGraph/Private/AnimGraphNode_LinkedAnimGraph.cpp
@@ -3,11 +3,14 @@
 #include "AnimGraphNode_LinkedAnimGraph.h"
 
 #include "AnimGraphNode_AssetPlayerBase.h"
+#include "AnimGraphNode_CallFunction.h"
 #include "BlueprintNodeSpawner.h"
 #include "Animation/AnimBlueprint.h"
 #include "AssetRegistry/AssetRegistryModule.h"
 #include "Engine/Blueprint.h"
 #include "BlueprintActionDatabaseRegistrar.h"
+#include "IAnimBlueprintCopyTermDefaultsContext.h"
+#include "KismetCompiler.h"
 
 #define LOCTEXT_NAMESPACE "UAnimGraphNode_LinkedAnimGraph"
 
@@ -20,7 +23,7 @@ void UAnimGraphNode_LinkedAnimGraph::PostPasteNode()
 		{
 			if(UAnimBlueprint* ThisBlueprint = GetAnimBlueprint())
 			{
-				if(LinkedBlueprint->TargetSkeleton != ThisBlueprint->TargetSkeleton)
+				if(!LinkedBlueprint->bIsTemplate && !ThisBlueprint->bIsTemplate && LinkedBlueprint->TargetSkeleton != ThisBlueprint->TargetSkeleton)
 				{
 					Node.InstanceClass = nullptr;
 				}
@@ -29,22 +32,45 @@ void UAnimGraphNode_LinkedAnimGraph::PostPasteNode()
 	}
 }
 
+TArray<UEdGraph*> UAnimGraphNode_LinkedAnimGraph::GetExternalGraphs() const
+{
+	if(UClass* InstanceClass = GetTargetClass())
+	{
+		if(UAnimBlueprint* LinkedBlueprint = Cast<UAnimBlueprint>(UBlueprint::GetBlueprintFromClass(InstanceClass)))
+		{
+			for(UEdGraph* Graph : LinkedBlueprint->FunctionGraphs)
+			{
+				if(Graph->GetFName() == UEdGraphSchema_K2::GN_AnimGraph)
+				{
+					return { Graph };
+				}
+			}
+		}
+	}
+
+	return TArray<UEdGraph*>();
+}
+
 void UAnimGraphNode_LinkedAnimGraph::SetupFromAsset(const FAssetData& InAssetData, bool bInIsTemplateNode)
 {
-	InAssetData.GetTagValue("TargetSkeleton", SkeletonName);
-	
-	if(!bInIsTemplateNode)
+	if(InAssetData.IsValid())
 	{
-		UAnimBlueprint* AnimBlueprint = CastChecked<UAnimBlueprint>(InAssetData.GetAsset());
-		Node.InstanceClass = AnimBlueprint->GeneratedClass.Get();
+		InAssetData.GetTagValue("TargetSkeleton", SkeletonName);
+		if(SkeletonName == TEXT("None"))
+		{
+			SkeletonName.Empty();
+		}
+
+		if(!bInIsTemplateNode)
+		{
+			UAnimBlueprint* AnimBlueprint = CastChecked<UAnimBlueprint>(InAssetData.GetAsset());
+			Node.InstanceClass = AnimBlueprint->GeneratedClass.Get();
+		}	
 	}
 }
 
 void UAnimGraphNode_LinkedAnimGraph::GetMenuActions(FBlueprintActionDatabaseRegistrar& InActionRegistrar) const
 {
-	// Anim graph node base class will allow us to spawn an 'empty' node
-	UAnimGraphNode_Base::GetMenuActions(InActionRegistrar);
-
 	UAnimGraphNode_AssetPlayerBase::GetMenuActionsHelper(
 		InActionRegistrar,
 		GetClass(),
@@ -52,11 +78,25 @@ void UAnimGraphNode_LinkedAnimGraph::GetMenuActions(FBlueprintActionDatabaseRegi
 		{ },
 		[](const FAssetData& InAssetData)
 		{
-			return FText::Format(LOCTEXT("MenuDescFormat", "{0} - Linked Anim Graph"), FText::FromName(InAssetData.AssetName));
+			if(InAssetData.IsValid())
+			{
+				return FText::Format(LOCTEXT("MenuDescFormat", "{0} - Linked Anim Graph"), FText::FromName(InAssetData.AssetName));
+			}
+			else
+			{
+				return LOCTEXT("MenuDesc", "Linked Anim Graph");
+			}
 		},
 		[](const FAssetData& InAssetData)
 		{
-			return FText::Format(LOCTEXT("MenuDescTooltipFormat", "Linked Anim Graph\n'{0}'"), FText::FromName(InAssetData.ObjectPath));
+			if(InAssetData.IsValid())
+			{
+				return FText::Format(LOCTEXT("MenuDescTooltipFormat", "Linked Anim Graph\n'{0}'"), FText::FromName(InAssetData.ObjectPath));
+			}
+			else
+			{
+				return LOCTEXT("MenuDescTooltip", "Linked Anim Graph");
+			}
 		},
 		[](UEdGraphNode* InNewNode, bool bInIsTemplateNode, const FAssetData InAssetData)
 		{
@@ -77,7 +117,7 @@ bool UAnimGraphNode_LinkedAnimGraph::IsActionFilteredOut(class FBlueprintActionF
 		{
 			if (UAnimBlueprint* AnimBlueprint = Cast<UAnimBlueprint>(Blueprint))
 			{
-				if(!AnimBlueprint->TargetSkeleton->IsCompatibleSkeletonByAssetString(SkeletonName))
+				if(AnimBlueprint->TargetSkeleton != nullptr && !AnimBlueprint->TargetSkeleton->IsCompatibleSkeletonByAssetString(SkeletonName))
 				{
 					bIsFilteredOut = true;
 					break;
diff --git a/Engine/Source/Editor/AnimGraph/Private/AnimGraphNode_LinkedAnimGraphBase.cpp b/Engine/Source/Editor/AnimGraph/Private/AnimGraphNode_LinkedAnimGraphBase.cpp
--- a/Engine/Source/Editor/AnimGraph/Private/AnimGraphNode_LinkedAnimGraphBase.cpp
+++ b/Engine/Source/Editor/AnimGraph/Private/AnimGraphNode_LinkedAnimGraphBase.cpp
@@ -503,6 +503,8 @@ void UAnimGraphNode_LinkedAnimGraphBase::GetOutputLinkAttributes(FNodeAttributeA
 
 void UAnimGraphNode_LinkedAnimGraphBase::OnCopyTermDefaultsToDefaultObject(IAnimBlueprintCopyTermDefaultsContext& InCompilationContext, IAnimBlueprintNodeCopyTermDefaultsContext& InPerNodeContext, IAnimBlueprintGeneratedClassCompiledData& OutCompiledData)
 {
+	Super::OnCopyTermDefaultsToDefaultObject(InCompilationContext, InPerNodeContext, OutCompiledData);
+	
 	UAnimGraphNode_LinkedAnimGraphBase* TrueNode = InCompilationContext.GetMessageLog().FindSourceObjectTypeChecked<UAnimGraphNode_LinkedAnimGraphBase>(this);
 
 	FAnimNode_LinkedAnimGraph* DestinationNode = reinterpret_cast<FAnimNode_LinkedAnimGraph*>(InPerNodeContext.GetDestinationPtr());
diff --git a/Engine/Source/Editor/AnimGraph/Private/AnimGraphNode_LinkedInputPose.cpp b/Engine/Source/Editor/AnimGraph/Private/AnimGraphNode_LinkedInputPose.cpp
--- a/Engine/Source/Editor/AnimGraph/Private/AnimGraphNode_LinkedInputPose.cpp
+++ b/Engine/Source/Editor/AnimGraph/Private/AnimGraphNode_LinkedInputPose.cpp
@@ -348,10 +348,12 @@ void UAnimGraphNode_LinkedInputPose::PostPlacedNewNode()
 			if(UAnimGraphNode_LinkedInputPose* LinkedInputPoseNode = WeakThis.Get())
 			{
 				// refresh the BP editor's details panel in case we are viewing the graph
-				IAssetEditorInstance* AssetEditor = GEditor->GetEditorSubsystem<UAssetEditorSubsystem>()->FindEditorForAsset(LinkedInputPoseNode->GetAnimBlueprint(), false);
-				check(AssetEditor->GetEditorName() == "AnimationBlueprintEditor");
-				IAnimationBlueprintEditor* AnimationBlueprintEditor = static_cast<IAnimationBlueprintEditor*>(AssetEditor);
-				AnimationBlueprintEditor->RefreshInspector();
+				if(IAssetEditorInstance* AssetEditor = GEditor->GetEditorSubsystem<UAssetEditorSubsystem>()->FindEditorForAsset(LinkedInputPoseNode->GetAnimBlueprint(), false))
+				{
+					check(AssetEditor->GetEditorName() == "AnimationBlueprintEditor");
+					IAnimationBlueprintEditor* AnimationBlueprintEditor = static_cast<IAnimationBlueprintEditor*>(AssetEditor);
+					AnimationBlueprintEditor->RefreshInspector();
+				}
 			}
 			return false;
 		}));
diff --git a/Engine/Source/Editor/AnimGraph/Private/AnimGraphNode_PoseHandler.cpp b/Engine/Source/Editor/AnimGraph/Private/AnimGraphNode_PoseHandler.cpp
--- a/Engine/Source/Editor/AnimGraph/Private/AnimGraphNode_PoseHandler.cpp
+++ b/Engine/Source/Editor/AnimGraph/Private/AnimGraphNode_PoseHandler.cpp
@@ -14,31 +14,9 @@ UAnimGraphNode_PoseHandler::UAnimGraphNode_PoseHandler(const FObjectInitializer&
 
 void UAnimGraphNode_PoseHandler::ValidateAnimNodeDuringCompilation(USkeleton* ForSkeleton, FCompilerResultsLog& MessageLog)
 {
-	UPoseAsset* PoseAssetToCheck = GetPoseHandlerNode()->PoseAsset;
-	UEdGraphPin* PoseAssetPin = FindPin(GET_MEMBER_NAME_STRING_CHECKED(FAnimNode_PoseHandler, PoseAsset));
-	if (PoseAssetPin != nullptr && PoseAssetToCheck == nullptr)
-	{
-		PoseAssetToCheck = Cast<UPoseAsset>(PoseAssetPin->DefaultObject);
-	}
-
-	if (PoseAssetToCheck == nullptr)
-	{
-		if (IsPoseAssetRequired() && (PoseAssetPin == nullptr || PoseAssetPin->LinkedTo.Num() == 0))
-		{
-			MessageLog.Error(TEXT("@@ references an unknown poseasset"), this);
-		}
-	}
-	else
-	{
-		USkeleton* SeqSkeleton = PoseAssetToCheck->GetSkeleton();
-		if (SeqSkeleton && // if PoseAsset doesn't have skeleton, it might be due to PoseAsset not loaded yet, @todo: wait with anim blueprint compilation until all assets are loaded?
-			!ForSkeleton->IsCompatible(SeqSkeleton))
-		{
-			MessageLog.Error(TEXT("@@ references poseasset that uses an incompatible skeleton @@"), this, SeqSkeleton);
-		}
-	}
-
 	Super::ValidateAnimNodeDuringCompilation(ForSkeleton, MessageLog);
+	
+	ValidateAnimNodeDuringCompilationHelper(ForSkeleton, MessageLog, GetPoseHandlerNode()->PoseAsset, UPoseAsset::StaticClass(), FindPin(GET_MEMBER_NAME_STRING_CHECKED(FAnimNode_PoseHandler, PoseAsset)));
 }
 
 void UAnimGraphNode_PoseHandler::SetAnimationAsset(UAnimationAsset* Asset)
diff --git a/Engine/Source/Editor/AnimGraph/Private/AnimGraphNode_SequenceEvaluator.cpp b/Engine/Source/Editor/AnimGraph/Private/AnimGraphNode_SequenceEvaluator.cpp
--- a/Engine/Source/Editor/AnimGraph/Private/AnimGraphNode_SequenceEvaluator.cpp
+++ b/Engine/Source/Editor/AnimGraph/Private/AnimGraphNode_SequenceEvaluator.cpp
@@ -143,40 +143,7 @@ void UAnimGraphNode_SequenceEvaluator::ValidateAnimNodeDuringCompilation(class U
 {
 	Super::ValidateAnimNodeDuringCompilation(ForSkeleton, MessageLog);
 
-	UAnimSequenceBase* SequenceToCheck = Node.GetSequence();
-	UEdGraphPin* SequencePin = FindPin(GET_MEMBER_NAME_STRING_CHECKED(FAnimNode_SequenceEvaluator, Sequence));
-	if (SequencePin != nullptr && SequenceToCheck == nullptr)
-	{
-		SequenceToCheck = Cast<UAnimSequenceBase>(SequencePin->DefaultObject);
-	}
-
-	if (SequenceToCheck == nullptr)
-	{
-		// Check for bindings
-		bool bHasBinding = false;
-		if(SequencePin != nullptr)
-		{
-			if (FAnimGraphNodePropertyBinding* BindingPtr = PropertyBindings.Find(SequencePin->GetFName()))
-			{
-				bHasBinding = true;
-			}
-		}
-
-		// we may have a connected node or binding
-		if (SequencePin == nullptr || (SequencePin->LinkedTo.Num() == 0 && !bHasBinding))
-		{
-			MessageLog.Error(TEXT("@@ references an unknown sequence"), this);
-		}
-	}
-	else
-	{
-		USkeleton* SeqSkeleton = SequenceToCheck->GetSkeleton();
-		if (SeqSkeleton && // if anim sequence doesn't have skeleton, it might be due to anim sequence not loaded yet, @todo: wait with anim blueprint compilation until all assets are loaded?
-			!ForSkeleton->IsCompatible(SeqSkeleton))
-		{
-			MessageLog.Error(TEXT("@@ references sequence that uses an incompatible skeleton @@"), this, SeqSkeleton);
-		}
-	}
+	ValidateAnimNodeDuringCompilationHelper(ForSkeleton, MessageLog, Node.GetSequence(), UAnimSequenceBase::StaticClass(), FindPin(GET_MEMBER_NAME_STRING_CHECKED(FAnimNode_SequenceEvaluator, Sequence)));
 }
 
 void UAnimGraphNode_SequenceEvaluator::GetNodeContextMenuActions(UToolMenu* Menu, UGraphNodeContextMenuContext* Context) const
diff --git a/Engine/Source/Editor/AnimGraph/Private/AnimGraphNode_SequencePlayer.cpp b/Engine/Source/Editor/AnimGraph/Private/AnimGraphNode_SequencePlayer.cpp
--- a/Engine/Source/Editor/AnimGraph/Private/AnimGraphNode_SequencePlayer.cpp
+++ b/Engine/Source/Editor/AnimGraph/Private/AnimGraphNode_SequencePlayer.cpp
@@ -162,44 +162,7 @@ void UAnimGraphNode_SequencePlayer::ValidateAnimNodeDuringCompilation(class USke
 {
 	Super::ValidateAnimNodeDuringCompilation(ForSkeleton, MessageLog);
 
-	UAnimSequenceBase* SequenceToCheck = Node.GetSequence();
-	UEdGraphPin* SequencePin = FindPin(GET_MEMBER_NAME_STRING_CHECKED(FAnimNode_SequencePlayer, Sequence));
-	if (SequencePin != nullptr && SequenceToCheck == nullptr)
-	{
-		SequenceToCheck = Cast<UAnimSequenceBase>(SequencePin->DefaultObject);
-	}
-
-	if (SequenceToCheck == nullptr)
-	{
-		// Check for bindings
-		bool bHasBinding = false;
-		if(SequencePin != nullptr)
-		{
-			if (FAnimGraphNodePropertyBinding* BindingPtr = PropertyBindings.Find(SequencePin->GetFName()))
-			{
-				bHasBinding = true;
-			}
-		}
-
-		// we may have a connected node or binding
-		if (SequencePin == nullptr || (SequencePin->LinkedTo.Num() == 0 && !bHasBinding))
-		{
-			MessageLog.Error(TEXT("@@ references an unknown sequence"), this);
-		}
-	}
-	else if(SupportsAssetClass(SequenceToCheck->GetClass()) == EAnimAssetHandlerType::NotSupported)
-	{
-		MessageLog.Error(*FText::Format(LOCTEXT("UnsupportedAssetError", "@@ is trying to play a {0} as a sequence, which is not allowed."), SequenceToCheck->GetClass()->GetDisplayNameText()).ToString(), this);
-	}
-	else
-	{
-		USkeleton* SeqSkeleton = SequenceToCheck->GetSkeleton();
-		if (SeqSkeleton && // if anim sequence doesn't have skeleton, it might be due to anim sequence not loaded yet, @todo: wait with anim blueprint compilation until all assets are loaded?
-			!ForSkeleton->IsCompatible(SeqSkeleton))
-		{
-			MessageLog.Error(TEXT("@@ references sequence that uses an incompatible skeleton @@"), this, SeqSkeleton);
-		}
-	}
+	ValidateAnimNodeDuringCompilationHelper(ForSkeleton, MessageLog, Node.GetSequence(), UAnimSequenceBase::StaticClass(), FindPin(GET_MEMBER_NAME_STRING_CHECKED(FAnimNode_SequencePlayer, Sequence)));
 }
 
 void UAnimGraphNode_SequencePlayer::GetNodeContextMenuActions(UToolMenu* Menu, UGraphNodeContextMenuContext* Context) const
diff --git a/Engine/Source/Editor/AnimGraph/Public/AnimGraphNode_AssetPlayerBase.h b/Engine/Source/Editor/AnimGraph/Public/AnimGraphNode_AssetPlayerBase.h
--- a/Engine/Source/Editor/AnimGraph/Public/AnimGraphNode_AssetPlayerBase.h
+++ b/Engine/Source/Editor/AnimGraph/Public/AnimGraphNode_AssetPlayerBase.h
@@ -46,6 +46,9 @@ public:
 	// Helper function to gather menu actions from specific asset types supported by this node
 	static void GetMenuActionsHelper(FBlueprintActionDatabaseRegistrar& InActionRegistrar, TSubclassOf<UAnimGraphNode_Base> InNodeClass, const TArray<TSubclassOf<UObject>>& InAssetTypes, const TArray<TSubclassOf<UObject>>& InExcludedAssetTypes, const TFunctionRef<FText(const FAssetData&)>& InMenuNameFunction, const TFunctionRef<FText(const FAssetData&)>& InMenuTooltipFunction, const TFunction<void(UEdGraphNode*, bool, const FAssetData)>& InSetupNewNodeFunction);
 
+	// Helper function to validate player nodes
+	void ValidateAnimNodeDuringCompilationHelper(USkeleton* ForSkeleton, FCompilerResultsLog& MessageLog, UAnimationAsset* InAsset, TSubclassOf<UAnimationAsset> InAssetType, UEdGraphPin* InExposedPin);
+
 protected:
 	// Helper functions to build a title for an asset player node
 	FText GetNodeTitleHelper(ENodeTitleType::Type InTitleType, UEdGraphPin* InAssetPin, const FText& InAssetDesc, const TFunction<FText(UAnimationAsset*)> InPostFixFunctionRef = nullptr) const;
diff --git a/Engine/Source/Editor/AnimGraph/Public/AnimGraphNode_Base.h b/Engine/Source/Editor/AnimGraph/Public/AnimGraphNode_Base.h
--- a/Engine/Source/Editor/AnimGraph/Public/AnimGraphNode_Base.h
+++ b/Engine/Source/Editor/AnimGraph/Public/AnimGraphNode_Base.h
@@ -12,6 +12,8 @@
 #include "Editor.h"
 #include "IPropertyAccessEditor.h"
 #include "K2Node.h"
+#include "Engine/MemberReference.h"
+
 #include "AnimGraphNode_Base.generated.h"
 
 class FAnimGraphNodeDetails;
@@ -200,9 +202,29 @@ class ANIMGRAPH_API UAnimGraphNode_Base : public UK2Node
  	UPROPERTY(EditAnywhere, Category=PinOptions)
  	TMap<FName, FAnimGraphNodePropertyBinding> PropertyBindings;
 
+	/** Properties marked as always dynamic, so they can be set externally */
+	UPROPERTY()
+	TSet<FName> AlwaysDynamicProperties;
+	
 	UPROPERTY(Transient)
 	EBlueprintUsage BlueprintUsage;
 
+	// Function called when the node is initialized
+	UPROPERTY(EditAnywhere, Category = Functions, meta=(FunctionReference, PrototypeFunction="/Script/AnimGraphRuntime.AnimNodeLibrary.Prototype_ThreadSafeAnimNodeCall"), DisplayName="On Initialize")
+	FMemberReference InitializeFunction;
+
+	// Function called when the node becomes relevant
+	UPROPERTY(EditAnywhere, Category = Functions, meta=(FunctionReference, PrototypeFunction="/Script/AnimGraphRuntime.AnimNodeLibrary.Prototype_ThreadSafeAnimNodeCall"), DisplayName="On Become Relevant")
+	FMemberReference BecomeRelevantFunction;
+
+	// Function called when the node is updated
+	UPROPERTY(EditAnywhere, Category = Functions, meta=(FunctionReference, PrototypeFunction="/Script/AnimGraphRuntime.AnimNodeLibrary.Prototype_ThreadSafeAnimNodeCall"), DisplayName="On Update")
+	FMemberReference UpdateFunction;
+
+	// Function called when the node is evaluated
+	UPROPERTY(EditAnywhere, Category = Functions, meta=(FunctionReference, PrototypeFunction="/Script/AnimGraphRuntime.AnimNodeLibrary.Prototype_ThreadSafeAnimNodeCall"), DisplayName="On Evaluate")
+	FMemberReference EvaluateFunction;
+	
 	// UObject interface
 	virtual void PostEditChangeProperty(FPropertyChangedEvent& PropertyChangedEvent) override;
 	virtual void PreEditChange(FProperty* PropertyAboutToChange) override;
@@ -325,6 +347,7 @@ class ANIMGRAPH_API UAnimGraphNode_Base : public UK2Node
 	// you should implement all below functions
 	virtual bool DoesSupportTimeForTransitionGetter() const { return false; }
 	virtual UAnimationAsset* GetAnimationAsset() const { return nullptr; }
+	virtual TSubclassOf<UAnimationAsset> GetAnimationAssetClass() const { return nullptr; }
 	virtual const TCHAR* GetTimePropertyName() const { return nullptr; }
 	virtual UScriptStruct* GetTimePropertyStruct() const { return nullptr; }
 	// END Interface to support transition getter
diff --git a/Engine/Source/Editor/AnimGraph/Public/AnimGraphNode_BlendSpaceBase.h b/Engine/Source/Editor/AnimGraph/Public/AnimGraphNode_BlendSpaceBase.h
--- a/Engine/Source/Editor/AnimGraph/Public/AnimGraphNode_BlendSpaceBase.h
+++ b/Engine/Source/Editor/AnimGraph/Public/AnimGraphNode_BlendSpaceBase.h
@@ -24,6 +24,7 @@ class UAnimGraphNode_BlendSpaceBase : public UAnimGraphNode_AssetPlayerBase
 	ANIMGRAPH_API virtual void PreloadRequiredAssets() override;
 	ANIMGRAPH_API virtual void PostProcessPinName(const UEdGraphPin* Pin, FString& DisplayName) const override;
 	ANIMGRAPH_API virtual void GetOutputLinkAttributes(FNodeAttributeArray& OutAttributes) const;
+	virtual TSubclassOf<UAnimationAsset> GetAnimationAssetClass() const { return UBlendSpace::StaticClass(); }
 	// End of UAnimGraphNode_Base interface
 
 protected:
diff --git a/Engine/Source/Editor/AnimGraph/Public/AnimGraphNode_CustomProperty.h b/Engine/Source/Editor/AnimGraph/Public/AnimGraphNode_CustomProperty.h
--- a/Engine/Source/Editor/AnimGraph/Public/AnimGraphNode_CustomProperty.h
+++ b/Engine/Source/Editor/AnimGraph/Public/AnimGraphNode_CustomProperty.h
@@ -35,6 +35,7 @@ public:
 	// UAnimGraphNode_Base interface
 	virtual void CustomizeDetails(IDetailLayoutBuilder& DetailBuilder) override;
 	virtual void OnProcessDuringCompilation(IAnimBlueprintCompilationContext& InCompilationContext, IAnimBlueprintGeneratedClassCompiledData& OutCompiledData) override;
+	virtual void OnCopyTermDefaultsToDefaultObject(IAnimBlueprintCopyTermDefaultsContext& InCompilationContext, IAnimBlueprintNodeCopyTermDefaultsContext& InPerNodeContext, IAnimBlueprintGeneratedClassCompiledData& OutCompiledData) override;
 
 	// Gets the property on InOwnerInstanceClass that corresponds to InInputPin
 	void GetInstancePinProperty(const IAnimBlueprintCompilationContext& InCompilationContext, UEdGraphPin* InInputPin, FProperty*& OutProperty);
diff --git a/Engine/Source/Editor/AnimGraph/Public/AnimGraphNode_LinkedAnimGraph.h b/Engine/Source/Editor/AnimGraph/Public/AnimGraphNode_LinkedAnimGraph.h
--- a/Engine/Source/Editor/AnimGraph/Public/AnimGraphNode_LinkedAnimGraph.h
+++ b/Engine/Source/Editor/AnimGraph/Public/AnimGraphNode_LinkedAnimGraph.h
@@ -3,15 +3,15 @@
 #pragma once
 
 #include "CoreMinimal.h"
-#include "UObject/ObjectMacros.h"
-#include "Misc/Guid.h"
 #include "AnimGraphNode_LinkedAnimGraphBase.h"
+#include "K2Node_ExternalGraphInterface.h"
 #include "Animation/AnimNode_LinkedAnimGraph.h"
+#include "Engine/MemberReference.h"
 
 #include "AnimGraphNode_LinkedAnimGraph.generated.h"
 
 UCLASS(MinimalAPI)
-class UAnimGraphNode_LinkedAnimGraph : public UAnimGraphNode_LinkedAnimGraphBase
+class UAnimGraphNode_LinkedAnimGraph : public UAnimGraphNode_LinkedAnimGraphBase, public IK2Node_ExternalGraphInterface
 {
 	GENERATED_BODY()
 
@@ -21,7 +21,7 @@ public:
 	FAnimNode_LinkedAnimGraph Node;
 
 	// Begin UEdGraphNode
-	virtual void PostPasteNode();
+	virtual void PostPasteNode() override;
 
 	// Begin UK2Node
 	virtual void GetMenuActions(FBlueprintActionDatabaseRegistrar& ActionRegistrar) const override;
@@ -35,6 +35,9 @@ public:
 	virtual FAnimNode_LinkedAnimGraph* GetLinkedAnimGraphNode() override { return &Node; }
 	virtual const FAnimNode_LinkedAnimGraph* GetLinkedAnimGraphNode() const override { return &Node; }
 
+	// IK2Node_ExternalGraphInterface interface
+	virtual TArray<UEdGraph*> GetExternalGraphs() const override;
+
 private:
 	// Setup the node from a specified anim BP
 	void SetupFromAsset(const FAssetData& InAssetData, bool bInIsTemplateNode);
diff --git a/Engine/Source/Editor/AnimGraph/Public/AnimGraphNode_SequencePlayer.h b/Engine/Source/Editor/AnimGraph/Public/AnimGraphNode_SequencePlayer.h
--- a/Engine/Source/Editor/AnimGraph/Public/AnimGraphNode_SequencePlayer.h
+++ b/Engine/Source/Editor/AnimGraph/Public/AnimGraphNode_SequencePlayer.h
@@ -46,6 +46,7 @@ class UAnimGraphNode_SequencePlayer : public UAnimGraphNode_AssetPlayerBase
 	virtual EAnimAssetHandlerType SupportsAssetClass(const UClass* AssetClass) const override;
 	virtual void OnOverrideAssets(IAnimBlueprintNodeOverrideAssetsContext& InContext) const override;
 	virtual void GetOutputLinkAttributes(FNodeAttributeArray& OutAttributes) const override;
+	virtual TSubclassOf<UAnimationAsset> GetAnimationAssetClass() const { return UAnimSequenceBase::StaticClass(); }
 	// End of UAnimGraphNode_Base interface
 
 	// UK2Node interface
diff --git a/Engine/Source/Editor/AnimGraph/Public/IAnimBlueprintCopyTermDefaultsContext.h b/Engine/Source/Editor/AnimGraph/Public/IAnimBlueprintCopyTermDefaultsContext.h
--- a/Engine/Source/Editor/AnimGraph/Public/IAnimBlueprintCopyTermDefaultsContext.h
+++ b/Engine/Source/Editor/AnimGraph/Public/IAnimBlueprintCopyTermDefaultsContext.h
@@ -34,6 +34,9 @@ class IAnimBlueprintNodeCopyTermDefaultsContext
 public:	
 	virtual ~IAnimBlueprintNodeCopyTermDefaultsContext() = default;
 
+	// Get the CDO that we are writing to
+	UObject* GetClassDefaultObject() const { return GetClassDefaultObjectImpl(); }
+	
 	// Get the property that we are writing to
 	const FProperty* GetTargetProperty() const { return GetTargetPropertyImpl(); }
 
@@ -46,7 +49,28 @@ public:
 	// Get the property index for this node
 	int32 GetNodePropertyIndex() const { return GetNodePropertyIndexImpl(); }
 
+	// Get the source node cast to the correct type
+	template <typename NodeType>
+	const NodeType& GetSourceNode() const
+	{
+		const FStructProperty* TargetProperty = CastFieldChecked<FStructProperty>(GetTargetProperty());
+		check(TargetProperty->Struct->IsChildOf(NodeType::StaticStruct()));
+		return *reinterpret_cast<const NodeType*>(GetSourcePtr());
+	}
+
+	// Get the destination node cast to the correct type
+	template <typename NodeType>
+	NodeType& GetDestinationNode() const
+	{
+		const FStructProperty* TargetProperty = CastFieldChecked<FStructProperty>(GetTargetProperty());
+		check(TargetProperty->Struct->IsChildOf(NodeType::StaticStruct()));
+		return *reinterpret_cast<NodeType*>(GetDestinationPtr());
+	}
+
 protected:
+	// Get the CDO that we are writing to
+	virtual UObject* GetClassDefaultObjectImpl() const = 0;
+	
 	// Get the property that we are writing to
 	virtual const FProperty* GetTargetPropertyImpl() const = 0;
 
@@ -66,6 +90,9 @@ class IAnimBlueprintExtensionCopyTermDefaultsContext
 public:	
 	virtual ~IAnimBlueprintExtensionCopyTermDefaultsContext() = default;
 
+	// Get the CDO that we are writing to
+	UObject* GetClassDefaultObject() const { return GetClassDefaultObjectImpl(); }
+	
 	// Get the property that we are writing to
 	const FProperty* GetTargetProperty() const { return GetTargetPropertyImpl(); }
 
@@ -79,6 +106,9 @@ public:
 	int32 GetNodePropertyIndex() const { return GetNodePropertyIndexImpl(); }
 
 protected:
+	// Get the CDO that we are writing to
+	virtual UObject* GetClassDefaultObjectImpl() const = 0;
+	
 	// Get the property that we are writing to
 	virtual const FProperty* GetTargetPropertyImpl() const = 0;
 
diff --git a/Engine/Source/Editor/AnimationBlueprintEditor/Private/AnimationBlueprintEditor.cpp b/Engine/Source/Editor/AnimationBlueprintEditor/Private/AnimationBlueprintEditor.cpp
--- a/Engine/Source/Editor/AnimationBlueprintEditor/Private/AnimationBlueprintEditor.cpp
+++ b/Engine/Source/Editor/AnimationBlueprintEditor/Private/AnimationBlueprintEditor.cpp
@@ -66,6 +66,7 @@
 #include "Widgets/Input/SButton.h"
 #include "EditorFontGlyphs.h"
 #include "AnimationBlueprintInterfaceEditorMode.h"
+#include "AnimationBlueprintTemplateEditorMode.h"
 #include "ToolMenus.h"
 
 #include "PersonaToolMenuContext.h"
@@ -94,6 +95,7 @@ const FName AnimationBlueprintEditorAppName(TEXT("AnimationBlueprintEditorApp"))
 
 const FName FAnimationBlueprintEditorModes::AnimationBlueprintEditorMode("GraphName");	// For backwards compatibility we keep the old mode name here
 const FName FAnimationBlueprintEditorModes::AnimationBlueprintInterfaceEditorMode("Interface");
+const FName FAnimationBlueprintEditorModes::AnimationBlueprintTemplateEditorMode("Template");
 
 namespace AnimationBlueprintEditorTabs
 {
@@ -262,7 +264,7 @@ void FAnimationBlueprintEditor::InitAnimationBlueprintEditor(const EToolkitMode:
 	TSharedRef<IAssetFamily> AssetFamily = PersonaModule.CreatePersonaAssetFamily(InAnimBlueprint);
 	AssetFamily->RecordAssetOpened(FAssetData(InAnimBlueprint));
 
-	if(InAnimBlueprint->BlueprintType != BPTYPE_Interface)
+	if(InAnimBlueprint->BlueprintType != BPTYPE_Interface && !InAnimBlueprint->bIsTemplate)
 	{
 		// create the skeleton tree
 		FSkeletonTreeArgs SkeletonTreeArgs;
@@ -313,6 +315,19 @@ void FAnimationBlueprintEditor::InitAnimationBlueprintEditor(const EToolkitMode:
 		// Activate the initial mode (which will populate with a real layout)
 		SetCurrentMode(FAnimationBlueprintEditorModes::AnimationBlueprintInterfaceEditorMode);
 	}
+	else if(InAnimBlueprint->bIsTemplate)
+	{
+		AddApplicationMode(
+			FAnimationBlueprintEditorModes::AnimationBlueprintTemplateEditorMode,
+			MakeShareable(new FAnimationBlueprintTemplateEditorMode(SharedThis(this))));
+		
+		ExtendMenu();
+		ExtendToolbar();
+		RegenerateMenusAndToolbars();
+
+		// Activate the initial mode (which will populate with a real layout)
+		SetCurrentMode(FAnimationBlueprintEditorModes::AnimationBlueprintTemplateEditorMode);
+	}
 	else
 	{
 		AddApplicationMode(
@@ -393,7 +408,7 @@ void FAnimationBlueprintEditor::ExtendToolbar()
 	}
 
 	UAnimBlueprint* AnimBlueprint = PersonaToolkit->GetAnimBlueprint();
-	if(AnimBlueprint && AnimBlueprint->BlueprintType != BPTYPE_Interface)
+	if(AnimBlueprint && AnimBlueprint->BlueprintType != BPTYPE_Interface && !AnimBlueprint->bIsTemplate)
 	{
 		ToolbarExtender->AddToolBarExtension(
 			"Asset",
@@ -1553,7 +1568,7 @@ FText FAnimationBlueprintEditor::GetGraphDecorationString(UEdGraph* InGraph) con
 {
 	if (!IsGraphInCurrentBlueprint(InGraph))
 	{
-		return LOCTEXT("PersonaExternalGraphDecoration", " Parent Graph Preview");
+		return LOCTEXT("PersonaExternalGraphDecoration", " External Graph Preview");
 	}
 	return FText::GetEmpty();
 }
diff --git a/Engine/Source/Editor/AnimationBlueprintEditor/Private/AnimationBlueprintEditor.h b/Engine/Source/Editor/AnimationBlueprintEditor/Private/AnimationBlueprintEditor.h
--- a/Engine/Source/Editor/AnimationBlueprintEditor/Private/AnimationBlueprintEditor.h
+++ b/Engine/Source/Editor/AnimationBlueprintEditor/Private/AnimationBlueprintEditor.h
@@ -29,6 +29,7 @@ struct FAnimationBlueprintEditorModes
 	// Mode constants
 	static const FName AnimationBlueprintEditorMode;
 	static const FName AnimationBlueprintInterfaceEditorMode;
+	static const FName AnimationBlueprintTemplateEditorMode;
 
 	static FText GetLocalizedMode(const FName InMode)
 	{
@@ -38,6 +39,7 @@ struct FAnimationBlueprintEditorModes
 		{
 			LocModes.Add(AnimationBlueprintEditorMode, NSLOCTEXT("AnimationBlueprintEditorModes", "AnimationBlueprintEditorMode", "Animation Blueprint"));
 			LocModes.Add(AnimationBlueprintInterfaceEditorMode, NSLOCTEXT("AnimationBlueprintEditorModes", "AnimationBlueprintInterface EditorMode", "Animation Blueprint Interface"));
+			LocModes.Add(AnimationBlueprintTemplateEditorMode, NSLOCTEXT("AnimationBlueprintEditorModes", "AnimationBlueprintTemplate EditorMode", "Animation Blueprint Template"));
 		}
 
 		check(InMode != NAME_None);
diff --git a/Engine/Source/Editor/AnimationBlueprintEditor/Private/AnimationBlueprintTemplateEditorMode.cpp b/Engine/Source/Editor/AnimationBlueprintEditor/Private/AnimationBlueprintTemplateEditorMode.cpp
new file mode 100644
--- /dev/null
+++ b/Engine/Source/Editor/AnimationBlueprintEditor/Private/AnimationBlueprintTemplateEditorMode.cpp
@@ -0,0 +1,91 @@
+// Copyright Epic Games, Inc. All Rights Reserved.
+
+#include "AnimationBlueprintTemplateEditorMode.h"
+#include "BlueprintEditorTabs.h"
+#include "PersonaModule.h"
+#include "SBlueprintEditorToolbar.h"
+#include "ToolMenus.h"
+
+FAnimationBlueprintTemplateEditorMode::FAnimationBlueprintTemplateEditorMode(const TSharedRef<FAnimationBlueprintEditor>& InAnimationBlueprintEditor)
+	: FBlueprintEditorApplicationMode(InAnimationBlueprintEditor, FAnimationBlueprintEditorModes::AnimationBlueprintTemplateEditorMode, FAnimationBlueprintEditorModes::GetLocalizedMode)
+{
+	AnimBlueprintPtr = CastChecked<UAnimBlueprint>(InAnimationBlueprintEditor->GetBlueprintObj());
+
+	TabLayout = FTabManager::NewLayout( "Standalone_AnimationBlueprintTemplateEditMode_Layout_v1.0" )
+		->AddArea
+		(
+			FTabManager::NewPrimaryArea()
+			->SetOrientation(Orient_Vertical)
+			->Split
+			(
+				// Main application area
+				FTabManager::NewSplitter()
+				->SetOrientation(Orient_Horizontal)
+				->Split
+				(
+					//	Left - My Blueprint
+					FTabManager::NewStack()
+					->SetSizeCoefficient(0.25f)
+					->AddTab(FBlueprintEditorTabs::MyBlueprintID, ETabState::OpenedTab)
+				)
+				->Split
+				(
+					// Middle 
+					FTabManager::NewSplitter()
+					->SetOrientation(Orient_Vertical)
+					->SetSizeCoefficient(0.55f)
+					->Split
+					(
+						// Middle top - document edit area
+						FTabManager::NewStack()
+						->SetSizeCoefficient(0.8f)
+						->AddTab("Document", ETabState::ClosedTab)
+					)
+					->Split
+					(
+						// Middle bottom - compiler results & find
+						FTabManager::NewStack()
+						->SetSizeCoefficient(0.2f)
+						->AddTab(FBlueprintEditorTabs::CompilerResultsID, ETabState::ClosedTab)
+						->AddTab(FBlueprintEditorTabs::FindResultsID, ETabState::ClosedTab)
+					)
+				)
+				->Split
+				(
+					// Right top - selection details panel
+					FTabManager::NewStack()
+					->SetSizeCoefficient(0.2f)
+					->AddTab(FBlueprintEditorTabs::DetailsID, ETabState::OpenedTab)
+				)
+			)
+		);
+
+	FPersonaModule& PersonaModule = FModuleManager::LoadModuleChecked<FPersonaModule>("Persona");
+
+	ToolbarExtender = MakeShareable(new FExtender);
+
+	if (UToolMenu* Toolbar = InAnimationBlueprintEditor->RegisterModeToolbarIfUnregistered(GetModeName()))
+	{
+		InAnimationBlueprintEditor->GetToolbarBuilder()->AddCompileToolbar(Toolbar);
+		InAnimationBlueprintEditor->GetToolbarBuilder()->AddScriptingToolbar(Toolbar);
+		InAnimationBlueprintEditor->GetToolbarBuilder()->AddBlueprintGlobalOptionsToolbar(Toolbar);
+		InAnimationBlueprintEditor->GetToolbarBuilder()->AddDebuggingToolbar(Toolbar);
+	}
+
+	PersonaModule.OnRegisterTabs().Broadcast(TabFactories, InAnimationBlueprintEditor);
+	LayoutExtender = MakeShared<FLayoutExtender>();
+	PersonaModule.OnRegisterLayoutExtensions().Broadcast(*LayoutExtender.Get());
+	TabLayout->ProcessExtensions(*LayoutExtender.Get());
+}
+
+void FAnimationBlueprintTemplateEditorMode::RegisterTabFactories(TSharedPtr<FTabManager> InTabManager)
+{
+	TSharedPtr<FBlueprintEditor> BP = MyBlueprintEditor.Pin();
+	
+	BP->RegisterToolbarTab(InTabManager.ToSharedRef());
+
+	// Mode-specific setup
+	BP->PushTabFactories(CoreTabFactories);
+	BP->PushTabFactories(BlueprintEditorTabFactories);
+	BP->PushTabFactories(TabFactories);
+}
diff --git a/Engine/Source/Editor/AnimationBlueprintEditor/Private/AnimationBlueprintTemplateEditorMode.h b/Engine/Source/Editor/AnimationBlueprintEditor/Private/AnimationBlueprintTemplateEditorMode.h
new file mode 100644
--- /dev/null
+++ b/Engine/Source/Editor/AnimationBlueprintEditor/Private/AnimationBlueprintTemplateEditorMode.h
@@ -0,0 +1,28 @@
+// Copyright Epic Games, Inc. All Rights Reserved.
+
+#pragma once
+
+#include "CoreMinimal.h"
+#include "WorkflowOrientedApp/WorkflowTabManager.h"
+#include "AnimationBlueprintEditor.h"
+#include "Animation/AnimBlueprint.h"
+#include "BlueprintEditorModes.h"
+
+class IPersonaPreviewScene;
+
+class FAnimationBlueprintTemplateEditorMode : public FBlueprintEditorApplicationMode
+{
+protected:
+	// Set of spawnable tabs in persona mode (@TODO: Multiple lists!)
+	FWorkflowAllowedTabSet TabFactories;
+
+public:
+	FAnimationBlueprintTemplateEditorMode(const TSharedRef<FAnimationBlueprintEditor>& InAnimationBlueprintEditor);
+
+	// FApplicationMode interface
+	virtual void RegisterTabFactories(TSharedPtr<FTabManager> InTabManager) override;
+	// End of FApplicationMode interface
+
+private:
+	TWeakObjectPtr<class UAnimBlueprint> AnimBlueprintPtr;
+};
diff --git a/Engine/Source/Editor/AnimationBlueprintEditor/Private/AnimationNodes/SAnimationGraphNode.cpp b/Engine/Source/Editor/AnimationBlueprintEditor/Private/AnimationNodes/SAnimationGraphNode.cpp
--- a/Engine/Source/Editor/AnimationBlueprintEditor/Private/AnimationNodes/SAnimationGraphNode.cpp
+++ b/Engine/Source/Editor/AnimationBlueprintEditor/Private/AnimationNodes/SAnimationGraphNode.cpp
@@ -16,9 +16,14 @@
 #include "SLevelOfDetailBranchNode.h"
 #include "Widgets/Layout/SSpacer.h"
 #include "AnimationGraphSchema.h"
+#include "BlueprintMemberReferenceCustomization.h"
 #include "SGraphPin.h"
 #include "Widgets/Layout/SWrapBox.h"
 #include "Brushes/SlateColorBrush.h"
+#include "PropertyEditorModule.h"
+#include "IPropertyRowGenerator.h"
+#include "IDetailTreeNode.h"
+#include "Widgets/Layout/SGridPanel.h"
 
 #define LOCTEXT_NAMESPACE "AnimationGraphNode"
 
@@ -132,7 +137,7 @@ void SAnimationGraphNode::Construct(const FArguments& InArgs, UAnimGraphNode_Bas
 
 	this->UpdateGraphNode();
 
-	ReconfigurePinWidgetsForPropertyBindings(CastChecked<UAnimGraphNode_Base>(GraphNode), [this](UEdGraphPin* InPin){ return FindWidgetForPin(InPin); });
+	ReconfigurePinWidgetsForPropertyBindings(CastChecked<UAnimGraphNode_Base>(GraphNode), SharedThis(this), [this](UEdGraphPin* InPin){ return FindWidgetForPin(InPin); });
 
 	const FSlateBrush* ImageBrush = FEditorStyle::Get().GetBrush(TEXT("Graph.AnimationFastPathIndicator"));
 
@@ -271,7 +276,118 @@ void SAnimationGraphNode::GetNodeInfoPopups(FNodeInfoContext* Context, TArray<FG
 	}
 }
 
-void SAnimationGraphNode::ReconfigurePinWidgetsForPropertyBindings(UAnimGraphNode_Base* InAnimGraphNode, TFunctionRef<TSharedPtr<SGraphPin>(UEdGraphPin*)> InFindWidgetForPin)
+void SAnimationGraphNode::CreateBelowPinControls(TSharedPtr<SVerticalBox> MainBox)
+{
+	if (UAnimGraphNode_Base* AnimNode = CastChecked<UAnimGraphNode_Base>(GraphNode, ECastCheckedType::NullAllowed))
+	{
+		FPropertyEditorModule& PropertyEditorModule = FModuleManager::Get().LoadModuleChecked<FPropertyEditorModule>("PropertyEditor");
+		PropertyRowGenerator = PropertyEditorModule.CreatePropertyRowGenerator(FPropertyRowGeneratorArgs());
+		PropertyRowGenerator->RegisterInstancedCustomPropertyTypeLayout(FMemberReference::StaticStruct()->GetFName(), FOnGetPropertyTypeCustomizationInstance::CreateStatic(&FBlueprintMemberReferenceDetails::MakeInstance));
+		PropertyRowGenerator->SetObjects({ AnimNode });
+
+		TSharedPtr<SGridPanel> GridPanel;
+		
+		MainBox->AddSlot()
+		.AutoHeight()
+		.Padding(4.0f)
+		[
+			SNew(SLevelOfDetailBranchNode)
+			.UseLowDetailSlot(this, &SAnimationGraphNode::UseLowDetailNodeTitles)
+			.LowDetail()
+			[
+				SNew(SSpacer)
+				.Size(FVector2D(17.0f, 17.f))
+			]
+			.HighDetail()
+			[
+				SAssignNew(GridPanel, SGridPanel)
+				.IsEnabled_Lambda([this](){ return IsNodeEditable(); })
+			]
+		];
+
+		GridPanel->SetVisibility(EVisibility::Collapsed);
+
+		int32 RowIndex = 0;
+		
+		// Add bound functions
+		auto AddFunctionBindingWidget = [this, &GridPanel, &RowIndex](FName InCategory, FName InMemberName)
+		{
+			GridPanel->SetVisibility(EVisibility::Visible);
+			
+			// Find row
+			TSharedPtr<IPropertyHandle> PropertyHandle;
+			TSharedPtr<IDetailTreeNode> DetailTreeNode;
+
+			for (const TSharedRef<IDetailTreeNode>& RootTreeNode : PropertyRowGenerator->GetRootTreeNodes())
+			{
+				if(RootTreeNode->GetNodeName() == InCategory)
+				{
+					TArray<TSharedRef<IDetailTreeNode>> Children;
+					RootTreeNode->GetChildren(Children);
+
+					for (int32 ChildIdx = 0; ChildIdx < Children.Num(); ChildIdx++)
+					{
+						TSharedPtr<IPropertyHandle> ChildPropertyHandle = Children[ChildIdx]->CreatePropertyHandle();
+						if (ChildPropertyHandle.IsValid() && ChildPropertyHandle->GetProperty() && ChildPropertyHandle->GetProperty()->GetFName() == InMemberName)
+						{
+							DetailTreeNode = Children[ChildIdx];
+							PropertyHandle = ChildPropertyHandle;
+							PropertyHandle->SetOnPropertyValueChanged(FSimpleDelegate::CreateLambda([this]()
+							{
+								GraphNode->ReconstructNode();
+							}));
+							break;
+						}
+					}
+				}
+			}
+
+			if(DetailTreeNode.IsValid() && PropertyHandle.IsValid())
+			{
+				DetailNodes.Add(DetailTreeNode);
+				
+				FNodeWidgets NodeWidgets = DetailTreeNode->CreateNodeWidgets();
+
+				GridPanel->AddSlot(0, RowIndex)
+				.HAlign(HAlign_Left)
+				.VAlign(VAlign_Center)
+				.Padding(10.0f, 2.0f, 2.0f, 2.0f)
+				[
+					NodeWidgets.NameWidget.ToSharedRef()
+				];
+
+				GridPanel->AddSlot(1, RowIndex)
+				.HAlign(HAlign_Left)
+				.VAlign(VAlign_Center)
+				.Padding(2.0f, 2.0f, 10.0f, 2.0f)
+				[
+					NodeWidgets.ValueWidget.ToSharedRef()
+				];
+
+				RowIndex++;
+			}
+		};
+
+		if(AnimNode->InitializeFunction.GetMemberGuid().IsValid())
+		{
+			AddFunctionBindingWidget("Functions", GET_MEMBER_NAME_CHECKED(UAnimGraphNode_Base, InitializeFunction));
+		}
+		if(AnimNode->BecomeRelevantFunction.GetMemberGuid().IsValid())
+		{
+			AddFunctionBindingWidget("Functions", GET_MEMBER_NAME_CHECKED(UAnimGraphNode_Base, BecomeRelevantFunction));
+		}
+		if(AnimNode->UpdateFunction.GetMemberGuid().IsValid())
+		{
+			AddFunctionBindingWidget("Functions", GET_MEMBER_NAME_CHECKED(UAnimGraphNode_Base, UpdateFunction));
+		}
+		if(AnimNode->EvaluateFunction.GetMemberGuid().IsValid())
+		{
+			AddFunctionBindingWidget("Functions", GET_MEMBER_NAME_CHECKED(UAnimGraphNode_Base, EvaluateFunction));
+		}
+	}
+}
+
+void SAnimationGraphNode::ReconfigurePinWidgetsForPropertyBindings(UAnimGraphNode_Base* InAnimGraphNode, TSharedRef<SGraphNode> InGraphNodeWidget, TFunctionRef<TSharedPtr<SGraphPin>(UEdGraphPin*)> InFindWidgetForPin)
 {
 	for(UEdGraphPin* Pin : InAnimGraphNode->Pins)
 	{
@@ -328,6 +444,10 @@ void SAnimationGraphNode::ReconfigurePinWidgetsForPropertyBindings(UAnimGraphNod
 					PinWidget->GetLabelAndValue()->AddSlot()
 					[
 						SNew(SBox)
+						.IsEnabled_Lambda([WeakWidget = TWeakPtr<SGraphNode>(InGraphNodeWidget)]()
+						{
+							return WeakWidget.IsValid() ? WeakWidget.Pin()->IsNodeEditable() : false;
+						})
 						.Visibility_Lambda([PinName, InAnimGraphNode]()
 						{
 							if (FAnimGraphNodePropertyBinding* BindingPtr = InAnimGraphNode->PropertyBindings.Find(PinName))
diff --git a/Engine/Source/Editor/AnimationBlueprintEditor/Private/AnimationNodes/SGraphNodeBlendSpaceGraph.cpp b/Engine/Source/Editor/AnimationBlueprintEditor/Private/AnimationNodes/SGraphNodeBlendSpaceGraph.cpp
--- a/Engine/Source/Editor/AnimationBlueprintEditor/Private/AnimationNodes/SGraphNodeBlendSpaceGraph.cpp
+++ b/Engine/Source/Editor/AnimationBlueprintEditor/Private/AnimationNodes/SGraphNodeBlendSpaceGraph.cpp
@@ -20,7 +20,7 @@ void SGraphNodeBlendSpaceGraph::Construct(const FArguments& InArgs, UAnimGraphNo
 
 	UpdateGraphNode();
 
-	SAnimationGraphNode::ReconfigurePinWidgetsForPropertyBindings(CastChecked<UAnimGraphNode_Base>(GraphNode), [this](UEdGraphPin* InPin){ return FindWidgetForPin(InPin); });
+	SAnimationGraphNode::ReconfigurePinWidgetsForPropertyBindings(CastChecked<UAnimGraphNode_Base>(GraphNode), SharedThis(this), [this](UEdGraphPin* InPin){ return FindWidgetForPin(InPin); });
 }
 
 UEdGraph* SGraphNodeBlendSpaceGraph::GetInnerGraph() const
diff --git a/Engine/Source/Editor/AnimationBlueprintEditor/Public/AnimationNodes/SAnimationGraphNode.h b/Engine/Source/Editor/AnimationBlueprintEditor/Public/AnimationNodes/SAnimationGraphNode.h
--- a/Engine/Source/Editor/AnimationBlueprintEditor/Public/AnimationNodes/SAnimationGraphNode.h
+++ b/Engine/Source/Editor/AnimationBlueprintEditor/Public/AnimationNodes/SAnimationGraphNode.h
@@ -11,6 +11,8 @@
 #include "Engine/PoseWatch.h"
 
 class UAnimGraphNode_Base;
+class IDetailTreeNode;
+class IPropertyRowGenerator;
 
 class ANIMATIONBLUEPRINTEDITOR_API SAnimationGraphNode : public SGraphNodeK2Base
 {
@@ -21,7 +23,7 @@ public:
 	void Construct(const FArguments& InArgs, UAnimGraphNode_Base* InNode);
 
 	// Tweak any created pin widgets so they respond to bindings
-	static void ReconfigurePinWidgetsForPropertyBindings(UAnimGraphNode_Base* InAnimGraphNode, TFunctionRef<TSharedPtr<SGraphPin>(UEdGraphPin*)> InFindWidgetForPin);
+	static void ReconfigurePinWidgetsForPropertyBindings(UAnimGraphNode_Base* InAnimGraphNode, TSharedRef<SGraphNode> InGraphNodeWidget, TFunctionRef<TSharedPtr<SGraphPin>(UEdGraphPin*)> InFindWidgetForPin);
 
 protected:
 	// SWidget interface
@@ -32,6 +34,7 @@ protected:
 	virtual TArray<FOverlayWidgetInfo> GetOverlayWidgets(bool bSelected, const FVector2D& WidgetSize) const override;
 	virtual TSharedRef<SWidget> CreateTitleWidget(TSharedPtr<SNodeTitle> InNodeTitle) override;
 	virtual void GetNodeInfoPopups(FNodeInfoContext* Context, TArray<FGraphInformationPopupInfo>& Popups) const override;
+	virtual void CreateBelowPinControls(TSharedPtr<SVerticalBox> MainBox) override;
 	// End of SGraphNode interface
 
 private:
@@ -53,4 +56,10 @@ private:
 	TSharedPtr<SNodeTitle> NodeTitle;
 
 	TWeakObjectPtr<class UPoseWatch> PoseWatch;
+
+	// Property row generator used to display function properties on nodes
+	TSharedPtr<IPropertyRowGenerator> PropertyRowGenerator;
+	
+	// Hold a reference to the root ptr of the details tree we use to display function properties
+	TArray<TSharedPtr<IDetailTreeNode>> DetailNodes;
 };
diff --git a/Engine/Source/Editor/EditorStyle/Private/StarshipStyle.cpp b/Engine/Source/Editor/EditorStyle/Private/StarshipStyle.cpp
--- a/Engine/Source/Editor/EditorStyle/Private/StarshipStyle.cpp
+++ b/Engine/Source/Editor/EditorStyle/Private/StarshipStyle.cpp
@@ -5127,6 +5127,22 @@ void FStarshipEditorStyle::FStyle::SetupPersonaStyle()
 		Set("SkeletalMeshEditor.MeshSectionSelection", new IMAGE_BRUSH_SVG("Starship/Persona/SectionSelection", Icon20x20));
 	}
 
+	// New anim blueprint dialog
+	{
+		Set("NewAnimBlueprintDialog.AreaBorder", new FSlateRoundedBoxBrush(FStyleColors::Panel, 4.0f));
+
+		const FCheckBoxStyle CheckBoxAreaStyle = FCheckBoxStyle()
+			.SetCheckBoxType(ESlateCheckBoxType::ToggleButton)
+			.SetUncheckedImage(FSlateRoundedBoxBrush(FStyleColors::Secondary, 4.0f))
+			.SetUncheckedHoveredImage(FSlateRoundedBoxBrush(FStyleColors::Hover, 4.0f))
+			.SetUncheckedPressedImage(FSlateRoundedBoxBrush(FStyleColors::PrimaryPress, 4.0f))
+			.SetCheckedImage(FSlateRoundedBoxBrush(FStyleColors::Transparent, 4.0f, FStyleColors::Primary, 1.0f))
+			.SetCheckedHoveredImage(FSlateRoundedBoxBrush(FStyleColors::Transparent, 4.0f, FStyleColors::PrimaryHover, 1.0f))
+			.SetCheckedPressedImage(FSlateRoundedBoxBrush(FStyleColors::Transparent, 4.0f, FStyleColors::PrimaryPress, 1.0f));
+
+		Set("NewAnimBlueprintDialog.CheckBoxArea", CheckBoxAreaStyle);
+	}
+
 	// Kismet 2
 	{
 		Set( "FullBlueprintEditor.SwitchToScriptingMode", new IMAGE_BRUSH_SVG( "Starship/Blueprints/icon_BlueprintEditor_EventGraph", Icon20x20));
diff --git a/Engine/Source/Editor/Kismet/Private/BlueprintDetailsCustomization.cpp b/Engine/Source/Editor/Kismet/Private/BlueprintDetailsCustomization.cpp
--- a/Engine/Source/Editor/Kismet/Private/BlueprintDetailsCustomization.cpp
+++ b/Engine/Source/Editor/Kismet/Private/BlueprintDetailsCustomization.cpp
@@ -6928,197 +6928,4 @@ TSharedRef<SWidget> FBlueprintDocumentationDetails::GenerateExcerptList()
 		];
 }
 
-void FBlueprintMemberReferenceDetails::CustomizeHeader(TSharedRef<IPropertyHandle> InStructPropertyHandle, FDetailWidgetRow& InHeaderRow, IPropertyTypeCustomizationUtils& InStructCustomizationUtils)
-{
-	if(IModularFeatures::Get().IsModularFeatureAvailable("PropertyAccessEditor"))
-	{
-		UBlueprint* Blueprint = MyBlueprint.Pin()->GetBlueprintObj();
-		if(Blueprint == nullptr)
-		{
-			return;
-		}
-
-		const bool bFunctionReference = InStructPropertyHandle->HasMetaData("FunctionReference");
-		const bool bPropertyReference = InStructPropertyHandle->HasMetaData("PropertyReference");
-		if(bPropertyReference || !bFunctionReference)
-		{
-			// Only function references are supported right now
-			return;
-		}
-		
-		// Try to find the prototype function
-		const FString& PrototypeFunctionName = InStructPropertyHandle->GetMetaData("PrototypeFunction");
-		UFunction* PrototypeFunction = PrototypeFunctionName.IsEmpty() ? nullptr : FindObject<UFunction>(nullptr, *PrototypeFunctionName);
-
-		auto OnGenerateBindingName = []() -> FString
-		{
-			return TEXT("NewFunction");
-		};
-		
-		auto OnCanBindProperty = [](FProperty* InProperty)
-		{
-			return true;
-		};
-		
-		auto OnCanBindFunction = [PrototypeFunction](UFunction* InFunction)
-		{
-			if(PrototypeFunction != nullptr)
-			{
-				return PrototypeFunction->IsSignatureCompatibleWith(InFunction)
-					&& FBlueprintEditorUtils::HasFunctionBlueprintThreadSafeMetaData(PrototypeFunction) == FBlueprintEditorUtils::HasFunctionBlueprintThreadSafeMetaData(InFunction);
-			}
-			
-			return false;
-		};
-
-		auto OnAddBinding = [InStructPropertyHandle, Blueprint](FName InPropertyName, const TArray<FBindingChainElement>& InBindingChain)
-		{
-			void* StructData = nullptr;
-			const FPropertyAccess::Result Result = InStructPropertyHandle->GetValueData(StructData);
-			if(Result == FPropertyAccess::Success)
-			{
-				InStructPropertyHandle->NotifyPreChange();
-				
-				check(StructData);
-				FMemberReference* MemberReference = static_cast<FMemberReference*>(StructData);
-				MemberReference->SetFromField<UFunction>(InBindingChain[0].Field.Get<UFunction>(), true);
-
-				InStructPropertyHandle->NotifyPostChange(EPropertyChangeType::ValueSet);
-			}
-		};
-
-		auto OnRemoveBinding = [InStructPropertyHandle](FName InPropertyName)
-		{
-			
-			
-			void* StructData = nullptr;
-			const FPropertyAccess::Result Result = InStructPropertyHandle->GetValueData(StructData);
-			if(Result == FPropertyAccess::Success)
-			{
-				InStructPropertyHandle->NotifyPreChange();
-				
-				check(StructData);
-				FMemberReference* MemberReference = static_cast<FMemberReference*>(StructData);
-				FGuid EmptyGuid;
-				MemberReference->SetSelfMember(NAME_None, EmptyGuid);
-
-				InStructPropertyHandle->NotifyPostChange(EPropertyChangeType::ValueSet);
-			}
-		};
-
-		auto CanRemoveBinding = [InStructPropertyHandle](FName InPropertyName)
-		{
-			void* StructData = nullptr;
-			const FPropertyAccess::Result Result = InStructPropertyHandle->GetValueData(StructData);
-			if(Result == FPropertyAccess::Success)
-			{
-				check(StructData);
-				FMemberReference* MemberReference = static_cast<FMemberReference*>(StructData);
-				return MemberReference->GetMemberName() != NAME_None;
-			}
-			
-			return false;
-		};
-
-		auto OnNewFunctionBindingCreated = [PrototypeFunction](UEdGraph* InFunctionGraph, UFunction* InFunction)
-		{
-			// Ensure newly created function is thread safe 
-			if(PrototypeFunction && FBlueprintEditorUtils::HasFunctionBlueprintThreadSafeMetaData(PrototypeFunction))
-			{
-				TArray<UK2Node_FunctionEntry*> EntryNodes;
-				InFunctionGraph->GetNodesOfClass<UK2Node_FunctionEntry>(EntryNodes);
-				if(EntryNodes.Num() > 0)
-				{
-					EntryNodes[0]->MetaData.bThreadSafe = true;
-				}
-			}
-		};
-
-		auto CurrentBindingText = [bFunctionReference, Blueprint, InStructPropertyHandle]()
-		{
-			void* StructData = nullptr;
-			const FPropertyAccess::Result Result = InStructPropertyHandle->GetValueData(StructData);
-			if(Result == FPropertyAccess::Success)
-			{
-				check(StructData);
-				FMemberReference* MemberReference = static_cast<FMemberReference*>(StructData);
-				if(bFunctionReference)
-				{
-					UFunction* Function = MemberReference->ResolveMember<UFunction>(Blueprint->SkeletonGeneratedClass);
-					if(Function)
-					{
-						return FText::FromName(Function->GetFName());
-					}
-					else
-					{
-						return FText::FromName(MemberReference->GetMemberName());
-					}
-				}
-			}
-			else if(Result == FPropertyAccess::MultipleValues)
-			{
-				return LOCTEXT("MultipleValues", "Multiple Values");
-			}
-
-			return FText::GetEmpty();
-		};
-
-		auto CurrentBindingToolTipText = [InStructPropertyHandle]()
-		{
-			void* StructData = nullptr;
-			const FPropertyAccess::Result Result = InStructPropertyHandle->GetValueData(StructData);
-			if(Result == FPropertyAccess::Success)
-			{
-				check(StructData);
-				FMemberReference* MemberReference = static_cast<FMemberReference*>(StructData);
-				return FText::FromName(MemberReference->GetMemberName());
-			}
-			else if(Result == FPropertyAccess::MultipleValues)
-			{
-				return LOCTEXT("MultipleValues", "Multiple Values");
-			}
-
-			return FText::GetEmpty();
-		};
-	
-		FPropertyBindingWidgetArgs Args;
-		Args.BindableSignature = PrototypeFunction;
-		Args.OnGenerateBindingName = FOnGenerateBindingName::CreateLambda(OnGenerateBindingName);
-		Args.OnCanBindProperty = FOnCanBindProperty::CreateLambda(OnCanBindProperty);
-		Args.OnCanBindFunction = FOnCanBindFunction::CreateLambda(OnCanBindFunction);
-		Args.OnCanBindToClass = FOnCanBindToClass::CreateLambda([](UClass* InClass){ return true; });
-		Args.OnAddBinding = FOnAddBinding::CreateLambda(OnAddBinding);
-		Args.OnRemoveBinding = FOnRemoveBinding::CreateLambda(OnRemoveBinding);
-		Args.OnCanRemoveBinding = FOnCanRemoveBinding::CreateLambda(CanRemoveBinding);
-		Args.OnNewFunctionBindingCreated = FOnNewFunctionBindingCreated::CreateLambda(OnNewFunctionBindingCreated);
-		Args.CurrentBindingText = MakeAttributeLambda(CurrentBindingText);
-		Args.CurrentBindingToolTipText = MakeAttributeLambda(CurrentBindingToolTipText);
-		Args.CurrentBindingImage = FEditorStyle::GetBrush("GraphEditor.Function_16x");
-		Args.CurrentBindingColor = FEditorStyle::GetSlateColor("Colors.Foreground").GetSpecifiedColor();
-		Args.bGeneratePureBindings = false;
-		Args.bAllowFunctionBindings = bFunctionReference;
-		Args.bAllowPropertyBindings = false;
-		Args.bAllowNewBindings = true;
-		Args.bAllowArrayElementBindings = false;
-		Args.bAllowUObjectFunctions = false;
-		Args.bAllowStructFunctions = false;
-		Args.bAllowStructMemberBindings = false;
-
-		IPropertyAccessEditor& PropertyAccessEditor = IModularFeatures::Get().GetModularFeature<IPropertyAccessEditor>("PropertyAccessEditor");
-		InHeaderRow
-		.NameContent()
-		[
-			InStructPropertyHandle->CreatePropertyNameWidget()
-		]
-		.ValueContent()
-		[
-			SNew(SBox)
-			.MaxDesiredWidth(200.0f)
-			[
-				PropertyAccessEditor.MakePropertyBindingWidget(Blueprint, Args)
-			]
-		];
-	}
-}
-
 #undef LOCTEXT_NAMESPACE
diff --git a/Engine/Source/Editor/Kismet/Private/BlueprintDetailsCustomization.h b/Engine/Source/Editor/Kismet/Private/BlueprintDetailsCustomization.h
--- a/Engine/Source/Editor/Kismet/Private/BlueprintDetailsCustomization.h
+++ b/Engine/Source/Editor/Kismet/Private/BlueprintDetailsCustomization.h
@@ -1048,25 +1048,3 @@ private:
 	TArray<TSharedPtr<FString>> ExcerptList;
 
 };
-
-class FBlueprintMemberReferenceDetails : public IPropertyTypeCustomization
-{
-public:
-	static TSharedRef<IPropertyTypeCustomization> MakeInstance(TWeakPtr<SMyBlueprint> InMyBlueprint)
-	{
-		return MakeShareable(new FBlueprintMemberReferenceDetails(InMyBlueprint));
-	}
-
-private:
-	FBlueprintMemberReferenceDetails(TWeakPtr<SMyBlueprint> InMyBlueprint)
-		: MyBlueprint(InMyBlueprint)
-	{
-	}
-	
-	// IPropertyTypeCustomization interface
-	virtual void CustomizeHeader(TSharedRef<IPropertyHandle> InStructPropertyHandle, FDetailWidgetRow& InHeaderRow, IPropertyTypeCustomizationUtils& InStructCustomizationUtils) override;
-	virtual void CustomizeChildren(TSharedRef<IPropertyHandle> InStructPropertyHandle, IDetailChildrenBuilder& InChildBuilder, IPropertyTypeCustomizationUtils& InStructCustomizationUtils) override {};
-
-private:
-	TWeakPtr<SMyBlueprint> MyBlueprint;
-};
diff --git a/Engine/Source/Editor/Kismet/Private/BlueprintMemberReferenceCustomization.cpp b/Engine/Source/Editor/Kismet/Private/BlueprintMemberReferenceCustomization.cpp
new file mode 100644
--- /dev/null
+++ b/Engine/Source/Editor/Kismet/Private/BlueprintMemberReferenceCustomization.cpp
@@ -0,0 +1,259 @@
+// Copyright Epic Games, Inc. All Rights Reserved.
+
+#include "BlueprintMemberReferenceCustomization.h"
+#include "IPropertyAccessEditor.h"
+#include "SMyBlueprint.h"
+#include "Engine/MemberReference.h"
+#include "Kismet2/BlueprintEditorUtils.h"
+#include "K2Node_FunctionEntry.h"
+#include "DetailWidgetRow.h"
+#include "Subsystems/AssetEditorSubsystem.h"
+
+#define LOCTEXT_NAMESPACE "BlueprintMemberReferenceCustomization"
+
+void FBlueprintMemberReferenceDetails::CustomizeHeader(TSharedRef<IPropertyHandle> InStructPropertyHandle, FDetailWidgetRow& InHeaderRow, IPropertyTypeCustomizationUtils& InStructCustomizationUtils)
+{
+	if(IModularFeatures::Get().IsModularFeatureAvailable("PropertyAccessEditor"))
+	{
+		UBlueprint* Blueprint = MyBlueprint.IsValid() ? MyBlueprint.Pin()->GetBlueprintObj() : nullptr;
+		if(Blueprint == nullptr)
+		{
+			// Try to get the BP from a node that we are member of
+			TArray<UObject*> OuterObjects;
+			InStructPropertyHandle->GetOuterObjects(OuterObjects);
+			if(OuterObjects.Num() > 0)
+			{
+				if(UK2Node* Node = Cast<UK2Node>(OuterObjects[0]))
+				{
+					Blueprint = Node->HasValidBlueprint() ? Node->GetBlueprint() : nullptr;
+				}
+			}
+		}
+
+		if(Blueprint == nullptr)
+		{
+			return;
+		}
+		
+		const bool bFunctionReference = InStructPropertyHandle->HasMetaData("FunctionReference");
+		const bool bPropertyReference = InStructPropertyHandle->HasMetaData("PropertyReference");
+		if(bPropertyReference || !bFunctionReference)
+		{
+			// Only function references are supported right now
+			return;
+		}
+		
+		// Try to find the prototype function
+		const FString& PrototypeFunctionName = InStructPropertyHandle->GetMetaData("PrototypeFunction");
+		UFunction* PrototypeFunction = PrototypeFunctionName.IsEmpty() ? nullptr : FindObject<UFunction>(nullptr, *PrototypeFunctionName);
+
+		auto OnGenerateBindingName = []() -> FString
+		{
+			return TEXT("NewFunction");
+		};
+		
+		auto OnCanBindProperty = [](FProperty* InProperty)
+		{
+			return true;
+		};
+
+		auto OnGoToBinding = [InStructPropertyHandle, Blueprint](FName InPropertyName)
+		{
+			void* StructData = nullptr;
+			const FPropertyAccess::Result Result = InStructPropertyHandle->GetValueData(StructData);
+			if(Result == FPropertyAccess::Success)
+			{
+				check(StructData);
+				FMemberReference* MemberReference = static_cast<FMemberReference*>(StructData);
+				UFunction* Function = MemberReference->ResolveMember<UFunction>(Blueprint->SkeletonGeneratedClass);
+				if(Function)
+				{
+					GEditor->GetEditorSubsystem<UAssetEditorSubsystem>()->OpenEditorForAsset(Blueprint);
+					
+					if(IBlueprintEditor* BlueprintEditor = static_cast<IBlueprintEditor*>(GEditor->GetEditorSubsystem<UAssetEditorSubsystem>()->FindEditorForAsset(Blueprint, true)))
+					{
+						BlueprintEditor->JumpToHyperlink(Function, false);
+						return true;
+					}
+				}
+			}
+
+			return false;
+		};
+		
+		auto OnCanGotoBinding = [InStructPropertyHandle](FName InPropertyName)
+		{
+			void* StructData = nullptr;
+			const FPropertyAccess::Result Result = InStructPropertyHandle->GetValueData(StructData);
+			if(Result == FPropertyAccess::Success)
+			{
+				check(StructData);
+				FMemberReference* MemberReference = static_cast<FMemberReference*>(StructData);
+				return MemberReference->GetMemberName() != NAME_None;
+			}
+			
+			return false;
+		};
+		
+		auto OnCanBindFunction = [PrototypeFunction](UFunction* InFunction)
+		{
+			if(PrototypeFunction != nullptr)
+			{
+				return PrototypeFunction->IsSignatureCompatibleWith(InFunction)
+					&& FBlueprintEditorUtils::HasFunctionBlueprintThreadSafeMetaData(PrototypeFunction) == FBlueprintEditorUtils::HasFunctionBlueprintThreadSafeMetaData(InFunction);
+			}
+			
+			return false;
+		};
+
+		auto OnAddBinding = [InStructPropertyHandle, Blueprint](FName InPropertyName, const TArray<FBindingChainElement>& InBindingChain)
+		{
+			void* StructData = nullptr;
+			const FPropertyAccess::Result Result = InStructPropertyHandle->GetValueData(StructData);
+			if(Result == FPropertyAccess::Success)
+			{
+				InStructPropertyHandle->NotifyPreChange();
+				
+				check(StructData);
+				FMemberReference* MemberReference = static_cast<FMemberReference*>(StructData);
+				MemberReference->SetFromField<UFunction>(InBindingChain[0].Field.Get<UFunction>(), true);
+
+				InStructPropertyHandle->NotifyPostChange(EPropertyChangeType::ValueSet);
+			}
+		};
+
+		auto OnRemoveBinding = [InStructPropertyHandle](FName InPropertyName)
+		{
+			void* StructData = nullptr;
+			const FPropertyAccess::Result Result = InStructPropertyHandle->GetValueData(StructData);
+			if(Result == FPropertyAccess::Success)
+			{
+				InStructPropertyHandle->NotifyPreChange();
+				
+				check(StructData);
+				FMemberReference* MemberReference = static_cast<FMemberReference*>(StructData);
+				FGuid EmptyGuid;
+				MemberReference->SetSelfMember(NAME_None, EmptyGuid);
+
+				InStructPropertyHandle->NotifyPostChange(EPropertyChangeType::ValueSet);
+			}
+		};
+
+		auto CanRemoveBinding = [InStructPropertyHandle](FName InPropertyName)
+		{
+			void* StructData = nullptr;
+			const FPropertyAccess::Result Result = InStructPropertyHandle->GetValueData(StructData);
+			if(Result == FPropertyAccess::Success)
+			{
+				check(StructData);
+				FMemberReference* MemberReference = static_cast<FMemberReference*>(StructData);
+				return MemberReference->GetMemberName() != NAME_None;
+			}
+			
+			return false;
+		};
+
+		auto OnNewFunctionBindingCreated = [PrototypeFunction](UEdGraph* InFunctionGraph, UFunction* InFunction)
+		{
+			// Ensure newly created function is thread safe 
+			if(PrototypeFunction && FBlueprintEditorUtils::HasFunctionBlueprintThreadSafeMetaData(PrototypeFunction))
+			{
+				TArray<UK2Node_FunctionEntry*> EntryNodes;
+				InFunctionGraph->GetNodesOfClass<UK2Node_FunctionEntry>(EntryNodes);
+				if(EntryNodes.Num() > 0)
+				{
+					EntryNodes[0]->MetaData.bThreadSafe = true;
+				}
+			}
+		};
+
+		auto CurrentBindingText = [bFunctionReference, Blueprint, InStructPropertyHandle]()
+		{
+			void* StructData = nullptr;
+			const FPropertyAccess::Result Result = InStructPropertyHandle->GetValueData(StructData);
+			if(Result == FPropertyAccess::Success)
+			{
+				check(StructData);
+				FMemberReference* MemberReference = static_cast<FMemberReference*>(StructData);
+				if(bFunctionReference)
+				{
+					UFunction* Function = MemberReference->ResolveMember<UFunction>(Blueprint->SkeletonGeneratedClass);
+					if(Function)
+					{
+						return FText::FromName(Function->GetFName());
+					}
+					else
+					{
+						return FText::FromName(MemberReference->GetMemberName());
+					}
+				}
+			}
+			else if(Result == FPropertyAccess::MultipleValues)
+			{
+				return LOCTEXT("MultipleValues", "Multiple Values");
+			}
+
+			return FText::GetEmpty();
+		};
+
+		auto CurrentBindingToolTipText = [InStructPropertyHandle]()
+		{
+			void* StructData = nullptr;
+			const FPropertyAccess::Result Result = InStructPropertyHandle->GetValueData(StructData);
+			if(Result == FPropertyAccess::Success)
+			{
+				check(StructData);
+				FMemberReference* MemberReference = static_cast<FMemberReference*>(StructData);
+				return FText::FromName(MemberReference->GetMemberName());
+			}
+			else if(Result == FPropertyAccess::MultipleValues)
+			{
+				return LOCTEXT("MultipleValues", "Multiple Values");
+			}
+
+			return FText::GetEmpty();
+		};
+	
+		FPropertyBindingWidgetArgs Args;
+		Args.BindableSignature = PrototypeFunction;
+		Args.OnGenerateBindingName = FOnGenerateBindingName::CreateLambda(OnGenerateBindingName);
+		Args.OnCanBindProperty = FOnCanBindProperty::CreateLambda(OnCanBindProperty);
+		Args.OnGotoBinding = FOnGotoBinding::CreateLambda(OnGoToBinding);
+		Args.OnCanGotoBinding = FOnCanGotoBinding::CreateLambda(OnCanGotoBinding);
+		Args.OnCanBindFunction = FOnCanBindFunction::CreateLambda(OnCanBindFunction);
+		Args.OnCanBindToClass = FOnCanBindToClass::CreateLambda([](UClass* InClass){ return true; });
+		Args.OnAddBinding = FOnAddBinding::CreateLambda(OnAddBinding);
+		Args.OnRemoveBinding = FOnRemoveBinding::CreateLambda(OnRemoveBinding);
+		Args.OnCanRemoveBinding = FOnCanRemoveBinding::CreateLambda(CanRemoveBinding);
+		Args.OnNewFunctionBindingCreated = FOnNewFunctionBindingCreated::CreateLambda(OnNewFunctionBindingCreated);
+		Args.CurrentBindingText = MakeAttributeLambda(CurrentBindingText);
+		Args.CurrentBindingToolTipText = MakeAttributeLambda(CurrentBindingToolTipText);
+		Args.CurrentBindingImage = FEditorStyle::GetBrush("GraphEditor.Function_16x");
+		Args.CurrentBindingColor = FEditorStyle::GetSlateColor("Colors.Foreground").GetSpecifiedColor();
+		Args.bGeneratePureBindings = false;
+		Args.bAllowFunctionBindings = bFunctionReference;
+		Args.bAllowPropertyBindings = false;
+		Args.bAllowNewBindings = true;
+		Args.bAllowArrayElementBindings = false;
+		Args.bAllowUObjectFunctions = false;
+		Args.bAllowStructFunctions = false;
+		Args.bAllowStructMemberBindings = false;
+
+		IPropertyAccessEditor& PropertyAccessEditor = IModularFeatures::Get().GetModularFeature<IPropertyAccessEditor>("PropertyAccessEditor");
+		InHeaderRow
+		.NameContent()
+		[
+			InStructPropertyHandle->CreatePropertyNameWidget()
+		]
+		.ValueContent()
+		[
+			SNew(SBox)
+			.MaxDesiredWidth(200.0f)
+			[
+				PropertyAccessEditor.MakePropertyBindingWidget(Blueprint, Args)
+			]
+		];
+	}
+}
+
+#undef LOCTEXT_NAMESPACE
diff --git a/Engine/Source/Editor/Kismet/Private/SKismetInspector.cpp b/Engine/Source/Editor/Kismet/Private/SKismetInspector.cpp
--- a/Engine/Source/Editor/Kismet/Private/SKismetInspector.cpp
+++ b/Engine/Source/Editor/Kismet/Private/SKismetInspector.cpp
@@ -41,6 +41,7 @@
 #include "BlueprintDetailsCustomization.h"
 #include "K2Node_BitmaskLiteral.h"
 #include "BitmaskLiteralDetails.h"
+#include "BlueprintMemberReferenceCustomization.h"
 #include "FormatTextDetails.h"
 #include "Kismet2/BlueprintEditorUtils.h"
 #include "ClassViewerFilter.h"
diff --git a/Engine/Source/Editor/Kismet/Public/BlueprintMemberReferenceCustomization.h b/Engine/Source/Editor/Kismet/Public/BlueprintMemberReferenceCustomization.h
new file mode 100644
--- /dev/null
+++ b/Engine/Source/Editor/Kismet/Public/BlueprintMemberReferenceCustomization.h
@@ -0,0 +1,36 @@
+// Copyright Epic Games, Inc. All Rights Reserved.
+
+#pragma once
+
+#include "CoreMinimal.h"
+#include "IPropertyTypeCustomization.h"
+
+class SMyBlueprint;
+
+// Property type customization for FMemberReference
+class KISMET_API FBlueprintMemberReferenceDetails : public IPropertyTypeCustomization
+{
+public:
+	static TSharedRef<IPropertyTypeCustomization> MakeInstance()
+	{
+		return MakeShareable(new FBlueprintMemberReferenceDetails(TWeakPtr<SMyBlueprint>()));
+	}
+	
+	static TSharedRef<IPropertyTypeCustomization> MakeInstance(TWeakPtr<SMyBlueprint> InMyBlueprint)
+	{
+		return MakeShareable(new FBlueprintMemberReferenceDetails(InMyBlueprint));
+	}
+
+private:
+	FBlueprintMemberReferenceDetails(TWeakPtr<SMyBlueprint> InMyBlueprint)
+		: MyBlueprint(InMyBlueprint)
+	{
+	}
+	
+	// IPropertyTypeCustomization interface
+	virtual void CustomizeHeader(TSharedRef<IPropertyHandle> InStructPropertyHandle, FDetailWidgetRow& InHeaderRow, IPropertyTypeCustomizationUtils& InStructCustomizationUtils) override;
+	virtual void CustomizeChildren(TSharedRef<IPropertyHandle> InStructPropertyHandle, IDetailChildrenBuilder& InChildBuilder, IPropertyTypeCustomizationUtils& InStructCustomizationUtils) override {};
+
+private:
+	TWeakPtr<SMyBlueprint> MyBlueprint;
+};
diff --git a/Engine/Source/Editor/Persona/Private/AnimGraphNodeDetails.cpp b/Engine/Source/Editor/Persona/Private/AnimGraphNodeDetails.cpp
--- a/Engine/Source/Editor/Persona/Private/AnimGraphNodeDetails.cpp
+++ b/Engine/Source/Editor/Persona/Private/AnimGraphNodeDetails.cpp
@@ -992,35 +992,38 @@ TSharedRef<SWidget> SParentPlayerTreeRow::GenerateWidgetForColumn(const FName& C
 				]
 				
 			];
-		
-		TArray<const UClass*> AllowedClasses;
-		AllowedClasses.Add(UAnimationAsset::StaticClass());
-		HorizBox->AddSlot()
-			.VAlign(VAlign_Center)
-			.FillWidth(1.f)
-			[
-				SNew(SObjectPropertyEntryBox)
-				.ObjectPath(this, &SParentPlayerTreeRow::GetCurrentAssetPath)
-				.OnShouldFilterAsset(this, &SParentPlayerTreeRow::OnShouldFilterAsset)
-				.OnObjectChanged(this, &SParentPlayerTreeRow::OnAssetSelected)
-				.AllowedClass(GetCurrentAssetToUse()->GetClass())
-			];
 
-		HorizBox->AddSlot()
-			.VAlign(VAlign_Center)
-			.AutoWidth()
-			[
-				SNew(SButton)
-				.ButtonStyle(FEditorStyle::Get(), "NoBorder")
-				.Visibility(this, &SParentPlayerTreeRow::GetResetToDefaultVisibility)
-				.OnClicked(this, &SParentPlayerTreeRow::OnResetButtonClicked)
-				.ToolTip(IDocumentation::Get()->CreateToolTip(LOCTEXT("ResetToParentButtonTip", "Undo the override, returning to the default asset for this node"), NULL, "Shared/Editors/Persona", "ResetToParentButton"))
-				.Content()
+		if(GraphNode)
+		{
+			TArray<const UClass*> AllowedClasses;
+			AllowedClasses.Add(UAnimationAsset::StaticClass());
+			HorizBox->AddSlot()
+				.VAlign(VAlign_Center)
+				.FillWidth(1.f)
 				[
-					SNew(SImage)
-					.Image(FEditorStyle::GetBrush("PropertyWindow.DiffersFromDefault"))
-				]
-			];
+					SNew(SObjectPropertyEntryBox)
+					.ObjectPath(this, &SParentPlayerTreeRow::GetCurrentAssetPath)
+					.OnShouldFilterAsset(this, &SParentPlayerTreeRow::OnShouldFilterAsset)
+					.OnObjectChanged(this, &SParentPlayerTreeRow::OnAssetSelected)
+					.AllowedClass(GraphNode->GetAnimationAssetClass())
+				];
+
+			HorizBox->AddSlot()
+				.VAlign(VAlign_Center)
+				.AutoWidth()
+				[
+					SNew(SButton)
+					.ButtonStyle(FEditorStyle::Get(), "NoBorder")
+					.Visibility(this, &SParentPlayerTreeRow::GetResetToDefaultVisibility)
+					.OnClicked(this, &SParentPlayerTreeRow::OnResetButtonClicked)
+					.ToolTip(IDocumentation::Get()->CreateToolTip(LOCTEXT("ResetToParentButtonTip", "Undo the override, returning to the default asset for this node"), NULL, "Shared/Editors/Persona", "ResetToParentButton"))
+					.Content()
+					[
+						SNew(SImage)
+						.Image(FEditorStyle::GetBrush("PropertyWindow.DiffersFromDefault"))
+					]
+				];
+		}
 	}
 
 	return HorizBox.ToSharedRef();
@@ -1028,8 +1031,8 @@ TSharedRef<SWidget> SParentPlayerTreeRow::GenerateWidgetForColumn(const FName& C
 
 bool SParentPlayerTreeRow::OnShouldFilterAsset(const FAssetData& AssetData)
 {
-	const USkeleton* CurrentSkeleton = GraphNode->GetAnimBlueprint()->TargetSkeleton;
-	return !CurrentSkeleton->IsCompatibleSkeletonByAssetData(AssetData);
+	const USkeleton* CurrentSkeleton = CastChecked<UAnimBlueprint>(BlueprintEditor.Pin()->GetBlueprintObj())->TargetSkeleton;
+	return CurrentSkeleton != nullptr && !CurrentSkeleton->IsCompatibleSkeletonByAssetData(AssetData);
 }
 
 void SParentPlayerTreeRow::OnAssetSelected(const FAssetData& AssetData)
@@ -1070,7 +1073,7 @@ const UAnimationAsset* SParentPlayerTreeRow::GetCurrentAssetToUse() const
 		return GraphNode->GetAnimationAsset();
 	}
 
-	return NULL;
+	return nullptr;
 }
 
 EVisibility SParentPlayerTreeRow::GetResetToDefaultVisibility() const
diff --git a/Engine/Source/Editor/Persona/Private/PersonaAssetFamily.cpp b/Engine/Source/Editor/Persona/Private/PersonaAssetFamily.cpp
--- a/Engine/Source/Editor/Persona/Private/PersonaAssetFamily.cpp
+++ b/Engine/Source/Editor/Persona/Private/PersonaAssetFamily.cpp
@@ -411,7 +411,7 @@ void FPersonaAssetFamily::FindCounterpartAssets(const UObject* InAsset, const US
 		const UAnimBlueprint* AnimBlueprint = CastChecked<const UAnimBlueprint>(InAsset);
 		OutSkeleton = AnimBlueprint->TargetSkeleton;
 		OutMesh = AnimBlueprint->GetPreviewMesh();
-		check(AnimBlueprint->BlueprintType == BPTYPE_Interface || AnimBlueprint->TargetSkeleton != nullptr);
+		check(AnimBlueprint->BlueprintType == BPTYPE_Interface || AnimBlueprint->bIsTemplate || AnimBlueprint->TargetSkeleton != nullptr);
 		if(OutMesh == nullptr && AnimBlueprint->TargetSkeleton)
 		{
 			OutMesh = AnimBlueprint->TargetSkeleton->GetPreviewMesh();
diff --git a/Engine/Source/Editor/UnrealEd/Classes/Factories/AnimBlueprintFactory.h b/Engine/Source/Editor/UnrealEd/Classes/Factories/AnimBlueprintFactory.h
--- a/Engine/Source/Editor/UnrealEd/Classes/Factories/AnimBlueprintFactory.h
+++ b/Engine/Source/Editor/UnrealEd/Classes/Factories/AnimBlueprintFactory.h
@@ -34,6 +34,10 @@ class UAnimBlueprintFactory : public UFactory
 	UPROPERTY(EditAnywhere, Category=AnimBlueprintFactory)
 	TObjectPtr<class USkeletalMesh> PreviewSkeletalMesh;
 
+	// Whether the created blueprint should be a template with no target skeleton
+	UPROPERTY(EditAnywhere, Category=AnimBlueprintFactory)
+	bool bTemplate;
+	
 	//~ Begin UFactory Interface
 	virtual bool ConfigureProperties() override;
 	virtual UObject* FactoryCreateNew(UClass* Class, UObject* InParent, FName Name, EObjectFlags Flags, UObject* Context, FFeedbackContext* Warn, FName CallingContext) override;
diff --git a/Engine/Source/Editor/UnrealEd/Private/Animation/EditorParentPlayerListObj.cpp b/Engine/Source/Editor/UnrealEd/Private/Animation/EditorParentPlayerListObj.cpp
--- a/Engine/Source/Editor/UnrealEd/Private/Animation/EditorParentPlayerListObj.cpp
+++ b/Engine/Source/Editor/UnrealEd/Private/Animation/EditorParentPlayerListObj.cpp
@@ -53,7 +53,7 @@ void UEditorParentPlayerListObj::InitialiseFromBlueprint(UAnimBlueprint* Bluepri
 			for(UEdGraphNode* Node : Graph->Nodes)
 			{
 				UAnimGraphNode_Base* AnimNode = Cast<UAnimGraphNode_Base>(Node);
-				if(AnimNode && AnimNode->GetAnimationAsset())
+				if(AnimNode && (AnimNode->GetAnimationAsset() || AnimNode->GetAnimationAssetClass()))
 				{
 					FAnimParentNodeAssetOverride& Override = AddOverridableNode(AnimNode);
 
diff --git a/Engine/Source/Editor/UnrealEd/Private/AnimationEditorUtils.cpp b/Engine/Source/Editor/UnrealEd/Private/AnimationEditorUtils.cpp
--- a/Engine/Source/Editor/UnrealEd/Private/AnimationEditorUtils.cpp
+++ b/Engine/Source/Editor/UnrealEd/Private/AnimationEditorUtils.cpp
@@ -952,6 +952,46 @@ namespace AnimationEditorUtils
 		}
 #endif
 	}
+
+	void SetupDebugLinkedAnimInstances(UAnimBlueprint* InAnimBlueprint, UObject* InRootObjectBeingDebugged)
+	{
+		check(IsInGameThread());
+		
+		static bool bSettingDebugInstances = false;
+
+		if(!bSettingDebugInstances)
+		{
+			TGuardValue<bool> GuardValue(bSettingDebugInstances, true);
+			if(InRootObjectBeingDebugged)
+			{
+				if(const USkeletalMeshComponent* Component = Cast<USkeletalMeshComponent>(InRootObjectBeingDebugged->GetOuter()))
+				{
+					// See if we have any linked instances
+					for(UAnimInstance* LinkedInstance : Component->GetLinkedAnimInstances())
+					{
+						if(UAnimBlueprint* LinkedAnimBlueprint = Cast<UAnimBlueprint>(LinkedInstance->GetClass()->ClassGeneratedBy))
+						{
+							LinkedAnimBlueprint->SetObjectBeingDebugged(LinkedInstance);
+						}
+					}
+				}
+			}
+			else if(UObject* OldDebuggedObject = InAnimBlueprint->GetObjectBeingDebugged())
+			{
+				if(const USkeletalMeshComponent* Component = Cast<USkeletalMeshComponent>(OldDebuggedObject->GetOuter()))
+				{
+					// See if we have any linked instances
+					for(UAnimInstance* LinkedInstance : Component->GetLinkedAnimInstances())
+					{
+						if(UAnimBlueprint* LinkedAnimBlueprint = Cast<UAnimBlueprint>(LinkedInstance->GetClass()->ClassGeneratedBy))
+						{
+							LinkedAnimBlueprint->SetObjectBeingDebugged(nullptr);
+						}
+					}
+				}
+			}
+		}
+	}
 }
 
 #undef LOCTEXT_NAMESPACE
diff --git a/Engine/Source/Editor/UnrealEd/Private/Factories/AnimBlueprintFactory.cpp b/Engine/Source/Editor/UnrealEd/Private/Factories/AnimBlueprintFactory.cpp
--- a/Engine/Source/Editor/UnrealEd/Private/Factories/AnimBlueprintFactory.cpp
+++ b/Engine/Source/Editor/UnrealEd/Private/Factories/AnimBlueprintFactory.cpp
@@ -35,6 +35,8 @@
 #include "ClassViewerModule.h"
 #include "ClassViewerFilter.h"
 #include "Animation/AnimLayerInterface.h"
+#include "Widgets/Input/SSegmentedControl.h"
+#include "Widgets/Layout/SWidgetSwitcher.h"
 
 #define LOCTEXT_NAMESPACE "AnimBlueprintFactory"
 
@@ -44,7 +46,7 @@ static bool CanCreateAnimBlueprint(const FAssetData& Skeleton, UClass const * Pa
 	{
 		if (UAnimBlueprintGeneratedClass const * GeneratedParent = Cast<const UAnimBlueprintGeneratedClass>(ParentClass))
 		{
-			if (Skeleton.GetExportTextName() != FAssetData(GeneratedParent->GetTargetSkeleton()).GetExportTextName())
+			if (GeneratedParent->GetTargetSkeleton() != nullptr && Skeleton.GetExportTextName() != FAssetData(GeneratedParent->GetTargetSkeleton()).GetExportTextName())
 			{
 				return false;
 			}
@@ -73,35 +75,89 @@ public:
 		[
 			SNew(SBorder)
 			.Visibility(EVisibility::Visible)
-			.BorderImage(FEditorStyle::GetBrush("Menu.Background"))
+			.BorderImage(FEditorStyle::GetBrush("ChildWindow.Background"))
 			[
 				SNew(SBox)
 				.Visibility(EVisibility::Visible)
-				.WidthOverride(500.0f)
+				.WidthOverride(800.0f)
+				.HeightOverride(500.0f)
 				[
 					SNew(SVerticalBox)
 					+SVerticalBox::Slot()
-					.FillHeight(1)
+					.FillHeight(1.0f)
 					[
-						SNew(SBorder)
-						.BorderImage( FEditorStyle::GetBrush("ToolPanel.GroupBorder") )
-						.Content()
+						SNew(SHorizontalBox)
+						+SHorizontalBox::Slot()
+						.FillWidth(0.5f)
+						.Padding(10.0f)
 						[
-							SAssignNew(ParentClassContainer, SVerticalBox)
+							SNew(SBorder)
+							.BorderImage(FEditorStyle::GetBrush("NewAnimBlueprintDialog.AreaBorder"))
+							[
+								MakeParentClassPicker()
+							]
 						]
-					]
-					+SVerticalBox::Slot()
-					.FillHeight(1)
-					.Padding(0.0f, 10.0f, 0.0f, 0.0f)
-					[
-						SNew(SBorder)
-						.BorderImage( FEditorStyle::GetBrush("ToolPanel.GroupBorder") )
-						.Content()
+						+SHorizontalBox::Slot()
+						.FillWidth(0.5f)
+						.Padding(10.0f)
 						[
-							SAssignNew(SkeletonContainer, SVerticalBox)
+							SNew(SBorder)
+							.BorderImage(FEditorStyle::GetBrush("NewAnimBlueprintDialog.AreaBorder"))
+							[	
+								SNew(SVerticalBox)
+								+SVerticalBox::Slot()
+								.AutoHeight()
+								.Padding(0.0f, 0.0f, 0.0f, 5.0f)
+								[
+									SNew(SSegmentedControl<bool>)
+									.UniformPadding(FMargin(25.0f,5.0f))
+									.OnValueChanged_Lambda([this](bool bInNewValue)
+									{
+										bTemplate = bInNewValue;
+										if(bTemplate)
+										{
+											TargetSkeleton = FAssetData(); 
+										}
+										RefreshSkeletonPicker();
+									})
+									.Value_Lambda([this](){ return bTemplate; })	
+									+SSegmentedControl<bool>::Slot(false)
+									[
+										SNew(STextBlock)
+										.Text(LOCTEXT("SpecificSkeleton", "SPECIFIC SKELETON"))
+										.TextStyle( FEditorStyle::Get(), "NormalText" )
+									]
+									+SSegmentedControl<bool>::Slot(true)
+									[
+										SNew(STextBlock)
+										.Text(LOCTEXT("Template", "TEMPLATE"))
+										.TextStyle(FEditorStyle::Get(), "NormalText")
+									]
+								]
+								+SVerticalBox::Slot()
+								.FillHeight(1.0f)
+								[
+									SNew(SWidgetSwitcher)
+									.WidgetIndex_Lambda([this](){ return bTemplate ? 1 : 0; })
+									+SWidgetSwitcher::Slot()
+									[
+										MakeSkeletonPickerArea()
+									]
+									+SWidgetSwitcher::Slot()
+									.Padding(5.0f)
+									.VAlign(VAlign_Center)
+									[
+										SNew(STextBlock)
+										.Justification(ETextJustify::Center)
+										.AutoWrapText(true)
+										.Text(LOCTEXT("TemplateDesc", "A Template Animation Blueprint has no target skeleton.\nTemplates cannot have direct references to animation assets placed inside of their animation graphs."))
+										.TextStyle(FEditorStyle::Get(), "NormalText")
+									]
+								]
+							]
 						]
 					]
-
+					
 					// Ok/Cancel buttons
 					+SVerticalBox::Slot()
 					.AutoHeight()
@@ -116,10 +172,14 @@ public:
 						+SUniformGridPanel::Slot(0,0)
 						[
 							SNew(SButton)
+							.IsEnabled_Lambda([this]()
+							{
+								return ParentClass.Get() != nullptr && (bTemplate || TargetSkeleton.IsValid());
+							})
 							.HAlign(HAlign_Center)
 							.ContentPadding( FEditorStyle::GetMargin("StandardDialog.ContentPadding") )
-							.OnClicked(this, &SAnimBlueprintCreateDialog::OkClicked)	
-							.Text(LOCTEXT("CreateAnimBlueprintOk", "OK"))
+							.OnClicked(this, &SAnimBlueprintCreateDialog::OkClicked)
+							.Text(LOCTEXT("CreateAnimBlueprintCreate", "Create"))
 						]
 						+SUniformGridPanel::Slot(1,0)
 						[
@@ -134,8 +194,7 @@ public:
 			]
 		];
 
-		MakeParentClassPicker();
-		MakeSkeletonPicker();
+		RefreshSkeletonPicker();
 	}
 	
 	/** Sets properties for the supplied AnimBlueprintFactory */
@@ -145,7 +204,7 @@ public:
 
 		TSharedRef<SWindow> Window = SNew(SWindow)
 		.Title( LOCTEXT("CreateAnimBlueprintOptions", "Create Animation Blueprint") )
-		.ClientSize(FVector2D(400, 700))
+		.ClientSize(FVector2D(800, 500))
 		.SupportsMinimize(false) .SupportsMaximize(false)
 		[
 			AsShared()
@@ -165,18 +224,11 @@ private:
 	public:
 		/** All children of these classes will be included unless filtered out by another setting. */
 		TSet< const UClass* > AllowedChildrenOfClasses;
-		const FAssetData& ShouldBeCompatibleWithSkeleton;
-
-		FAnimBlueprintParentFilter(const FAssetData& Skeleton) : ShouldBeCompatibleWithSkeleton(Skeleton) {}
 
 		virtual bool IsClassAllowed(const FClassViewerInitializationOptions& InInitOptions, const UClass* InClass, TSharedRef< FClassViewerFilterFuncs > InFilterFuncs ) override
 		{
 			// If it appears on the allowed child-of classes list (or there is nothing on that list)
-			if (InFilterFuncs->IfInChildOfClassesSet( AllowedChildrenOfClasses, InClass) != EFilterReturn::Failed)
-			{
-				return CanCreateAnimBlueprint(ShouldBeCompatibleWithSkeleton, InClass);
-			}
-			return false;
+			return InFilterFuncs->IfInChildOfClassesSet( AllowedChildrenOfClasses, InClass) != EFilterReturn::Failed;
 		}
 
 		virtual bool IsUnloadedClassAllowed(const FClassViewerInitializationOptions& InInitOptions, const TSharedRef< const IUnloadedBlueprintData > InUnloadedClassData, TSharedRef< FClassViewerFilterFuncs > InFilterFuncs) override
@@ -187,7 +239,7 @@ private:
 	};
 
 	/** Creates the combo menu for the parent class */
-	void MakeParentClassPicker()
+	TSharedRef<SWidget> MakeParentClassPicker()
 	{
 		// Load the classviewer module to display a class picker
 		FClassViewerModule& ClassViewerModule = FModuleManager::LoadModuleChecked<FClassViewerModule>("ClassViewer");
@@ -195,40 +247,45 @@ private:
 		// Fill in options
 		FClassViewerInitializationOptions Options;
 		Options.Mode = EClassViewerMode::ClassPicker;
+		Options.DisplayMode = EClassViewerDisplayMode::TreeView;
+		Options.InitiallySelectedClass = UAnimInstance::StaticClass();
 		
 		// Only allow parenting to base blueprints.
 		Options.bIsBlueprintBaseOnly = true;
 
-		TSharedPtr<FAnimBlueprintParentFilter> Filter = MakeShareable(new FAnimBlueprintParentFilter(TargetSkeleton));
+		TSharedPtr<FAnimBlueprintParentFilter> Filter = MakeShared<FAnimBlueprintParentFilter>();
 		Options.ClassFilters.Add(Filter.ToSharedRef());
 
 		// All child child classes of UAnimInstance are valid.
 		Filter->AllowedChildrenOfClasses.Add(UAnimInstance::StaticClass());
 
-		ParentClassContainer->ClearChildren();
-		ParentClassContainer->AddSlot()
-		.AutoHeight()
-		[
-			SNew( STextBlock )
-			.Text( LOCTEXT("ParentClass", "Parent Class:") )
-			.ShadowOffset( FVector2D(1.0f, 1.0f) )
-		];
-
-		ParentClassContainer->AddSlot()
-		[
-			ClassViewerModule.CreateClassViewer(Options, FOnClassPicked::CreateSP(this, &SAnimBlueprintCreateDialog::OnClassPicked))
-		];
+		return SNew(SVerticalBox)
+			+SVerticalBox::Slot()
+			.AutoHeight()
+			.Padding(5.0f)
+			.HAlign(HAlign_Center)
+			[
+				SNew(STextBlock)
+				.Text(LOCTEXT("ParentClass", "PARENT CLASS"))
+				.TextStyle( FEditorStyle::Get(), "NormalText" )
+			]
+			+SVerticalBox::Slot()
+			.FillHeight(1.0f)
+			.Padding(5.0f)
+			[
+				ClassViewerModule.CreateClassViewer(Options, FOnClassPicked::CreateSP(this, &SAnimBlueprintCreateDialog::OnClassPicked))
+			];
 	}
 
 	/** Handler for when a parent class is selected */
 	void OnClassPicked(UClass* ChosenClass)
 	{
 		ParentClass = ChosenClass;
-		MakeSkeletonPicker();
+		RefreshSkeletonPicker();
 	}
 
 	/** Creates the combo menu for the target skeleton */
-	void MakeSkeletonPicker()
+	void RefreshSkeletonPicker()
 	{
 		FContentBrowserModule& ContentBrowserModule = FModuleManager::Get().LoadModuleChecked<FContentBrowserModule>(TEXT("ContentBrowser"));
 
@@ -236,24 +293,36 @@ private:
 		AssetPickerConfig.Filter.ClassNames.Add(USkeleton::StaticClass()->GetFName());
 		AssetPickerConfig.OnAssetSelected = FOnAssetSelected::CreateSP(this, &SAnimBlueprintCreateDialog::OnSkeletonSelected);
 		AssetPickerConfig.OnShouldFilterAsset = FOnShouldFilterAsset::CreateSP(this, &SAnimBlueprintCreateDialog::FilterSkeletonBasedOnParentClass);
-		AssetPickerConfig.bAllowNullSelection = true;
+		AssetPickerConfig.bAllowNullSelection = false;
 		AssetPickerConfig.InitialAssetViewType = EAssetViewType::Column;
 		AssetPickerConfig.InitialAssetSelection = TargetSkeleton;
+		AssetPickerConfig.HiddenColumnNames =
+		{
+			"DiskSize",
+			"AdditionalPreviewSkeletalMeshes",
+			"PreviewSkeletalMesh"
+		};
+		AssetPickerConfig.bShowPathInColumnView = false;
+		AssetPickerConfig.bShowTypeInColumnView = false;
 
 		SkeletonContainer->ClearChildren();
-		SkeletonContainer->AddSlot()
-		.AutoHeight()
-		[
-			SNew( STextBlock )
-			.Text( LOCTEXT("TargetSkeleton", "Target Skeleton:") )
-			.ShadowOffset( FVector2D(1.0f, 1.0f) )
-		];
-
 		SkeletonContainer->AddSlot()
 		[
 			ContentBrowserModule.Get().CreateAssetPicker(AssetPickerConfig)
 		];
 	}
+	
+	/** Creates the widgets for the target skeleton area */
+	TSharedRef<SWidget> MakeSkeletonPickerArea()
+	{
+		return SNew(SVerticalBox)
+			+SVerticalBox::Slot()
+			.FillHeight(1.0f)
+			.Padding(5.0f)
+			[
+				SAssignNew(SkeletonContainer, SVerticalBox)
+			];
+	}
 
 	bool FilterSkeletonBasedOnParentClass(const FAssetData& AssetData)
 	{
@@ -273,16 +342,10 @@ private:
 		{
 			AnimBlueprintFactory->BlueprintType = BPTYPE_Normal;
 			AnimBlueprintFactory->ParentClass = ParentClass.Get();
+			AnimBlueprintFactory->bTemplate = bTemplate;
 			AnimBlueprintFactory->TargetSkeleton = Cast<USkeleton>(TargetSkeleton.GetAsset());
 		}
 
-		if ( !TargetSkeleton.IsValid() )
-		{
-			// if TargetSkeleton is not valid
-			FMessageDialog::Open( EAppMsgType::Ok, LOCTEXT("NeedValidSkeleton", "Must specify a valid skeleton for the Anim Blueprint to target."));			
-			return FReply::Handled();
-		}
-
 		if (! CanCreateAnimBlueprint(TargetSkeleton, ParentClass.Get()))
 		{
 			FMessageDialog::Open( EAppMsgType::Ok, LOCTEXT("NeedCompatibleSkeleton", "Selected skeleton has to be compatible with selected parent class."));
@@ -338,6 +401,9 @@ private:
 	/** The selected skeleton */
 	FAssetData TargetSkeleton;
 
+	/** Whether we have a template selected or not */
+	bool bTemplate = false;
+	
 	/** True if Ok was clicked */
 	bool bOkClicked;
 };
@@ -367,13 +433,6 @@ UObject* UAnimBlueprintFactory::FactoryCreateNew(UClass* Class, UObject* InParen
 	// Make sure we are trying to factory a Anim Blueprint, then create and init one
 	check(Class->IsChildOf(UAnimBlueprint::StaticClass()));
 
-	// If they selected an interface, we dont need a target skeleton
-	if (BlueprintType != BPTYPE_Interface && TargetSkeleton == nullptr)
-	{
-		FMessageDialog::Open( EAppMsgType::Ok, LOCTEXT("NeedValidSkeleton", "Must specify a valid skeleton for the Anim Blueprint to target."));
-		return nullptr;
-	}
-
 	if (BlueprintType != BPTYPE_Interface && ((ParentClass == nullptr) || !FKismetEditorUtilities::CanCreateBlueprintOfClass(ParentClass) || !ParentClass->IsChildOf(UAnimInstance::StaticClass())))
 	{
 		FFormatNamedArguments Args;
@@ -393,7 +452,16 @@ UObject* UAnimBlueprintFactory::FactoryCreateNew(UClass* Class, UObject* InParen
 			FKismetEditorUtilities::CompileBlueprint(NewBP);
 		}
 		
-		NewBP->TargetSkeleton = TargetSkeleton;
+		if(bTemplate)
+		{
+			NewBP->bIsTemplate = true;
+			NewBP->TargetSkeleton = nullptr;
+		}
+		else
+		{
+			NewBP->bIsTemplate = false;
+			NewBP->TargetSkeleton = TargetSkeleton;
+		}
 
 		// Because the BP itself didn't have the skeleton set when the initial compile occured, it's not set on the generated classes either
 		if (UAnimBlueprintGeneratedClass* TypedNewClass = Cast<UAnimBlueprintGeneratedClass>(NewBP->GeneratedClass))
@@ -405,7 +473,7 @@ UObject* UAnimBlueprintFactory::FactoryCreateNew(UClass* Class, UObject* InParen
 			TypedNewClass_SKEL->TargetSkeleton = TargetSkeleton;
 		}
 
-		if (PreviewSkeletalMesh)
+		if (TargetSkeleton && PreviewSkeletalMesh)
 		{
 			NewBP->SetPreviewMesh(PreviewSkeletalMesh);
 		}
diff --git a/Engine/Source/Editor/UnrealEd/Private/ThumbnailRendering/AnimBlueprintThumbnailRenderer.cpp b/Engine/Source/Editor/UnrealEd/Private/ThumbnailRendering/AnimBlueprintThumbnailRenderer.cpp
--- a/Engine/Source/Editor/UnrealEd/Private/ThumbnailRendering/AnimBlueprintThumbnailRenderer.cpp
+++ b/Engine/Source/Editor/UnrealEd/Private/ThumbnailRendering/AnimBlueprintThumbnailRenderer.cpp
@@ -14,7 +14,7 @@ UAnimBlueprintThumbnailRenderer::UAnimBlueprintThumbnailRenderer(const FObjectIn
 bool UAnimBlueprintThumbnailRenderer::CanVisualizeAsset(UObject* Object)
 {
 	UAnimBlueprint* AnimBlueprint = Cast<UAnimBlueprint>(Object);
-	return AnimBlueprint->BlueprintType != BPTYPE_Interface;
+	return AnimBlueprint->BlueprintType != BPTYPE_Interface && !AnimBlueprint->bIsTemplate;
 }
 
 void UAnimBlueprintThumbnailRenderer::Draw(UObject* Object, int32 X, int32 Y, uint32 Width, uint32 Height, FRenderTarget* RenderTarget, FCanvas* Canvas, bool bAdditionalViewFamily)
diff --git a/Engine/Source/Editor/UnrealEd/Public/AnimationEditorUtils.h b/Engine/Source/Editor/UnrealEd/Public/AnimationEditorUtils.h
--- a/Engine/Source/Editor/UnrealEd/Public/AnimationEditorUtils.h
+++ b/Engine/Source/Editor/UnrealEd/Public/AnimationEditorUtils.h
@@ -198,6 +198,7 @@ namespace AnimationEditorUtils
 	UNREALED_API void RemovePoseWatch(UPoseWatch* PoseWatch, UAnimBlueprint* AnimBlueprintIfKnown=nullptr);
 	UNREALED_API void UpdatePoseWatchColour(UPoseWatch* PoseWatch, FColor NewPoseWatchColour);
 	UNREALED_API TArrayView<const FColor> GetPoseWatchColorPalette();
+	UNREALED_API void SetupDebugLinkedAnimInstances(UAnimBlueprint* InAnimBlueprint, UObject* InRootObjectBeingDebugged);
 
 	//////////////////////////////////////////////////////////////////////////////////////////
 
diff --git a/Engine/Source/Runtime/AnimGraphRuntime/Private/AnimNodeContextLibrary.cpp b/Engine/Source/Runtime/AnimGraphRuntime/Private/AnimNodeContextLibrary.cpp
new file mode 100644
--- /dev/null
+++ b/Engine/Source/Runtime/AnimGraphRuntime/Private/AnimNodeContextLibrary.cpp
@@ -0,0 +1,32 @@
+// Copyright Epic Games, Inc. All Rights Reserved.
+
+#include "AnimNodeContextLibrary.h"
+#include "Animation/AnimNodeContext.h"
+#include "Animation/AnimInstance.h"
+#include "Animation/AnimNodeBase.h"
+#include "Animation/AnimNode_LinkedAnimGraph.h"
+
+DEFINE_LOG_CATEGORY_STATIC(LogAnimNodeContextLibrary, Verbose, All);
+
+UAnimInstance* UAnimNodeContextLibrary::GetAnimInstance(const FAnimNodeContext& NodeContext)
+{
+	return CastChecked<UAnimInstance>(NodeContext.GetContext()->GetAnimInstanceObject());
+}
+
+bool UAnimNodeContextLibrary::HasLinkedAnimInstance(const FAnimNodeContext& NodeContext)
+{
+	if(FAnimNode_LinkedAnimGraph* LinkedAnimGraphNode = NodeContext.GetAnimNode<FAnimNode_LinkedAnimGraph>())
+	{
+		return LinkedAnimGraphNode->GetTargetInstance<UAnimInstance>() != nullptr;
+	}
+	return false;
+}
+
+UAnimInstance* UAnimNodeContextLibrary::GetLinkedAnimInstance(const FAnimNodeContext& NodeContext)
+{
+	if(FAnimNode_LinkedAnimGraph* LinkedAnimGraphNode = NodeContext.GetAnimNode<FAnimNode_LinkedAnimGraph>())
+	{
+		return LinkedAnimGraphNode->GetTargetInstance<UAnimInstance>();
+	}
+	return nullptr;
+}
\ No newline at end of file
diff --git a/Engine/Source/Runtime/AnimGraphRuntime/Private/AnimNodes/AnimNode_BlendSpacePlayer.cpp b/Engine/Source/Runtime/AnimGraphRuntime/Private/AnimNodes/AnimNode_BlendSpacePlayer.cpp
--- a/Engine/Source/Runtime/AnimGraphRuntime/Private/AnimNodes/AnimNode_BlendSpacePlayer.cpp
+++ b/Engine/Source/Runtime/AnimGraphRuntime/Private/AnimNodes/AnimNode_BlendSpacePlayer.cpp
@@ -199,13 +199,21 @@ void FAnimNode_BlendSpacePlayer::Reinitialize(bool bResetTime)
 	}
 }
 
-#if WITH_EDITORONLY_DATA
-void FAnimNode_BlendSpacePlayer::SetBlendSpace(UBlendSpace* InBlendSpace)
+bool FAnimNode_BlendSpacePlayer::SetBlendSpace(UBlendSpace* InBlendSpace)
 {
-	TObjectPtr<UBlendSpace>& BlendSpaceToSet = GET_MUTABLE_ANIM_NODE_DATA(TObjectPtr<UBlendSpace>, BlendSpace);
-	BlendSpaceToSet = InBlendSpace;
-}
+#if WITH_EDITORONLY_DATA
+	BlendSpace = InBlendSpace;
+	GET_MUTABLE_ANIM_NODE_DATA(TObjectPtr<UBlendSpace>, BlendSpace) = InBlendSpace;
 #endif
+	
+	if(TObjectPtr<UBlendSpace>* BlendSpacePtr = GET_INSTANCE_ANIM_NODE_DATA_PTR(TObjectPtr<UBlendSpace>, BlendSpace))
+	{
+		*BlendSpacePtr = InBlendSpace;
+		return true;
+	}
+
+	return false;
+}
 
 FVector FAnimNode_BlendSpacePlayer::GetPosition() const
 {
@@ -257,30 +265,62 @@ bool FAnimNode_BlendSpacePlayer::GetIgnoreForRelevancyTest() const
 	return GET_ANIM_NODE_DATA(bool, bIgnoreForRelevancyTest);
 }
 
-void FAnimNode_BlendSpacePlayer::SetGroupName(FName InGroupName)
+bool FAnimNode_BlendSpacePlayer::SetGroupName(FName InGroupName)
 {
-#if WITH_EDITORONLY_DATA	
+#if WITH_EDITORONLY_DATA
 	GroupName = InGroupName;
 #endif
+
+	if(FName* GroupNamePtr = GET_INSTANCE_ANIM_NODE_DATA_PTR(FName, GroupName))
+	{
+		*GroupNamePtr = InGroupName;
+		return true;
+	}
+
+	return false;
 }
 
-void FAnimNode_BlendSpacePlayer::SetGroupRole(EAnimGroupRole::Type InRole)
+bool FAnimNode_BlendSpacePlayer::SetGroupRole(EAnimGroupRole::Type InRole)
 {
 #if WITH_EDITORONLY_DATA
 	GroupRole = InRole;
 #endif
+	
+	if(TEnumAsByte<EAnimGroupRole::Type>* GroupRolePtr = GET_INSTANCE_ANIM_NODE_DATA_PTR(TEnumAsByte<EAnimGroupRole::Type>, GroupRole))
+	{
+		*GroupRolePtr = InRole;
+		return true;
+	}
+
+	return false;
 }
 
-void FAnimNode_BlendSpacePlayer::SetGroupMethod(EAnimSyncMethod InMethod)
+bool FAnimNode_BlendSpacePlayer::SetGroupMethod(EAnimSyncMethod InMethod)
 {
 #if WITH_EDITORONLY_DATA
 	Method = InMethod;
 #endif
+
+	if(EAnimSyncMethod* MethodPtr = GET_INSTANCE_ANIM_NODE_DATA_PTR(EAnimSyncMethod, Method))
+	{
+		*MethodPtr = InMethod;
+		return true;
+	}
+
+	return false;
 }
 
-void FAnimNode_BlendSpacePlayer::SetIgnoreForRelevancyTest(bool bInIgnoreForRelevancyTest)
+bool FAnimNode_BlendSpacePlayer::SetIgnoreForRelevancyTest(bool bInIgnoreForRelevancyTest)
 {
 #if WITH_EDITORONLY_DATA
 	bIgnoreForRelevancyTest = bInIgnoreForRelevancyTest;
 #endif
+
+	if(bool* bIgnoreForRelevancyTestPtr = GET_INSTANCE_ANIM_NODE_DATA_PTR(bool, bIgnoreForRelevancyTest))
+	{
+		*bIgnoreForRelevancyTestPtr = bInIgnoreForRelevancyTest;
+		return true;
+	}
+
+	return false;
 }
\ No newline at end of file
diff --git a/Engine/Source/Runtime/AnimGraphRuntime/Private/AnimNodes/AnimNode_CallFunction.cpp b/Engine/Source/Runtime/AnimGraphRuntime/Private/AnimNodes/AnimNode_CallFunction.cpp
--- a/Engine/Source/Runtime/AnimGraphRuntime/Private/AnimNodes/AnimNode_CallFunction.cpp
+++ b/Engine/Source/Runtime/AnimGraphRuntime/Private/AnimNodes/AnimNode_CallFunction.cpp
@@ -5,11 +5,6 @@
 
 void FAnimNode_CallFunction::OnInitializeAnimInstance(const FAnimInstanceProxy* InProxy, const UAnimInstance* InAnimInstance)
 {
-	if(FunctionName != NAME_None)
-	{
-		Function = InAnimInstance->FindFunction(FunctionName);
-	}
-	
 	Counter.Reset();
 	CurrentWeight = 0.0f;
 }
@@ -23,10 +18,9 @@ void FAnimNode_CallFunction::GatherDebugData(FNodeDebugData& DebugData)
 
 void FAnimNode_CallFunction::CallFunctionFromCallSite(EAnimFunctionCallSite InCallSite, const FAnimationBaseContext& InContext) const
 {
-	if(CallSite == InCallSite && Function != nullptr)
+	if(CallSite == InCallSite)
 	{
-		UObject* AnimInstanceObject = InContext.AnimInstanceProxy->GetAnimInstanceObject();
-		AnimInstanceObject->ProcessEvent(Function, nullptr);
+		GetFunction().Call(InContext.AnimInstanceProxy->GetAnimInstanceObject());
 	}
 }
 
@@ -90,3 +84,8 @@ void FAnimNode_CallFunction::Initialize_AnyThread(const FAnimationInitializeCont
 
 	CallFunctionFromCallSite(EAnimFunctionCallSite::OnInitializePostRecursion, InContext);
 }
+
+const FAnimNodeFunctionRef& FAnimNode_CallFunction::GetFunction() const
+{
+	return GET_ANIM_NODE_DATA(FAnimNodeFunctionRef, Function);
+}
\ No newline at end of file
diff --git a/Engine/Source/Runtime/AnimGraphRuntime/Private/AnimNodes/AnimNode_SequenceEvaluator.cpp b/Engine/Source/Runtime/AnimGraphRuntime/Private/AnimNodes/AnimNode_SequenceEvaluator.cpp
--- a/Engine/Source/Runtime/AnimGraphRuntime/Private/AnimNodes/AnimNode_SequenceEvaluator.cpp
+++ b/Engine/Source/Runtime/AnimGraphRuntime/Private/AnimNodes/AnimNode_SequenceEvaluator.cpp
@@ -131,12 +131,20 @@ float FAnimNode_SequenceEvaluatorBase::GetEffectiveDeltaTime(float ExplicitTime,
 	return DeltaTime;
 }
 
-void FAnimNode_SequenceEvaluator::SetSequence(UAnimSequenceBase* InSequence)
+bool FAnimNode_SequenceEvaluator::SetSequence(UAnimSequenceBase* InSequence)
 {
 #if WITH_EDITORONLY_DATA
-	TObjectPtr<UAnimSequenceBase>& SequenceToSet = GET_MUTABLE_ANIM_NODE_DATA(TObjectPtr<UAnimSequenceBase>, Sequence);
-	SequenceToSet = InSequence;
+	Sequence = InSequence;
+	GET_MUTABLE_ANIM_NODE_DATA(TObjectPtr<UAnimSequenceBase>, Sequence) = InSequence;
 #endif
+
+	if(TObjectPtr<UAnimSequenceBase>* SequencePtr = GET_INSTANCE_ANIM_NODE_DATA_PTR(TObjectPtr<UAnimSequenceBase>, Sequence))
+	{
+		*SequencePtr = InSequence;
+		return true;
+	}
+
+	return false;
 }
 
 UAnimSequenceBase* FAnimNode_SequenceEvaluator::GetSequence() const
@@ -189,30 +197,62 @@ bool FAnimNode_SequenceEvaluator::GetIgnoreForRelevancyTest() const
 	return GET_ANIM_NODE_DATA(bool, bIgnoreForRelevancyTest);
 }
 
-void FAnimNode_SequenceEvaluator::SetGroupName(FName InGroupName)
+bool FAnimNode_SequenceEvaluator::SetGroupName(FName InGroupName)
 {
 #if WITH_EDITORONLY_DATA	
 	GroupName = InGroupName;
 #endif
+
+	if(FName* GroupNamePtr = GET_INSTANCE_ANIM_NODE_DATA_PTR(FName, GroupName))
+	{
+		*GroupNamePtr = InGroupName;
+		return true;
+	}
+
+	return false;
 }
 
-void FAnimNode_SequenceEvaluator::SetGroupRole(EAnimGroupRole::Type InRole)
+bool FAnimNode_SequenceEvaluator::SetGroupRole(EAnimGroupRole::Type InRole)
 {
 #if WITH_EDITORONLY_DATA
 	GroupRole = InRole;
 #endif
+
+	if(TEnumAsByte<EAnimGroupRole::Type>* GroupRolePtr = GET_INSTANCE_ANIM_NODE_DATA_PTR(TEnumAsByte<EAnimGroupRole::Type>, GroupRole))
+	{
+		*GroupRolePtr = InRole;
+		return true;
+	}
+
+	return false;	
 }
 
-void FAnimNode_SequenceEvaluator::SetGroupMethod(EAnimSyncMethod InMethod)
+bool FAnimNode_SequenceEvaluator::SetGroupMethod(EAnimSyncMethod InMethod)
 {
 #if WITH_EDITORONLY_DATA
 	Method = InMethod;
 #endif
+
+	if(EAnimSyncMethod* MethodPtr = GET_INSTANCE_ANIM_NODE_DATA_PTR(EAnimSyncMethod, Method))
+	{
+		*MethodPtr = InMethod;
+		return true;
+	}
+
+	return false;
 }
 
-void FAnimNode_SequenceEvaluator::SetIgnoreForRelevancyTest(bool bInIgnoreForRelevancyTest)
+bool FAnimNode_SequenceEvaluator::SetIgnoreForRelevancyTest(bool bInIgnoreForRelevancyTest)
 {
 #if WITH_EDITORONLY_DATA
 	bIgnoreForRelevancyTest = bInIgnoreForRelevancyTest;
 #endif
+
+	if(bool* bIgnoreForRelevancyTestPtr = GET_INSTANCE_ANIM_NODE_DATA_PTR(bool, bIgnoreForRelevancyTest))
+	{
+		*bIgnoreForRelevancyTestPtr = bInIgnoreForRelevancyTest;
+		return true;
+	}
+
+	return false;
 }
\ No newline at end of file
diff --git a/Engine/Source/Runtime/AnimGraphRuntime/Private/SequencePlayerLibrary.cpp b/Engine/Source/Runtime/AnimGraphRuntime/Private/SequencePlayerLibrary.cpp
new file mode 100644
--- /dev/null
+++ b/Engine/Source/Runtime/AnimGraphRuntime/Private/SequencePlayerLibrary.cpp
@@ -0,0 +1,49 @@
+// Copyright Epic Games, Inc. All Rights Reserved.
+
+#include "SequencePlayerLibrary.h"
+#include "Animation/AnimNodeContext.h"
+#include "Animation/AnimNode_SequencePlayer.h"
+
+DEFINE_LOG_CATEGORY_STATIC(LogSequencePlayerLibrary, Verbose, All);
+
+void USequencePlayerLibrary::SetAccumulatedTime(const FAnimNodeContext& NodeContext, float Time)
+{
+	if(FAnimNode_SequencePlayer* SequencePlayer = NodeContext.GetAnimNode<FAnimNode_SequencePlayer>())
+	{
+		SequencePlayer->SetAccumulatedTime(Time);
+	}
+	else
+	{
+		UE_LOG(LogSequencePlayerLibrary, Warning, TEXT("SetAccumulatedTime called on a non-sequence player node"));
+	}
+}
+
+void USequencePlayerLibrary::SetStartPosition(const FAnimNodeContext& NodeContext, float StartPosition)
+{
+	if(FAnimNode_SequencePlayer* SequencePlayer = NodeContext.GetAnimNode<FAnimNode_SequencePlayer>())
+	{
+		if(!SequencePlayer->SetStartPosition(StartPosition))
+		{
+			UE_LOG(LogSequencePlayerLibrary, Warning, TEXT("Could not set start time on sequence player, value is not dynamic. Set it as Always Dynamic or expose it as a pin"));
+		}
+	}
+	else
+	{
+		UE_LOG(LogSequencePlayerLibrary, Warning, TEXT("SetStartTime called on a non-sequence player node"));
+	}
+}
+
+void USequencePlayerLibrary::SetPlayRate(const FAnimNodeContext& NodeContext, float PlayRate)
+{
+	if(FAnimNode_SequencePlayer* SequencePlayer = NodeContext.GetAnimNode<FAnimNode_SequencePlayer>())
+	{
+		if(!SequencePlayer->SetPlayRate(PlayRate))
+		{
+			UE_LOG(LogSequencePlayerLibrary, Warning, TEXT("Could not set start time on sequence player, value is not dynamic. Set it as Always Dynamic or expose it as a pin"));
+		}
+	}
+	else
+	{
+		UE_LOG(LogSequencePlayerLibrary, Warning, TEXT("SetStartTime called on a non-sequence player node"));
+	}
+}
\ No newline at end of file
diff --git a/Engine/Source/Runtime/AnimGraphRuntime/Public/AnimNodeContextLibrary.h b/Engine/Source/Runtime/AnimGraphRuntime/Public/AnimNodeContextLibrary.h
new file mode 100644
--- /dev/null
+++ b/Engine/Source/Runtime/AnimGraphRuntime/Public/AnimNodeContextLibrary.h
@@ -0,0 +1,33 @@
+// Copyright Epic Games, Inc. All Rights Reserved.
+
+#pragma once
+
+#include "CoreMinimal.h"
+#include "Kismet/BlueprintFunctionLibrary.h"
+#include "Animation/AnimNodeContext.h"
+#include "AnimNodeContextLibrary.generated.h"
+
+// Exposes operations to be performed on anim node contexts
+UCLASS()
+class ANIMGRAPHRUNTIME_API UAnimNodeContextLibrary : public UBlueprintFunctionLibrary
+{
+	GENERATED_BODY()
+
+#if WITH_EDITOR
+	/** Prototype function for thread-safe anim node calls */
+	UFUNCTION(BlueprintInternalUseOnly, meta=(BlueprintThreadSafe))
+	void Prototype_ThreadSafeAnimNodeCall(const FAnimNodeContext& NodeContext) {}
+#endif
+	
+	/** Get the anim instance that hosts this anim node */
+	UFUNCTION(BlueprintPure, Category = "Anim Node", meta=(BlueprintThreadSafe))
+	static UAnimInstance* GetAnimInstance(const FAnimNodeContext& NodeContext);
+
+	/** Returns whether the node hosts an instance (e.g. linked anim graph or layer) */
+	UFUNCTION(BlueprintPure, Category = "Anim Node", meta=(BlueprintThreadSafe))
+	static bool HasLinkedAnimInstance(const FAnimNodeContext& NodeContext);
+
+	/** Get the linked instance is hosted by this node. If the node does not host an instance then HasLinkedAnimInstance will return false */
+	UFUNCTION(BlueprintPure, Category = "Anim Node", meta=(BlueprintThreadSafe))
+	static UAnimInstance* GetLinkedAnimInstance(const FAnimNodeContext& NodeContext);
+};
\ No newline at end of file
diff --git a/Engine/Source/Runtime/AnimGraphRuntime/Public/AnimNodes/AnimNode_BlendSpacePlayer.h b/Engine/Source/Runtime/AnimGraphRuntime/Public/AnimNodes/AnimNode_BlendSpacePlayer.h
--- a/Engine/Source/Runtime/AnimGraphRuntime/Public/AnimNodes/AnimNode_BlendSpacePlayer.h
+++ b/Engine/Source/Runtime/AnimGraphRuntime/Public/AnimNodes/AnimNode_BlendSpacePlayer.h
@@ -101,10 +101,10 @@ public:
 	virtual EAnimGroupRole::Type GetGroupRole() const override;
 	virtual EAnimSyncMethod GetGroupMethod() const override;
 	virtual bool GetIgnoreForRelevancyTest() const override;
-	virtual void SetGroupName(FName InGroupName) override;
-	virtual void SetGroupRole(EAnimGroupRole::Type InRole) override;
-	virtual void SetGroupMethod(EAnimSyncMethod InMethod) override;
-	virtual void SetIgnoreForRelevancyTest(bool bInIgnoreForRelevancyTest) override;
+	virtual bool SetGroupName(FName InGroupName) override;
+	virtual bool SetGroupRole(EAnimGroupRole::Type InRole) override;
+	virtual bool SetGroupMethod(EAnimSyncMethod InMethod) override;
+	virtual bool SetIgnoreForRelevancyTest(bool bInIgnoreForRelevancyTest) override;
 	// End of FAnimNode_AssetPlayerBase interface
 
 	// FAnimNode_Base interface
@@ -117,10 +117,8 @@ public:
 
 	float GetTimeFromEnd(float CurrentTime) const;
 
-#if WITH_EDITORONLY_DATA
 	// Set the blendspace asset to play
-	void SetBlendSpace(UBlendSpace* InBlendSpace);
-#endif
+	bool SetBlendSpace(UBlendSpace* InBlendSpace);
 
 	// Get the coordinates that are currently being sampled by the blendspace
 	virtual FVector GetPosition() const;
diff --git a/Engine/Source/Runtime/AnimGraphRuntime/Public/AnimNodes/AnimNode_CallFunction.h b/Engine/Source/Runtime/AnimGraphRuntime/Public/AnimNodes/AnimNode_CallFunction.h
--- a/Engine/Source/Runtime/AnimGraphRuntime/Public/AnimNodes/AnimNode_CallFunction.h
+++ b/Engine/Source/Runtime/AnimGraphRuntime/Public/AnimNodes/AnimNode_CallFunction.h
@@ -3,6 +3,7 @@
 
 #include "CoreMinimal.h"
 #include "Animation/AnimNodeBase.h"
+#include "Animation/AnimNodeFunctionRef.h"
 #include "AnimNode_CallFunction.generated.h"
 
 // When to call the function during the execution of the animation graph
@@ -59,14 +60,12 @@ struct ANIMGRAPHRUNTIME_API FAnimNode_CallFunction : public FAnimNode_Base
 	UPROPERTY(EditAnywhere, Category = Links)
 	FPoseLink Source;
 
-	// The name of the function to call
-	UPROPERTY()
-	FName FunctionName = NAME_None;
-
-	// The function to call, recovered by looking for a function of name FunctionName
-	UPROPERTY()
-	UFunction* Function = nullptr;
-
+#if WITH_EDITORONLY_DATA
+	// Function to call
+	UPROPERTY(meta=(FoldProperty))
+	FAnimNodeFunctionRef Function;
+#endif
+	
 	// Counter used to determine relevancy
 	FGraphTraversalCounter Counter;
 
@@ -88,4 +87,7 @@ struct ANIMGRAPHRUNTIME_API FAnimNode_CallFunction : public FAnimNode_Base
 
 	// Calls the function we hold if the callsite matches the one we have set
 	void CallFunctionFromCallSite(EAnimFunctionCallSite InCallSite, const FAnimationBaseContext& InContext) const;
+
+	// Get the function held on this node
+	const FAnimNodeFunctionRef& GetFunction() const;
 };
\ No newline at end of file
diff --git a/Engine/Source/Runtime/AnimGraphRuntime/Public/AnimNodes/AnimNode_SequenceEvaluator.h b/Engine/Source/Runtime/AnimGraphRuntime/Public/AnimNodes/AnimNode_SequenceEvaluator.h
--- a/Engine/Source/Runtime/AnimGraphRuntime/Public/AnimNodes/AnimNode_SequenceEvaluator.h
+++ b/Engine/Source/Runtime/AnimGraphRuntime/Public/AnimNodes/AnimNode_SequenceEvaluator.h
@@ -59,7 +59,7 @@ public:
 	virtual float GetEffectiveDeltaTime(float ExplicitTime, float PrevExplicitTime) const;
 
 	// Set the animation sequence asset to evaluate
-	virtual void SetSequence(UAnimSequenceBase* InSequence) {}
+	virtual bool SetSequence(UAnimSequenceBase* InSequence) { return false; }
 
 	// Set the time at which to evaluate the associated sequence
 	virtual void SetExplicitTime(float InTime) {}
@@ -146,7 +146,7 @@ private:
 
 public:
 	// FAnimNode_SequenceEvaluatorBase interface
-	virtual void SetSequence(UAnimSequenceBase* InSequence) override;
+	virtual bool SetSequence(UAnimSequenceBase* InSequence) override;
 	virtual UAnimSequenceBase* GetSequence() const override;
 	virtual float GetExplicitTime() const override;
 	virtual bool GetShouldLoop() const override;
@@ -159,10 +159,10 @@ public:
 	virtual EAnimGroupRole::Type GetGroupRole() const override;
 	virtual EAnimSyncMethod GetGroupMethod() const override;
 	virtual bool GetIgnoreForRelevancyTest() const override;
-	virtual void SetGroupName(FName InGroupName) override;
-	virtual void SetGroupRole(EAnimGroupRole::Type InRole) override;
-	virtual void SetGroupMethod(EAnimSyncMethod InMethod) override;
-	virtual void SetIgnoreForRelevancyTest(bool bInIgnoreForRelevancyTest) override;
+	virtual bool SetGroupName(FName InGroupName) override;
+	virtual bool SetGroupRole(EAnimGroupRole::Type InRole) override;
+	virtual bool SetGroupMethod(EAnimSyncMethod InMethod) override;
+	virtual bool SetIgnoreForRelevancyTest(bool bInIgnoreForRelevancyTest) override;
 };
 
 // Sequence evaluator node that can be used standalone (without constant folding)
@@ -216,7 +216,7 @@ private:
 
 public:
 	// FAnimNode_SequenceEvaluatorBase interface
-	virtual void SetSequence(UAnimSequenceBase* InSequence) override { Sequence = InSequence; }
+	virtual bool SetSequence(UAnimSequenceBase* InSequence) override { Sequence = InSequence; return true; }
 	virtual void SetExplicitTime(float InTime) override { ExplicitTime = InTime; }
 	virtual void SetTeleportToExplicitTime(bool bInTeleport) override { bTeleportToExplicitTime = bInTeleport; }
 	virtual void SetReinitializationBehavior(TEnumAsByte<ESequenceEvalReinit::Type> InBehavior) override { ReinitializationBehavior = InBehavior; }
@@ -232,8 +232,8 @@ public:
 	virtual EAnimGroupRole::Type GetGroupRole() const override { return GroupRole; }
 	virtual EAnimSyncMethod GetGroupMethod() const override { return Method; }
 	virtual bool GetIgnoreForRelevancyTest() const override { return bIgnoreForRelevancyTest; }
-	virtual void SetGroupName(FName InGroupName) override { GroupName = InGroupName; }
-	virtual void SetGroupRole(EAnimGroupRole::Type InRole) override { GroupRole = InRole; }
-	virtual void SetGroupMethod(EAnimSyncMethod InMethod) override { Method = InMethod; }
-	virtual void SetIgnoreForRelevancyTest(bool bInIgnoreForRelevancyTest) override { bIgnoreForRelevancyTest = bInIgnoreForRelevancyTest; }
+	virtual bool SetGroupName(FName InGroupName) override { GroupName = InGroupName; return true; }
+	virtual bool SetGroupRole(EAnimGroupRole::Type InRole) override { GroupRole = InRole; return true; }
+	virtual bool SetGroupMethod(EAnimSyncMethod InMethod) override { Method = InMethod; return true; }
+	virtual bool SetIgnoreForRelevancyTest(bool bInIgnoreForRelevancyTest) override { bIgnoreForRelevancyTest = bInIgnoreForRelevancyTest; return true; }
 };
\ No newline at end of file
diff --git a/Engine/Source/Runtime/AnimGraphRuntime/Public/SequencePlayerLibrary.h b/Engine/Source/Runtime/AnimGraphRuntime/Public/SequencePlayerLibrary.h
new file mode 100644
--- /dev/null
+++ b/Engine/Source/Runtime/AnimGraphRuntime/Public/SequencePlayerLibrary.h
@@ -0,0 +1,27 @@
+// Copyright Epic Games, Inc. All Rights Reserved.
+
+#pragma once
+
+#include "CoreMinimal.h"
+#include "Kismet/BlueprintFunctionLibrary.h"
+#include "Animation/AnimNodeContext.h"
+#include "SequencePlayerLibrary.generated.h"
+
+// Exposes operations to be performed on a sequence player anim node
+// Note: Experimental and subject to change!
+UCLASS(Experimental)
+class ANIMGRAPHRUNTIME_API USequencePlayerLibrary : public UBlueprintFunctionLibrary
+{
+	GENERATED_BODY()
+
+	/** Set the current accumulated time of the sequence player */
+	UFUNCTION(BlueprintCallable, Category = "Sequence Player", meta=(BlueprintThreadSafe))
+	static void SetAccumulatedTime(const FAnimNodeContext& NodeContext, float Time);
+
+	/** Set the start position of the sequence player */
+	UFUNCTION(BlueprintCallable, Category = "Sequence Player", meta=(BlueprintThreadSafe))
+	static void SetStartPosition(const FAnimNodeContext& NodeContext, float StartPosition);
+
+	UFUNCTION(BlueprintCallable, Category = "Sequence Player", meta=(BlueprintThreadSafe))
+	static void SetPlayRate(const FAnimNodeContext& NodeContext, float PlayRate);
+};
\ No newline at end of file
diff --git a/Engine/Source/Runtime/Engine/Classes/Animation/AnimBlueprint.h b/Engine/Source/Runtime/Engine/Classes/Animation/AnimBlueprint.h
--- a/Engine/Source/Runtime/Engine/Classes/Animation/AnimBlueprint.h
+++ b/Engine/Source/Runtime/Engine/Classes/Animation/AnimBlueprint.h
@@ -90,6 +90,11 @@ class ENGINE_API UAnimBlueprint : public UBlueprint, public IInterface_PreviewMe
 	UPROPERTY()
 	TArray<FAnimGroupInfo> Groups;
 
+	// This is an anim blueprint that acts as a set of template functionality without being tied to a specific skeleton.
+	// Implies a null TargetSkeleton.
+	UPROPERTY(AssetRegistrySearchable)
+	bool bIsTemplate;
+	
 	/**
 	 * Allows this anim Blueprint to update its native update, blend tree, montages and asset players on
 	 * a worker thread. The compiler will attempt to pick up any issues that may occur with threaded update.
@@ -168,6 +173,7 @@ class ENGINE_API UAnimBlueprint : public UBlueprint, public IInterface_PreviewMe
 
 	virtual void PostLoad() override;
 	virtual bool FindDiffs(const UBlueprint* OtherBlueprint, FDiffResults& Results) const override;
+	virtual void SetObjectBeingDebugged(UObject* NewObject) override;
 	
 protected:
 	// Broadcast when an override is changed, allowing derived blueprints to be updated
diff --git a/Engine/Source/Runtime/Engine/Classes/Animation/AnimInstance.h b/Engine/Source/Runtime/Engine/Classes/Animation/AnimInstance.h
--- a/Engine/Source/Runtime/Engine/Classes/Animation/AnimInstance.h
+++ b/Engine/Source/Runtime/Engine/Classes/Animation/AnimInstance.h
@@ -13,6 +13,7 @@
 #include "BonePose.h"
 #include "Components/SkeletalMeshComponent.h"
 #include "Animation/AnimNotifyQueue.h"
+#include "Animation/AnimSubsystemInstance.h"
 #include "Animation/AnimNotifies/AnimNotify.h"
 #include "AnimInstance.generated.h"
 
@@ -877,6 +878,27 @@ public:
 	/** Other UObjects should call this to remove themselves from the callbacks */
 	void RemoveExternalNotifyHandler(UObject* ExternalHandlerObject, FName NotifyEventName);
 
+	// Find a subsystem's instance-resident data. If no subsystem of the type exists this will return nullptr.
+	// @param	InSubsystemType	The subsystem's type
+	FAnimSubsystemInstance* FindSubsystem(UScriptStruct* InSubsystemType);
+
+	// Get a subsystem's instance-resident data. If no subsystem of the type exists this will return nullptr.
+	template<typename SubsystemType>
+	SubsystemType* FindSubsystem()
+	{
+		FAnimSubsystemInstance* Subsystem = FindSubsystem(SubsystemType::StaticStruct());
+		return static_cast<SubsystemType*>(Subsystem);
+	}
+	
+	// Get a subsystem's instance-resident data. If no subsystem of the type exists this will assert.
+	template<typename SubsystemType>
+	SubsystemType& GetSubsystem()
+	{
+		FAnimSubsystemInstance* Subsystem = FindSubsystem(SubsystemType::StaticStruct());
+		check(Subsystem);
+		return static_cast<SubsystemType&>(*Subsystem);
+	}
+
 private:
 	/** Helper function to perform layer overlay actions (set, clear) */
 	void PerformLinkedLayerOverlayOperation(TSubclassOf<UAnimInstance> InClass, TFunctionRef<UClass*(UClass*, FAnimNode_LinkedAnimLayer*)> InClassSelectorFunction, bool bInDeferSubGraphInitialization = false);
diff --git a/Engine/Source/Runtime/Engine/Classes/Animation/AnimNodeBase.h b/Engine/Source/Runtime/Engine/Classes/Animation/AnimNodeBase.h
--- a/Engine/Source/Runtime/Engine/Classes/Animation/AnimNodeBase.h
+++ b/Engine/Source/Runtime/Engine/Classes/Animation/AnimNodeBase.h
@@ -17,6 +17,7 @@
 #include "Animation/AnimNodeMessages.h"
 #include "Animation/AnimNodeData.h"
 #include "Animation/ExposedValueHandler.h"
+#include "AnimNodeFunctionRef.h"
 
 // WARNING: This should always be the last include in any file that needs it (except .generated.h)
 #include "UObject/UndefineUPropertyMacros.h"
@@ -909,11 +910,13 @@ protected:
 #endif
 	}
 
-#if WITH_EDITORONLY_DATA
 	// Get anim node constant/folded data of the specified type given the identifier. Do not use directly - use GET_MUTABLE_ANIM_NODE_DATA
+	// Note: will assert if data is not held on the instance/dynamic. Use GetInstanceDataPtr/GET_INSTANCE_ANIM_NODE_DATA_PTR if the value
+	// might not be mutable, which will return null.
 	template<typename DataType>
 	DataType& GetMutableData(UE::Anim::FNodeDataId InId, UObject* InObject = nullptr)
 	{
+#if WITH_EDITORONLY_DATA
 		if(NodeData)
 		{
 			return *static_cast<DataType*>(NodeData->GetMutableData(InId, this, InObject));
@@ -922,9 +925,32 @@ protected:
 		{
 			return *InId.GetProperty()->ContainerPtrToValuePtr<DataType>(this);
 		}
-	}
+#else
+		check(NodeData);
+		return *static_cast<DataType*>(NodeData->GetMutableData(InId, this, InObject));
 #endif
+	}
 
+	// Get anim node mutable data of the specified type given the identifier. Do not use directly - use GET_INSTANCE_ANIM_NODE_DATA_PTR
+	// @return nullptr if the data is not mutable/dynamic
+	template<typename DataType>
+	DataType* GetInstanceDataPtr(UE::Anim::FNodeDataId InId, UObject* InObject = nullptr)
+	{
+#if WITH_EDITORONLY_DATA	
+		if(NodeData)
+		{
+			return static_cast<DataType*>(NodeData->GetInstanceData(InId, this, InObject));
+		}
+		else
+		{
+			return InId.GetProperty()->ContainerPtrToValuePtr<DataType>(this);
+		}
+#else
+		check(NodeData);
+		return static_cast<DataType*>(NodeData->GetInstanceData(InId, this, InObject));
+#endif
+	}
+	
 protected:
 	/** return true if enabled, otherwise, return false. This is utility function that can be used per node level */
 	bool IsLODEnabled(FAnimInstanceProxy* AnimInstanceProxy);
@@ -941,16 +967,43 @@ protected:
 
 	friend struct FAnimInstanceProxy;
 
+private:
+	// Access functions
+	const FAnimNodeFunctionRef& GetInitializeFunction() const;
+	const FAnimNodeFunctionRef& GetBecomeRelevantFunction() const;
+	const FAnimNodeFunctionRef& GetUpdateFunction() const;
+	const FAnimNodeFunctionRef& GetEvaluateFunction() const;
+	
 private:
 	friend class IAnimClassInterface;
 	friend class UAnimBlueprintGeneratedClass;
 	friend struct UE::Anim::FNodeDataId;
+	friend struct FAnimNodeFunctionCaller;
+	friend class UAnimGraphNode_Base;
 
 	// Set the cached ptr to the constant/folded data for this node
 	void SetNodeData(const FAnimNodeData& InNodeData) { NodeData = &InNodeData; }
 
 	// Reference to the constant/folded data for this node
 	const FAnimNodeData* NodeData = nullptr;
+
+#if WITH_EDITORONLY_DATA
+	// Function called on initialize
+	UPROPERTY(meta=(FoldProperty))
+	FAnimNodeFunctionRef InitializeFunction;
+
+	// Function called on become relevant
+	UPROPERTY(meta=(FoldProperty))
+	FAnimNodeFunctionRef BecomeRelevantFunction;
+
+	// Function called on update
+	UPROPERTY(meta=(FoldProperty))
+	FAnimNodeFunctionRef UpdateFunction;
+
+	// Function called on evaluate
+	UPROPERTY(meta=(FoldProperty))
+	FAnimNodeFunctionRef EvaluateFunction;
+#endif
 };
 
 #if WITH_EDITORONLY_DATA
@@ -959,42 +1012,31 @@ private:
 #define VERIFY_ANIM_NODE_MEMBER_TYPE(Type, Identifier)
 #endif
 
+#define GET_ANIM_NODE_DATA_ID_INTERNAL(Type, Identifier) \
+	[this]() -> UE::Anim::FNodeDataId \
+	{ \
+		VERIFY_ANIM_NODE_MEMBER_TYPE(Type, Identifier) \
+		static UE::Anim::FNodeDataId CachedId_##Identifier; \
+		if(!CachedId_##Identifier.IsValid()) \
+		{ \
+			static const FName AnimName_##Identifier(#Identifier); \
+			CachedId_##Identifier = UE::Anim::FNodeDataId(AnimName_##Identifier, this, StaticStruct()); \
+		} \
+		return CachedId_##Identifier; \
+	}() \
+
 // Get some (potentially folded) anim node data. Only usable from within an anim node.
 // This caches the node data ID in static contained in a local lambda for improved performance
-#define GET_ANIM_NODE_DATA(Type, Identifier) \
-	( \
-		GetData<Type>( \
-			[this]() -> UE::Anim::FNodeDataId \
-			{ \
-				VERIFY_ANIM_NODE_MEMBER_TYPE(Type, Identifier) \
-				static UE::Anim::FNodeDataId CachedId_##Identifier; \
-				if(!CachedId_##Identifier.IsValid()) \
-				{ \
-					static const FName AnimName_##Identifier(#Identifier); \
-					CachedId_##Identifier = UE::Anim::FNodeDataId(AnimName_##Identifier, this, StaticStruct()); \
-				} \
-				return CachedId_##Identifier; \
-			}() \
-		) \
-	)
+#define GET_ANIM_NODE_DATA(Type, Identifier) (GetData<Type>(GET_ANIM_NODE_DATA_ID_INTERNAL(Type, Identifier)))
+
+// Get some anim node data that should be held on an instance. Only usable from within an anim node.
+// @return nullptr if the data is not held on an instance (i.e. it is in constant sparse class data)
+// This caches the node data ID in static contained in a local lambda for improved performance
+#define GET_INSTANCE_ANIM_NODE_DATA_PTR(Type, Identifier) (GetInstanceDataPtr<Type>(GET_ANIM_NODE_DATA_ID_INTERNAL(Type, Identifier)))
 
 #if WITH_EDITORONLY_DATA
-#define GET_MUTABLE_ANIM_NODE_DATA(Type, Identifier) \
-	( \
-		GetMutableData<Type>( \
-			[this]() -> UE::Anim::FNodeDataId \
-			{ \
-				VERIFY_ANIM_NODE_MEMBER_TYPE(Type, Identifier) \
-				static UE::Anim::FNodeDataId CachedId_##Identifier; \
-				if(!CachedId_##Identifier.IsValid()) \
-				{ \
-					static const FName AnimName_##Identifier(#Identifier); \
-					CachedId_##Identifier = UE::Anim::FNodeDataId(AnimName_##Identifier, this, StaticStruct()); \
-				} \
-				return CachedId_##Identifier; \
-			}() \
-		) \
-	)
+// Editor-only way of accessing mutable anim node data but with internal checks
+#define GET_MUTABLE_ANIM_NODE_DATA(Type, Identifier) (GetMutableData<Type>(GET_ANIM_NODE_DATA_ID_INTERNAL(Type, Identifier)))
 #endif
 
 #include "UObject/DefineUPropertyMacros.h"
\ No newline at end of file
diff --git a/Engine/Source/Runtime/Engine/Classes/Animation/AnimNodeContext.h b/Engine/Source/Runtime/Engine/Classes/Animation/AnimNodeContext.h
new file mode 100644
--- /dev/null
+++ b/Engine/Source/Runtime/Engine/Classes/Animation/AnimNodeContext.h
@@ -0,0 +1,108 @@
+// Copyright Epic Games, Inc. All Rights Reserved.
+
+#pragma once
+
+#include "CoreMinimal.h"
+#include "AnimNodeContext.generated.h"
+
+struct FAnimNode_Base;
+struct FAnimationBaseContext;
+struct FAnimationInitializeContext;
+struct FAnimationUpdateContext;
+struct FPoseContext;
+struct FComponentSpacePoseContext;
+
+
+
+// Context used to expose anim nodes to BP function libraries
+USTRUCT(BlueprintType)
+struct ENGINE_API FAnimNodeContext
+{
+	GENERATED_BODY()
+
+public:
+	// Internal data, weakly referenced
+	struct FData
+	{
+	public:
+		FData(const FAnimationInitializeContext& InContext, FAnimNode_Base& InAnimNode, UScriptStruct* InAnimNodeStruct);
+
+		FData(const FAnimationUpdateContext& InContext, FAnimNode_Base& InAnimNode, UScriptStruct* InAnimNodeStruct);
+
+		FData(FPoseContext& InContext, FAnimNode_Base& InAnimNode, UScriptStruct* InAnimNodeStruct);
+
+		FData(FComponentSpacePoseContext& InContext, FAnimNode_Base& InAnimNode, UScriptStruct* InAnimNodeStruct);
+
+	private:
+		friend struct ::FAnimNodeContext;
+		
+		enum class EContextType
+		{
+			None,
+			Initialize,
+			Update,
+			Pose,
+			ComponentSpacePose,
+		};
+	
+		// The node we wrap
+		FAnimNode_Base* AnimNode = nullptr;
+
+		// The struct type of the anim node
+		UScriptStruct* AnimNodeStruct = nullptr;
+
+		// The context used when executing this node, e.g. FAnimationUpdateContext, FPoseContext etc.
+		FAnimationBaseContext* Context = nullptr;
+
+		// The phase we are in
+		EContextType ContextType = EContextType::None;	
+	};
+	
+	FAnimNodeContext() = default;
+
+	FAnimNodeContext(TWeakPtr<FData> InData)
+		: Data(InData)
+	{}
+
+	// Get the node we wrap. If the node is not of the specified type then this will return nullptr.
+	template<typename NodeType>
+	NodeType* GetAnimNode() const
+	{
+		if(TSharedPtr<FData> PinnedData = Data.Pin())
+		{
+			if(NodeType::StaticStruct()->IsChildOf(PinnedData->AnimNodeStruct))
+			{
+				return static_cast<NodeType*>(PinnedData->AnimNode);
+			}
+		}
+		return nullptr;
+	}
+
+	// Get the context we wrap.
+	FAnimationBaseContext* GetContext() const
+	{
+		if(TSharedPtr<FData> PinnedData = Data.Pin())
+		{
+			return PinnedData->Context;
+		}
+		return nullptr;
+	}
+
+	// Get the context we wrap. If the context is not an initialize context then this will return nullptr
+	FAnimationInitializeContext* GetInitializeContext() const;
+	
+	// Get the context we wrap. If the context is not an update context then this will return nullptr
+	FAnimationUpdateContext* GetUpdateContext() const;
+
+	// Get the context we wrap. If the context is not a pose context then this will return nullptr
+	FPoseContext* GetPoseContext() const;
+
+	// Get the context we wrap. If the context is not a component space pose context then this will return nullptr
+	FComponentSpacePoseContext* GetComponentSpacePoseContext() const;
+
+private:
+	friend struct FAnimNodeFunctionCaller;
+
+	// Internal data
+	TWeakPtr<FData> Data; 
+};
\ No newline at end of file
diff --git a/Engine/Source/Runtime/Engine/Classes/Animation/AnimNodeFunctionRef.h b/Engine/Source/Runtime/Engine/Classes/Animation/AnimNodeFunctionRef.h
new file mode 100644
--- /dev/null
+++ b/Engine/Source/Runtime/Engine/Classes/Animation/AnimNodeFunctionRef.h
@@ -0,0 +1,46 @@
+// Copyright Epic Games, Inc. All Rights Reserved.
+
+#pragma once
+
+#include "CoreMinimal.h"
+#include "AnimNodeFunctionRef.generated.h"
+
+/**
+ * Cached function name/ptr that is resolved at init time
+ */
+USTRUCT()
+struct ENGINE_API FAnimNodeFunctionRef
+{
+	GENERATED_BODY()
+
+public:
+	// Cache the function ptr from the name
+	void Initialize(const UClass* InClass);
+	
+	// Call the function
+	void Call(UObject* InObject, void* InParameters = nullptr) const;
+
+	// Set the function via name
+	void SetFromFunctionName(FName InName) { FunctionName = InName; }
+
+	// Set the function via a function
+	void SetFromFunction(UFunction* InFunction) { FunctionName = InFunction ? InFunction->GetFName() : NAME_None; }
+	
+	// Get the function name
+	FName GetFunctionName() const { return FunctionName; }
+	
+	// Get the function we reference
+	UFunction* GetFunction() const { return Function; }
+	
+	// Check if we reference a valid function
+	bool IsValid() const { return Function != nullptr; }
+
+private:
+	// The name of the function to call
+	UPROPERTY()
+	FName FunctionName = NAME_None;
+
+	// The function to call, recovered by looking for a function of name FunctionName
+	UPROPERTY(Transient)
+	UFunction* Function = nullptr;
+};
\ No newline at end of file
diff --git a/Engine/Source/Runtime/Engine/Classes/Animation/AnimNode_AssetPlayerBase.h b/Engine/Source/Runtime/Engine/Classes/Animation/AnimNode_AssetPlayerBase.h
--- a/Engine/Source/Runtime/Engine/Classes/Animation/AnimNode_AssetPlayerBase.h
+++ b/Engine/Source/Runtime/Engine/Classes/Animation/AnimNode_AssetPlayerBase.h
@@ -65,16 +65,16 @@ struct ENGINE_API FAnimNode_AssetPlayerBase : public FAnimNode_Base
 	virtual bool GetIgnoreForRelevancyTest() const { return false; }
 
 	// Set the sync group name we are using
-	virtual void SetGroupName(FName InGroupName) { }
+	virtual bool SetGroupName(FName InGroupName) { return false; }
 
 	// Set the sync group role we are using
-	virtual void SetGroupRole(EAnimGroupRole::Type InRole) { }
+	virtual bool SetGroupRole(EAnimGroupRole::Type InRole) { return false; }
 
 	// Set the sync group method we are using
-	virtual void SetGroupMethod(EAnimSyncMethod InMethod) { }
+	virtual bool SetGroupMethod(EAnimSyncMethod InMethod) { return false; }
 
 	// Set whether this node should be ignored when testing for relevancy in state machines
-	virtual void SetIgnoreForRelevancyTest(bool bInIgnoreForRelevancyTest) { }
+	virtual bool SetIgnoreForRelevancyTest(bool bInIgnoreForRelevancyTest) { return false; }
 
 private:
 #if WITH_EDITORONLY_DATA
diff --git a/Engine/Source/Runtime/Engine/Classes/Animation/AnimNode_SequencePlayer.h b/Engine/Source/Runtime/Engine/Classes/Animation/AnimNode_SequencePlayer.h
--- a/Engine/Source/Runtime/Engine/Classes/Animation/AnimNode_SequencePlayer.h
+++ b/Engine/Source/Runtime/Engine/Classes/Animation/AnimNode_SequencePlayer.h
@@ -45,10 +45,10 @@ public:
 
 protected:
 	// Set the animation sequence asset to play
-	virtual void SetSequence(UAnimSequenceBase* InSequence) {}
+	virtual bool SetSequence(UAnimSequenceBase* InSequence) { return false; }
 
 	// Set the animation to continue looping when it reaches the end
-	virtual void SetLoopAnimation(bool bInLoopAnimation) {}
+	virtual bool SetLoopAnimation(bool bInLoopAnimation) { return false; }
 	
 	// The Basis in which the PlayRate is expressed in. This is used to rescale PlayRate inputs.
 	// For example a Basis of 100 means that the PlayRate input will be divided by 100.
@@ -70,6 +70,14 @@ protected:
 
 	// Use pose matching to choose the start position. Requires experimental PoseSearch plugin.
 	virtual bool GetStartFromMatchingPose() const { return false; }
+
+	// Set the start time of this node.
+	// @return true if the value was set (it is dynamic), or false if it could not (it is not dynamic or pin exposed)
+	virtual bool SetStartPosition(float InStartPosition) { return false; }
+
+	// Set the play rate of this node.
+	// @return true if the value was set (it is dynamic), or false if it could not (it is not dynamic or pin exposed)	
+	virtual bool SetPlayRate(float InPlayRate) { return false; }
 };
 
 // Sequence player node that can be used with constant folding
@@ -135,8 +143,8 @@ protected:
 
 public:
 	// FAnimNode_SequencePlayerBase interface
-	virtual void SetSequence(UAnimSequenceBase* InSequence) override;
-	virtual void SetLoopAnimation(bool bInLoopAnimation) override;
+	virtual bool SetSequence(UAnimSequenceBase* InSequence) override;
+	virtual bool SetLoopAnimation(bool bInLoopAnimation) override;
 	virtual UAnimSequenceBase* GetSequence() const override;
 	virtual float GetPlayRateBasis() const override;
 	virtual float GetPlayRate() const override;
@@ -144,16 +152,18 @@ public:
 	virtual float GetStartPosition() const override;
 	virtual bool GetLoopAnimation() const override;
 	virtual bool GetStartFromMatchingPose() const override;
+	virtual bool SetStartPosition(float InStartPosition) override;
+	virtual bool SetPlayRate(float InPlayRate) override;
 
 	// FAnimNode_AssetPlayerBase interface
 	virtual FName GetGroupName() const override;
 	virtual EAnimGroupRole::Type GetGroupRole() const override;
 	virtual EAnimSyncMethod GetGroupMethod() const override;
 	virtual bool GetIgnoreForRelevancyTest() const override;
-	virtual void SetGroupName(FName InGroupName) override;
-	virtual void SetGroupRole(EAnimGroupRole::Type InRole) override;
-	virtual void SetGroupMethod(EAnimSyncMethod InMethod) override;
-	virtual void SetIgnoreForRelevancyTest(bool bInIgnoreForRelevancyTest) override;	
+	virtual bool SetGroupName(FName InGroupName) override;
+	virtual bool SetGroupRole(EAnimGroupRole::Type InRole) override;
+	virtual bool SetGroupMethod(EAnimSyncMethod InMethod) override;
+	virtual bool SetIgnoreForRelevancyTest(bool bInIgnoreForRelevancyTest) override;
 };
 
 // Sequence player node that can be used standalone (without constant folding)
@@ -212,8 +222,8 @@ protected:
 
 public:
 	// FAnimNode_SequencePlayerBase interface
-	virtual void SetSequence(UAnimSequenceBase* InSequence) override { Sequence = InSequence; }
-	virtual void SetLoopAnimation(bool bInLoopAnimation) override { bLoopAnimation = bInLoopAnimation; }
+	virtual bool SetSequence(UAnimSequenceBase* InSequence) override { Sequence = InSequence; return true; }
+	virtual bool SetLoopAnimation(bool bInLoopAnimation) override { bLoopAnimation = bInLoopAnimation; return true; }
 	virtual UAnimSequenceBase* GetSequence() const override { return Sequence; }
 	virtual float GetPlayRateBasis() const override { return PlayRateBasis; }
 	virtual float GetPlayRate() const override { return PlayRate; }
@@ -227,8 +237,10 @@ public:
 	virtual EAnimGroupRole::Type GetGroupRole() const override { return GroupRole; }
 	virtual EAnimSyncMethod GetGroupMethod() const override { return Method; }
 	virtual bool GetIgnoreForRelevancyTest() const override { return bIgnoreForRelevancyTest; }
-	virtual void SetGroupName(FName InGroupName) override { GroupName = InGroupName; }
-	virtual void SetGroupRole(EAnimGroupRole::Type InRole) override { GroupRole = InRole; }
-	virtual void SetGroupMethod(EAnimSyncMethod InMethod) override { Method = InMethod; }
-	virtual void SetIgnoreForRelevancyTest(bool bInIgnoreForRelevancyTest) override { bIgnoreForRelevancyTest = bInIgnoreForRelevancyTest; }
+	virtual bool SetGroupName(FName InGroupName) override { GroupName = InGroupName; return true; }
+	virtual bool SetGroupRole(EAnimGroupRole::Type InRole) override { GroupRole = InRole; return true; }
+	virtual bool SetGroupMethod(EAnimSyncMethod InMethod) override { Method = InMethod; return true; }
+	virtual bool SetIgnoreForRelevancyTest(bool bInIgnoreForRelevancyTest) override { bIgnoreForRelevancyTest = bInIgnoreForRelevancyTest; return true; }
+	virtual bool SetStartPosition(float InStartPosition) override { StartPosition = InStartPosition; return  true; }
+	virtual bool SetPlayRate(float InPlayRate) override { PlayRate = InPlayRate; return true; }
 };
diff --git a/Engine/Source/Runtime/Engine/Private/AnimBlueprint.cpp b/Engine/Source/Runtime/Engine/Private/AnimBlueprint.cpp
--- a/Engine/Source/Runtime/Engine/Private/AnimBlueprint.cpp
+++ b/Engine/Source/Runtime/Engine/Private/AnimBlueprint.cpp
@@ -216,6 +216,14 @@ bool UAnimBlueprint::FindDiffs(const UBlueprint* OtherBlueprint, FDiffResults& R
 	return true;
 }
 
+void UAnimBlueprint::SetObjectBeingDebugged(UObject* NewObject)
+{
+	// Look for any linked instances and set them up too if they are not already open in an editor
+	AnimationEditorUtils::SetupDebugLinkedAnimInstances(this, NewObject);
+
+	Super::SetObjectBeingDebugged(NewObject);
+}
+
 #endif
 
 USkeletalMesh* UAnimBlueprint::GetPreviewMesh(bool bFindIfNotSet/*=false*/)
@@ -339,4 +347,4 @@ FAnimBlueprintDebugData* UAnimBlueprint::GetDebugData() const
 #else
 	return nullptr;
 #endif // WITH_EDITORONLY_DATA
-}
\ No newline at end of file
+}
diff --git a/Engine/Source/Runtime/Engine/Private/Animation/AnimBlueprintGeneratedClass.cpp b/Engine/Source/Runtime/Engine/Private/Animation/AnimBlueprintGeneratedClass.cpp
--- a/Engine/Source/Runtime/Engine/Private/Animation/AnimBlueprintGeneratedClass.cpp
+++ b/Engine/Source/Runtime/Engine/Private/Animation/AnimBlueprintGeneratedClass.cpp
@@ -878,6 +878,12 @@ void UAnimBlueprintGeneratedClass::BuildConstantProperties()
 				{
 					ConstantSubsystemProperties.Add(StructProp);
 				}
+				// Initialize any anim node functions held in sparse data
+				else if(StructProp->Struct == FAnimNodeFunctionRef::StaticStruct())
+				{
+					FAnimNodeFunctionRef& AnimNodeFunction = *StructProp->ContainerPtrToValuePtr<FAnimNodeFunctionRef>(SparseClassData);
+					AnimNodeFunction.Initialize(this);
+				}
 			}
 			
 			ConstantProperties.Add(*It);
diff --git a/Engine/Source/Runtime/Engine/Private/Animation/AnimInstance.cpp b/Engine/Source/Runtime/Engine/Private/Animation/AnimInstance.cpp
--- a/Engine/Source/Runtime/Engine/Private/Animation/AnimInstance.cpp
+++ b/Engine/Source/Runtime/Engine/Private/Animation/AnimInstance.cpp
@@ -3067,6 +3067,29 @@ void UAnimInstance::RemoveExternalNotifyHandler(UObject* ExternalHandlerObject,
 	}
 }
 
+FAnimSubsystemInstance* UAnimInstance::FindSubsystem(UScriptStruct* InSubsystemType)
+{
+	if (IAnimClassInterface* AnimBlueprintClass = IAnimClassInterface::GetFromClass(GetClass()))
+	{
+		FAnimSubsystemInstance* Subsystem = nullptr;
+		
+		AnimBlueprintClass->ForEachSubsystem(this, [&Subsystem, InSubsystemType](const FAnimSubsystemInstanceContext& InContext)
+		{
+			if(InContext.SubsystemInstanceStruct == InSubsystemType)
+			{
+				Subsystem = &InContext.SubsystemInstance;
+				return EAnimSubsystemEnumeration::Stop;
+			}
+			
+			return EAnimSubsystemEnumeration::Continue;
+		});
+
+		return Subsystem;
+	}
+
+	return nullptr;
+}
+
 UAnimInstance* UAnimInstance::GetLinkedAnimLayerInstanceByGroup(FName InGroup) const
 {
 	if (IAnimClassInterface* AnimBlueprintClass = IAnimClassInterface::GetFromClass(GetClass()))
diff --git a/Engine/Source/Runtime/Engine/Private/Animation/AnimInstanceProxy.cpp b/Engine/Source/Runtime/Engine/Private/Animation/AnimInstanceProxy.cpp
--- a/Engine/Source/Runtime/Engine/Private/Animation/AnimInstanceProxy.cpp
+++ b/Engine/Source/Runtime/Engine/Private/Animation/AnimInstanceProxy.cpp
@@ -282,6 +282,8 @@ void FAnimInstanceProxy::InitializeRootNode(bool bInDeferRootNodeInitialization)
 	{
 		bDeferRootNodeInitialization = true;
 	}
+
+	bInitializeSubsystems = true;
 }
 
 void FAnimInstanceProxy::InitializeRootNode_WithRoot(FAnimNode_Base* InRootNode)
@@ -303,8 +305,8 @@ void FAnimInstanceProxy::InitializeRootNode_WithRoot(FAnimNode_Base* InRootNode)
 		}
 		else
 		{
-		InRootNode->Initialize_AnyThread(InitContext);
-	}
+			InRootNode->Initialize_AnyThread(InitContext);
+		}
 	}
 }
 
@@ -1065,6 +1067,16 @@ void FAnimInstanceProxy::UpdateAnimation_WithRoot(const FAnimationUpdateContext&
 
 	if(InRootNode == RootNode)
 	{
+		if(bInitializeSubsystems && AnimClassInterface)
+		{
+			AnimClassInterface->ForEachSubsystem(GetAnimInstanceObject(), [this](const FAnimSubsystemInstanceContext& InContext)
+			{
+				InContext.SubsystemInstance.Initialize_WorkerThread();
+				return EAnimSubsystemEnumeration::Continue;
+			});
+			bInitializeSubsystems = false;
+		}
+		
 		if(bDeferRootNodeInitialization)
 		{
 			InitializeRootNode_WithRoot(RootNode);
diff --git a/Engine/Source/Runtime/Engine/Private/Animation/AnimNodeBase.cpp b/Engine/Source/Runtime/Engine/Private/Animation/AnimNodeBase.cpp
--- a/Engine/Source/Runtime/Engine/Private/Animation/AnimNodeBase.cpp
+++ b/Engine/Source/Runtime/Engine/Private/Animation/AnimNodeBase.cpp
@@ -6,8 +6,11 @@
 #include "Animation/AnimTrace.h"
 #include "UObject/CoreObjectVersion.h"
 #include "Animation/AnimAttributes.h"
+#include "Animation/AnimNode_LinkedAnimGraph.h"
 #include "Animation/AnimSubsystem_Base.h"
 #include "Animation/AnimRootMotionProvider.h"
+#include "Animation/AnimSubsystem_NodeRelevancy.h"
+#include "Animation/AnimNodeContext.h"
 
 /////////////////////////////////////////////////////
 // FAnimationBaseContext
@@ -150,6 +153,81 @@ void FAnimNode_Base::ResetDynamics(ETeleportType InTeleportType)
 	PRAGMA_ENABLE_DEPRECATION_WARNINGS
 }
 
+const FAnimNodeFunctionRef& FAnimNode_Base::GetInitializeFunction() const
+{
+	return GET_ANIM_NODE_DATA(FAnimNodeFunctionRef, InitializeFunction);
+}
+
+const FAnimNodeFunctionRef& FAnimNode_Base::GetBecomeRelevantFunction() const
+{
+	return GET_ANIM_NODE_DATA(FAnimNodeFunctionRef, BecomeRelevantFunction);
+}
+
+const FAnimNodeFunctionRef& FAnimNode_Base::GetUpdateFunction() const
+{
+	return GET_ANIM_NODE_DATA(FAnimNodeFunctionRef, UpdateFunction);
+}
+
+const FAnimNodeFunctionRef& FAnimNode_Base::GetEvaluateFunction() const
+{
+	return GET_ANIM_NODE_DATA(FAnimNodeFunctionRef, EvaluateFunction);
+}
+
+// Private impl used to call anim node functions from pose links
+struct FAnimNodeFunctionCaller
+{
+private:
+	template<typename ContextType>
+	static void CallFunctionHelper(const FAnimNodeFunctionRef& InFunction, ContextType InContext, FAnimNode_Base& InNode)
+	{
+		if(InFunction.IsValid())
+		{
+			IAnimClassInterface* AnimBlueprintClass = InContext.GetAnimClass();
+			check(AnimBlueprintClass);
+
+			const TArray<FStructProperty*>& AnimNodeProperties = AnimBlueprintClass->GetAnimNodeProperties();
+			UScriptStruct* ScriptStruct = AnimNodeProperties[InNode.GetNodeIndex()]->Struct;
+			TSharedRef<FAnimNodeContext::FData> ContextData = MakeShared<FAnimNodeContext::FData>(InContext, InNode, ScriptStruct);
+			FAnimNodeContext NodeContext(ContextData);
+			InFunction.Call(InContext.GetAnimInstanceObject(), &NodeContext);
+		}
+	}
+
+public:
+	static void Initialize(const FAnimationInitializeContext& InContext, FAnimNode_Base& InNode)
+	{
+		CallFunctionHelper(InNode.GetInitializeFunction(), InContext, InNode);
+	}
+	
+	static void BecomeRelevant(const FAnimationUpdateContext& InContext, FAnimNode_Base& InNode)
+	{
+		const FAnimNodeFunctionRef& Function = InNode.GetBecomeRelevantFunction();
+		if(Function.IsValid())
+		{
+			FAnimSubsystemInstance_NodeRelevancy& RelevancySubsystem = CastChecked<UAnimInstance>(InContext.GetAnimInstanceObject())->GetSubsystem<FAnimSubsystemInstance_NodeRelevancy>();
+			FAnimNodeRelevancyStatus Status = RelevancySubsystem.UpdateNodeRelevancy(InContext, InNode);
+			if(Status.HasJustBecomeRelevant())
+			{
+				CallFunctionHelper(Function, InContext, InNode);
+			}
+		}
+	}
+	
+	static void Update(const FAnimationUpdateContext& InContext, FAnimNode_Base& InNode)
+	{
+		CallFunctionHelper(InNode.GetUpdateFunction(), InContext, InNode);
+	}
+	
+	static void Evaluate(FPoseContext& InContext, FAnimNode_Base& InNode)
+	{
+		CallFunctionHelper(InNode.GetEvaluateFunction(), InContext, InNode);
+	}
+
+	static void EvaluateComponentSpace(FComponentSpacePoseContext& InContext, FAnimNode_Base& InNode)
+	{
+		CallFunctionHelper(InNode.GetEvaluateFunction(), InContext, InNode);
+	}
+};
 
 /////////////////////////////////////////////////////
 // FPoseLinkBase
@@ -196,6 +274,7 @@ void FPoseLinkBase::Initialize(const FAnimationInitializeContext& InContext)
 		FAnimationInitializeContext LinkContext(InContext);
 		LinkContext.SetNodeId(LinkID);
 		TRACE_SCOPED_ANIM_NODE(LinkContext);
+		FAnimNodeFunctionCaller::Initialize(LinkContext, *LinkedNode);
 		LinkedNode->Initialize_AnyThread(LinkContext);
 	}
 }
@@ -306,6 +385,8 @@ void FPoseLinkBase::Update(const FAnimationUpdateContext& InContext)
 	{
 		FAnimationUpdateContext LinkContext(InContext.WithNodeId(LinkID));
 		TRACE_SCOPED_ANIM_NODE(LinkContext);
+		FAnimNodeFunctionCaller::BecomeRelevant(LinkContext, *LinkedNode);
+		FAnimNodeFunctionCaller::Update(LinkContext, *LinkedNode);
 		LinkedNode->Update_AnyThread(LinkContext);
 	}
 }
@@ -355,6 +436,7 @@ void FPoseLink::Evaluate(FPoseContext& Output)
 		{
 			Output.SetNodeId(LinkID);
 			TRACE_SCOPED_ANIM_NODE(Output);
+			FAnimNodeFunctionCaller::Evaluate(Output, *LinkedNode);
 			LinkedNode->Evaluate_AnyThread(Output);
 			TRACE_ANIM_NODE_BLENDABLE_ATTRIBUTES(Output, SourceID, LinkID);
 		}
@@ -465,6 +547,7 @@ void FComponentSpacePoseLink::EvaluateComponentSpace(FComponentSpacePoseContext&
 		{
 			Output.SetNodeId(LinkID);
 			TRACE_SCOPED_ANIM_NODE(Output);
+			FAnimNodeFunctionCaller::EvaluateComponentSpace(Output, *LinkedNode);
 			LinkedNode->EvaluateComponentSpace_AnyThread(Output);
 			TRACE_ANIM_NODE_BLENDABLE_ATTRIBUTES(Output, SourceID, LinkID);
 		}
diff --git a/Engine/Source/Runtime/Engine/Private/Animation/AnimNodeContext.cpp b/Engine/Source/Runtime/Engine/Private/Animation/AnimNodeContext.cpp
new file mode 100644
--- /dev/null
+++ b/Engine/Source/Runtime/Engine/Private/Animation/AnimNodeContext.cpp
@@ -0,0 +1,84 @@
+// Copyright Epic Games, Inc. All Rights Reserved.
+
+#include "Animation/AnimNodeContext.h"
+#include "Animation/AnimNodeBase.h"
+
+FAnimNodeContext::FData::FData(const FAnimationInitializeContext& InContext, FAnimNode_Base& InAnimNode, UScriptStruct* InAnimNodeStruct)
+{
+	AnimNode = &InAnimNode;
+    AnimNodeStruct = InAnimNodeStruct;
+    Context = const_cast<FAnimationInitializeContext*>(&InContext);
+	ContextType = EContextType::Initialize;
+}
+
+FAnimNodeContext::FData::FData(const FAnimationUpdateContext& InContext, FAnimNode_Base& InAnimNode, UScriptStruct* InAnimNodeStruct)
+{
+	AnimNode = &InAnimNode;
+	AnimNodeStruct = InAnimNodeStruct;
+	Context = const_cast<FAnimationUpdateContext*>(&InContext);
+	ContextType = EContextType::Update;
+}
+
+FAnimNodeContext::FData::FData(FPoseContext& InContext, FAnimNode_Base& InAnimNode, UScriptStruct* InAnimNodeStruct)
+{
+	AnimNode = &InAnimNode;
+	AnimNodeStruct = InAnimNodeStruct;
+	Context = &InContext;
+	ContextType = EContextType::Pose;
+}
+
+FAnimNodeContext::FData::FData(FComponentSpacePoseContext& InContext, FAnimNode_Base& InAnimNode, UScriptStruct* InAnimNodeStruct)
+{
+	AnimNode = &InAnimNode;
+	AnimNodeStruct = InAnimNodeStruct;
+	Context = &InContext;
+	ContextType = EContextType::ComponentSpacePose;
+}
+
+FAnimationInitializeContext* FAnimNodeContext::GetInitializeContext() const
+{
+	if(TSharedPtr<FData> PinnedData = Data.Pin())
+	{
+		if(PinnedData->ContextType == FData::EContextType::Initialize)
+		{
+			return static_cast<FAnimationInitializeContext*>(PinnedData->Context);
+		}
+	}
+	return nullptr;
+}
+
+FAnimationUpdateContext* FAnimNodeContext::GetUpdateContext() const
+{
+	if(TSharedPtr<FData> PinnedData = Data.Pin())
+	{
+		if(PinnedData->ContextType == FData::EContextType::Update)
+		{
+			return static_cast<FAnimationUpdateContext*>(PinnedData->Context);
+		}
+	}
+	return nullptr;
+}
+
+FPoseContext* FAnimNodeContext::GetPoseContext() const
+{
+	if(TSharedPtr<FData> PinnedData = Data.Pin())
+	{
+		if(PinnedData->ContextType == FData::EContextType::Pose)
+		{
+			return static_cast<FPoseContext*>(PinnedData->Context);
+		}
+	}
+	return nullptr;
+}
+
+FComponentSpacePoseContext* FAnimNodeContext::GetComponentSpacePoseContext() const
+{
+	if(TSharedPtr<FData> PinnedData = Data.Pin())
+	{
+		if(PinnedData->ContextType == FData::EContextType::ComponentSpacePose)
+		{
+			return static_cast<FComponentSpacePoseContext*>(PinnedData->Context);
+		}
+	}
+	return nullptr;
+}
\ No newline at end of file
diff --git a/Engine/Source/Runtime/Engine/Private/Animation/AnimNodeData.cpp b/Engine/Source/Runtime/Engine/Private/Animation/AnimNodeData.cpp
--- a/Engine/Source/Runtime/Engine/Private/Animation/AnimNodeData.cpp
+++ b/Engine/Source/Runtime/Engine/Private/Animation/AnimNodeData.cpp
@@ -64,6 +64,28 @@ void* FAnimNodeData::GetMutableData(UE::Anim::FNodeDataId InId, FAnimNode_Base*
 }
 #endif
 
+void* FAnimNodeData::GetInstanceData(UE::Anim::FNodeDataId InId, FAnimNode_Base* InNode, UObject* InCurrentObject) const
+{
+	check(Entries.IsValidIndex(InId.Index));
+	
+	const uint32 Entry = Entries[InId.Index];
+	check(Entry != ANIM_NODE_DATA_INVALID_ENTRY);
+	
+	uint32 EntryIndex = (Entry & ANIM_NODE_DATA_INSTANCE_DATA_MASK);
+	if((Entry & ANIM_NODE_DATA_INSTANCE_DATA_FLAG) != 0)
+	{
+		// Use the supplied object or find the object ptr by walking the property chain from this node
+		const UObject* CurrentObject = InCurrentObject ? InCurrentObject : IAnimClassInterface::GetObjectPtrFromAnimNode(&(*AnimClassInterface), InNode);
+
+		// Check the current context is expected
+		check(CurrentObject && CurrentObject->GetClass()->IsChildOf(IAnimClassInterface::GetActualAnimClass(&(*AnimClassInterface))));
+
+		return const_cast<void*>(AnimClassInterface->GetMutableNodeValueRaw(EntryIndex, CurrentObject));
+	}
+
+	return nullptr;
+}
+
 FAnimNodeStructData::FAnimNodeStructData(const UScriptStruct* InNodeType)
 {
 	// Iterated properties are not in the same order they are laid out in memory, so we extract and sort them here.
diff --git a/Engine/Source/Runtime/Engine/Private/Animation/AnimNodeFunctionRef.cpp b/Engine/Source/Runtime/Engine/Private/Animation/AnimNodeFunctionRef.cpp
new file mode 100644
--- /dev/null
+++ b/Engine/Source/Runtime/Engine/Private/Animation/AnimNodeFunctionRef.cpp
@@ -0,0 +1,19 @@
+// Copyright Epic Games, Inc. All Rights Reserved.
+
+#include "Animation/AnimNodeFunctionRef.h"
+
+void FAnimNodeFunctionRef::Initialize(const UClass* InClass)
+{
+	if(FunctionName != NAME_None)
+	{
+		Function = InClass->FindFunctionByName(FunctionName);
+	}
+}
+
+void FAnimNodeFunctionRef::Call(UObject* InObject, void* InParameters) const
+{
+	if(IsValid())
+	{
+		InObject->ProcessEvent(Function, InParameters);
+	}
+}
\ No newline at end of file
diff --git a/Engine/Source/Runtime/Engine/Private/Animation/AnimNode_SequencePlayer.cpp b/Engine/Source/Runtime/Engine/Private/Animation/AnimNode_SequencePlayer.cpp
--- a/Engine/Source/Runtime/Engine/Private/Animation/AnimNode_SequencePlayer.cpp
+++ b/Engine/Source/Runtime/Engine/Private/Animation/AnimNode_SequencePlayer.cpp
@@ -168,19 +168,35 @@ float FAnimNode_SequencePlayerBase::GetEffectiveStartPosition(const FAnimationBa
 	return GetStartPosition();
 }
 
-void FAnimNode_SequencePlayer::SetSequence(UAnimSequenceBase* InSequence)
+bool FAnimNode_SequencePlayer::SetSequence(UAnimSequenceBase* InSequence)
 {
 #if WITH_EDITORONLY_DATA
-	TObjectPtr<UAnimSequenceBase>& SequenceToSet = GET_MUTABLE_ANIM_NODE_DATA(TObjectPtr<UAnimSequenceBase>, Sequence);
-	SequenceToSet = InSequence;
+	Sequence = InSequence;
+	GET_MUTABLE_ANIM_NODE_DATA(TObjectPtr<UAnimSequenceBase>, Sequence) = InSequence;
 #endif
+	
+	if(TObjectPtr<UAnimSequenceBase>* SequencePtr = GET_INSTANCE_ANIM_NODE_DATA_PTR(TObjectPtr<UAnimSequenceBase>, Sequence))
+	{
+		*SequencePtr = InSequence;
+		return true;
+	}
+
+	return false;
 }
 
-void FAnimNode_SequencePlayer::SetLoopAnimation(bool bInLoopAnimation)
+bool FAnimNode_SequencePlayer::SetLoopAnimation(bool bInLoopAnimation)
 {
 #if WITH_EDITORONLY_DATA
-	bLoopAnimation = bInLoopAnimation; 
+	bLoopAnimation = bInLoopAnimation;
 #endif
+	
+	if(bool* bLoopAnimationPtr = GET_INSTANCE_ANIM_NODE_DATA_PTR(bool, bLoopAnimation))
+	{
+		*bLoopAnimationPtr = bInLoopAnimation;
+		return true;
+	}
+
+	return false;
 }
 
 UAnimSequenceBase* FAnimNode_SequencePlayer::GetSequence() const
@@ -238,32 +254,87 @@ bool FAnimNode_SequencePlayer::GetIgnoreForRelevancyTest() const
 	return GET_ANIM_NODE_DATA(bool, bIgnoreForRelevancyTest);
 }
 
-void FAnimNode_SequencePlayer::SetGroupName(FName InGroupName)
+bool FAnimNode_SequencePlayer::SetGroupName(FName InGroupName)
 {
 #if WITH_EDITORONLY_DATA
 	GroupName = InGroupName;
 #endif
+
+	if(FName* GroupNamePtr = GET_INSTANCE_ANIM_NODE_DATA_PTR(FName, GroupName))
+	{
+		*GroupNamePtr = InGroupName;
+		return true;
+	}
+
+	return false;
 }
 
-void FAnimNode_SequencePlayer::SetGroupRole(EAnimGroupRole::Type InRole)
+bool FAnimNode_SequencePlayer::SetGroupRole(EAnimGroupRole::Type InRole)
 {
 #if WITH_EDITORONLY_DATA
 	GroupRole = InRole;
 #endif
+
+	if(TEnumAsByte<EAnimGroupRole::Type>* GroupRolePtr = GET_INSTANCE_ANIM_NODE_DATA_PTR(TEnumAsByte<EAnimGroupRole::Type>, GroupRole))
+	{
+		*GroupRolePtr = InRole;
+		return true;
+	}
+
+	return false;
 }
 
-void FAnimNode_SequencePlayer::SetGroupMethod(EAnimSyncMethod InMethod)
+bool FAnimNode_SequencePlayer::SetGroupMethod(EAnimSyncMethod InMethod)
 {
 #if WITH_EDITORONLY_DATA
 	Method = InMethod;
 #endif
+
+	if(EAnimSyncMethod* MethodPtr = GET_INSTANCE_ANIM_NODE_DATA_PTR(EAnimSyncMethod, Method))
+	{
+		*MethodPtr = InMethod;
+		return true;
+	}
+
+	return false;
 }
 
-void FAnimNode_SequencePlayer::SetIgnoreForRelevancyTest(bool bInIgnoreForRelevancyTest)
+bool FAnimNode_SequencePlayer::SetIgnoreForRelevancyTest(bool bInIgnoreForRelevancyTest)
 {
 #if WITH_EDITORONLY_DATA
 	bIgnoreForRelevancyTest = bInIgnoreForRelevancyTest;
 #endif
+
+	if(bool* bIgnoreForRelevancyTestPtr = GET_INSTANCE_ANIM_NODE_DATA_PTR(bool, bIgnoreForRelevancyTest))
+	{
+		*bIgnoreForRelevancyTestPtr = bInIgnoreForRelevancyTest;
+		return true;
+	}
+
+	return false;
+}
+
+bool FAnimNode_SequencePlayer::SetStartPosition(float InStartPosition)
+{
+	if(float* StartPositionPtr = GET_INSTANCE_ANIM_NODE_DATA_PTR(float, StartPosition))
+	{
+		*StartPositionPtr = InStartPosition;
+		return true;
+	}
+
+	return false;
+}
+
+
+bool FAnimNode_SequencePlayer::SetPlayRate(float InPlayRate)
+{
+	if(float* PlayRatePtr = GET_INSTANCE_ANIM_NODE_DATA_PTR(float, PlayRate))
+	{
+		*PlayRatePtr = InPlayRate;
+		return true;
+	}
+
+	return false;
 }
 
 #undef LOCTEXT_NAMESPACE
diff --git a/Engine/Source/Runtime/Engine/Private/Animation/AnimSubsystem_NodeRelevancy.cpp b/Engine/Source/Runtime/Engine/Private/Animation/AnimSubsystem_NodeRelevancy.cpp
new file mode 100644
--- /dev/null
+++ b/Engine/Source/Runtime/Engine/Private/Animation/AnimSubsystem_NodeRelevancy.cpp
@@ -0,0 +1,78 @@
+﻿// Copyright Epic Games, Inc. All Rights Reserved.
+
+#include "Animation/AnimSubsystem_NodeRelevancy.h"
+#include "Animation/AnimInstanceProxy.h"
+
+float FAnimNodeRelevancyStatus::GetCurrentWeight() const
+{
+	return CurrentWeight;
+}
+
+float FAnimNodeRelevancyStatus::GetPreviousWeight() const
+{
+	return PreviousWeight;
+}
+
+bool FAnimNodeRelevancyStatus::IsStartingBlendingOut() const
+{
+	return FAnimWeight::IsFullWeight(PreviousWeight) && !FAnimWeight::IsFullWeight(CurrentWeight);
+}
+
+bool FAnimNodeRelevancyStatus::IsStartingBlendingIn() const
+{
+	return !FAnimWeight::IsRelevant(PreviousWeight) && FAnimWeight::IsRelevant(CurrentWeight);
+}
+
+bool FAnimNodeRelevancyStatus::IsFinishingBlendingIn() const
+{
+	return !FAnimWeight::IsFullWeight(PreviousWeight) && FAnimWeight::IsFullWeight(CurrentWeight);
+}
+
+bool FAnimNodeRelevancyStatus::IsFinishingBlendingOut() const
+{
+	return FAnimWeight::IsRelevant(PreviousWeight) && !FAnimWeight::IsRelevant(CurrentWeight);
+}
+
+bool FAnimNodeRelevancyStatus::HasFullyBlendedIn() const
+{
+	return FAnimWeight::IsFullWeight(PreviousWeight) && FAnimWeight::IsFullWeight(CurrentWeight);
+}
+
+bool FAnimNodeRelevancyStatus::HasFullyBlendedOut() const
+{
+	return !FAnimWeight::IsRelevant(PreviousWeight) && !FAnimWeight::IsRelevant(CurrentWeight);
+}
+
+void FAnimSubsystemInstance_NodeRelevancy::Initialize_WorkerThread()
+{
+	NodeTrackers.Empty();
+}
+
+FAnimNodeRelevancyStatus FAnimSubsystemInstance_NodeRelevancy::UpdateNodeRelevancy(const FAnimationUpdateContext& InContext, const FAnimNode_Base& InNode)
+{
+	FTracker& Tracker = NodeTrackers.FindOrAdd(&InNode);
+	
+	if(!Tracker.Counter.HasEverBeenUpdated() || !Tracker.Counter.WasSynchronizedCounter(InContext.AnimInstanceProxy->GetUpdateCounter()))
+	{
+		Tracker.Status.CurrentWeight = 0.0f;
+	}
+
+	if(!Tracker.Counter.IsSynchronized_All(InContext.AnimInstanceProxy->GetUpdateCounter()))
+	{
+		Tracker.Counter.SynchronizeWith(InContext.AnimInstanceProxy->GetUpdateCounter());
+		Tracker.Status.PreviousWeight = Tracker.Status.CurrentWeight;
+		Tracker.Status.CurrentWeight = InContext.GetFinalBlendWeight();
+	}
+	
+	return Tracker.Status;
+}
+
+FAnimNodeRelevancyStatus FAnimSubsystemInstance_NodeRelevancy::GetNodeRelevancy(const FAnimNode_Base& InNode) const
+{
+	if(const FTracker* Tracker = NodeTrackers.Find(&InNode))
+	{
+		return Tracker->Status;
+	}
+
+	return FAnimNodeRelevancyStatus();
+}
diff --git a/Engine/Source/Runtime/Engine/Public/Animation/AnimInstanceProxy.h b/Engine/Source/Runtime/Engine/Public/Animation/AnimInstanceProxy.h
--- a/Engine/Source/Runtime/Engine/Public/Animation/AnimInstanceProxy.h
+++ b/Engine/Source/Runtime/Engine/Public/Animation/AnimInstanceProxy.h
@@ -140,6 +140,7 @@ public:
 #if WITH_EDITORONLY_DATA
 		, bIsBeingDebugged(false)
 #endif
+		, bInitializeSubsystems(false)
 	{
 	}
 
@@ -166,6 +167,7 @@ public:
 #if WITH_EDITORONLY_DATA
 		, bIsBeingDebugged(false)
 #endif
+		, bInitializeSubsystems(false)
 	{
 	}
 
@@ -1157,4 +1159,7 @@ private:
 	/** Whether this UAnimInstance is currently being debugged in the editor */
 	uint8 bIsBeingDebugged : 1;
 #endif
+
+	// Whether subsystems should be initialized
+	uint8 bInitializeSubsystems : 1;
 };
diff --git a/Engine/Source/Runtime/Engine/Public/Animation/AnimNodeData.h b/Engine/Source/Runtime/Engine/Public/Animation/AnimNodeData.h
--- a/Engine/Source/Runtime/Engine/Public/Animation/AnimNodeData.h
+++ b/Engine/Source/Runtime/Engine/Public/Animation/AnimNodeData.h
@@ -81,12 +81,18 @@ public:
 	 * @return the raw data value, either held on sparse class data or on the instance 
 	 */
 	const void* GetData(UE::Anim::FNodeDataId InId, const FAnimNode_Base* InNode, const UObject* InCurrentObject = nullptr) const;
-	
+
 #if WITH_EDITORONLY_DATA
-	/** Get the specified mutable data for the specified instance. The data may or may not reside on an anim instance itself as it may have been folded into constants. */
+	/**
+	 * Get the specified mutable data for the specified instance. The data may or may not reside on an anim instance itself as it may have been folded into constants.
+	 * Only available in editor to support patching constant data during compilation. Not for use at runtime. 
+	 */
 	void* GetMutableData(UE::Anim::FNodeDataId InId, FAnimNode_Base* InNode, UObject* InCurrentObject = nullptr) const;
 #endif
 
+	/** Get the specified mutable data for the specified instance. If the data is not held on the instance this will return nullptr. */
+	void* GetInstanceData(UE::Anim::FNodeDataId InId, FAnimNode_Base* InNode, UObject* InCurrentObject = nullptr) const;
+	
 	/** The class we are part of */
 	const IAnimClassInterface& GetAnimClassInterface() const { check(AnimClassInterface); return *AnimClassInterface; }
 
@@ -101,19 +107,19 @@ private:
 	UPROPERTY()
 	TScriptInterface<IAnimClassInterface> AnimClassInterface = nullptr;
 
-	/** 
-	 * The index of the node for this constant data block in the class that it is held in. 
-	 * INDEX_NONE if this node is not in a generated class or is per-instance data. 
-	 */
-	UPROPERTY()
-	int32 NodeIndex = INDEX_NONE;
-
 	/**
 	 * Flags & Indices for table entries. Used to look up indices from NodePropertyIndex->FoldedDataIndex.
 	 * If the MSB is set for an entry it is assumed to be on an instance, if not, it is assumed to be stored on the class.
 	 */
 	UPROPERTY()
 	TArray<uint32> Entries;
+	
+	/** 
+	 * The index of the node for this constant data block in the class that it is held in. 
+	 * INDEX_NONE if this node is not in a generated class or is per-instance data. 
+	 */
+	UPROPERTY()
+	int32 NodeIndex = INDEX_NONE;
 };
 
 /**
diff --git a/Engine/Source/Runtime/Engine/Public/Animation/AnimSubsystemInstance.h b/Engine/Source/Runtime/Engine/Public/Animation/AnimSubsystemInstance.h
--- a/Engine/Source/Runtime/Engine/Public/Animation/AnimSubsystemInstance.h
+++ b/Engine/Source/Runtime/Engine/Public/Animation/AnimSubsystemInstance.h
@@ -9,4 +9,9 @@ USTRUCT()
 struct FAnimSubsystemInstance
 {
 	GENERATED_BODY()
+
+	virtual ~FAnimSubsystemInstance() {}
+
+	// Override point used to initialize per-instance data. Called on a worker thread.
+	virtual void Initialize_WorkerThread() {}
 };
\ No newline at end of file
diff --git a/Engine/Source/Runtime/Engine/Public/Animation/AnimSubsystem_NodeRelevancy.h b/Engine/Source/Runtime/Engine/Public/Animation/AnimSubsystem_NodeRelevancy.h
new file mode 100644
--- /dev/null
+++ b/Engine/Source/Runtime/Engine/Public/Animation/AnimSubsystem_NodeRelevancy.h
@@ -0,0 +1,89 @@
+// Copyright Epic Games, Inc. All Rights Reserved.
+
+#pragma once
+
+#include "AnimSubsystemInstance.h"
+#include "Animation/AnimTypes.h"
+#include "AnimSubsystem_NodeRelevancy.generated.h"
+
+struct FAnimationUpdateContext;
+struct FAnimNode_Base;
+class UAnimInstance;
+
+// Relevancy status of a animation node tracked by this subsystem
+struct ENGINE_API FAnimNodeRelevancyStatus
+{
+public:
+	// Get the current weight
+	float GetCurrentWeight() const;
+
+	// Get the previous weight
+	float GetPreviousWeight() const;
+	
+	// Node has zero weight
+	bool HasFullyBlendedOut() const;
+
+	// Node has just become relevant
+	bool HasJustBecomeRelevant() const { return IsStartingBlendingIn(); }
+	
+	// Node has weight and had zero weight last update
+	bool IsStartingBlendingIn() const;
+
+	// Node has full weight and had non-full weight last update
+	bool IsFinishingBlendingIn() const;
+
+	// Node has full weight and had full weight last update
+	bool HasFullyBlendedIn() const;
+
+	// Node has non-full weight and had full weight last update
+	bool IsStartingBlendingOut() const;
+
+	// Node has zero weight and had non-zero weight last update
+	bool IsFinishingBlendingOut() const;
+
+private:
+	friend struct FAnimSubsystemInstance_NodeRelevancy;
+	
+	// Previous weight 
+	float PreviousWeight = 0.0f;
+
+	// Current weight
+	float CurrentWeight = 0.0f;
+};
+
+/** Allows anim node's relevancy (when they receive and lose weight in the graph) to be tracked */
+USTRUCT()
+struct ENGINE_API FAnimSubsystemInstance_NodeRelevancy : public FAnimSubsystemInstance
+{
+	GENERATED_BODY()
+
+	// FAnimSubsystemInstance interface
+	virtual void Initialize_WorkerThread() override;
+	
+	// Update the relevancy of the passed-in node using the supplied context
+	FAnimNodeRelevancyStatus UpdateNodeRelevancy(const FAnimationUpdateContext& InContext, const FAnimNode_Base& InNode);
+
+	// Update the relevancy of the passed-in instance using the supplied context
+	FAnimNodeRelevancyStatus UpdateInstanceRelevancy(const FAnimationUpdateContext& InContext, const UAnimInstance& InInstance);
+	
+	// Get the tracked relevancy of the passed-in node. If the node is not tracked the relevancy will be default (zero weighted).
+	FAnimNodeRelevancyStatus GetNodeRelevancy(const FAnimNode_Base& InNode) const;
+
+	// Get the tracked relevancy of the passed-in instance. If the instance is not tracked the relevancy will be default (zero weighted).
+	FAnimNodeRelevancyStatus GetInstanceRelevancy(const UAnimInstance& InInstance) const;
+
+private:
+	// Tracks the relevancy of a node
+	struct FTracker
+	{
+		// Counter used to determine relevancy
+		FGraphTraversalCounter Counter;
+
+		// Status of the node's relevancy
+		FAnimNodeRelevancyStatus Status;
+	};
+
+private:
+	// Map of tracked anim nodes
+	TMap<const FAnimNode_Base*, FTracker> NodeTrackers;
+};
\ No newline at end of file
