Matinee/Editor: Removed more Matinee related functions from various Editor callbacks.

diff --git a/Engine/Source/Editor/LevelEditor/Private/SLevelEditor.cpp b/Engine/Source/Editor/LevelEditor/Private/SLevelEditor.cpp
--- a/Engine/Source/Editor/LevelEditor/Private/SLevelEditor.cpp
+++ b/Engine/Source/Editor/LevelEditor/Private/SLevelEditor.cpp
@@ -1548,23 +1548,6 @@ void SLevelEditor::ToggleEditorMode( FEditorModeID ModeID )
 	FLevelEditorModule& LevelEditorModule = FModuleManager::GetModuleChecked<FLevelEditorModule>("LevelEditor");
 	TSharedPtr<FTabManager> LevelEditorTabManager = LevelEditorModule.GetLevelEditorTabManager();
 
-	// Prompt the user if Matinee must be closed before activating new mode
-	if (ModeID != FBuiltinEditorModes::EM_InterpEdit)
-	{
-		FEdMode* MatineeMode = GetEditorModeManager().GetActiveMode(FBuiltinEditorModes::EM_InterpEdit);
-		if (MatineeMode && !MatineeMode->IsCompatibleWith(ModeID))
-		{
-			FFormatNamedArguments Args;
-			Args.Add(TEXT("ModeName"), MatineeMode->GetModeInfo().Name);
-			FText Msg = FText::Format(NSLOCTEXT("LevelEditor", "ModeSwitchCloseMatineeQ", "Activating '{ModeName}' editor mode will close UnrealMatinee.  Continue?"), Args);
-
-			if (EAppReturnType::Yes != FMessageDialog::Open(EAppMsgType::YesNo, Msg))
-			{
-				return;
-			}
-		}
-	}
-
 	// Abort viewport tracking when switching editor mode
 	if (GCurrentLevelEditingViewportClient)
 	{
diff --git a/Engine/Source/Editor/UnrealEd/Classes/Editor/UnrealEdEngine.h b/Engine/Source/Editor/UnrealEd/Classes/Editor/UnrealEdEngine.h
--- a/Engine/Source/Editor/UnrealEd/Classes/Editor/UnrealEdEngine.h
+++ b/Engine/Source/Editor/UnrealEd/Classes/Editor/UnrealEdEngine.h
@@ -123,10 +123,6 @@ public:
 	UPROPERTY()
 	TObjectPtr<class UMaterial> MaterialCopyPasteBuffer;
 
-	/** A buffer for implementing matinee track/group copy/paste. */
-	UPROPERTY()
-	TArray<TObjectPtr<class UObject>> MatineeCopyPasteBuffer;
-
 	/** A buffer for implementing sound cue nodes copy/paste. */
 	UPROPERTY()
 	TObjectPtr<class USoundCue> SoundCueCopyPasteBuffer;
@@ -826,7 +822,8 @@ public:
 	virtual bool CanSavePackage( UPackage* PackageToSave );
 
 	/** Converts kismet based matinees in the current level to matinees controlled via matinee actors */
-	void ConvertMatinees();
+	UE_DEPRECATED(5.0, "Matinee is no longer part of the editor.")
+	void ConvertMatinees() {}
 
 	/**
 	 * Updates the volume actor visibility for all viewports based on the passed in volume class
@@ -909,12 +906,8 @@ public:
 	bool HandleBuildPathsCommand( const TCHAR* Str, FOutputDevice& Ar, UWorld* InWorld );
 	bool HandleRecreateLandscapeCollisionCommand(const TCHAR* Str, FOutputDevice& Ar, UWorld* InWorld);
 	bool HandleRemoveLandscapeXYOffsetsCommand(const TCHAR* Str, FOutputDevice& Ar, UWorld* InWorld);
-	bool HandleConvertMatineesCommand( const TCHAR* Str, FOutputDevice& Ar, UWorld* InWorld );
 	bool HandleDisasmScriptCommand( const TCHAR* Str, FOutputDevice& Ar );	
 
-	/** OnEditorModeChanged delegate which looks for Matinee editor closing */
-	void UpdateEdModeOnMatineeClose(const FEditorModeID& EditorModeID, bool IsEntering);
-
 	bool IsComponentSelected(const UPrimitiveComponent* PrimComponent);
 
 	/** Return if we have write permission under the mount point this packages lives in. */
diff --git a/Engine/Source/Editor/UnrealEd/Private/AssetSelection.cpp b/Engine/Source/Editor/UnrealEd/Private/AssetSelection.cpp
--- a/Engine/Source/Editor/UnrealEd/Private/AssetSelection.cpp
+++ b/Engine/Source/Editor/UnrealEd/Private/AssetSelection.cpp
@@ -294,11 +294,6 @@ namespace AssetSelectionUtils
 						ActorInfo.bHaveEmitter = true;
 					}
 
-					if ( CurrentActor->IsA( AMatineeActor::StaticClass() ) )
-					{
-						ActorInfo.bHaveMatinee = true;
-					}
-
 					if ( CurrentActor->IsTemporarilyHiddenInEditor() )
 					{
 						ActorInfo.bHaveHidden = true;
diff --git a/Engine/Source/Editor/UnrealEd/Private/CameraController.cpp b/Engine/Source/Editor/UnrealEd/Private/CameraController.cpp
--- a/Engine/Source/Editor/UnrealEd/Private/CameraController.cpp
+++ b/Engine/Source/Editor/UnrealEd/Private/CameraController.cpp
@@ -356,8 +356,7 @@ void FEditorCameraController::UpdateRotation( const FCameraControllerUserImpulse
 				Angle += 360.f;
 			}
 			
-			//allow for unlocked pitch constraints while in matinee
-			if (Config.bLockedPitch || !GLevelEditorModeTools().GetActiveMode(FBuiltinEditorModes::EM_InterpEdit)) 
+			if (Config.bLockedPitch) 
 			{
 				// Clamp the angle.
 				InOutCameraEuler[ CurRotationAxis ] =
diff --git a/Engine/Source/Editor/UnrealEd/Private/EditorActor.cpp b/Engine/Source/Editor/UnrealEd/Private/EditorActor.cpp
--- a/Engine/Source/Editor/UnrealEd/Private/EditorActor.cpp
+++ b/Engine/Source/Editor/UnrealEd/Private/EditorActor.cpp
@@ -713,15 +713,6 @@ bool UUnrealEdEngine::CanDeleteActor(const AActor* InActor, FText* OutReason) co
 
 bool UUnrealEdEngine::ShouldAbortComponentDeletion(const TArray<UActorComponent*>& InComponentsToDelete, FText* OutReason) const
 {
-	if (!GLevelEditorModeTools().EnsureNotInMode(FBuiltinEditorModes::EM_InterpEdit))
-	{
-		if (OutReason)
-		{
-			*OutReason = LOCTEXT("ShouldAbortComponentDeletion_Error_WrongModeForComponentDeletion", "Cannot delete components while Matinee is open.");
-		}
-		return true;
-	}
-
 	for (UActorComponent* Component : InComponentsToDelete)
 	{
 		AActor* OwnerActor = Component->GetOwner();
@@ -740,15 +731,6 @@ bool UUnrealEdEngine::ShouldAbortComponentDeletion(const TArray<UActorComponent*
 
 bool UUnrealEdEngine::ShouldAbortActorDeletion(const TArray<AActor*>& InActorsToDelete, FText* OutReason) const
 {
-	if (!GLevelEditorModeTools().EnsureNotInMode(FBuiltinEditorModes::EM_InterpEdit))
-	{
-		if (OutReason)
-		{
-			*OutReason = LOCTEXT("ShouldAbortActorDeletion_WrongModeForActorDeletion", "Cannot delete actors while Matinee is open.");
-		}
-		return true;
-	}
-
 	for (AActor* Actor : InActorsToDelete)
 	{
 		if (FLevelUtils::IsLevelLocked(Actor))
diff --git a/Engine/Source/Editor/UnrealEd/Private/EditorServer.cpp b/Engine/Source/Editor/UnrealEd/Private/EditorServer.cpp
--- a/Engine/Source/Editor/UnrealEd/Private/EditorServer.cpp
+++ b/Engine/Source/Editor/UnrealEd/Private/EditorServer.cpp
@@ -2248,12 +2248,6 @@ void UEditorEngine::CreateNewMapForEditing()
 	
 	const FScopedBusyCursor BusyCursor;
 
-	// Change out of Matinee when opening new map, so we avoid editing data in the old one.
-	if( GLevelEditorModeTools().IsModeActive( FBuiltinEditorModes::EM_InterpEdit ) )
-	{
-		GLevelEditorModeTools().DeactivateMode( FBuiltinEditorModes::EM_InterpEdit );
-	}
-
 	// Also change out of Landscape mode to ensure all references are cleared.
 	if( GLevelEditorModeTools().IsModeActive( FBuiltinEditorModes::EM_Landscape ) )
 	{
diff --git a/Engine/Source/Editor/UnrealEd/Private/FileHelpers.cpp b/Engine/Source/Editor/UnrealEd/Private/FileHelpers.cpp
--- a/Engine/Source/Editor/UnrealEd/Private/FileHelpers.cpp
+++ b/Engine/Source/Editor/UnrealEd/Private/FileHelpers.cpp
@@ -210,29 +210,6 @@ namespace FileDialogHelpers
 	}
 }
 
-
-
-/**
- * Queries the user if they want to quit out of interpolation editing before save.
- *
- * @return		true if in interpolation editing mode, false otherwise.
- */
-static bool InInterpEditMode()
-{
-	// Must exit Interpolation Editing mode before you can save - so it can reset everything to its initial state.
-	if( GLevelEditorModeTools().IsModeActive( FBuiltinEditorModes::EM_InterpEdit ) )
-	{
-		const bool ExitInterp = EAppReturnType::Yes == FMessageDialog::Open( EAppMsgType::YesNo, EAppReturnType::Yes, NSLOCTEXT("UnrealEd", "Prompt_21", "You must close Matinee before saving level.\nDo you wish to do this now and continue?") );
-		if(!ExitInterp)
-		{
-			return true;
-		}
-
-		GLevelEditorModeTools().DeactivateMode( FBuiltinEditorModes::EM_InterpEdit );
-	}
-	return false;
-}
-
 /**
 * Prompts user with a confirmation dialog if there are checkouts or modifications in other branches
 *
@@ -2754,21 +2731,17 @@ bool FEditorFileUtils::SaveMap(UWorld* InWorld, const FString& Filename )
 {
 	bool bLevelWasSaved = false;
 
-	// Disallow the save if in interpolation editing mode and the user doesn't want to exit interpolation mode.
-	if ( !InInterpEditMode() )
-	{
-		const double SaveStartTime = FPlatformTime::Seconds();
+	const double SaveStartTime = FPlatformTime::Seconds();
 
-		FString FinalFilename;
-		bLevelWasSaved = SaveWorld( InWorld, &Filename,
-									nullptr, nullptr,
-									true, false,
-									FinalFilename,
-									false, false );
+	FString FinalFilename;
+	bLevelWasSaved = SaveWorld( InWorld, &Filename,
+								nullptr, nullptr,
+								true, false,
+								FinalFilename,
+								false, false );
 
-		// Track time spent saving map.
-		UE_LOG(LogFileHelpers, Log, TEXT("Saving map '%s' took %.3f"), *FPaths::GetBaseFilename(Filename), FPlatformTime::Seconds() - SaveStartTime );
-	}
+	// Track time spent saving map.
+	UE_LOG(LogFileHelpers, Log, TEXT("Saving map '%s' took %.3f"), *FPaths::GetBaseFilename(Filename), FPlatformTime::Seconds() - SaveStartTime );
 
 	return bLevelWasSaved;
 }
@@ -3559,8 +3532,7 @@ bool FEditorFileUtils::SaveLevel(ULevel* Level, const FString& DefaultFilename,
 {
 	bool bLevelWasSaved = false;
 
-	// Disallow the save if in interpolation editing mode and the user doesn't want to exit interpolation mode.
-	if ( Level && !InInterpEditMode() )
+	if (Level)
 	{
 		// Check and see if this is a new map.
 		const bool bIsPersistentLevelCurrent = Level->IsPersistentLevel();
diff --git a/Engine/Source/Editor/UnrealEd/Private/Lightmass/Lightmass.cpp b/Engine/Source/Editor/UnrealEd/Private/Lightmass/Lightmass.cpp
--- a/Engine/Source/Editor/UnrealEd/Private/Lightmass/Lightmass.cpp
+++ b/Engine/Source/Editor/UnrealEd/Private/Lightmass/Lightmass.cpp
@@ -39,11 +39,6 @@
 #include "ShadowMap.h"
 #include "LandscapeProxy.h"
 #include "LandscapeComponent.h"
-#include "Matinee/MatineeActor.h"
-#include "Matinee/InterpGroup.h"
-#include "Matinee/InterpGroupInst.h"
-#include "Matinee/InterpTrackMove.h"
-#include "Matinee/InterpTrackInstMove.h"
 #include "Components/SplineMeshComponent.h"
 #include "Lightmass/PrecomputedVisibilityVolume.h"
 #include "Lightmass/PrecomputedVisibilityOverrideVolume.h"
@@ -929,71 +924,9 @@ void FLightmassExporter::WriteVisibilityData( int32 Channel )
 
 	const float CellSize = World->GetWorldSettings()->VisibilityCellSize;
 
+	// This used to walk along Matinee sequences and write the positions for pre-computed visibility.
+	// Leaving the swarm channel code here for fear of breaking swarm.
 	TArray<FVector4> CameraTrackPositions;
-	if (World->GetWorldSettings()->bPrecomputeVisibility)
-	{
-		// Export positions along matinee camera tracks
-		// Lightmass needs to know these positions in order to place visibility cells containing them, since they may be outside any visibility volumes
-		for( TObjectIterator<ACameraActor> CameraIt; CameraIt; ++CameraIt )
-		{
-			ACameraActor* Camera = *CameraIt;
-			if (World->ContainsActor(Camera) && !Camera->IsPendingKill())
-			{
-				for( TActorIterator<AMatineeActor> It(World); It; ++It )
-				{
-					AMatineeActor* MatineeActor = *It;
-
-					bool bNeedsTermInterp = false;
-					if (MatineeActor->GroupInst.Num() == 0)
-					{
-						// If matinee is closed, GroupInst will be empty, so we need to populate it
-						bNeedsTermInterp = true;
-						MatineeActor->InitInterp();
-					}
-					UInterpGroupInst* GroupInstance = MatineeActor->FindGroupInst(Camera);
-					if (GroupInstance && GroupInstance->Group)
-					{
-						for (int32 TrackIndex = 0; TrackIndex < GroupInstance->Group->InterpTracks.Num(); TrackIndex++)
-						{
-							UInterpTrackMove* MoveTrack = Cast<UInterpTrackMove>(GroupInstance->Group->InterpTracks[TrackIndex]);
-							if (MoveTrack)
-							{
-								float StartTime;
-								float EndTime;
-								//@todo - look at the camera cuts to only process time ranges where the camera is actually active
-								MoveTrack->GetTimeRange(StartTime, EndTime);
-								for (int32 TrackInstanceIndex = 0; TrackInstanceIndex < GroupInstance->TrackInst.Num(); TrackInstanceIndex++)
-								{
-									UInterpTrackInst* TrackInstance = GroupInstance->TrackInst[TrackInstanceIndex];
-									UInterpTrackInstMove* MoveTrackInstance = Cast<UInterpTrackInstMove>(TrackInstance);
-									if (MoveTrackInstance)
-									{
-										//@todo - handle long camera paths
-										for (float Time = StartTime; Time < EndTime; Time += FMath::Max((EndTime - StartTime) * .001f, .001f))
-										{
-											const FVector RelativePosition = MoveTrack->EvalPositionAtTime(TrackInstance, Time);
-											FVector CurrentPosition;
-											FRotator CurrentRotation;
-											MoveTrack->ComputeWorldSpaceKeyTransform(MoveTrackInstance, RelativePosition, FRotator::ZeroRotator, CurrentPosition, CurrentRotation);
-											if (CameraTrackPositions.Num() == 0 || !CurrentPosition.Equals(CameraTrackPositions.Last(), CellSize * .1f))
-											{
-												CameraTrackPositions.Add(CurrentPosition);
-											}
-										}
-									}
-								}
-							}
-						}
-					}
-					if (bNeedsTermInterp)
-					{
-						MatineeActor->TermInterp();
-					}
-				}
-			}
-		}
-	}
-	
 	int32 NumCameraPositions = CameraTrackPositions.Num();
 	Swarm.WriteChannel( Channel, &NumCameraPositions, sizeof(NumCameraPositions) );
 	Swarm.WriteChannel( Channel, CameraTrackPositions.GetData(), CameraTrackPositions.Num() * CameraTrackPositions.GetTypeSize() );
diff --git a/Engine/Source/Editor/UnrealEd/Private/PackageAutoSaver.cpp b/Engine/Source/Editor/UnrealEd/Private/PackageAutoSaver.cpp
--- a/Engine/Source/Editor/UnrealEd/Private/PackageAutoSaver.cpp
+++ b/Engine/Source/Editor/UnrealEd/Private/PackageAutoSaver.cpp
@@ -507,8 +507,7 @@ void FPackageAutoSaver::UpdateAutoSaveNotification()
 
 	if (UserAllowsAutosave && // The user has set to allow auto-save in preferences
 		TimeInSecondsUntilAutosave < LoadingSavingSettings->AutoSaveWarningInSeconds && 
-		!InGame && // we want to hide auto-save if we are simulating/playing
-		!GLevelEditorModeTools().IsModeActive(FBuiltinEditorModes::EM_InterpEdit) // we want to hide auto-save if we are in matinee 
+		!InGame // we want to hide auto-save if we are simulating/playing
 		)
 	{		
 		if (!bAutoSaveNotificationLaunched && !bDelayingDueToFailedSave)
diff --git a/Engine/Source/Editor/UnrealEd/Private/PlayLevel.cpp b/Engine/Source/Editor/UnrealEd/Private/PlayLevel.cpp
--- a/Engine/Source/Editor/UnrealEd/Private/PlayLevel.cpp
+++ b/Engine/Source/Editor/UnrealEd/Private/PlayLevel.cpp
@@ -125,8 +125,6 @@ const static FName NAME_CategoryPIE("PIE");
 
 // Forward declare local utility functions
 FText GeneratePIEViewportWindowTitle(const EPlayNetMode InNetMode, const ERHIFeatureLevel::Type InFeatureLevel, const FRequestPlaySessionParams& InSessionParams, const int32 ClientIndex, const float FixedTick);
-bool PromptMatineeClose();
-
 
 // This class listens to output log messages, and forwards warnings and errors to the message log
 class FOutputLogErrorsToMessageLogProxy : public FOutputDevice
@@ -190,13 +188,6 @@ void UEditorEngine::EndPlayMap()
 		GEngine->XRSystem->OnEndPlay(*GEngine->GetWorldContextFromWorld(PlayWorld));
 	}
 
-	// Matinee must be closed before PIE can stop - matinee during PIE will be editing a PIE-world actor
-	if( GLevelEditorModeTools().IsModeActive(FBuiltinEditorModes::EM_InterpEdit) )
-	{
-		FMessageDialog::Open( EAppMsgType::Ok, NSLOCTEXT("UnrealEd", "PIENeedsToCloseMatineeMessage", "Closing 'Play in Editor' must close UnrealMatinee.") );
-		GLevelEditorModeTools().DeactivateMode( FBuiltinEditorModes::EM_InterpEdit );
-	}
-
 	EndPlayOnLocalPc();
 
 	const FScopedBusyCursor BusyCursor;
@@ -2473,14 +2464,6 @@ void UEditorEngine::StartPlayInEditorSession(FRequestPlaySessionParams& InReques
 		UE_LOG(LogPlayLevel, Warning, TEXT("Cancelling Open Transaction '%s' to start PIE session."), *TransactionName.ToString());
 	}
 
-	// Prompt the user that Matinee must be closed before PIE can occur. If they don't want
-	// to close Matinee, we can't PIE.
-	if (!PromptMatineeClose())
-	{
-		CancelRequestPlaySession();
-		return;
-	}
-
 	TArray<IPIEAuthorizer*> PlayAuthorizers = IModularFeatures::Get().GetModularFeatureImplementations<IPIEAuthorizer>(IPIEAuthorizer::GetModularFeatureName());
 	for (const IPIEAuthorizer* Authority : PlayAuthorizers)
 	{
@@ -3406,21 +3389,6 @@ void UEditorEngine::StoreWindowSizeAndPositionForInstanceIndex(const int32 InIns
 	}
 }
 
-bool PromptMatineeClose()
-{
-	if (GLevelEditorModeTools().IsModeActive(FBuiltinEditorModes::EM_InterpEdit))
-	{
-		const bool bContinuePIE = EAppReturnType::Yes == FMessageDialog::Open(EAppMsgType::YesNo, NSLOCTEXT("UnrealEd", "PIENeedsToCloseMatineeQ", "'Play in Editor' must close UnrealMatinee.  Continue?"));
-		if (!bContinuePIE)
-		{
-			return false;
-		}
-		GLevelEditorModeTools().DeactivateMode(FBuiltinEditorModes::EM_InterpEdit);
-	}
-
-	return true;
-}
-
 // Deprecated Stubs
 PRAGMA_DISABLE_DEPRECATION_WARNINGS
 UGameInstance* UEditorEngine::CreatePIEGameInstance(int32 InPIEInstance, bool bInSimulateInEditor, bool bAnyBlueprintErrors, bool bStartInSpectatorMode, bool bPlayNetDedicated, bool bPlayStereoscopic, float PIEStartTime)
diff --git a/Engine/Source/Editor/UnrealEd/Private/UnrealEdEngine.cpp b/Engine/Source/Editor/UnrealEd/Private/UnrealEdEngine.cpp
--- a/Engine/Source/Editor/UnrealEd/Private/UnrealEdEngine.cpp
+++ b/Engine/Source/Editor/UnrealEd/Private/UnrealEdEngine.cpp
@@ -36,8 +36,6 @@
 #include "UnrealEdMisc.h"
 #include "UnrealEdGlobals.h"
 
-#include "Matinee/InterpData.h"
-#include "Matinee/MatineeActor.h"
 #include "Animation/AnimCompress.h"
 
 #include "EditorSupportDelegates.h"
@@ -452,7 +450,7 @@ void UUnrealEdEngine::Tick(float DeltaSeconds, bool bIdleMode)
 	// Increment the "seconds since last autosave" counter, then try to autosave.
 	if (!GSlowTaskOccurred)
 	{
-		// Don't increment autosave count while in game/pie/automation testing or while in Matinee
+		// Don't increment autosave count while in game/pie/automation testing
 		const bool PauseAutosave = (PlayWorld != nullptr) || GIsAutomationTesting;
 		if (!PauseAutosave && PackageAutoSaver.Get())
 		{
@@ -759,50 +757,6 @@ bool UUnrealEdEngine::IsAutosaving() const
 	return false;
 }
 
-
-void UUnrealEdEngine::ConvertMatinees()
-{
-	FVector StartLocation= FVector::ZeroVector;
-	UWorld* World = GWorld;
-	if( World )
-	{
-		ULevel* Level = World->GetCurrentLevel();
-		if( !Level )
-		{
-			Level = World->PersistentLevel;
-		}
-		check(Level);
-		for( TObjectIterator<UInterpData> It; It; ++It )
-		{
-			UInterpData* InterpData = *It;
-			if( InterpData->IsIn( Level ) ) 
-			{
-				// We dont care about renaming references or adding redirectors.  References to this will be old seqact_interps
-				RenameObject( InterpData, Level->GetOutermost(), *InterpData->GetName() );
-
-				AMatineeActor* MatineeActor = Level->OwningWorld->SpawnActor<AMatineeActor>(StartLocation, FRotator::ZeroRotator);
-				StartLocation.Y += 50;
-								
-				MatineeActor->MatineeData = InterpData;
-				FProperty* MatineeDataProp = NULL;
-				for( FProperty* Property = MatineeActor->GetClass()->PropertyLink; Property != NULL; Property = Property->PropertyLinkNext )
-				{
-					if( Property->GetName() == TEXT("MatineeData") )
-					{
-						MatineeDataProp = Property;
-						break;
-					}
-				}
-
-				FPropertyChangedEvent PropertyChangedEvent( MatineeDataProp ); 
-				MatineeActor->PostEditChangeProperty( PropertyChangedEvent );
-			}
-		}
-	}
-
-}
-
-
 void UUnrealEdEngine::ShowActorProperties()
 {
 	// See if we have any unlocked property windows available.  If not, create a new one.
@@ -1010,7 +964,6 @@ void UUnrealEdEngine::Serialize(FArchive& Ar)
 	Super::Serialize(Ar);
 	Ar << MaterialCopyPasteBuffer;
 	Ar << AnimationCompressionAlgorithms;
-	Ar << MatineeCopyPasteBuffer;
 }
 
 void UUnrealEdEngine::DuplicateSelectedActors(UWorld* InWorld)
@@ -1529,21 +1482,6 @@ void UUnrealEdEngine::OnContentPathDismounted(const FString& AssetPath, const FS
 {
 	MountPointCheckedForWritePermission.Remove(*AssetPath);
 }
-void UUnrealEdEngine::UpdateEdModeOnMatineeClose(const FEditorModeID& EditorModeID, bool IsEntering)
-{
-	// if we are closing the Matinee editor
-	if (!IsEntering && EditorModeID == FBuiltinEditorModes::EM_InterpEdit)
-	{
-		// set the autosave timer to save soon
-		if (PackageAutoSaver)
-		{
-			PackageAutoSaver->ForceMinimumTimeTillAutoSave();
-		}
-
-		// Remove this delegate. 
-		GLevelEditorModeTools().OnEditorModeIDChanged().Remove(UpdateEdModeOnMatineeCloseDelegateHandle);
-	}
-}
 
 bool IsBelowFreeDiskSpaceLimit(const TCHAR* TestDir, FText& OutAppendMessage, const FText& LocationDescriptor, const uint64 MinMB = 5120)
 {
diff --git a/Engine/Source/Editor/UnrealEd/Private/UnrealEdPrivatePCH.h b/Engine/Source/Editor/UnrealEd/Private/UnrealEdPrivatePCH.h
--- a/Engine/Source/Editor/UnrealEd/Private/UnrealEdPrivatePCH.h
+++ b/Engine/Source/Editor/UnrealEd/Private/UnrealEdPrivatePCH.h
@@ -589,7 +589,6 @@
 #include "Engine/Polys.h"
 #include "Kismet/BlueprintFunctionLibrary.h"
 #include "Animation/SkeletalMeshActor.h"
-#include "Matinee/MatineeAnimInterface.h"
 #include "Engine/LevelStreaming.h"
 #include "LatentActions.h"
 #include "PreviewScene.h"
diff --git a/Engine/Source/Editor/UnrealEd/Private/UnrealEdSrv.cpp b/Engine/Source/Editor/UnrealEd/Private/UnrealEdSrv.cpp
--- a/Engine/Source/Editor/UnrealEd/Private/UnrealEdSrv.cpp
+++ b/Engine/Source/Editor/UnrealEd/Private/UnrealEdSrv.cpp
@@ -80,8 +80,6 @@
 #include "IContentBrowserSingleton.h"
 #include "ContentBrowserModule.h"
 #include "AssetRegistryModule.h"
-#include "Matinee/MatineeActor.h"
-#include "MatineeExporter.h"
 #include "FbxExporter.h"
 #include "DesktopPlatformModule.h"
 #include "Elements/Framework/TypedElementList.h"
@@ -92,7 +90,6 @@
 #include "HighResScreenshot.h"
 #include "ActorEditorUtils.h"
 #include "Editor/ActorPositioning.h"
-#include "Matinee/InterpData.h"
 #include "LandscapeInfo.h"
 #include "LandscapeInfoMap.h"
 #include "Framework/Notifications/NotificationManager.h"
@@ -618,47 +615,6 @@ bool UUnrealEdEngine::HandleRemoveLandscapeXYOffsetsCommand(const TCHAR* Str, FO
 	return true;
 }
 
-bool UUnrealEdEngine::HandleConvertMatineesCommand( const TCHAR* Str, FOutputDevice& Ar, UWorld* InWorld )
-{
-	FVector StartLocation= FVector::ZeroVector;
-	if( InWorld )
-	{
-		ULevel* Level = InWorld->GetCurrentLevel();
-		if( !Level )
-		{
-			Level = InWorld->PersistentLevel;
-	}
-		check(Level);
-		for( TObjectIterator<UInterpData> It; It; ++It )
-		{
-			UInterpData* InterpData = *It;
-			if( InterpData->IsIn( Level ) ) 
-			{
-				// We dont care about renaming references or adding redirectors.  References to this will be old seqact_interps
-				GEditor->RenameObject( InterpData, Level->GetOutermost(), *InterpData->GetName() );
-
-				AMatineeActor* MatineeActor = Level->OwningWorld->SpawnActor<AMatineeActor>(StartLocation, FRotator::ZeroRotator);
-				StartLocation.Y += 50;
-
-				MatineeActor->MatineeData = InterpData;
-				FProperty* MatineeDataProp = NULL;
-				for( FProperty* Property = MatineeActor->GetClass()->PropertyLink; Property != NULL; Property = Property->PropertyLinkNext )
-				{
-					if( Property->GetName() == TEXT("MatineeData") )
-					{
-						MatineeDataProp = Property;
-						break;
-					}
-				}
-
-				FPropertyChangedEvent PropertyChangedEvent( MatineeDataProp ); 
-				MatineeActor->PostEditChangeProperty( PropertyChangedEvent );
-			}
-		}
-	}
-		return true;
-	}
-
 bool UUnrealEdEngine::HandleDisasmScriptCommand(const TCHAR* Str, FOutputDevice& Ar)
 {
 	FString ClassName;
@@ -809,10 +765,6 @@ bool UUnrealEdEngine::Exec( UWorld* InWorld, const TCHAR* Stream, FOutputDevice&
 		return HandleRemoveLandscapeXYOffsetsCommand(Str, Ar, World);
 	}
 #endif // WITH_EDITOR
-	else if( FParse::Command(&Str, TEXT("CONVERTMATINEES")) )
-	{
-		return HandleConvertMatineesCommand( Str, Ar, InWorld );
-	}
 	else if( FParse::Command(&Str, TEXT("DISASMSCRIPT")) )
 	{
 		return HandleDisasmScriptCommand( Str, Ar );
diff --git a/Engine/Source/Editor/UnrealEd/Public/AssetSelection.h b/Engine/Source/Editor/UnrealEd/Public/AssetSelection.h
--- a/Engine/Source/Editor/UnrealEd/Public/AssetSelection.h
+++ b/Engine/Source/Editor/UnrealEd/Public/AssetSelection.h
@@ -105,8 +105,6 @@ struct FSelectedActorInfo
 	uint32 bHaveSkeletalMesh:1;
 	/** true if an actor in the selection is an emitter */
 	uint32 bHaveEmitter:1;
-	/** true if an actor in the selection is a matinee actor */
-	uint32 bHaveMatinee:1;
 	/** true if an actor in the selection is hidden */
 	uint32 bHaveHidden:1;
 	/** true if a landscape is in the selection */
@@ -144,7 +142,6 @@ struct FSelectedActorInfo
 		, bHavePawn(false)
 		, bHaveSkeletalMesh(false)
 		, bHaveEmitter(false)
-		, bHaveMatinee(false)
 		, bHaveHidden(false)
 		, bHaveLandscape(false)
 		, bHaveExperimentalClass(false)
