Editor: Fixed an assert if the Late Join button was bound to a key and triggered while in PIE.

diff --git a/Engine/Source/Editor/UnrealEd/Classes/Editor/EditorEngine.h b/Engine/Source/Editor/UnrealEd/Classes/Editor/EditorEngine.h
--- a/Engine/Source/Editor/UnrealEd/Classes/Editor/EditorEngine.h
+++ b/Engine/Source/Editor/UnrealEd/Classes/Editor/EditorEngine.h
@@ -2895,6 +2895,9 @@ protected:
 
 	/** Start a Play in New Process session with the given parameters. Called by StartQueuedPlaySessionRequestImpl based on request settings. */
 	virtual void StartPlayInNewProcessSession(FRequestPlaySessionParams& InRequestParams);
+
+	/** Does the actual late join process. Don't call this directly, use RequestLateJoin() which will defer it to the next frame. */
+	void AddPendingLateJoinClient();
 	
 	// This chunk is used for Play Via Launcher settings.
 public:
diff --git a/Engine/Source/Editor/UnrealEd/Private/EditorEngine.cpp b/Engine/Source/Editor/UnrealEd/Private/EditorEngine.cpp
--- a/Engine/Source/Editor/UnrealEd/Private/EditorEngine.cpp
+++ b/Engine/Source/Editor/UnrealEd/Private/EditorEngine.cpp
@@ -1671,6 +1671,12 @@ void UEditorEngine::Tick( float DeltaSeconds, bool bIdleMode )
 		ToggleBetweenPIEandSIE();
 	}
 
+	// Deferred until here so it doesn't happen mid iteration of worlds.
+	if (PlayInEditorSessionInfo.IsSet() && PlayInEditorSessionInfo->bLateJoinRequested)
+	{
+		AddPendingLateJoinClient();
+	}
+
 	static bool bFirstTick = true;
 
 	// Skip updating reflection captures on the first update as the level will not be ready to display
diff --git a/Engine/Source/Editor/UnrealEd/Private/PlayLevel.cpp b/Engine/Source/Editor/UnrealEd/Private/PlayLevel.cpp
--- a/Engine/Source/Editor/UnrealEd/Private/PlayLevel.cpp
+++ b/Engine/Source/Editor/UnrealEd/Private/PlayLevel.cpp
@@ -1608,11 +1608,26 @@ void UEditorEngine::GiveFocusToLastClientPIEViewport()
 }
 
 void UEditorEngine::RequestLateJoin()
+{
+	if (PlayInEditorSessionInfo.IsSet())
+	{
+		PlayInEditorSessionInfo->bLateJoinRequested = true;
+	}
+}
+
+void UEditorEngine::AddPendingLateJoinClient()
 {
 	if (!ensureMsgf(PlayInEditorSessionInfo.IsSet(), TEXT("RequestLateJoin shouldn't be called if no session is in progress!")))
 	{
 		return;
 	}
+	
+	if(!ensureMsgf(PlayInEditorSessionInfo->bLateJoinRequested, TEXT("AddPendingLateJoinClient() shouldn't be called directly, use RequestLateJoin() instead!")))
+	{
+		return;
+	}
+
+	PlayInEditorSessionInfo->bLateJoinRequested = false;
 
 	if (!ensureMsgf(PlayInEditorSessionInfo->OriginalRequestParams.WorldType != EPlaySessionWorldType::SimulateInEditor, TEXT("RequestLateJoin shouldn't be called for SIE!")))
 	{
diff --git a/Engine/Source/Editor/UnrealEd/Public/PlayInEditorDataTypes.h b/Engine/Source/Editor/UnrealEd/Public/PlayInEditorDataTypes.h
--- a/Engine/Source/Editor/UnrealEd/Public/PlayInEditorDataTypes.h
+++ b/Engine/Source/Editor/UnrealEd/Public/PlayInEditorDataTypes.h
@@ -243,6 +243,7 @@ public:
 		, bUsingOnlinePlatform(false)
 		, bAnyBlueprintErrors(false)
 		, bServerWasLaunched(false)
+		, bLateJoinRequested(false)
 	{
 	}
 
@@ -289,6 +290,9 @@ public:
 	/** Have we launched a server for this PIE session yet? */
 	bool bServerWasLaunched;
 
+	/** If true, a late join client will be added on the next tick. */
+	bool bLateJoinRequested;
+
 	/** Transient information about window sizes/positions. This gets loaded from settings at start and saved at end. */
 	TArray<FWindowSizeAndPos> CachedWindowInfo;
 
