Optimize the ddc fill job
  - Avoid redundant iterations and processing of UPackages by implementing a listener instead
  - Process many packages asynchronously (not just one) by keeping them in memory longer and finishing them later
  - Make the ddc fill job wait for async tasks properly
  - Fix texture always being single threaded by allowing async at all time for ddc fill jobs
  - Avoid calling IsCachedCookedPlatformDataLoaded too often on objects still being compiled
  - Use a more sensible max memory usage in WorldPartitionHelpers to avoid calling the GC too often

  - 30m16s to 7m44s when running -run=DerivedDataCache -fill -map=P_Construct_WP -mapsonly -projectonly with hot cache
  - 2h30m to 1h11s when running -run=DerivedDataCache -fill -map=P_Construct_WP -mapsonly -projectonly with cold cache

diff --git a/Engine/Source/Editor/UnrealEd/Classes/Commandlets/DerivedDataCacheCommandlet.h b/Engine/Source/Editor/UnrealEd/Classes/Commandlets/DerivedDataCacheCommandlet.h
--- a/Engine/Source/Editor/UnrealEd/Classes/Commandlets/DerivedDataCacheCommandlet.h
+++ b/Engine/Source/Editor/UnrealEd/Classes/Commandlets/DerivedDataCacheCommandlet.h
@@ -15,20 +15,32 @@ UCLASS()
 class UDerivedDataCacheCommandlet : public UCommandlet
 {
 	GENERATED_UCLASS_BODY()
+	UDerivedDataCacheCommandlet(FVTableHelper& Helper); // Declare the FVTableHelper constructor manually so that we can forward-declare-only TUniquePtrs in the header without getting compile error in generated cpp
+
 	//~ Begin UCommandlet Interface
 	virtual int32 Main(const FString& Params) override;
 	//~ End UCommandlet Interface
 
-	// We hook this up to a delegate to avoid reloading textures and whatnot
-	TSet<FName> ProcessedPackages;
-	TSet<FName> PackagesToProcess;
+	// Objects currently being cached along with when was the last time in seconds we verified if they were still compiling
+	// The last time is used to do some throttling on the IsCachedCookedPlatformDataLoaded which can be quite expensive on some objects
+	TMap<UObject*, double> CachingObjects;
+	TSet<FName>    ProcessedPackages;
+	TSet<FName>    PackagesToProcess;
 	double FinishCacheTime = 0.0;
 	double BeginCacheTime = 0.0;
 
+	class FPackageListener;
+	TUniquePtr<FPackageListener> PackageListener;
+
+	class FObjectReferencer;
+	TUniquePtr<FObjectReferencer> ObjectReferencer;
+
 	void MaybeMarkPackageAsAlreadyLoaded(UPackage *Package);
 
 	void CacheLoadedPackages(UPackage* CurrentPackage, uint8 PackageFilter, const TArray<ITargetPlatform*>& Platforms);
 	void CacheWorldPackages(UWorld* World, uint8 PackageFilter, const TArray<ITargetPlatform*>& Platforms);
+	bool ProcessCachingObjects(const TArray<ITargetPlatform*>& Platforms);
+	void FinishCachingObjects(const TArray<ITargetPlatform*>& Platforms);
 };
 
 
diff --git a/Engine/Source/Editor/UnrealEd/Private/Commandlets/DerivedDataCacheCommandlet.cpp b/Engine/Source/Editor/UnrealEd/Private/Commandlets/DerivedDataCacheCommandlet.cpp
--- a/Engine/Source/Editor/UnrealEd/Private/Commandlets/DerivedDataCacheCommandlet.cpp
+++ b/Engine/Source/Editor/UnrealEd/Private/Commandlets/DerivedDataCacheCommandlet.cpp
@@ -32,8 +32,92 @@ DerivedDataCacheCommandlet.cpp: Commandlet for DDC maintenence
 
 DEFINE_LOG_CATEGORY_STATIC(LogDerivedDataCacheCommandlet, Log, All);
 
+class UDerivedDataCacheCommandlet::FObjectReferencer : public FGCObject
+{
+public:
+	FObjectReferencer(TMap<UObject*, double>& InReferencedObjects)
+		: ReferencedObjects(InReferencedObjects)
+	{
+	}
+
+private:
+	void AddReferencedObjects(FReferenceCollector& Collector) override
+	{
+		Collector.AllowEliminatingReferences(false);
+		Collector.AddReferencedObjects(ReferencedObjects);
+		Collector.AllowEliminatingReferences(true);
+	}
+
+	FString GetReferencerName() const override
+	{
+		return TEXT("UDerivedDataCacheCommandlet");
+	}
+
+	FString ReferencerName;
+	TMap<UObject*, double>& ReferencedObjects;
+};
+
+class UDerivedDataCacheCommandlet::FPackageListener : public FUObjectArray::FUObjectCreateListener, public FUObjectArray::FUObjectDeleteListener
+{
+public:
+	FPackageListener()
+	{
+		GUObjectArray.AddUObjectDeleteListener(this);
+		GUObjectArray.AddUObjectCreateListener(this);
+
+		// We might be late to the party, check if some UPackage already have been created
+		for (TObjectIterator<UPackage> PackageIter; PackageIter; ++PackageIter)
+		{
+			NewPackages.Add(*PackageIter);
+		}
+	}
+
+	~FPackageListener()
+	{
+		GUObjectArray.RemoveUObjectDeleteListener(this);
+		GUObjectArray.RemoveUObjectCreateListener(this);
+	}
+
+	TSet<UPackage*> GetNewPackages()
+	{
+		return MoveTemp(NewPackages);
+	}
+
+private:
+	void NotifyUObjectCreated(const class UObjectBase* Object, int32 Index) override
+	{
+		if (Object->GetClass() == UPackage::StaticClass())
+		{
+			NewPackages.Add(const_cast<UPackage*>(static_cast<const UPackage*>(Object)));
+		}
+	}
+
+	void NotifyUObjectDeleted(const class UObjectBase* Object, int32 Index) override
+	{
+		if (Object->GetClass() == UPackage::StaticClass())
+		{
+			NewPackages.Remove(const_cast<UPackage*>(static_cast<const UPackage*>(Object)));
+		}
+	}
+
+	void OnUObjectArrayShutdown() override
+	{
+		GUObjectArray.RemoveUObjectDeleteListener(this);
+		GUObjectArray.RemoveUObjectCreateListener(this);
+	}
+
+	TSet<UPackage*> NewPackages;
+};
+
+UDerivedDataCacheCommandlet::UDerivedDataCacheCommandlet(FVTableHelper& Helper)
+	: Super(Helper)
+{
+}
+
 UDerivedDataCacheCommandlet::UDerivedDataCacheCommandlet(const FObjectInitializer& ObjectInitializer)
 	: Super(ObjectInitializer)
+	, PackageListener(MakeUnique<FPackageListener>())
+	, ObjectReferencer(MakeUnique<FObjectReferencer>(CachingObjects))
 {
 	LogToConsole = false;
 }
@@ -47,57 +131,58 @@ void UDerivedDataCacheCommandlet::MaybeMarkPackageAsAlreadyLoaded(UPackage *Pack
 	}
 }
 
-static void WaitForCurrentShaderCompilationToFinish(bool& bInOutHadActivity)
+static void WaitForCompilationToFinish(bool& bInOutHadActivity)
 {
-	if (GShaderCompilingManager->IsCompiling())
+	auto LogStatus =
+		[](IAssetCompilingManager* CompilingManager)
 	{
-		bInOutHadActivity = true;
-		int32 CachedShaderCount = GShaderCompilingManager->GetNumRemainingJobs();
-		if (CachedShaderCount > 0)
+		int32 AssetCount = CompilingManager->GetNumRemainingAssets();
+		if (AssetCount > 0)
 		{
-			UE_LOG(LogDerivedDataCacheCommandlet, Display, TEXT("Waiting for %d shaders to finish."), CachedShaderCount);
+			UE_LOG(LogDerivedDataCacheCommandlet, Display, TEXT("Waiting for %d %s to finish."), AssetCount, *FText::Format(CompilingManager->GetAssetNameFormat(), FText::AsNumber(AssetCount)).ToString());
 		}
-		int32 NumCompletedShadersSinceLastLog = 0;
-		while (GShaderCompilingManager->IsCompiling())
+		else
 		{
-			const int32 CurrentShaderCount = GShaderCompilingManager->GetNumRemainingJobs();
-			NumCompletedShadersSinceLastLog += (CachedShaderCount - CurrentShaderCount);
-			CachedShaderCount = CurrentShaderCount;
+			UE_LOG(LogDerivedDataCacheCommandlet, Display, TEXT("Done waiting for %s to finish."), *FText::Format(CompilingManager->GetAssetNameFormat(), FText::AsNumber(100)).ToString());
+		}
+	};
 
-			if (NumCompletedShadersSinceLastLog >= 1000)
+	while (FAssetCompilingManager::Get().GetNumRemainingAssets() > 0)
+	{
+		for (IAssetCompilingManager* CompilingManager : FAssetCompilingManager::Get().GetRegisteredManagers())
+		{
+			int32 CachedAssetCount = CompilingManager->GetNumRemainingAssets();
+			if (CachedAssetCount)
 			{
-				UE_LOG(LogDerivedDataCacheCommandlet, Display, TEXT("Waiting for %d shaders to finish."), CachedShaderCount);
-				NumCompletedShadersSinceLastLog = 0;
-			}
+				bInOutHadActivity = true;
+				LogStatus(CompilingManager);
+				int32 NumCompletedAssetsSinceLastLog = 0;
+				while (CompilingManager->GetNumRemainingAssets() > 0)
+				{
+					const int32 CurrentAssetCount = CompilingManager->GetNumRemainingAssets();
+					NumCompletedAssetsSinceLastLog += (CachedAssetCount - CurrentAssetCount);
+					CachedAssetCount = CurrentAssetCount;
 
-			// Process any asynchronous shader compile results that are ready, limit execution time
-			GShaderCompilingManager->ProcessAsyncResults(true, false);
-			GDistanceFieldAsyncQueue->ProcessAsyncTasks();
-			GCardRepresentationAsyncQueue->ProcessAsyncTasks();
-		}
-		GShaderCompilingManager->FinishAllCompilation(); // Final blocking check as IsCompiling() may be non-deterministic
-		UE_LOG(LogDerivedDataCacheCommandlet, Display, TEXT("Done waiting for shaders to finish."));
-	}
+					if (NumCompletedAssetsSinceLastLog >= 1000)
+					{
+						LogStatus(CompilingManager);
+						NumCompletedAssetsSinceLastLog = 0;
+					}
 
-	// these shouldn't be predicated on whether the shaders were being compiled
-	GDistanceFieldAsyncQueue->BlockUntilAllBuildsComplete();
-	GCardRepresentationAsyncQueue->BlockUntilAllBuildsComplete();
-}
+					// Process any asynchronous Asset compile results that are ready, limit execution time
+					FAssetCompilingManager::Get().ProcessAsyncTasks(true);
+				}
 
-static void WaitForCurrentTextureBuildingToFinish(bool& bInOutHadActivity)
-{
-	for (TObjectIterator<UTexture> Texture; Texture; ++Texture)
-	{
-		Texture->FinishCachePlatformData();
+				LogStatus(CompilingManager);
+			}
+		}
 	}
-};
+}
 
 static void PumpAsync(bool* bInOutHadActivity = nullptr)
 {
 	bool bHadActivity = false;
-	WaitForCurrentShaderCompilationToFinish(bHadActivity);
-	WaitForCurrentTextureBuildingToFinish(bHadActivity);
-	FAssetCompilingManager::Get().ProcessAsyncTasks(true);
+	WaitForCompilationToFinish(bHadActivity);
 	if (bInOutHadActivity)
 	{
 		*bInOutHadActivity = *bInOutHadActivity || bHadActivity;
@@ -106,133 +191,134 @@ static void PumpAsync(bool* bInOutHadActivity = nullptr)
 
 void UDerivedDataCacheCommandlet::CacheLoadedPackages(UPackage* CurrentPackage, uint8 PackageFilter, const TArray<ITargetPlatform*>& Platforms)
 {
-	// Timing variables
-	double DDCCommandletMaxWaitSeconds = 60. * 10.;
-	GConfig->GetDouble(TEXT("CookSettings"), TEXT("DDCCommandletMaxWaitSeconds"), DDCCommandletMaxWaitSeconds, GEditorIni);
-
-	TArray<UObject*> CachingObjects;
-	TArray<UObject*> ObjectBuffer;
-	TArray<UPackage*> NewPackages;
+	TRACE_CPUPROFILER_EVENT_SCOPE(UDerivedDataCacheCommandlet::CacheLoadedPackages);
 
-	bool bIsCaching = false;
+	const double BeginCacheTimeStart = FPlatformTime::Seconds();
 
+	TArray<UObject*> ObjectsWithOuter;
+	for (UPackage* NewPackage : PackageListener->GetNewPackages())
 	{
-		const double BeginCacheTimeStart = FPlatformTime::Seconds();
-		CachingObjects.Reset();
-		NewPackages.Reset();
-
-		for (TObjectIterator<UPackage> PackageIter; PackageIter; ++PackageIter)
+		const FName NewPackageName = NewPackage->GetFName();
+		if (!ProcessedPackages.Contains(NewPackageName))
 		{
-			UPackage* ExistingPackage = *PackageIter;
-			if ((PackageFilter & NORMALIZE_ExcludeEnginePackages) == 0 || !ExistingPackage->GetName().StartsWith(TEXT("/Engine")))
+			if ((PackageFilter & NORMALIZE_ExcludeEnginePackages) != 0 && NewPackage->GetName().StartsWith(TEXT("/Engine")))
 			{
-				FName ExistingPackageName = ExistingPackage->GetFName();
-				if (ExistingPackage == CurrentPackage || !PackagesToProcess.Contains(ExistingPackageName))
+				// Add it so we don't convert the FName to a string everytime we encounter this package
+				ProcessedPackages.Add(NewPackageName);
+			}
+			else if (NewPackage == CurrentPackage || !PackagesToProcess.Contains(NewPackageName))
+			{
+				UE_LOG(LogDerivedDataCacheCommandlet, Display, TEXT("Processing %s"), *NewPackageName.ToString());
+
+				ProcessedPackages.Add(NewPackageName);
+
+				ObjectsWithOuter.Reset();
+				GetObjectsWithOuter(NewPackage, ObjectsWithOuter, true /* bIncludeNestedObjects */, RF_ClassDefaultObject /* ExclusionFlags */);
+				for (UObject* Object : ObjectsWithOuter)
 				{
-					if (!ProcessedPackages.Contains(ExistingPackageName))
+					for (auto Platform : Platforms)
 					{
-						UE_LOG(LogDerivedDataCacheCommandlet, Display, TEXT("Processing %s"), *ExistingPackageName.ToString());
-
-						ProcessedPackages.Add(ExistingPackageName);
-						NewPackages.Add(ExistingPackage);
-						check((ExistingPackage->GetPackageFlags() & PKG_ReloadingForCooker) == 0);
-
-						ObjectBuffer.Reset();
-						GetObjectsWithOuter(ExistingPackage, ObjectBuffer, true /* bIncludeNestedObjects */, RF_ClassDefaultObject /* ExclusionFlags */);
-						for (UObject* Object : ObjectBuffer)
-						{
-							for (auto Platform : Platforms)
-							{
-								Object->BeginCacheForCookedPlatformData(Platform);
-								bIsCaching |= !Object->IsCachedCookedPlatformDataLoaded(Platform);
-							}
-							CachingObjects.Add(Object);
-						}
+						Object->BeginCacheForCookedPlatformData(Platform);
 					}
+
+					CachingObjects.Add(Object);
 				}
 			}
 		}
-
-		BeginCacheTime += FPlatformTime::Seconds() - BeginCacheTimeStart;
 	}
 
-	{
-		const double FinishCacheTimeStart = FPlatformTime::Seconds();
+	BeginCacheTime += FPlatformTime::Seconds() - BeginCacheTimeStart;
 
-		if (bIsCaching)
-		{
-			PumpAsync();
+	ProcessCachingObjects(Platforms);
+}
+
+bool UDerivedDataCacheCommandlet::ProcessCachingObjects(const TArray<ITargetPlatform*>& Platforms)
+{
+	TRACE_CPUPROFILER_EVENT_SCOPE(UDerivedDataCacheCommandlet::ProcessCachingObjects);
 
-			ObjectBuffer.Reset();
-			ObjectBuffer.Append(CachingObjects);
+	bool bHadActivity = false;
+	if (CachingObjects.Num() > 0)
+	{
+		FAssetCompilingManager::Get().ProcessAsyncTasks(true);
 
-			double LastActivityTime = FinishCacheTimeStart;
-			const double WaitingForCacheSleepTime = 0.050;
-			while (ObjectBuffer.Num() > 0)
+		double CurrentTime = FPlatformTime::Seconds();
+		for (auto It = CachingObjects.CreateIterator(); It; ++It)
+		{
+			// Call IsCachedCookedPlatformDataLoaded once a second per object since it can be quite expensive
+			if (CurrentTime - It->Value > 1.0)
 			{
-				bool bHadActivity = false;
-				for (int32 ObjectIndex = 0; ObjectIndex < ObjectBuffer.Num();)
+				UObject* Object = It->Key;
+				bool bIsFinished = true;
+				const IInterface_AsyncCompilation* Interface_AsyncCompilation = Cast<IInterface_AsyncCompilation>(Object);
+				if (Interface_AsyncCompilation && Interface_AsyncCompilation->IsCompiling())
 				{
-					UObject* Object = ObjectBuffer[ObjectIndex];
-					bool bIsFinished = true;
-					for (auto Platform : Platforms)
-					{
-						bIsFinished = Object->IsCachedCookedPlatformDataLoaded(Platform) && bIsFinished;
-					}
-					if (bIsFinished)
-					{
-						ObjectBuffer.RemoveAtSwap(ObjectIndex);
-						bHadActivity = true;
-					}
-					else
-					{
-						++ObjectIndex;
-					}
+					bIsFinished = false;
 				}
 
-				double CurrentTime = FPlatformTime::Seconds();
-				if (!bHadActivity)
+				for (auto Platform : Platforms)
 				{
-					PumpAsync(&bHadActivity);
+					// IsCachedCookedPlatformDataLoaded can be quite slow for some objects
+					// Do not call it if bIsFinished is already false
+					bIsFinished = bIsFinished && Object->IsCachedCookedPlatformDataLoaded(Platform);
 				}
-				if (!bHadActivity)
+
+				if (bIsFinished)
 				{
-					if (CurrentTime - LastActivityTime >= DDCCommandletMaxWaitSeconds)
-					{
-						UE_LOG(LogDerivedDataCacheCommandlet, Error, TEXT("Timed out for %.2lfs waiting for %d objects to finish caching. First object: %s."),
-							DDCCommandletMaxWaitSeconds, ObjectBuffer.Num(), *ObjectBuffer[0]->GetFullName());
-						ObjectBuffer.Reset();
-					}
-					else
-					{
-						FPlatformProcess::Sleep(WaitingForCacheSleepTime);
-					}
+					bHadActivity = true;
+					Object->WillNeverCacheCookedPlatformDataAgain();
+					Object->ClearAllCachedCookedPlatformData();
+					It.RemoveCurrent();
 				}
 				else
 				{
-					LastActivityTime = CurrentTime;
+					It->Value = CurrentTime;
 				}
 			}
-
-			PumpAsync();
 		}
+	}
+
+	return bHadActivity;
+}
+
+void UDerivedDataCacheCommandlet::FinishCachingObjects(const TArray<ITargetPlatform*>& Platforms)
+{
+	// Timing variables
+	double DDCCommandletMaxWaitSeconds = 60. * 10.;
+	GConfig->GetDouble(TEXT("CookSettings"), TEXT("DDCCommandletMaxWaitSeconds"), DDCCommandletMaxWaitSeconds, GEditorIni);
+
+	const double FinishCacheTimeStart = FPlatformTime::Seconds();
+	double LastActivityTime = FinishCacheTimeStart;
 
-		// Tear down all of the Cached data; we do this only after all objects have finished because we need to not
-		// tear down any object until all objects in its package have finished
-		for (UObject* Object : CachingObjects)
+	while (CachingObjects.Num() > 0)
+	{
+		bool bHadActivity = ProcessCachingObjects(Platforms);
+
+		double CurrentTime = FPlatformTime::Seconds();
+		if (!bHadActivity)
 		{
-			Object->WillNeverCacheCookedPlatformDataAgain();
-			Object->ClearAllCachedCookedPlatformData();
+			PumpAsync(&bHadActivity);
 		}
-
-		// Mark the packages as processed
-		for (UPackage* NewPackage : NewPackages)
+		if (!bHadActivity)
+		{
+			if (CurrentTime - LastActivityTime >= DDCCommandletMaxWaitSeconds)
+			{
+				UObject* Object = CachingObjects.CreateIterator()->Key;
+				UE_LOG(LogDerivedDataCacheCommandlet, Error, TEXT("Timed out for %.2lfs waiting for %d objects to finish caching. First object: %s."),
+					DDCCommandletMaxWaitSeconds, CachingObjects.Num(), *Object->GetFullName());
+			}
+			else
+			{
+				const double WaitingForCacheSleepTime = 0.050;
+				FPlatformProcess::Sleep(WaitingForCacheSleepTime);
+			}
+		}
+		else
 		{
-			NewPackage->SetPackageFlags(PKG_ReloadingForCooker);
+			LastActivityTime = CurrentTime;
 		}
-
-		FinishCacheTime += FPlatformTime::Seconds() - FinishCacheTimeStart;
 	}
+
+	FinishCacheTime += FPlatformTime::Seconds() - FinishCacheTimeStart;
 }
 
 void UDerivedDataCacheCommandlet::CacheWorldPackages(UWorld* World, uint8 PackageFilter, const TArray<ITargetPlatform*>& Platforms)
@@ -542,6 +628,8 @@ int32 UDerivedDataCacheCommandlet::Main( const FString& Params )
 		}
 	}
 
+	FinishCachingObjects(GetTargetPlatformManager()->GetActiveTargetPlatforms());
+
 	GetDerivedDataCacheRef().WaitForQuiescence(true);
 
 	UE_LOG(LogDerivedDataCacheCommandlet, Display, TEXT("BeginCacheTime=%.2lfs, FinishCacheTime=%.2lfs, GCTime=%.2lfs."), BeginCacheTime, FinishCacheTime, GCTime);
diff --git a/Engine/Source/Runtime/CoreUObject/Private/UObject/GarbageCollection.cpp b/Engine/Source/Runtime/CoreUObject/Private/UObject/GarbageCollection.cpp
--- a/Engine/Source/Runtime/CoreUObject/Private/UObject/GarbageCollection.cpp
+++ b/Engine/Source/Runtime/CoreUObject/Private/UObject/GarbageCollection.cpp
@@ -70,6 +70,7 @@ static bool GObjCurrentPurgeObjectIndexNeedsReset = true;
 static TArray<FUObjectItem*> GUnreachableObjects;
 static FCriticalSection GUnreachableObjectsCritical;
 static int32 GUnrechableObjectIndex = 0;
+static double GLastGCTime = 0;
 
 /** Helpful constant for determining how many token slots we need to store a pointer **/
 static const uint32 GNumTokensPerPointer = sizeof(void*) / sizeof(uint32); //-V514
@@ -2119,10 +2120,16 @@ void CollectGarbageInternal(EObjectFlags KeepFlags, bool bPerformFullPurge)
 		FCoreUObjectDelegates::GetPostGarbageCollect().Broadcast();
 	}
 
+	GLastGCTime = FPlatformTime::Seconds();
 	STAT_ADD_CUSTOMMESSAGE_NAME( STAT_NamedMarker, TEXT( "GarbageCollection - End" ) );
 #endif	// UE_WITH_GC
 }
 
+double GetLastGCTime()
+{
+	return GLastGCTime;
+}
+
 bool IsIncrementalUnhashPending()
 {
 	return GUnrechableObjectIndex < GUnreachableObjects.Num();
diff --git a/Engine/Source/Runtime/CoreUObject/Public/UObject/GarbageCollection.h b/Engine/Source/Runtime/CoreUObject/Public/UObject/GarbageCollection.h
--- a/Engine/Source/Runtime/CoreUObject/Public/UObject/GarbageCollection.h
+++ b/Engine/Source/Runtime/CoreUObject/Public/UObject/GarbageCollection.h
@@ -565,6 +565,13 @@ public:
 /** True if Garbage Collection is running. Use IsGarbageCollecting() functio n instead of using this variable directly */
 extern COREUOBJECT_API FThreadSafeBool GIsGarbageCollecting;
 
+/**
+ * Gets the last time that the GC was run.
+ *
+ * @return	Returns the FPlatformTime::Seconds() for the last garbage collection, 0 if GC has never run.
+ */
+COREUOBJECT_API double GetLastGCTime();
+
 /**
 * Whether we are inside garbage collection
 */
diff --git a/Engine/Source/Runtime/Engine/Private/TextureDerivedData.cpp b/Engine/Source/Runtime/Engine/Private/TextureDerivedData.cpp
--- a/Engine/Source/Runtime/Engine/Private/TextureDerivedData.cpp
+++ b/Engine/Source/Runtime/Engine/Private/TextureDerivedData.cpp
@@ -744,7 +744,7 @@ void FTexturePlatformData::Cache(
 	}
 
 	bool bForceRebuild = (Flags & ETextureCacheFlags::ForceRebuild) != 0;
-	bool bAsync = !bForDDC && (Flags & ETextureCacheFlags::Async) != 0;
+	bool bAsync = (Flags & ETextureCacheFlags::Async) != 0;
 	GetTextureDerivedDataKey(InTexture, InSettingsPerLayer, DerivedDataKey);
 
 	if (!bForDDC && !bForceRebuild && InSettingsPerLayer[0].bHasEditorOnlyData)
@@ -1730,25 +1730,6 @@ void UTexture::BeginCacheForCookedPlatformData( const ITargetPlatform *TargetPla
 		TArray< TArray<FTextureBuildSettings> > BuildSettingsToCache;
 		GetBuildSettingsPerFormat(*this, BuildSettings, TargetPlatform, BuildSettingsToCache);
 
-		
-		/*TargetPlatform->GetTextureFormats(this, 0, PlatformFormats);
-		for (int32 FormatIndex = 0; FormatIndex < PlatformFormats.Num(); ++FormatIndex)
-		{
-			BuildSettings.TextureFormatName = PlatformFormats[FormatIndex];
-			BuildSettingsToCache.Add(BuildSettings);
-		}*/
-
-		uint32 CacheFlags = ETextureCacheFlags::Async | ETextureCacheFlags::InlineMips;
-
-		// If source data is resident in memory then allow the texture to be built
-		// in a background thread.
-		bool bAllowAsyncBuild = Source.IsBulkDataLoaded();
-
-		if (bAllowAsyncBuild)
-		{
-			CacheFlags |= ETextureCacheFlags::AllowAsyncBuild;
-		}
-
 		// Cull redundant settings by comparing derived data keys.
 		TArray<FString> BuildSettingsCacheKeys;
 		for (TArray<TArray<FTextureBuildSettings>>::TIterator It = BuildSettingsToCache.CreateIterator(); It; ++It)
@@ -1768,32 +1749,16 @@ void UTexture::BeginCacheForCookedPlatformData( const ITargetPlatform *TargetPla
 			}
 		}
 
-		TBitArray<> BuildSettingsProbablyCached;
-		if (BuildSettingsCacheKeys.Num())
-		{
-			BuildSettingsProbablyCached = GetDerivedDataCacheRef().CachedDataProbablyExistsBatch(BuildSettingsCacheKeys);
-		}
-
 		for (int32 SettingsIndex = 0; SettingsIndex < BuildSettingsToCache.Num(); ++SettingsIndex)
-			{
-				// UE_LOG(LogTemp, Warning, TEXT("Caching data for texture %s with id %s"), *GetName(), *DerivedDataKey);
-				uint32 CurrentCacheFlags = CacheFlags;
-				// if the cached data key exists already then we don't need to allowasync build
-				// if it doesn't then allow async builds
-			if (!BuildSettingsProbablyCached[SettingsIndex])
-					{
-						CurrentCacheFlags |= ETextureCacheFlags::AllowAsyncBuild;
-						CurrentCacheFlags |= ETextureCacheFlags::AllowAsyncLoading;
-					}
-
-				FTexturePlatformData* PlatformDataToCache;
-				PlatformDataToCache = new FTexturePlatformData();
-				PlatformDataToCache->Cache(
-					*this,
-					BuildSettingsToCache[SettingsIndex].GetData(),
-					CurrentCacheFlags,
-					nullptr
-					);
+		{
+			// UE_LOG(LogTemp, Warning, TEXT("Caching data for texture %s with id %s"), *GetName(), *DerivedDataKey);
+			FTexturePlatformData* PlatformDataToCache = new FTexturePlatformData();
+			PlatformDataToCache->Cache(
+				*this,
+				BuildSettingsToCache[SettingsIndex].GetData(),
+				ETextureCacheFlags::Async | ETextureCacheFlags::InlineMips | ETextureCacheFlags::AllowAsyncBuild | ETextureCacheFlags::AllowAsyncLoading,
+				nullptr
+				);
 			CookedPlatformData.Add(BuildSettingsCacheKeys[SettingsIndex], PlatformDataToCache);
 		}
 	}
diff --git a/Engine/Source/Runtime/Engine/Private/WorldPartition/WorldPartitionHelpers.cpp b/Engine/Source/Runtime/Engine/Private/WorldPartition/WorldPartitionHelpers.cpp
--- a/Engine/Source/Runtime/Engine/Private/WorldPartition/WorldPartitionHelpers.cpp
+++ b/Engine/Source/Runtime/Engine/Private/WorldPartition/WorldPartitionHelpers.cpp
@@ -88,15 +88,17 @@ void FWorldPartitionHelpers::ForEachActorWithLoading(UWorldPartition* WorldParti
 
 bool FWorldPartitionHelpers::HasExceededMaxMemory()
 {
-	const uint64 MemoryMinFreePhysical = 1 * 1024ll * 1024ll * 1024ll;
-	const uint64 MemoryMaxUsedPhysical = 32 * 1024ll * 1024ll * 1024ll;
 	const FPlatformMemoryStats MemStats = FPlatformMemory::GetStats();
 
+	const uint64 MemoryMinFreePhysical = 1llu * 1024 * 1024 * 1024;
+	const uint64 MemoryMaxUsedPhysical = FMath::Max(32llu * 1024 * 1024 * 1024, MemStats.TotalPhysical / 2);
+
 	const bool bHasExceededMinFreePhysical = MemStats.AvailablePhysical < MemoryMinFreePhysical;
 	const bool bHasExceededMaxUsedPhysical = MemStats.UsedPhysical >= MemoryMaxUsedPhysical;
 	const bool bHasExceededMaxMemory = bHasExceededMinFreePhysical || bHasExceededMaxUsedPhysical;
 
-	return bHasExceededMaxMemory;
+	// Even if we're not exhausting memory, GC should be run at periodic intervals
+	return bHasExceededMaxMemory || (FPlatformTime::Seconds() - GetLastGCTime()) > 30;
 };
 
 #endif // #if WITH_EDITOR
