IterationProfile: Add support for Baseline and Unit columns

I had to restructure the baseline analysis code to run before we generate the telemetry file, so the diffs will be hard to read.

This change was tested with a local run of IterationProfile with ShooterGame.

diff --git a/Engine/Source/Editor/UnrealEd/Private/Commandlets/SummarizeTraceCommandlet.cpp b/Engine/Source/Editor/UnrealEd/Private/Commandlets/SummarizeTraceCommandlet.cpp
--- a/Engine/Source/Editor/UnrealEd/Private/Commandlets/SummarizeTraceCommandlet.cpp
+++ b/Engine/Source/Editor/UnrealEd/Private/Commandlets/SummarizeTraceCommandlet.cpp
@@ -613,9 +613,10 @@ struct StatisticDefinition
 	StatisticDefinition()
 	{}
 
-	StatisticDefinition(const FString& InName, const FString& InStatistic, const FString& InWarningThreshold, const FString& InErrorThreshold)
+	StatisticDefinition(const FString& InName, const FString& InStatistic, const FString& InUnit, const FString& InWarningThreshold, const FString& InErrorThreshold)
 		: Name(InName)
 		, Statistic(InStatistic)
+		, Unit(InUnit)
 		, WarningThreshold(InWarningThreshold)
 		, ErrorThreshold(InErrorThreshold)
 	{}
@@ -623,6 +624,7 @@ struct StatisticDefinition
 	StatisticDefinition(const StatisticDefinition& InStatistic)
 		: Name(InStatistic.Name)
 		, Statistic(InStatistic.Statistic)
+		, Unit(InStatistic.Unit)
 		, WarningThreshold(InStatistic.WarningThreshold)
 		, ErrorThreshold(InStatistic.ErrorThreshold)
 	{}
@@ -631,6 +633,7 @@ struct StatisticDefinition
 	{
 		return Name == InStatistic.Name
 			&& Statistic == InStatistic.Statistic
+			&& Unit == InStatistic.Unit
 			&& WarningThreshold == InStatistic.WarningThreshold
 			&& ErrorThreshold == InStatistic.ErrorThreshold;
 	}
@@ -639,6 +642,7 @@ struct StatisticDefinition
 
 	FString Name;
 	FString Statistic;
+	FString Unit;
 	FString WarningThreshold;
 	FString ErrorThreshold;
 };
@@ -650,6 +654,7 @@ bool StatisticDefinition::LoadFromCSV(const FString& FilePath, TMultiMap<FString
 
 	int NameColumn = -1;
 	int StatisticColumn = -1;
+	int UnitColumn = -1;
 	int WarningThresholdColumn = -1;
 	int ErrorThresholdColumn = -1;
 	struct Column
@@ -661,6 +666,7 @@ bool StatisticDefinition::LoadFromCSV(const FString& FilePath, TMultiMap<FString
 	{
 		{ TEXT("Name"), &NameColumn },
 		{ TEXT("Statistic"), &StatisticColumn },
+		{ TEXT("Unit"), &UnitColumn },
 		{ TEXT("WarningThreshold"), &WarningThresholdColumn },
 		{ TEXT("ErrorThreshold"), &ErrorThresholdColumn },
 	};
@@ -699,9 +705,10 @@ bool StatisticDefinition::LoadFromCSV(const FString& FilePath, TMultiMap<FString
 		{
 			const FString& Name(Fields[NameColumn]);
 			const FString& Statistic(Fields[StatisticColumn]);
+			const FString& Unit(Fields[UnitColumn]);
 			const FString& WarningThreshold(Fields[WarningThresholdColumn]);
 			const FString& ErrorThreshold(Fields[ErrorThresholdColumn]);
-			NameToDefinitionMap.AddUnique(Name, StatisticDefinition(Name, Statistic, WarningThreshold, ErrorThreshold));
+			NameToDefinitionMap.AddUnique(Name, StatisticDefinition(Name, Statistic, Unit, WarningThreshold, ErrorThreshold));
 		}
 	}
 
@@ -717,11 +724,15 @@ struct TelemetryDefinition
 	TelemetryDefinition()
 	{}
 
-	TelemetryDefinition(const FString& InTestName, const FString& InContext, const FString& InDataPoint, const FString& InMeasurement)
+	TelemetryDefinition(const FString& InTestName, const FString& InContext, const FString& InDataPoint, const FString& InMeasurement,
+		const FString* InBaseline = nullptr,
+		const FString* InUnit = nullptr)
 		: TestName(InTestName)
 		, Context(InContext)
 		, DataPoint(InDataPoint)
 		, Measurement(InMeasurement)
+		, Baseline(InBaseline ? *InBaseline : FString ())
+		, Unit(InUnit ? *InUnit : FString ())
 	{}
 
 	TelemetryDefinition(const TelemetryDefinition& InStatistic)
@@ -729,6 +740,8 @@ struct TelemetryDefinition
 		, Context(InStatistic.Context)
 		, DataPoint(InStatistic.DataPoint)
 		, Measurement(InStatistic.Measurement)
+		, Baseline(InStatistic.Baseline)
+		, Unit(InStatistic.Unit)
 	{}
 
 	bool operator==(const TelemetryDefinition& InStatistic) const
@@ -736,7 +749,9 @@ struct TelemetryDefinition
 		return TestName == InStatistic.TestName
 			&& Context == InStatistic.Context
 			&& DataPoint == InStatistic.DataPoint
-			&& Measurement == InStatistic.Measurement;
+			&& Measurement == InStatistic.Measurement
+			&& Baseline == InStatistic.Baseline
+			&& Unit == InStatistic.Unit;
 	}
 
 	static bool LoadFromCSV(const FString& FilePath, TMap<TPair<FString,FString>, TelemetryDefinition>& ContextAndDataPointToDefinitionMap);
@@ -747,6 +762,8 @@ struct TelemetryDefinition
 	FString Context;
 	FString DataPoint;
 	FString Measurement;
+	FString Baseline;
+	FString Unit;
 };
 
 bool TelemetryDefinition::LoadFromCSV(const FString& FilePath, TMap<TPair<FString, FString>, TelemetryDefinition>& ContextAndDataPointToDefinitionMap)
@@ -758,10 +775,13 @@ bool TelemetryDefinition::LoadFromCSV(const FString& FilePath, TMap<TPair<FStrin
 	int ContextColumn = -1;
 	int DataPointColumn = -1;
 	int MeasurementColumn = -1;
+	int BaselineColumn = -1;
+	int UnitColumn = -1;
 	struct Column
 	{
 		const TCHAR* Name = nullptr;
 		int* Index = nullptr;
+		bool bRequired = true;
 	}
 	Columns[] =
 	{
@@ -769,6 +789,8 @@ bool TelemetryDefinition::LoadFromCSV(const FString& FilePath, TMap<TPair<FStrin
 		{ TEXT("Context"), &ContextColumn },
 		{ TEXT("DataPoint"), &DataPointColumn },
 		{ TEXT("Measurement"), &MeasurementColumn },
+		{ TEXT("Baseline"), &BaselineColumn, false },
+		{ TEXT("Unit"), &UnitColumn, false },
 	};
 
 	bool bValidColumns = true;
@@ -795,7 +817,7 @@ bool TelemetryDefinition::LoadFromCSV(const FString& FilePath, TMap<TPair<FStrin
 					}
 				}
 
-				if (*Column.Index == -1)
+				if (*Column.Index == -1 && Column.bRequired)
 				{
 					bValidColumns = false;
 				}
@@ -807,7 +829,20 @@ bool TelemetryDefinition::LoadFromCSV(const FString& FilePath, TMap<TPair<FStrin
 			const FString& Context(Fields[ContextColumn]);
 			const FString& DataPoint(Fields[DataPointColumn]);
 			const FString& Measurement(Fields[MeasurementColumn]);
-			ContextAndDataPointToDefinitionMap.Add(TPair<FString, FString>(Context, DataPoint), TelemetryDefinition(TestName, Context, DataPoint, Measurement));
+
+			FString Baseline;
+			if (BaselineColumn != -1)
+			{
+				Baseline = Fields[BaselineColumn];
+			}
+
+			FString Unit;
+			if (UnitColumn != -1)
+			{
+				Unit = Fields[UnitColumn];
+			}
+
+			ContextAndDataPointToDefinitionMap.Add(TPair<FString, FString>(Context, DataPoint), TelemetryDefinition(TestName, Context, DataPoint, Measurement, &Baseline, &Unit));
 		}
 	}
 
@@ -928,7 +963,6 @@ int32 USummarizeTraceCommandlet::Main(const FString& CmdLineParams)
 		UE_LOG(LogSummarizeTrace, Log, TEXT(" Required: -inputfile=<utrace path>   (The utrace you wish to process)"));
 		UE_LOG(LogSummarizeTrace, Log, TEXT(" Optional: -statsfile=<csv path>      (The csv of statistics to generate)"));
 		UE_LOG(LogSummarizeTrace, Log, TEXT(" Optional: -testname=<string>         (Test name to use in telemetry csv)"));
-		UE_LOG(LogSummarizeTrace, Log, TEXT(" Optional: -skipbaseline              (Skip comparing values vs. baselines)"));
 		return 0;
 	}
 
@@ -1126,54 +1160,33 @@ int32 USummarizeTraceCommandlet::Main(const FString& CmdLineParams)
 		FString TelemetryCsvFileName = TraceFileBasename + TEXT("Telemetry");
 		TelemetryCsvFileName = FPaths::Combine(TracePath, FPaths::SetExtension(TelemetryCsvFileName, "csv"));
 
-		UE_LOG(LogSummarizeTrace, Display, TEXT("Writing telemetry to %s..."), *TelemetryCsvFileName);
-		IFileHandle* TelemetryCsvHandle = FPlatformFileManager::Get().GetPlatformFile().OpenWrite(*TelemetryCsvFileName);
-		TArray<TelemetryDefinition> TelemetryWritten;
-		if (TelemetryCsvHandle)
-		{
-			// override the test name
-			FString TestName = TraceFileBasename;
-			FParse::Value(*CmdLineParams, TEXT("testname="), TestName, true);
+		// override the test name
+		FString TestName = TraceFileBasename;
+		FParse::Value(*CmdLineParams, TEXT("testname="), TestName, true);
 
-			// resolve scopes to telemetry
-			for (const FSummarizeCpuAnalyzer::FScope& Scope : SortedScopes)
-			{
-				TArray<StatisticDefinition> Statistics;
-				NameToDefinitionMap.MultiFind(Scope.Name, Statistics, true);
-				for (const StatisticDefinition& Statistic : Statistics)
-				{
-					TelemetryWritten.Add(TelemetryDefinition(TestName, Scope.Name, Statistic.Statistic, Scope.GetValue(Statistic.Statistic)));
-				}
-			}
+		TArray<TelemetryDefinition> TelemetryData;
 
-			// resolve counters to telemetry
-			for (const TMap<uint16, FSummarizeCountersAnalyzer::FCounter>::ElementType& Counter : CountersAnalyzer.Counters)
+		// resolve scopes to telemetry
+		for (const FSummarizeCpuAnalyzer::FScope& Scope : SortedScopes)
+		{
+			TArray<StatisticDefinition> Statistics;
+			NameToDefinitionMap.MultiFind(Scope.Name, Statistics, true);
+			for (const StatisticDefinition& Statistic : Statistics)
 			{
-				TArray<StatisticDefinition> Statistics;
-				NameToDefinitionMap.MultiFind(Counter.Value.Name, Statistics, true);
-				ensure(Statistics.Num() <= 1); // there should only be one, the counter value
-				for (const StatisticDefinition& Statistic : Statistics)
-				{
-					TelemetryWritten.Add(TelemetryDefinition(TestName, Counter.Value.Name, Statistic.Statistic, Counter.Value.GetValue()));
-				}
+				TelemetryData.Add(TelemetryDefinition(TestName, Scope.Name, Statistic.Statistic, Scope.GetValue(Statistic.Statistic)));
 			}
+		}
 
-			// no newline, see row printfs
-			WriteUTF8(TelemetryCsvHandle, FString::Printf(TEXT("TestName,Context,DataPoint,Measurement,")));
-			for (const TelemetryDefinition& Telemetry : TelemetryWritten)
+		// resolve counters to telemetry
+		for (const TMap<uint16, FSummarizeCountersAnalyzer::FCounter>::ElementType& Counter : CountersAnalyzer.Counters)
+		{
+			TArray<StatisticDefinition> Statistics;
+			NameToDefinitionMap.MultiFind(Counter.Value.Name, Statistics, true);
+			ensure(Statistics.Num() <= 1); // there should only be one, the counter value
+			for (const StatisticDefinition& Statistic : Statistics)
 			{
-				// note newline is at the front of every data line to prevent final extraneous newline, per customary for csv
-				WriteUTF8(TelemetryCsvHandle, FString::Printf(TEXT("\n%s,%s,%s,%s,"), *Telemetry.TestName, *Telemetry.Context, *Telemetry.DataPoint, *Telemetry.Measurement));
+				TelemetryData.Add(TelemetryDefinition(TestName, Counter.Value.Name, Statistic.Statistic, Counter.Value.GetValue()));
 			}
-
-			TelemetryCsvHandle->Flush();
-			delete TelemetryCsvHandle;
-			TelemetryCsvHandle = nullptr;
-		}
-		else
-		{
-			UE_LOG(LogSummarizeTrace, Error, TEXT("Unable to open telemetry csv '%s' for write"), *TelemetryCsvFileName);
-			return 1;
 		}
 
 		// compare vs. baseline telemetry file, if it exists
@@ -1193,32 +1206,36 @@ int32 USummarizeTraceCommandlet::Main(const FString& CmdLineParams)
 			check(bCSVOk);
 
 			// for every telemetry item we wrote for this trace...
-			for (const TelemetryDefinition& Telemetry : TelemetryWritten)
+			for (TelemetryDefinition& Telemetry : TelemetryData)
 			{
-				// find the corresponding keyed telemetry item in the baseline telemetry file...
-				TelemetryDefinition* BaselineTelemetry = ContextAndDataPointToDefinitionMap.Find(TPair<FString, FString>(Telemetry.Context, Telemetry.DataPoint));
-				if (BaselineTelemetry)
-				{
-					// the threshold is defined along with the original statistic map
-					const StatisticDefinition* RelatedStatistic = nullptr;
+				// the threshold is defined along with the original statistic map
+				const StatisticDefinition* RelatedStatistic = nullptr;
 
-					// find the statistic definition
-					TArray<StatisticDefinition> Statistics;
-					NameToDefinitionMap.MultiFind(Telemetry.Context, Statistics, true);
-					for (const StatisticDefinition& Statistic : Statistics)
+				// find the statistic definition
+				TArray<StatisticDefinition> Statistics;
+				NameToDefinitionMap.MultiFind(Telemetry.Context, Statistics, true);
+				for (const StatisticDefinition& Statistic : Statistics)
+				{
+					// the find will match on name, here we just need to find the right statistic for that named item
+					if (Statistic.Statistic == Telemetry.DataPoint)
 					{
-						// the find will match on name, here we just need to find the right statistic for that named item
-						if (Statistic.Statistic == Telemetry.DataPoint)
-						{
-							// we found it!
-							RelatedStatistic = &Statistic;
-							break;
-						}
+						// we found it!
+						RelatedStatistic = &Statistic;
+						break;
 					}
+				}
+
+				// do we still have the statistic definition in our current stats file? (if we don't that's fine, we don't care about it anymore)
+				if (RelatedStatistic)
+				{
+					Telemetry.Unit = RelatedStatistic->Unit;
 
-					// do we still have the statistic definition in our current stats file? (if we don't that's fine, we don't care about it anymore)
-					if (RelatedStatistic)
+					// find the corresponding keyed telemetry item in the baseline telemetry file...
+					TelemetryDefinition* BaselineTelemetry = ContextAndDataPointToDefinitionMap.Find(TPair<FString, FString>(Telemetry.Context, Telemetry.DataPoint));
+					if (BaselineTelemetry)
 					{
+						Telemetry.Baseline = BaselineTelemetry->Measurement;
+
 						// let's only report on statistics that have an assigned threshold, to keep things concise
 						if (!RelatedStatistic->WarningThreshold.IsEmpty() || !RelatedStatistic->ErrorThreshold.IsEmpty())
 						{
@@ -1266,13 +1283,35 @@ int32 USummarizeTraceCommandlet::Main(const FString& CmdLineParams)
 							}
 						}
 					}
+					else
+					{
+						UE_LOG(LogSummarizeTrace, Display, TEXT("Telemetry for %s,%s has no baseline measurement, skipping..."), *Telemetry.Context, *Telemetry.DataPoint);
+					}
 				}
-				else
-				{
-					UE_LOG(LogSummarizeTrace, Display, TEXT("Telemetry for %s,%s has no baseline measurement, skipping..."), *Telemetry.Context, *Telemetry.DataPoint);
-				}
 			}
 		}
+
+		UE_LOG(LogSummarizeTrace, Display, TEXT("Writing telemetry to %s..."), *TelemetryCsvFileName);
+		IFileHandle* TelemetryCsvHandle = FPlatformFileManager::Get().GetPlatformFile().OpenWrite(*TelemetryCsvFileName);
+		if (TelemetryCsvHandle)
+		{
+			// no newline, see row printfs
+			WriteUTF8(TelemetryCsvHandle, FString::Printf(TEXT("TestName,Context,DataPoint,Measurement,Baseline,Unit,")));
+			for (const TelemetryDefinition& Telemetry : TelemetryData)
+			{
+				// note newline is at the front of every data line to prevent final extraneous newline, per customary for csv
+				WriteUTF8(TelemetryCsvHandle, FString::Printf(TEXT("\n%s,%s,%s,%s,%s,%s,"), *Telemetry.TestName, *Telemetry.Context, *Telemetry.DataPoint, *Telemetry.Measurement, *Telemetry.Baseline, *Telemetry.Unit));
+			}
+
+			TelemetryCsvHandle->Flush();
+			delete TelemetryCsvHandle;
+			TelemetryCsvHandle = nullptr;
+		}
+		else
+		{
+			UE_LOG(LogSummarizeTrace, Error, TEXT("Unable to open telemetry csv '%s' for write"), *TelemetryCsvFileName);
+			return 1;
+		}
 	}
 
 	return 0;
