Enable current shader platform in material platform stats by default (UE-118073).

diff --git a/Engine/Source/Editor/UnrealEd/Private/PreferenceStubs.cpp b/Engine/Source/Editor/UnrealEd/Private/PreferenceStubs.cpp
--- a/Engine/Source/Editor/UnrealEd/Private/PreferenceStubs.cpp
+++ b/Engine/Source/Editor/UnrealEd/Private/PreferenceStubs.cpp
@@ -65,12 +65,11 @@ UMaterialEditorOptions::UMaterialEditorOptions(const FObjectInitializer& ObjectI
 UMaterialStatsOptions::UMaterialStatsOptions(const FObjectInitializer& ObjectInitializer)
 	: Super(ObjectInitializer)
 {
-#if PLATFORM_WINDOWS
-	//#todo-sm6
-	bPlatformUsed[/*GMaxRHIFeatureLevel == ERHIFeatureLevel::SM5 ? */SP_PCD3D_SM5/* : SP_PCD3D_SM4*/] = 1;
-#elif PLATFORM_IOS
-	bPlatformUsed[SP_METAL] = 1;
-#endif
+	uint32 CurrentlyUsedSP = static_cast<int32>(GMaxRHIShaderPlatform);
+	if (CurrentlyUsedSP < UE_ARRAY_COUNT(bPlatformUsed))
+	{
+		bPlatformUsed[CurrentlyUsedSP] = 1;
+	}
 
 	bMaterialQualityUsed[EMaterialQualityLevel::High] = 1;
 }
