BulkDataRegistry

Updating FVirtualizedUntypedBulkData and textures to use the BulkDataRegistry.
BulkDataRegistry: Add get/put accessors for the cached BulkDataList of packages.
EditorDomain: Move ClassDigests into a global variable that can be shared with BulkDataRegistry.
EditorDomain: Improve performance of GetFileSize by fetching metadata only.
Tickable Cook Objects, for systems used by the cooker that need to be ticked.
Implementation of the the BulkDataRegistry that uses the DDC cache for persistent storage of the BulkDataList.

diff --git a/Engine/Config/BaseEditor.ini b/Engine/Config/BaseEditor.ini
--- a/Engine/Config/BaseEditor.ini
+++ b/Engine/Config/BaseEditor.ini
@@ -329,6 +329,7 @@ MaxMemoryAllowance=16384
 [CookSettings]
 EditorDomainEnabled=false
 EditorDomainSaveUnversioned=false
+BulkDataRegistryEnabled=false
 PackagesPerGC=0
 IdleTimeToGC=20
 MemoryMaxUsedPhysical=16384
diff --git a/Engine/Source/Editor/UnrealEd/Classes/CookOnTheSide/CookOnTheFlyServer.h b/Engine/Source/Editor/UnrealEd/Classes/CookOnTheSide/CookOnTheFlyServer.h
--- a/Engine/Source/Editor/UnrealEd/Classes/CookOnTheSide/CookOnTheFlyServer.h
+++ b/Engine/Source/Editor/UnrealEd/Classes/CookOnTheSide/CookOnTheFlyServer.h
@@ -1132,6 +1132,8 @@ private:
 	float SaveBusyTimeStarted = 0.f;
 	float LoadBusyTimeLastRetry = 0.f;
 	float LoadBusyTimeStarted = 0.f;
+	/** Tracking for the ticking of tickable cook objects */
+	float LastCookableObjectTickTime = 0.f;
 
 	// These helper structs are all TUniquePtr rather than inline members so that we can keep their headers private.  See class header comments for their purpose.
 	TUniquePtr<UE::Cook::FPackageTracker> PackageTracker;
diff --git a/Engine/Source/Editor/UnrealEd/Private/CookOnTheFlyServer.cpp b/Engine/Source/Editor/UnrealEd/Private/CookOnTheFlyServer.cpp
--- a/Engine/Source/Editor/UnrealEd/Private/CookOnTheFlyServer.cpp
+++ b/Engine/Source/Editor/UnrealEd/Private/CookOnTheFlyServer.cpp
@@ -230,6 +230,7 @@ DECLARE_CYCLE_STAT(TEXT("Precache Derived data for platform"), STAT_TickPrecache
 DECLARE_CYCLE_STAT(TEXT("Tick cooking"), STAT_TickCooker, STATGROUP_Cooking);
 
 constexpr uint32 ExpectedMaxNumPlatforms = 32;
+constexpr float TickCookableObjectsFrameTime = .100f;
 
 namespace UE
 {
@@ -1529,8 +1530,17 @@ uint32 UCookOnTheFlyServer::TickCookOnTheSide(const float TimeSlice, uint32 &Coo
 void UCookOnTheFlyServer::TickCookStatus(UE::Cook::FTickStackData& StackData)
 {
 	UE_SCOPED_COOKTIMER(TickCookStatus);
-	UpdateDisplay(StackData.TickFlags, false /* bForceDisplay */);
 
+	// TODO: Calculate CurrentTime once per status update and share it with e.g. StackData.Timer
+	double CurrentTime = FPlatformTime::Seconds();
+	if (LastCookableObjectTickTime + TickCookableObjectsFrameTime <= CurrentTime)
+	{
+		UE_SCOPED_COOKTIMER(TickCookableObjects);
+		FTickableCookObject::TickObjects(CurrentTime - LastCookableObjectTickTime, false /* bTickComplete */);
+		LastCookableObjectTickTime = CurrentTime;
+	}
+
+	UpdateDisplay(StackData.TickFlags, false /* bForceDisplay */);
 	// prevent autosave from happening until we are finished cooking
 	// causes really bad hitches
 	if (GUnrealEd)
@@ -7270,6 +7280,13 @@ void UCookOnTheFlyServer::CookByTheBookFinished()
 
 	UE_LOG(LogCook, Display, TEXT("Finishing up..."));
 
+	{
+		UE_SCOPED_COOKTIMER(TickCookableObjects);
+		float CurrentTime = FPlatformTime::Seconds();
+		FTickableCookObject::TickObjects(CurrentTime - LastCookableObjectTickTime, true /* bTickComplete */);
+		LastCookableObjectTickTime = CurrentTime;
+	}
+
 	UPackage::WaitForAsyncFileWrites();
 	
 	GetDerivedDataCacheRef().WaitForQuiescence(true);
diff --git a/Engine/Source/Editor/UnrealEd/Private/EditorDomain/BulkDataRegistryEditorDomain.cpp b/Engine/Source/Editor/UnrealEd/Private/EditorDomain/BulkDataRegistryEditorDomain.cpp
new file mode 100644
--- /dev/null
+++ b/Engine/Source/Editor/UnrealEd/Private/EditorDomain/BulkDataRegistryEditorDomain.cpp
@@ -0,0 +1,495 @@
+// Copyright Epic Games, Inc. All Rights Reserved.
+
+#include "BulkDataRegistryEditorDomain.h"
+
+#include "Serialization/BulkDataRegistry.h"
+#include "Compression/CompressedBuffer.h"
+#include "DerivedDataCacheRecord.h"
+#include "EditorBuildInputResolver.h"
+#include "EditorDomain/EditorDomainUtils.h"
+#include "HAL/PlatformTime.h"
+#include "Interfaces/IPluginManager.h"
+#include "IO/IoHash.h"
+#include "Memory/SharedBuffer.h"
+#include "Misc/CoreDelegates.h"
+#include "Misc/ScopeLock.h"
+#include "Misc/ScopeRWLock.h"
+#include "Serialization/MemoryReader.h"
+#include "Templates/RefCounting.h"
+#include "TickableEditorObject.h"
+#include "UObject/Package.h"
+#include "UObject/UObjectIterator.h"
+
+namespace UE::BulkDataRegistry::Private
+{
+
+namespace Constants
+{
+	constexpr double UpdateHashesInitialWaitDuration = 1.0;
+	constexpr uint64 TempLoadedPayloadsSizeBudget = 1024 * 1024 * 100;
+	constexpr double TempLoadedPayloadsDuration = 60.;
+	constexpr double PendingPackageCooldown = 60.;
+
+}
+
+/** Add a hook to the BulkDataRegistry's startup delegate to use the EditorDomain as the BulkDataRegistry */
+class FEditorDomainRegisterAsBulkDataRegistry
+{
+public:
+	FEditorDomainRegisterAsBulkDataRegistry()
+	{
+		IBulkDataRegistry::GetSetBulkDataRegistryDelegate().BindStatic(SetBulkDataRegistry);
+	}
+
+	static IBulkDataRegistry* SetBulkDataRegistry()
+	{
+		return new FBulkDataRegistryEditorDomain();
+	}
+} GRegisterAsBulkDataRegistry;
+
+
+FBulkDataRegistryEditorDomain::FBulkDataRegistryEditorDomain()
+{
+	// We piggyback on the BulkDataRegistry hook to set the pointer to tunnel in the pointer to the EditorBuildInputResolver as well
+	SetGlobalBuildInputResolver(&UE::DerivedData::FEditorBuildInputResolver::Get());
+}
+
+FBulkDataRegistryEditorDomain::~FBulkDataRegistryEditorDomain()
+{
+	SetGlobalBuildInputResolver(nullptr);
+
+	{
+		FScopeLock PendingPackageScopeLock(&PendingPackageLock);
+		bAllowCacheUpdates = false;
+	}
+	PendingPackages.Empty();
+
+	FCoreDelegates::OnPostEngineInit.RemoveAll(this);
+}
+
+void FBulkDataRegistryEditorDomain::Register(UPackage* Owner, const UE::Virtualization::FVirtualizedUntypedBulkData& BulkData)
+{
+	if (!BulkData.GetIdentifier().IsValid())
+	{
+		return;
+	}
+
+	FName PackageName = NAME_None;
+	UE::Virtualization::FVirtualizedUntypedBulkData CopyBulk(BulkData.CopyTornOff());
+	if (Owner
+		&& Owner->GetFileSize() // We only record the BulkDataList for disk packages
+		&& CopyBulk.CanSaveForRegistry()
+		)
+	{
+		PackageName = Owner->GetFName();
+		AddPendingPackageBulkData(PackageName, CopyBulk);
+	}
+
+	FName PackageNameToUpdate = CopyBulk.HasPlaceholderPayloadId() ? PackageName : NAME_None;
+	FWriteScopeLock RegistryScopeLock(RegistryLock);
+	Registry.Add(BulkData.GetIdentifier(), FBulkSource(MoveTemp(CopyBulk), PackageNameToUpdate));
+}
+
+void FBulkDataRegistryEditorDomain::OnExitMemory(const UE::Virtualization::FVirtualizedUntypedBulkData& BulkData)
+{
+	const FGuid& Key = BulkData.GetIdentifier();
+	FWriteScopeLock RegistryScopeLock(RegistryLock);
+	FBulkSource* Existing = Registry.Find(Key);
+	if (Existing)
+	{
+		if (Existing->BulkData.IsMemoryOnlyPayload())
+		{
+			Registry.Remove(Key);
+		}
+	}
+}
+
+TFuture<UE::BulkDataRegistry::FMetaData> FBulkDataRegistryEditorDomain::GetMeta(const FGuid& BulkDataId)
+{
+	TPromise< UE::BulkDataRegistry::FMetaData> Promise;
+	UE::BulkDataRegistry::FMetaData Result;
+	Result.bValid = false;
+	auto TryExistingCompleteResult = [&BulkDataId, this, &Result](FBulkSource*& OutExisting)
+	{
+		OutExisting = Registry.Find(BulkDataId);
+		if (!OutExisting)
+		{
+			Result.bValid = false;
+			Result.RawHash = FIoHash();
+			Result.RawSize = 0;
+			return true;
+		}
+
+		if (!OutExisting->BulkData.HasPlaceholderPayloadId())
+		{
+			Result.bValid = true;
+			Result.RawHash = OutExisting->BulkData.GetPayloadId().GetIdentifier();
+			Result.RawSize = OutExisting->BulkData.GetPayloadSize();
+			return true;
+		}
+
+		return false;
+	};
+
+	TOptional<UE::Virtualization::FVirtualizedUntypedBulkData> CopyBulk;
+	FName PackageNameToUpdate;
+	{
+		FReadScopeLock RegistryScopeLock(RegistryLock);
+		FBulkSource* Existing = nullptr;
+		if (TryExistingCompleteResult(Existing))
+		{
+			Promise.SetValue(MoveTemp(Result));
+			return Promise.GetFuture();
+		}
+
+		CopyBulk.Emplace(Existing->BulkData);
+		PackageNameToUpdate = Existing->PackageNameToUpdate;
+	}
+
+	CopyBulk->UpdatePayloadId();
+	Result.bValid = true;
+	Result.RawHash= CopyBulk->GetPayloadId().GetIdentifier();
+	Result.RawSize = CopyBulk->GetPayloadSize();
+
+	{
+		FWriteScopeLock RegistryScopeLock(RegistryLock);
+		FBulkSource* Existing;
+		if (TryExistingCompleteResult(Existing))
+		{
+			// Some other thread has beaten us to updating the RawHash
+			Promise.SetValue(MoveTemp(Result));
+			return Promise.GetFuture();
+		}
+
+		Existing->BulkData = *CopyBulk;
+		Existing->PackageNameToUpdate = NAME_None;
+
+		AddTempLoadedPayload(CopyBulk->GetIdentifier(), CopyBulk->GetPayloadSize());
+		PruneTempLoadedPayloads();
+	}
+
+	if (!PackageNameToUpdate.IsNone())
+	{
+		AddPendingPackageBulkData(PackageNameToUpdate, *CopyBulk);
+	}
+
+	Promise.SetValue(MoveTemp(Result));
+	return Promise.GetFuture();
+}
+
+TFuture<UE::BulkDataRegistry::FData> FBulkDataRegistryEditorDomain::GetData(const FGuid& BulkDataId)
+{
+	TPromise<UE::BulkDataRegistry::FData> Result;
+	TOptional<UE::Virtualization::FVirtualizedUntypedBulkData> CopyBulk;
+	FName PackageNameToUpdate;
+	bool bNeedsUpdate = false;
+	{
+		FReadScopeLock RegistryScopeLock(RegistryLock);
+		FBulkSource* Existing = Registry.Find(BulkDataId);
+		if (!Existing)
+		{
+			Result.SetValue(UE::BulkDataRegistry::FData{ false, FCompressedBuffer() });
+			return Result.GetFuture();
+		}
+		CopyBulk.Emplace(Existing->BulkData);
+		PackageNameToUpdate = Existing->PackageNameToUpdate;
+		bNeedsUpdate = CopyBulk->HasPlaceholderPayloadId();
+	}
+
+	if (bNeedsUpdate)
+	{
+		CopyBulk->UpdatePayloadId();
+	}
+	FCompressedBuffer Payload = CopyBulk->GetCompressedPayload().Get();
+	bool bUpdatedByGetPayload = true; // BULKDATAREGISTRY_TODO: Calculate this and return it from GetPayload
+	bool bChanged = bNeedsUpdate || bUpdatedByGetPayload;
+	CopyBulk->UnloadData();
+
+	if (bChanged)
+	{
+		FWriteScopeLock RegistryScopeLock(RegistryLock);
+		FBulkSource* Existing = Registry.Find(BulkDataId);
+		if (Existing)
+		{
+			bNeedsUpdate = Existing->BulkData.HasPlaceholderPayloadId();
+			Existing->BulkData = *CopyBulk;
+			Existing->PackageNameToUpdate = NAME_None;
+		}
+	}
+
+	if (bNeedsUpdate && !PackageNameToUpdate.IsNone())
+	{
+		AddPendingPackageBulkData(PackageNameToUpdate, *CopyBulk);
+	}
+
+	Result.SetValue(UE::BulkDataRegistry::FData{ true, MoveTemp(Payload) });
+	return Result.GetFuture();
+}
+
+void FBulkDataRegistryEditorDomain::TickCook(float DeltaTime, bool bTickComplete)
+{
+	bool bWaitForCooldown = !bTickComplete;
+	PollPendingPackages(bWaitForCooldown);
+	{
+		FWriteScopeLock RegistryScopeLock(RegistryLock);
+		PruneTempLoadedPayloads();
+	}
+}
+
+void FBulkDataRegistryEditorDomain::Tick(float DeltaTime)
+{
+	TickCook(DeltaTime, false /* bTickComplete */);
+}
+
+void FBulkDataRegistryEditorDomain::AddPendingPackageBulkData(FName PackageName, const UE::Virtualization::FVirtualizedUntypedBulkData& BulkData)
+{
+	FScopeLock PendingPackageScopeLock(&PendingPackageLock);
+	TUniquePtr<FPendingPackage>& PendingPackage = PendingPackages.FindOrAdd(PackageName);
+	if (!PendingPackage.IsValid())
+	{
+		PendingPackage = MakeUnique<FPendingPackage>(PackageName, this);
+	}
+	PendingPackage->AddBulkData(BulkData);
+}
+
+void FBulkDataRegistryEditorDomain::PollPendingPackages(bool bWaitForCooldown)
+{
+	// BULKDATAREGISTRY_TODO: Reduce array iteration time; tick is called frequently in the editor
+	TArray<TPair<FName, TUniquePtr<FPendingPackage>>> PoppedPackages;
+	{
+		FScopeLock PendingPackageScopeLock(&PendingPackageLock);
+		if (PendingPackages.Num() == 0)
+		{
+			return;
+		}
+
+		double CreateTimeToPush = 0.0;
+		if (bWaitForCooldown)
+		{
+			CreateTimeToPush = FPlatformTime::Seconds() - Constants::PendingPackageCooldown;
+		}
+		for (TPair<FName, TUniquePtr<FPendingPackage>>& Pair : PendingPackages)
+		{
+			FPendingPackage* PendingPackage = Pair.Value.Get();
+			if (PendingPackage->IsReadFinished() && (!bWaitForCooldown || PendingPackage->GetCreateTime() <= CreateTimeToPush))
+			{
+				PoppedPackages.Emplace(Pair.Key, MoveTemp(Pair.Value));
+			}
+		}
+		for (TPair<FName, TUniquePtr<FPendingPackage>>& Pair : PoppedPackages)
+		{
+			PendingPackages.Remove(Pair.Get<0>());
+		}
+	}
+
+	TArray<UE::Virtualization::FVirtualizedUntypedBulkData> BulkDatas;
+	for (TPair<FName, TUniquePtr<FPendingPackage>>& Pair : PoppedPackages)
+	{
+		FName PackageName = Pair.Get<0>();
+		TUniquePtr<FPendingPackage>& PendingPackage = Pair.Get<1>();
+		PendingPackage->UpdateCache();
+	}
+}
+
+void FBulkDataRegistryEditorDomain::AddTempLoadedPayload(const FGuid& RegistryKey, uint64 PayloadSize)
+{
+	// Called within RegistryLock WriteLock
+	TempLoadedPayloads.Add(FTempLoadedPayload{ RegistryKey, PayloadSize, FPlatformTime::Seconds() + Constants::TempLoadedPayloadsDuration });
+	TempLoadedPayloadsSize += PayloadSize;
+}
+
+void FBulkDataRegistryEditorDomain::PruneTempLoadedPayloads()
+{
+	// Called within RegistryLock WriteLock
+	if (!TempLoadedPayloads.IsEmpty())
+	{
+		double CurrentTime = FPlatformTime::Seconds();
+		while (!TempLoadedPayloads.IsEmpty() &&
+			(TempLoadedPayloadsSize > Constants::TempLoadedPayloadsSizeBudget
+				|| TempLoadedPayloads[0].EndTime <= CurrentTime))
+		{
+			FTempLoadedPayload Payload = TempLoadedPayloads.PopFrontValue();
+			FBulkSource* Existing = Registry.Find(Payload.Guid);
+			if (Existing)
+			{
+				// UnloadData only unloads the in-memory data, and only if the BulkData can be reloaded from disk
+				Existing->BulkData.UnloadData();
+			}
+			TempLoadedPayloadsSize -= Payload.PayloadSize;
+		}
+	}
+}
+
+TConstArrayView<uint8> MakeArrayView(FSharedBuffer Buffer)
+{
+	return TConstArrayView<uint8>(reinterpret_cast<const uint8*>(Buffer.GetData()), Buffer.GetSize());
+}
+
+FPendingPackage::FPendingPackage(FName InPackageName, FBulkDataRegistryEditorDomain* InOwner)
+	: PackageName(InPackageName)
+	, Owner(InOwner)
+	, CreateTime(FPlatformTime::Seconds())
+{
+	CacheRequest = UE::EditorDomain::GetBulkDataList(PackageName, [this](FSharedBuffer Buffer)
+		{ OnCacheResults(Buffer); });
+}
+
+FPendingPackage::~FPendingPackage()
+{
+	CacheRequest.Cancel();
+}
+
+void FPendingPackage::OnCacheResults(FSharedBuffer Buffer)
+{
+	if (Buffer.IsNull())
+	{
+		return;
+	}
+	FMemoryReaderView Reader(MakeArrayView(Buffer));
+	Serialize(Reader, CachedBulkDatas);
+	if (Reader.IsError())
+	{
+		CachedBulkDatas.Empty();
+	}
+
+	// Add each CachedBulkData to the Registry. If the CachedBulkData exists, do not overwrite it, but do update the RawHash if it is missing it.
+	FWriteScopeLock RegistryScopeLock(Owner->RegistryLock);
+	if (!Owner->bAllowCacheUpdates)
+	{
+		return;
+	}
+	for (const UE::Virtualization::FVirtualizedUntypedBulkData& BulkData : CachedBulkDatas)
+	{
+		FBulkSource& TargetBulkSource = Owner->Registry.FindOrAdd(BulkData.GetIdentifier());
+		UE::Virtualization::FVirtualizedUntypedBulkData& TargetBulkData = TargetBulkSource.BulkData;
+		if (!TargetBulkData.GetIdentifier().IsValid())
+		{
+			TargetBulkData = BulkData;
+			TargetBulkSource.PackageNameToUpdate = BulkData.HasPlaceholderPayloadId() ? PackageName : NAME_None;
+		}
+		else
+		{
+			check(TargetBulkData.GetIdentifier() == BulkData.GetIdentifier());
+			if (!BulkData.HasPlaceholderPayloadId() && TargetBulkData.HasPlaceholderPayloadId())
+			{
+				// BULKDATAREGISTRY_TODO: Verify the payloads are the same
+				// && TargetBulkData.PayloadInfoMatches(BulkData)
+				TargetBulkData = BulkData;
+				// We now know that we don't need to update the package with the RawHash, so remove the PackageNameToUpdate
+				TargetBulkSource.PackageNameToUpdate = NAME_None;
+			}
+		}
+	}
+}
+
+bool FPendingPackage::IsReadFinished() const
+{
+	return CacheRequest.Poll();
+}
+
+void FPendingPackage::UpdateCache()
+{
+	CacheRequest.Wait();
+
+	// Remove any duplicates in the runtime BulkDatas; elements later in the list override were added after and override earlier elements
+	{
+		TSet<FGuid> BulkDataGuids;
+		for (int32 Index = BulkDataGuids.Num(); Index >= 0; --Index)
+		{
+			bool bAlreadyExists;
+			BulkDataGuids.Add(BulkDatas[Index].GetIdentifier(), &bAlreadyExists);
+			if (bAlreadyExists)
+			{
+				BulkDatas.RemoveAt(Index);
+			}
+		}
+	}
+
+	// If the current list of BulkDatas has any elements not in the cache, update the cache
+	TMap<FGuid, UE::Virtualization::FVirtualizedUntypedBulkData*> GuidToCachedBulkData;
+	for (UE::Virtualization::FVirtualizedUntypedBulkData& BulkData : CachedBulkDatas)
+	{
+		GuidToCachedBulkData.Add(BulkData.GetIdentifier(), &BulkData);
+	}
+
+	// Update any existing BulkDatas and add on any new ones
+	bool bUpdatedExisting = false;
+	TArray<UE::Virtualization::FVirtualizedUntypedBulkData*> BulkDatasToAdd;
+	for (UE::Virtualization::FVirtualizedUntypedBulkData& BulkData : BulkDatas)
+	{
+		UE::Virtualization::FVirtualizedUntypedBulkData** CachedBulkData = GuidToCachedBulkData.Find(BulkData.GetIdentifier());
+		if (CachedBulkData)
+		{
+			UE::Virtualization::FVirtualizedUntypedBulkData& Cached = **CachedBulkData;
+			if (Cached.HasPlaceholderPayloadId() && !BulkData.HasPlaceholderPayloadId())
+				// BULKDATAREGISTRY_TODO: Also update if the payloads are different
+				// || !TargetBulkData.PayloadInfoMatches(BulkData)
+			{
+				bUpdatedExisting = true;
+				Cached = BulkData;
+			}
+		}
+		else
+		{
+			BulkDatasToAdd.Add(&BulkData);
+		}
+	}
+
+	if (bUpdatedExisting || BulkDatasToAdd.Num() > 0)
+	{
+		// Add on all the new ones
+		GuidToCachedBulkData.Reset(); // This map is now invalidated since we are modifying the array it has pointers into
+		for (UE::Virtualization::FVirtualizedUntypedBulkData* BulkData : BulkDatasToAdd)
+		{
+			CachedBulkDatas.Add(*BulkData);
+		}
+
+		// Sort the list by guid, to avoid indeterminism in the list
+		CachedBulkDatas.Sort([](const UE::Virtualization::FVirtualizedUntypedBulkData& A, const UE::Virtualization::FVirtualizedUntypedBulkData& B)
+			{
+				return A.GetIdentifier() < B.GetIdentifier();
+			});
+
+		TArray<uint8> Bytes;
+		FMemoryWriter Writer(Bytes);
+		Serialize(Writer, BulkDatas);
+		// BULKDATAREGISTRY_TODO: Store the request in a list of pending writes so that we don't try to write a new entry for the package
+		// until the previous entry has finished storing, to avoid stale data overwriting new data
+		UE::EditorDomain::PutBulkDataList(PackageName, MakeSharedBufferFromArray(MoveTemp(Bytes)));
+	}
+}
+
+void Serialize(FArchive& Ar, TArray<UE::Virtualization::FVirtualizedUntypedBulkData>& Datas)
+{
+	int32 Num = Datas.Num();
+	Ar << Num;
+
+	const uint32 MinSize = 4;
+	if (Ar.IsLoading())
+	{
+		if (Ar.IsError() || Num * MinSize > Ar.TotalSize() - Ar.Tell())
+		{
+			Ar.SetError();
+			Datas.Empty();
+			return;
+		}
+		else
+		{
+			Datas.Empty(Num);
+		}
+		for (int32 n = 0; n < Num; ++n)
+		{
+			UE::Virtualization::FVirtualizedUntypedBulkData& BulkData = Datas.Emplace_GetRef();
+			BulkData.SerializeForRegistry(Ar);
+		}
+	}
+	else
+	{
+		for (UE::Virtualization::FVirtualizedUntypedBulkData& BulkData : Datas)
+		{
+			BulkData.SerializeForRegistry(Ar);
+		}
+	}
+}
+
+} // namespace UE::BulkDataRegistry::Private
diff --git a/Engine/Source/Editor/UnrealEd/Private/EditorDomain/BulkDataRegistryEditorDomain.h b/Engine/Source/Editor/UnrealEd/Private/EditorDomain/BulkDataRegistryEditorDomain.h
new file mode 100644
--- /dev/null
+++ b/Engine/Source/Editor/UnrealEd/Private/EditorDomain/BulkDataRegistryEditorDomain.h
@@ -0,0 +1,111 @@
+// Copyright Epic Games, Inc. All Rights Reserved.
+
+#include "Containers/Array.h"
+#include "Containers/Map.h"
+#include "Containers/RingBuffer.h"
+#include "DerivedDataRequest.h"
+#include "HAL/CriticalSection.h"
+#include "Serialization/BulkDataRegistry.h"
+#include "TickableEditorObject.h"
+#include "Virtualization/VirtualizedBulkData.h"
+
+class FCompressedBuffer;
+class UPackage;
+struct FIoHash;
+
+namespace UE::BulkDataRegistry::Private
+{
+
+class FBulkDataRegistryEditorDomain;
+
+/** Struct for storage of a BulkData in the registry, including the BulkData itself and data about cache status. */
+struct FBulkSource
+{
+	FBulkSource() = default;
+	FBulkSource(UE::Virtualization::FVirtualizedUntypedBulkData&& InBulkData, FName InPackageNameToUpdate = NAME_None)
+		:BulkData(MoveTemp(InBulkData)), PackageNameToUpdate(InPackageNameToUpdate)
+	{
+	}
+
+	UE::Virtualization::FVirtualizedUntypedBulkData BulkData;
+	FName PackageNameToUpdate;
+};
+
+/** Serialize an array of BulkDatas into or out of bytes saved/load from the registry's persistent cache. */
+void Serialize(FArchive& Ar, TArray<UE::Virtualization::FVirtualizedUntypedBulkData>& InDatas);
+
+/** A collection of bulkdatas that should be sent to the cache for the given package. */
+class FPendingPackage
+{
+public:
+	FPendingPackage(FName PackageName, FBulkDataRegistryEditorDomain* InOwner);
+	FPendingPackage(FPendingPackage&& Other) = default;
+	FPendingPackage(const FPendingPackage& Other) = delete;
+	~FPendingPackage();
+
+	bool IsReadFinished() const;
+	void UpdateCache();
+	double GetCreateTime() const
+	{
+		return CreateTime;
+	}
+	void AddBulkData(const UE::Virtualization::FVirtualizedUntypedBulkData& BulkData)
+	{
+		BulkDatas.Add(BulkData);
+	}
+
+private:
+	void OnCacheResults(FSharedBuffer Buffer);
+
+	FName PackageName;
+	TArray<UE::Virtualization::FVirtualizedUntypedBulkData> BulkDatas;
+	TArray<UE::Virtualization::FVirtualizedUntypedBulkData> CachedBulkDatas;
+	UE::DerivedData::FRequest CacheRequest;
+	FBulkDataRegistryEditorDomain* Owner;
+	double CreateTime;
+};
+
+/** Data about a BulkData that has loaded its payload for TryGetMeta and should drop it after GetData or a timeout. */
+struct FTempLoadedPayload
+{
+	FGuid Guid;
+	uint64 PayloadSize;
+	double EndTime;
+};
+
+/** Implementation of a BulkDataRegistry that stores its persistent data in a DDC bucket. */
+class FBulkDataRegistryEditorDomain : public IBulkDataRegistry, public FTickableEditorObject, public FTickableCookObject
+{
+public:
+	FBulkDataRegistryEditorDomain();
+	virtual ~FBulkDataRegistryEditorDomain();
+
+	// IBulkDataRegistry interface
+	virtual void Register(UPackage* Owner, const UE::Virtualization::FVirtualizedUntypedBulkData& BulkData) override;
+	virtual void OnExitMemory(const UE::Virtualization::FVirtualizedUntypedBulkData& BulkData) override;
+	virtual TFuture<UE::BulkDataRegistry::FMetaData> GetMeta(const FGuid& BulkDataId) override;
+	virtual TFuture<UE::BulkDataRegistry::FData> GetData(const FGuid& BulkDataId) override;
+
+	// FTickableEditorObject/FTickableCookObject interface
+	virtual void Tick(float DeltaTime) override;
+	virtual void TickCook(float DeltaTime, bool bTickComplete) override;
+	virtual ETickableTickType GetTickableTickType() const override { return ETickableTickType::Always; }
+	virtual TStatId GetStatId() const override { return TStatId(); }
+
+private:
+	void AddPendingPackageBulkData(FName PackageName, const UE::Virtualization::FVirtualizedUntypedBulkData& BulkData);
+	void PollPendingPackages(bool bWaitForCooldown);
+	void AddTempLoadedPayload(const FGuid& RegistryKey, uint64 PayloadSize);
+	void PruneTempLoadedPayloads();
+	friend class FPendingPackage;
+
+	FCriticalSection PendingPackageLock;
+	FRWLock RegistryLock;
+	TMap<FGuid, FBulkSource> Registry;
+	TMap<FName, TUniquePtr<FPendingPackage>> PendingPackages;
+	TRingBuffer<FTempLoadedPayload> TempLoadedPayloads;
+	uint64 TempLoadedPayloadsSize = 0;
+	bool bAllowCacheUpdates = true;
+};
+
+} // namespace UE::BulkDataRegistry::Private
diff --git a/Engine/Source/Editor/UnrealEd/Private/EditorDomain/EditorBuildInputResolver.cpp b/Engine/Source/Editor/UnrealEd/Private/EditorDomain/EditorBuildInputResolver.cpp
new file mode 100644
--- /dev/null
+++ b/Engine/Source/Editor/UnrealEd/Private/EditorDomain/EditorBuildInputResolver.cpp
@@ -0,0 +1,139 @@
+// Copyright Epic Games, Inc. All Rights Reserved.
+
+#include "EditorBuildInputResolver.h"
+
+#include "Serialization/BulkDataRegistry.h"
+
+namespace UE::DerivedData
+{
+
+FEditorBuildInputResolver& FEditorBuildInputResolver::Get()
+{
+	static FEditorBuildInputResolver Singleton;
+	return Singleton;
+}
+
+FRequest FEditorBuildInputResolver::ResolveKey(const FBuildKey& Key, FOnBuildKeyResolved&& OnResolved)
+{
+	// Not yet implemented
+	OnResolved({ Key, {}, EStatus::Error });
+	return FRequest();
+}
+
+FRequest FEditorBuildInputResolver::ResolveInputMeta(const FBuildDefinition& Definition, EPriority Priority,
+	FOnBuildInputMetaResolved&& OnResolved)
+{
+	EStatus Status = EStatus::Ok;
+	TArray<FString> InputKeys;
+	TArray<FBuildInputMetaByKey> Inputs;
+
+	Definition.IterateInputBuilds([&Status, &Definition](FStringView Key, const FBuildPayloadKey& PayloadKey)
+		{
+			UE_LOG(LogCore, Error, TEXT("FEditorBuildInputResolver::ResolveInputMeta: resolving InputBuilds is not yet implemented. Context=%.*s, Key=%.*s"),
+				Definition.GetName().Len(), Definition.GetName().GetData(), Key.Len(), Key.GetData());
+			Status = EStatus::Error;
+		});
+
+	/** Visits every input bulk data in order by key. */
+	Definition.IterateInputBulkData([&Status, &InputKeys, &Inputs, &Definition](FStringView Key, const FGuid& BulkDataId)
+		{
+			FBuildInputMetaByKey& Input = Inputs.Emplace_GetRef();
+			// TODO: MakeAsync: Create an IRequest that has an array of TFutures and calls .Then to populate the Inputs
+			TFuture<UE::BulkDataRegistry::FMetaData> Future = IBulkDataRegistry::Get().GetMeta(BulkDataId);
+			const UE::BulkDataRegistry::FMetaData& Result = Future.Get();
+			if (Result.bValid)
+			{
+				Input.Key = InputKeys.Emplace_GetRef(Key);
+				Input.RawHash = Result.RawHash;
+				Input.RawSize = Result.RawSize;
+			}
+			else
+			{
+				Inputs.Pop();
+				UE_LOG(LogCore, Error, TEXT("Failed to ResolveInputMetaData. Context=%.*s, Key=%.*s"),
+					Definition.GetName().Len(), Definition.GetName().GetData(), Key.Len(), Key.GetData())
+				Status = EStatus::Error;
+			}
+		});
+
+	Definition.IterateInputFiles([&Status, &Definition](FStringView Key, FStringView Path)
+		{
+			UE_LOG(LogCore, Error, TEXT("FEditorBuildInputResolver::ResolveInputMeta: resolving InputFiles is not yet implemented. Context=%.*s, Key=%.*s"),
+				Definition.GetName().Len(), Definition.GetName().GetData(), Key.Len(), Key.GetData());
+			Status = EStatus::Error;
+		});
+
+	Definition.IterateInputHashes([&Status, &Definition](FStringView Key, const FIoHash& RawHash)
+		{
+			UE_LOG(LogCore, Error, TEXT("FEditorBuildInputResolver::ResolveInputMeta: resolving InputHashes is not yet implemented. Context=%.*s, Key=%.*s"),
+				Definition.GetName().Len(), Definition.GetName().GetData(), Key.Len(), Key.GetData());
+			Status = EStatus::Error;
+		});
+
+	OnResolved({ Inputs, Status });
+	return FRequest();
+}
+
+FRequest FEditorBuildInputResolver::ResolveInputData(const FBuildDefinition& Definition, EPriority Priority,
+	FOnBuildInputDataResolved&& OnResolved, FBuildInputFilter&& Filter)
+{
+	EStatus Status = EStatus::Ok;
+	TArray<FString> InputKeys;
+	TArray<FBuildInputDataByKey> Inputs;
+
+	Definition.IterateInputBuilds([&Status, &Definition](FStringView Key, const FBuildPayloadKey& PayloadKey)
+		{
+			UE_LOG(LogCore, Error, TEXT("FEditorBuildInputResolver::ResolveInputData resolving InputBuilds is not yet implemented. Context=%.*s, Key=%.*s"),
+				Definition.GetName().Len(), Definition.GetName().GetData(), Key.Len(), Key.GetData());
+			Status = EStatus::Error;
+		});
+
+	/** Visits every input bulk data in order by key. */
+	Definition.IterateInputBulkData([&Filter, &Status, &InputKeys, &Inputs, &Definition](FStringView Key, const FGuid& BulkDataId)
+		{
+			if (!Filter || Filter(Key))
+			{
+				// TODO: MakeAsync: Create an IRequest that has an array of TFutures and calls .Then to populate the Inputs
+				TFuture<UE::BulkDataRegistry::FData> Future = IBulkDataRegistry::Get().GetData(BulkDataId);
+				const UE::BulkDataRegistry::FData& Result = Future.Get();
+				if (Result.bValid)
+				{
+					InputKeys.Emplace(Key);
+					Inputs.Add({ InputKeys.Last(), Result.Buffer });
+				}
+				else
+				{
+					UE_LOG(LogCore, Error, TEXT("Failed to ResolveInputData. Context=%.*s, Key=%.*s"),
+						Definition.GetName().Len(), Definition.GetName().GetData(), Key.Len(), Key.GetData())
+					Status = EStatus::Error;
+				}
+			}
+		});
+
+	Definition.IterateInputFiles([&Status, &Definition](FStringView Key, FStringView Path)
+		{
+			UE_LOG(LogCore, Error, TEXT("FEditorBuildInputResolver::ResolveInputData: resolving InputFiles is not yet implemented. Context=%.*s, Key=%.*s"),
+				Definition.GetName().Len(), Definition.GetName().GetData(), Key.Len(), Key.GetData());
+			Status = EStatus::Error;
+		});
+
+	Definition.IterateInputHashes([&Status, &Definition](FStringView Key, const FIoHash& RawHash)
+		{
+			UE_LOG(LogCore, Error, TEXT("FEditorBuildInputResolver::ResolveInputData: resolving InputHashes is not yet implemented. Context=%.*s, Key=%.*s"),
+				Definition.GetName().Len(), Definition.GetName().GetData(), Key.Len(), Key.GetData());
+			Status = EStatus::Error;
+		});
+
+	OnResolved({ Inputs, Status });
+	return FRequest();
+}
+
+FRequest FEditorBuildInputResolver::ResolveInputData(const FBuildAction& Action, EPriority Priority,
+	FOnBuildInputDataResolved&& OnResolved, FBuildInputFilter&& Filter)
+{
+	// Not yet implemented
+	OnResolved({ {}, EStatus::Error });
+	return FRequest();
+}
+
+} // namespace UE::DerivedData
diff --git a/Engine/Source/Editor/UnrealEd/Private/EditorDomain/EditorBuildInputResolver.h b/Engine/Source/Editor/UnrealEd/Private/EditorDomain/EditorBuildInputResolver.h
new file mode 100644
--- /dev/null
+++ b/Engine/Source/Editor/UnrealEd/Private/EditorDomain/EditorBuildInputResolver.h
@@ -0,0 +1,29 @@
+// Copyright Epic Games, Inc. All Rights Reserved.
+
+#pragma once
+
+#include "DerivedDataBuildInputResolver.h"
+
+#if WITH_EDITOR
+
+namespace UE::DerivedData
+{
+
+/** A BuildInputResolver that looks up the process-global inputs registered from package loads or a cache of the package loads. */
+class FEditorBuildInputResolver : public UE::DerivedData::IBuildInputResolver
+{
+public:
+	UNREALED_API static FEditorBuildInputResolver& Get();
+
+	virtual FRequest ResolveKey(const FBuildKey& Key, FOnBuildKeyResolved&& OnResolved) override;
+	virtual FRequest ResolveInputMeta(const FBuildDefinition& Definition, EPriority Priority,
+		FOnBuildInputMetaResolved&& OnResolved) override;
+	virtual FRequest ResolveInputData(const FBuildDefinition& Definition, EPriority Priority,
+		FOnBuildInputDataResolved&& OnResolved, FBuildInputFilter&& Filter = FBuildInputFilter()) override;
+	virtual FRequest ResolveInputData(const FBuildAction& Action, EPriority Priority,
+		FOnBuildInputDataResolved&& OnResolved, FBuildInputFilter&& Filter = FBuildInputFilter()) override;
+};
+
+} // namespace UE::DerivedData
+
+#endif // WITH_EDITOR
\ No newline at end of file
diff --git a/Engine/Source/Editor/UnrealEd/Private/EditorDomain/EditorDomain.cpp b/Engine/Source/Editor/UnrealEd/Private/EditorDomain/EditorDomain.cpp
--- a/Engine/Source/Editor/UnrealEd/Private/EditorDomain/EditorDomain.cpp
+++ b/Engine/Source/Editor/UnrealEd/Private/EditorDomain/EditorDomain.cpp
@@ -13,6 +13,7 @@
 #include "EditorDomain/EditorDomainSave.h"
 #include "EditorDomain/EditorDomainUtils.h"
 #include "HAL/CriticalSection.h"
+#include "Interfaces/IPluginManager.h"
 #include "Misc/CommandLine.h"
 #include "Misc/ConfigCacheIni.h"
 #include "Misc/CoreDelegates.h"
@@ -76,7 +77,15 @@ FEditorDomain::FEditorDomain()
 
 	bEditorDomainReadEnabled = !FParse::Param(FCommandLine::Get(), TEXT("noeditordomainread"));
 
-	FCoreDelegates::OnPostEngineInit.AddRaw(this, &FEditorDomain::OnPostEngineInit);
+	ELoadingPhase::Type CurrentPhase = IPluginManager::Get().GetLastCompletedLoadingPhase();
+	if (CurrentPhase == ELoadingPhase::None || CurrentPhase < ELoadingPhase::PostEngineInit)
+	{
+		FCoreDelegates::OnPostEngineInit.AddRaw(this, &FEditorDomain::OnPostEngineInit);
+	}
+	else
+	{
+		OnPostEngineInit();
+	}
 }
 
 FEditorDomain::~FEditorDomain()
@@ -143,7 +152,7 @@ bool FEditorDomain::TryFindOrAddPackageSource(const FPackagePath& PackagePath, T
 
 	FString ErrorMessage;
 	FPackageDigest PackageDigest;
-	EPackageDigestResult Result = GetPackageDigest(*AssetRegistry, PackageName, PackageDigest, ErrorMessage, ClassDigests);
+	EPackageDigestResult Result = GetPackageDigest(*AssetRegistry, PackageName, PackageDigest, ErrorMessage);
 	switch (Result)
 	{
 	case EPackageDigestResult::Success:
@@ -178,7 +187,7 @@ void FEditorDomain::PrecachePackageDigest(FName PackageName)
 	TOptional<FAssetPackageData> PackageData = AssetRegistry->GetAssetPackageDataCopy(PackageName);
 	if (PackageData)
 	{
-		UE::EditorDomain::PrecacheClassDigests(PackageData->ImportedClasses, ClassDigests);
+		UE::EditorDomain::PrecacheClassDigests(PackageData->ImportedClasses);
 	}
 }
 
@@ -215,6 +224,7 @@ int64 FEditorDomain::FileSize(const FPackagePath& PackagePath, EPackageSegment P
 	FPackagePath* OutUpdatedPath)
 {
 	using namespace UE::EditorDomain;
+	using namespace UE::DerivedData;
 
 	if (PackageSegment != EPackageSegment::Header)
 	{
@@ -263,8 +273,10 @@ int64 FEditorDomain::FileSize(const FPackagePath& PackagePath, EPackageSegment P
 				}
 			}
 		};
-		Request = RequestEditorDomainPackage(PackagePath, PackageSource->Digest,
-			UE::DerivedData::EPriority::Highest, MoveTemp(MetaDataGetComplete));
+		// Fetch meta-data only
+		ECachePolicy SkipFlags = ECachePolicy::SkipData & ~ECachePolicy::SkipMeta;
+		Request = RequestEditorDomainPackage(PackagePath, PackageSource->Digest, SkipFlags,
+			EPriority::Highest, MoveTemp(MetaDataGetComplete));
 	}
 	Request.Wait();
 	return FileSize;
@@ -294,7 +306,7 @@ FOpenPackageResult FEditorDomain::OpenReadPackage(const FPackagePath& PackagePat
 	ScopeLock.Unlock();
 
 	UE::DerivedData::FRequest Request = RequestEditorDomainPackage(PackagePath, PackageSourceDigest,
-		UE::DerivedData::EPriority::Normal,
+		UE::DerivedData::ECachePolicy::None, UE::DerivedData::EPriority::Normal,
 		[Result](UE::DerivedData::FCacheGetCompleteParams&& Params)
 		{
 			// Note that ~FEditorDomainReadArchive waits for this callback to be called, so Result cannot dangle
@@ -332,7 +344,7 @@ IAsyncReadFileHandle* FEditorDomain::OpenAsyncReadPackage(const FPackagePath& Pa
 
 	FEditorDomainAsyncReadFileHandle* Result = new FEditorDomainAsyncReadFileHandle(Locks, PackagePath, PackageSource);
 	UE::DerivedData::FRequest Request = RequestEditorDomainPackage(PackagePath, PackageSource->Digest,
-		UE::DerivedData::EPriority::Normal,
+		UE::DerivedData::ECachePolicy::None, UE::DerivedData::EPriority::Normal,
 		[Result](UE::DerivedData::FCacheGetCompleteParams&& Params)
 		{
 			// Note that ~FEditorDomainAsyncReadFileHandle waits for this callback to be called, so Result cannot dangle
@@ -430,7 +442,7 @@ void FEditorDomain::OnEndLoad(TConstArrayView<UPackage*> LoadedPackages)
 
 	for (UPackage* Package : PackagesToSave)
 	{
-		UE::EditorDomain::TrySavePackage(Package, ClassDigests);
+		UE::EditorDomain::TrySavePackage(Package);
 	}
 }
 
@@ -460,7 +472,7 @@ void FEditorDomain::OnPostEngineInit()
 
 	for (UPackage* Package : PackagesToSave)
 	{
-		UE::EditorDomain::TrySavePackage(Package, ClassDigests);
+		UE::EditorDomain::TrySavePackage(Package);
 	}
 }
 
diff --git a/Engine/Source/Editor/UnrealEd/Private/EditorDomain/EditorDomain.h b/Engine/Source/Editor/UnrealEd/Private/EditorDomain/EditorDomain.h
--- a/Engine/Source/Editor/UnrealEd/Private/EditorDomain/EditorDomain.h
+++ b/Engine/Source/Editor/UnrealEd/Private/EditorDomain/EditorDomain.h
@@ -21,19 +21,6 @@ class UPackage;
 namespace UE::EditorDomain
 {
 	typedef FIoHash FPackageDigest;
-
-	/** A UClass's data that is used in the EditorDomain Digest */
-	struct FClassDigestData
-	{
-		FBlake3Hash SchemaHash;
-		bool bNative;
-	};
-	/** Threadsafe cache of ClassName -> Digest data for calculating EditorDomain Digests */
-	struct FClassDigestMap
-	{
-		TMap<FName, FClassDigestData> Map;
-		FCriticalSection Lock;
-	};
 }
 
 DECLARE_LOG_CATEGORY_EXTERN(LogEditorDomain, Log, All);
@@ -157,8 +144,6 @@ private:
 	TRefCountPtr<FLocks> Locks;
 	/** Digests previously found for a package. Used for optimization, but also to record loaded-from-domain. */
 	TMap<FName, TRefCountPtr<FPackageSource>> PackageSources;
-	/** Cache of GetSchemaHash by class name */
-	UE::EditorDomain::FClassDigestMap ClassDigests;
 	/** True by default, set to false when reading is disabled for testing. */
 	bool bEditorDomainReadEnabled = true;
 	/** If true, use an out-of-process EditorDomainSaveServer for saves, else save in process in EndLoad */
diff --git a/Engine/Source/Editor/UnrealEd/Private/EditorDomain/EditorDomainSave.cpp b/Engine/Source/Editor/UnrealEd/Private/EditorDomain/EditorDomainSave.cpp
--- a/Engine/Source/Editor/UnrealEd/Private/EditorDomain/EditorDomainSave.cpp
+++ b/Engine/Source/Editor/UnrealEd/Private/EditorDomain/EditorDomainSave.cpp
@@ -713,7 +713,7 @@ bool FEditorDomainSaveServer::TrySavePackage(const FPackagePath& PackagePath, FS
 		return false;
 	}
 
-	if (!UE::EditorDomain::TrySavePackage(Package, this->ClassDigests))
+	if (!UE::EditorDomain::TrySavePackage(Package))
 	{
 		OutErrorMessage = FString::Printf(TEXT("Could not save package %s."), *PackagePath.GetDebugName());
 		return false;
diff --git a/Engine/Source/Editor/UnrealEd/Private/EditorDomain/EditorDomainSave.h b/Engine/Source/Editor/UnrealEd/Private/EditorDomain/EditorDomainSave.h
--- a/Engine/Source/Editor/UnrealEd/Private/EditorDomain/EditorDomainSave.h
+++ b/Engine/Source/Editor/UnrealEd/Private/EditorDomain/EditorDomainSave.h
@@ -139,8 +139,6 @@ private:
 	IAssetRegistry* AssetRegistry;
 	/** In-process handle to the interprocess synchronization object. */
 	FPlatformProcess::FSemaphore* ProcessLock = nullptr;
-	/** Cache of GetSchemaHash by class name */
-	UE::EditorDomain::FClassDigestMap ClassDigests;
 	/** The time at which the current idleness started. Used to decide abdication. */
 	double IdleStartTime = 0;
 	/** The time at which garbage was last collected. Used for scheduling the next collection. */
diff --git a/Engine/Source/Editor/UnrealEd/Private/EditorDomain/EditorDomainUtils.cpp b/Engine/Source/Editor/UnrealEd/Private/EditorDomain/EditorDomainUtils.cpp
--- a/Engine/Source/Editor/UnrealEd/Private/EditorDomain/EditorDomainUtils.cpp
+++ b/Engine/Source/Editor/UnrealEd/Private/EditorDomain/EditorDomainUtils.cpp
@@ -12,6 +12,7 @@
 #include "DerivedDataCacheRecord.h"
 #include "Editor.h"
 #include "HAL/FileManager.h"
+#include "Memory/SharedBuffer.h"
 #include "Misc/FileHelper.h"
 #include "Misc/PackagePath.h"
 #include "Misc/StringBuilder.h"
@@ -26,13 +27,21 @@ namespace UE
 namespace EditorDomain
 {
 
+FClassDigestMap GClassDigests;
+
+FClassDigestMap& GetClassDigests()
+{
+	return GClassDigests;
+}
+
 // Change to a new guid when EditorDomain needs to be invalidated
 const TCHAR* EditorDomainVersion = TEXT("C217EB656E9B4C04816D3DC0E21901F6");
 // Identifier of the CacheBucket for EditorDomainPackages
 const TCHAR* EditorDomainPackageBucketName = TEXT("EditorDomainPackage");
+const TCHAR* EditorDomainBulkDataListBucketName = TEXT("EditorDomainBulkDataList");
 
 EPackageDigestResult GetPackageDigest(FPackageDigest& OutDigest, FString& OutErrorMessage,
-	const FAssetPackageData& PackageData, FName PackageName, FClassDigestMap& ClassDigests)
+	const FAssetPackageData& PackageData, FName PackageName)
 {
 	FCbWriter Writer;
 	int32 CurrentFileVersionUE = GPackageFileUEVersion;
@@ -61,6 +70,7 @@ EPackageDigestResult GetPackageDigest(FPackageDigest& OutDigest, FString& OutErr
 	}
 	FNameBuilder NameBuilder;
 	int32 NextClass = 0;
+	FClassDigestMap& ClassDigests = GetClassDigests();
 	for (int32 Attempt = 0; NextClass < PackageData.ImportedClasses.Num(); ++Attempt)
 	{
 		if (Attempt > 0)
@@ -74,7 +84,7 @@ EPackageDigestResult GetPackageDigest(FPackageDigest& OutDigest, FString& OutErr
 			}
 			TConstArrayView<FName> RemainingClasses(PackageData.ImportedClasses);
 			RemainingClasses = RemainingClasses.Slice(NextClass, PackageData.ImportedClasses.Num() - NextClass);
-			PrecacheClassDigests(RemainingClasses, ClassDigests);
+			PrecacheClassDigests(RemainingClasses);
 		}
 		FScopeLock ClassDigestsScopeLock(&ClassDigests.Lock);
 		for (; NextClass < PackageData.ImportedClasses.Num(); ++NextClass)
@@ -95,8 +105,9 @@ EPackageDigestResult GetPackageDigest(FPackageDigest& OutDigest, FString& OutErr
 	return EPackageDigestResult::Success;
 }
 
-void PrecacheClassDigests(TConstArrayView<FName> ClassNames, FClassDigestMap& ClassDigests)
+void PrecacheClassDigests(TConstArrayView<FName> ClassNames)
 {
+	FClassDigestMap& ClassDigests = GetClassDigests();
 	TArray<TPair<FName,FClassDigestData>> ClassesToAdd;
 	FString ClassNameStr;
 	{
@@ -156,7 +167,7 @@ void PrecacheClassDigests(TConstArrayView<FName> ClassNames, FClassDigestMap& Cl
 }
 
 EPackageDigestResult GetPackageDigest(IAssetRegistry& AssetRegistry, FName PackageName,
-	FPackageDigest& OutPackageDigest, FString& OutErrorMessage, FClassDigestMap& ClassDigests)
+	FPackageDigest& OutPackageDigest, FString& OutErrorMessage)
 {
 	AssetRegistry.WaitForPackage(PackageName.ToString());
 	TOptional<FAssetPackageData> PackageData = AssetRegistry.GetAssetPackageDataCopy(PackageName);
@@ -166,7 +177,7 @@ EPackageDigestResult GetPackageDigest(IAssetRegistry& AssetRegistry, FName Packa
 			*PackageName.ToString());
 		return EPackageDigestResult::FileDoesNotExist;
 	}
-	return GetPackageDigest(OutPackageDigest, OutErrorMessage, *PackageData, PackageName, ClassDigests);
+	return GetPackageDigest(OutPackageDigest, OutErrorMessage, *PackageData, PackageName);
 }
 
 UE::DerivedData::FCacheKey GetEditorDomainPackageKey(const FPackageDigest& PackageDigest)
@@ -176,22 +187,31 @@ UE::DerivedData::FCacheKey GetEditorDomainPackageKey(const FPackageDigest& Packa
 	return UE::DerivedData::FCacheKey{EditorDomainPackageCacheBucket, PackageDigest};
 }
 
+UE::DerivedData::FCacheKey GetBulkDataListKey(const FPackageDigest& PackageDigest)
+{
+	static UE::DerivedData::FCacheBucket EditorDomainBulkDataListBucket =
+		GetDerivedDataCacheRef().CreateBucket(EditorDomainBulkDataListBucketName);
+	return UE::DerivedData::FCacheKey{ EditorDomainBulkDataListBucket, PackageDigest };
+}
+
 UE::DerivedData::FRequest RequestEditorDomainPackage(const FPackagePath& PackagePath,
-	const FPackageDigest& PackageDigest, UE::DerivedData::EPriority CachePriority,
+	const FPackageDigest& PackageDigest, UE::DerivedData::ECachePolicy SkipFlags, UE::DerivedData::EPriority CachePriority,
 	UE::DerivedData::FOnCacheGetComplete&& Callback)
 {
 	UE::DerivedData::ICache& Cache = GetDerivedDataCacheRef();
+	checkf((SkipFlags & (~UE::DerivedData::ECachePolicy::SkipData)) == UE::DerivedData::ECachePolicy::None,
+		TEXT("SkipFlags should only contain ECachePolicy::Skip* flags"));
 	return Cache.Get({ GetEditorDomainPackageKey(PackageDigest) },
 		PackagePath.GetDebugName(),
-		UE::DerivedData::ECachePolicy::QueryLocal, CachePriority, MoveTemp(Callback));
+		SkipFlags | UE::DerivedData::ECachePolicy::QueryLocal, CachePriority, MoveTemp(Callback));
 }
 
-bool TrySavePackage(UPackage* Package, FClassDigestMap& ClassDigests)
+bool TrySavePackage(UPackage* Package)
 {
 	FString ErrorMessage;
 	FPackageDigest PackageDigest;
 	EPackageDigestResult FindHashResult = GetPackageDigest(*IAssetRegistry::Get(), Package->GetFName(), PackageDigest,
-		ErrorMessage, ClassDigests);
+		ErrorMessage);
 	switch (FindHashResult)
 	{
 	case EPackageDigestResult::Success:
@@ -264,5 +284,57 @@ bool TrySavePackage(UPackage* Package, FClassDigestMap& ClassDigests)
 	return true;
 }
 
+UE::DerivedData::FRequest GetBulkDataList(FName PackageName, TUniqueFunction<void(FSharedBuffer Buffer)>&& Callback)
+{
+	UE::DerivedData::ICache& Cache = GetDerivedDataCacheRef();
+
+	FString ErrorMessage;
+	FPackageDigest PackageDigest;
+	EPackageDigestResult FindHashResult = GetPackageDigest(*IAssetRegistry::Get(), PackageName, PackageDigest,
+		ErrorMessage);
+	switch (FindHashResult)
+	{
+	case EPackageDigestResult::Success:
+		break;
+	default:
+	{
+		Callback(FSharedBuffer());
+		return UE::DerivedData::FRequest();
+	}
+	}
+
+	return Cache.Get({ GetBulkDataListKey(PackageDigest) },
+		WriteToString<128>(PackageName), UE::DerivedData::ECachePolicy::Default, UE::DerivedData::EPriority::Low,
+		[InnerCallback = MoveTemp(Callback)](UE::DerivedData::FCacheGetCompleteParams&& Params)
+		{
+			bool bOk = Params.Status == UE::DerivedData::EStatus::Ok;
+			InnerCallback(bOk ? Params.Record.GetValue() : FSharedBuffer());
+		});
+}
+
+void PutBulkDataList(FName PackageName, FSharedBuffer Buffer)
+{
+	UE::DerivedData::ICache& Cache = GetDerivedDataCacheRef();
+
+	FString ErrorMessage;
+	FPackageDigest PackageDigest;
+	EPackageDigestResult FindHashResult = GetPackageDigest(*IAssetRegistry::Get(), PackageName, PackageDigest,
+		ErrorMessage);
+	switch (FindHashResult)
+	{
+	case EPackageDigestResult::Success:
+		break;
+	default:
+	{
+		return;
+	}
+	}
+
+	UE::DerivedData::FCacheRecordBuilder RecordBuilder = Cache.CreateRecord(GetBulkDataListKey(PackageDigest));
+	RecordBuilder.SetValue(Buffer);
+	UE::DerivedData::FCacheRecord Record = RecordBuilder.Build();
+	Cache.Put(MakeArrayView(&Record, 1), WriteToString<128>(PackageName));
+}
+
 }
 }
diff --git a/Engine/Source/Editor/UnrealEd/Private/EditorDomain/EditorDomainUtils.h b/Engine/Source/Editor/UnrealEd/Private/EditorDomain/EditorDomainUtils.h
--- a/Engine/Source/Editor/UnrealEd/Private/EditorDomain/EditorDomainUtils.h
+++ b/Engine/Source/Editor/UnrealEd/Private/EditorDomain/EditorDomainUtils.h
@@ -7,12 +7,14 @@
 #include "DerivedDataCache.h"
 #include "EditorDomain/EditorDomain.h"
 #include "Misc/ConfigCacheIni.h"
+#include "Templates/Function.h"
 #include "UObject/NameTypes.h"
 
 DECLARE_LOG_CATEGORY_EXTERN(LogEditorDomainSave, Log, All);
 
 class FAssetPackageData;
 class FPackagePath;
+class FSharedBuffer;
 class IAssetRegistry;
 class UPackage;
 
@@ -27,24 +29,50 @@ enum class EPackageDigestResult
 	MissingClass,
 };
 
+/** A UClass's data that is used in the EditorDomain Digest */
+struct FClassDigestData
+{
+	FBlake3Hash SchemaHash;
+	bool bNative;
+};
+
+/** Threadsafe cache of ClassName -> Digest data for calculating EditorDomain Digests */
+struct FClassDigestMap
+{
+	TMap<FName, FClassDigestData> Map;
+	FCriticalSection Lock;
+};
+
 /**
  * Calculate the PackageDigest for the given packagePath.
  * Reads information from the AssetRegistry to compute the digest.
  */
 EPackageDigestResult GetPackageDigest(IAssetRegistry& AssetRegistry, FName PackageName,
-	FPackageDigest& OutPackageDigest, FString& OutErrorMessage, FClassDigestMap& ClassDigests);
+	FPackageDigest& OutPackageDigest, FString& OutErrorMessage);
 /** For any ClassNames not already in ClassDigests, look up their UStruct and add them. */
-void PrecacheClassDigests(TConstArrayView<FName> ClassNames, FClassDigestMap& ClassDigests);
+void PrecacheClassDigests(TConstArrayView<FName> ClassNames);
 
-/** Convert the given PackageDigest into CacheKey format. */
+/** Get the cachekey for the EditorDomainPackage for the given PackageDigest. */
 UE::DerivedData::FCacheKey GetEditorDomainPackageKey(const FPackageDigest& PackageDigest);
 
+/** Get the cachekey for the EditorDomainBulkDataList for the given PackageDigest. */
+UE::DerivedData::FCacheKey GetBulkDataListKey(const FPackageDigest& PackageDigest);
+
 /** Get the CacheRequest for the given package from the EditorDomain cache bucket. */
 UE::DerivedData::FRequest RequestEditorDomainPackage(const FPackagePath& PackagePath,
-	const FPackageDigest& PackageDigest, UE::DerivedData::EPriority CachePriority,
-	UE::DerivedData::FOnCacheGetComplete&& Callback);
+	const FPackageDigest& PackageDigest, UE::DerivedData::ECachePolicy SkipFlags,
+	UE::DerivedData::EPriority CachePriority, UE::DerivedData::FOnCacheGetComplete&& Callback);
 
 /** Save the given package into the EditorDomain. */
-bool TrySavePackage(UPackage* Package, FClassDigestMap& ClassDigests);
+bool TrySavePackage(UPackage* Package);
+
+/** Get the CacheRequest for the BulkDataList of the given package. */
+UE::DerivedData::FRequest GetBulkDataList(FName PackageName, TUniqueFunction<void(FSharedBuffer Buffer)>&& Callback);
+
+/** Write the data for the BulkDataList of the given package to the cache. */
+void PutBulkDataList(FName PackageName, FSharedBuffer Buffer);
+
+/** Accessor for the global ClassDigest map shared by systems needing to calculate PackageDigests. */
+FClassDigestMap& GetClassDigests();
 
 }
diff --git a/Engine/Source/Editor/UnrealEd/Private/TickableEditorObject.cpp b/Engine/Source/Editor/UnrealEd/Private/TickableEditorObject.cpp
--- a/Engine/Source/Editor/UnrealEd/Private/TickableEditorObject.cpp
+++ b/Engine/Source/Editor/UnrealEd/Private/TickableEditorObject.cpp
@@ -14,3 +14,19 @@ TArray<FTickableEditorObject*>& FTickableEditorObject::GetPendingTickableObjects
 	static TArray<FTickableEditorObject*> PendingTickableObjects;
 	return PendingTickableObjects;
 }
+
+TArray<FTickableObjectBase::FTickableObjectEntry>& FTickableCookObject::GetTickableObjects()
+{
+	static TTickableObjectsCollection TickableObjects;
+	return TickableObjects;
+}
+
+TArray<FTickableCookObject*>& FTickableCookObject::GetPendingTickableObjects()
+{
+	static TArray<FTickableCookObject*> PendingTickableObjects;
+	return PendingTickableObjects;
+}
+
+bool FTickableCookObject::bCollectionIntact = true;
+bool FTickableCookObject::bIsTickingObjects = false;
+FTickableCookObject* FTickableCookObject::ObjectBeingTicked = nullptr;
diff --git a/Engine/Source/Editor/UnrealEd/Public/TickableEditorObject.h b/Engine/Source/Editor/UnrealEd/Public/TickableEditorObject.h
--- a/Engine/Source/Editor/UnrealEd/Public/TickableEditorObject.h
+++ b/Engine/Source/Editor/UnrealEd/Public/TickableEditorObject.h
@@ -119,3 +119,121 @@ private:
 
 	static TArray<FTickableEditorObject*>& GetPendingTickableObjects();
 };
+
+/**
+ * The same as FTickableEditorObject, but for systems that need to be ticked periodically during
+ * cooking.
+ * If a system needs to be cooked both during cook commandlet and in editor without the cook commandlet,
+ * it should dual-inherit from both FTickableCookObject and FTickableEditorObject.
+ * TODO: Reduce duplication between  FTickableCookObject and FTickableEditorObject.
+ */
+class UNREALED_API FTickableCookObject : public FTickableObjectBase
+{
+public:
+
+	static void TickObjects(const float DeltaSeconds, bool bCookComplete)
+	{
+		TArray<FTickableCookObject*>& PendingTickableObjects = GetPendingTickableObjects();
+		TArray<FTickableObjectEntry>& TickableObjects = GetTickableObjects();
+
+		for (FTickableCookObject* PendingTickable : PendingTickableObjects)
+		{
+			AddTickableObject(TickableObjects, PendingTickable);
+		}
+		PendingTickableObjects.Empty();
+
+		if (TickableObjects.Num() > 0)
+		{
+			check(!bIsTickingObjects);
+			bIsTickingObjects = true;
+
+			bool bNeedsCleanup = false;
+
+			for (const FTickableObjectEntry& TickableEntry : TickableObjects)
+			{
+				if (FTickableObjectBase* TickableObject = TickableEntry.TickableObject)
+				{
+					if ((TickableEntry.TickType == ETickableTickType::Always) || TickableObject->IsTickable())
+					{
+						FTickableCookObject* CookTickableObject = static_cast<FTickableCookObject*>(TickableObject);
+						ObjectBeingTicked = CookTickableObject;
+						CookTickableObject->TickCook(DeltaSeconds, bCookComplete);
+						ObjectBeingTicked = nullptr;
+					}
+
+					// In case it was removed during tick
+					if (TickableEntry.TickableObject == nullptr)
+					{
+						bNeedsCleanup = true;
+					}
+				}
+				else
+				{
+					bNeedsCleanup = true;
+				}
+			}
+
+			if (bNeedsCleanup)
+			{
+				TickableObjects.RemoveAll([](const FTickableObjectEntry& Entry) { return (Entry.TickableObject == nullptr); });
+			}
+
+			bIsTickingObjects = false;
+		}
+	}
+
+	/** Registers this instance with the static array of tickable objects. */
+	FTickableCookObject()
+	{
+		ensure(IsInGameThread() || IsInAsyncLoadingThread());
+		check(!GetPendingTickableObjects().Contains(this));
+		check(!GetTickableObjects().Contains(this));
+		GetPendingTickableObjects().Add(this);
+	}
+
+	/** Removes this instance from the static array of tickable objects. */
+	virtual ~FTickableCookObject()
+	{
+		ensureMsgf(ObjectBeingTicked != this, TEXT("Detected possible memory stomp. We are in the Tickable objects Tick function but hit its deconstructor, the 'this' pointer for the Object will now be invalid"));
+
+		ensure(IsInGameThread() || IsInAsyncLoadingThread());
+		if (bCollectionIntact && GetPendingTickableObjects().Remove(this) == 0)
+		{
+			RemoveTickableObject(GetTickableObjects(), this, bIsTickingObjects);
+		}
+	}
+
+	virtual void TickCook(float DeltaTime, bool bCookCompete) = 0;
+
+private:
+
+	/**
+	 * Class that avoids crashes when unregistering a tickable editor object too late.
+	 *
+	 * Some tickable objects can outlive the collection
+	 * (global/static destructor order is unpredictable).
+	 */
+	class TTickableObjectsCollection : public TArray<FTickableObjectEntry>
+	{
+	public:
+		~TTickableObjectsCollection()
+		{
+			FTickableCookObject::bCollectionIntact = false;
+		}
+	};
+
+	friend class TTickableObjectsCollection;
+
+	/** True if collection of tickable objects is still intact. */
+	static bool bCollectionIntact;
+	/** True if currently ticking of tickable editor objects. */
+	static bool bIsTickingObjects;
+
+	/** Set if we are in the Tick function for an editor tickable object */
+	static FTickableCookObject* ObjectBeingTicked;
+
+
+	static TArray<FTickableObjectEntry>& GetTickableObjects();
+
+	static TArray<FTickableCookObject*>& GetPendingTickableObjects();
+};
diff --git a/Engine/Source/Runtime/Core/Public/Virtualization/PayloadId.h b/Engine/Source/Runtime/Core/Public/Virtualization/PayloadId.h
--- a/Engine/Source/Runtime/Core/Public/Virtualization/PayloadId.h
+++ b/Engine/Source/Runtime/Core/Public/Virtualization/PayloadId.h
@@ -49,6 +49,10 @@ public:
 	/** Returns a string representation of the id, useful for debug tools */
 	FString ToString() const;
 
+	const FIoHash& GetIdentifier() const
+	{
+		return Identifier;
+	}
 	bool operator == (const FPayloadId& InPayloadId) const;
 	bool operator != (const FPayloadId& InPayloadId) const { return !(*this == InPayloadId); }
 
diff --git a/Engine/Source/Runtime/CoreUObject/Private/Serialization/BulkDataRegistry.cpp b/Engine/Source/Runtime/CoreUObject/Private/Serialization/BulkDataRegistry.cpp
new file mode 100644
--- /dev/null
+++ b/Engine/Source/Runtime/CoreUObject/Private/Serialization/BulkDataRegistry.cpp
@@ -0,0 +1,100 @@
+// Copyright Epic Games, Inc. All Rights Reserved.
+
+#include "Serialization/BulkDataRegistry.h"
+
+#if WITH_EDITOR
+
+#include "Misc/ConfigCacheIni.h"
+
+namespace UE::BulkDataRegistry::Private
+{
+
+IBulkDataRegistry* GBulkDataRegistry = nullptr;
+FSetBulkDataRegistry GSetBulkDataRegistry;
+
+/** A stub class to provide a return value from IBulkDataRegistry::Get() when the registery is disabled. */
+class FBulkDataRegistryNull : public IBulkDataRegistry
+{
+public:
+	FBulkDataRegistryNull() = default;
+	virtual ~FBulkDataRegistryNull() {}
+
+	virtual void Register(UPackage* Owner, const UE::Virtualization::FVirtualizedUntypedBulkData& BulkData) override {}
+	virtual void OnExitMemory(const UE::Virtualization::FVirtualizedUntypedBulkData& BulkData) override {}
+	virtual TFuture<UE::BulkDataRegistry::FMetaData> GetMeta(const FGuid& BulkDataId) override
+	{
+		TPromise<UE::BulkDataRegistry::FMetaData> Promise;
+		Promise.SetValue(UE::BulkDataRegistry::FMetaData{ false, FIoHash(), 0 });
+		return Promise.GetFuture();
+	}
+	virtual TFuture<UE::BulkDataRegistry::FData> GetData(const FGuid& BulkDataId) override
+	{
+		TPromise<UE::BulkDataRegistry::FData> Promise;
+		Promise.SetValue(UE::BulkDataRegistry::FData{ false, FCompressedBuffer() });
+		return Promise.GetFuture();
+	}
+};
+
+}
+
+bool IBulkDataRegistry::IsEnabled()
+{
+	bool bEnabled = true;
+	GConfig->GetBool(TEXT("CookSettings"), TEXT("BulkDataRegistryEnabled"), bEnabled, GEditorIni);
+	return bEnabled;
+}
+
+IBulkDataRegistry& IBulkDataRegistry::Get()
+{
+	check(UE::BulkDataRegistry::Private::GBulkDataRegistry);
+	return *UE::BulkDataRegistry::Private::GBulkDataRegistry;
+}
+
+void IBulkDataRegistry::Initialize()
+{
+	using namespace UE::BulkDataRegistry::Private;
+
+	if (IsEnabled())
+	{
+		FSetBulkDataRegistry& SetRegistryDelegate = GetSetBulkDataRegistryDelegate();
+		if (SetRegistryDelegate.IsBound())
+		{
+			// Allow the editor or licensee project to define the BulkDataRegistry
+			GBulkDataRegistry = SetRegistryDelegate.Execute();
+		}
+	}
+
+	// Assign the null BulkDataRegistry if it was disabled or not set by a higher level source
+	if (!GBulkDataRegistry)
+	{
+		GBulkDataRegistry = new FBulkDataRegistryNull();
+	}
+}
+
+void IBulkDataRegistry::Shutdown()
+{
+	delete UE::BulkDataRegistry::Private::GBulkDataRegistry;
+	UE::BulkDataRegistry::Private::GBulkDataRegistry = nullptr;
+}
+
+FSetBulkDataRegistry& IBulkDataRegistry::GetSetBulkDataRegistryDelegate()
+{
+	return UE::BulkDataRegistry::Private::GSetBulkDataRegistry;
+}
+
+namespace UE::GlobalBuildInputResolver::Private
+{
+UE::DerivedData::IBuildInputResolver* GGlobalBuildInputResolver = nullptr;
+}
+
+UE::DerivedData::IBuildInputResolver* GetGlobalBuildInputResolver()
+{
+	return UE::GlobalBuildInputResolver::Private::GGlobalBuildInputResolver;
+}
+
+void SetGlobalBuildInputResolver(UE::DerivedData::IBuildInputResolver* InResolver)
+{
+	UE::GlobalBuildInputResolver::Private::GGlobalBuildInputResolver = InResolver;
+}
+
+#endif // WITH_EDITOR
\ No newline at end of file
diff --git a/Engine/Source/Runtime/CoreUObject/Private/Tests/Virtualization/VirtualizedBulkDataTests.cpp b/Engine/Source/Runtime/CoreUObject/Private/Tests/Virtualization/VirtualizedBulkDataTests.cpp
--- a/Engine/Source/Runtime/CoreUObject/Private/Tests/Virtualization/VirtualizedBulkDataTests.cpp
+++ b/Engine/Source/Runtime/CoreUObject/Private/Tests/Virtualization/VirtualizedBulkDataTests.cpp
@@ -332,13 +332,13 @@ bool FVirtualizationWrapperTestIdentifiers::RunTest(const FString& Parameters)
 		const FGuid OriginalGuid = BulkData.GetIdentifier();
 
 		BulkData.Reset();
-		TestFalse(TEXT("BulkData with no payload should return an invalid identifier"), BulkData.GetIdentifier().IsValid());
+		TestEqual(TEXT("BulkData with no payload should return its original identifier"), BulkData.GetIdentifier(), OriginalGuid);
 
 		BulkData.UpdatePayload(FUniqueBuffer::Alloc(32).MoveToShared());
 		TestEqual(TEXT("Removing a payload then adding a new one should return the original identifier"), BulkData.GetIdentifier(), OriginalGuid);
 
 		BulkData.UpdatePayload(FUniqueBuffer::Alloc(0).MoveToShared());
-		TestFalse(TEXT("Setting a zero length payload should return an invalid identifier"), BulkData.GetIdentifier().IsValid());
+		TestEqual(TEXT("Setting a zero length payload should keep the original identifier"), BulkData.GetIdentifier(), OriginalGuid);
 		
 		BulkData.UpdatePayload(FUniqueBuffer::Alloc(32).MoveToShared());
 		TestEqual(TEXT("Restoring a payload should return the original identifier"), BulkData.GetIdentifier(), OriginalGuid);
@@ -394,7 +394,7 @@ bool FVirtualizationWrapperTestIdentifiers::RunTest(const FString& Parameters)
 			DstData.Serialize(ReaderAr, nullptr);
 		}
 
-		TestFalse(TEXT("After serialization the identifier should be invalid"), DstData.GetIdentifier().IsValid());
+		TestEqual(TEXT("After serialization the identifier should keep its identifier "), DstData.GetIdentifier(), OriginalIdentifier);
 
 		DstData.UpdatePayload(FUniqueBuffer::Alloc(32).MoveToShared());
 		TestEqual(TEXT("After adding a new payload the object should have the original identifier"), DstData.GetIdentifier(), OriginalIdentifier);
diff --git a/Engine/Source/Runtime/CoreUObject/Private/Virtualization/VirtualizedBulkData.cpp b/Engine/Source/Runtime/CoreUObject/Private/Virtualization/VirtualizedBulkData.cpp
--- a/Engine/Source/Runtime/CoreUObject/Private/Virtualization/VirtualizedBulkData.cpp
+++ b/Engine/Source/Runtime/CoreUObject/Private/Virtualization/VirtualizedBulkData.cpp
@@ -9,6 +9,7 @@
 #include "Misc/SecureHash.h"
 #include "Misc/ScopeLock.h"
 #include "Serialization/BulkData.h"
+#include "Serialization/BulkDataRegistry.h"
 #include "UObject/LinkerLoad.h"
 #include "UObject/LinkerSave.h"
 #include "UObject/Object.h"
@@ -78,6 +79,33 @@ IVirtualizationSourceControlUtilities* GetSourceControlInterface()
 	return (IVirtualizationSourceControlUtilities*)IModularFeatures::Get().GetModularFeatureImplementation(FName("VirtualizationSourceControlUtilities"), 0);
 }
 
+FVirtualizedUntypedBulkData::FVirtualizedUntypedBulkData(FVirtualizedUntypedBulkData&& Other)
+{
+	*this = MoveTemp(Other);
+}
+
+FVirtualizedUntypedBulkData& FVirtualizedUntypedBulkData::operator=(FVirtualizedUntypedBulkData&& Other)
+{
+	// The same as the default move constructor, except we need to handle registration and deregistration
+	Unregister();
+	Other.Unregister();
+
+	BulkDataId = MoveTemp(Other.BulkDataId);
+	PayloadContentId = MoveTemp(Other.PayloadContentId);
+	Payload = MoveTemp(Other.Payload);
+	PayloadSize = MoveTemp(Other.PayloadSize);
+	CompressionFormatToUse = MoveTemp(Other.CompressionFormatToUse);
+	OffsetInFile = MoveTemp(Other.OffsetInFile);
+	PackagePath = MoveTemp(Other.PackagePath);
+	PackageSegment = MoveTemp(Other.PackageSegment);
+	Flags = MoveTemp(Other.Flags);
+	Other.Reset();
+
+	Register(nullptr);
+
+	return *this;
+}
+
 FVirtualizedUntypedBulkData::FVirtualizedUntypedBulkData(const FVirtualizedUntypedBulkData& Other)
 {
 	*this = Other;
@@ -85,9 +113,28 @@ FVirtualizedUntypedBulkData::FVirtualizedUntypedBulkData(const FVirtualizedUntyp
 
 FVirtualizedUntypedBulkData& FVirtualizedUntypedBulkData::operator=(const FVirtualizedUntypedBulkData& Other)
 {
-	if (!BulkDataId.IsValid() && Other.BulkDataId.IsValid())
+	// Torn-off BulkDatas remain torn-off even when being copied into from a non-torn-off BulkData
+	// Remaining torn-off is a work-around necessary for FTextureSource::CopyTornOff to avoid registering a new
+	// guid before setting the new BulkData to torn-off. The caller can call Reset to clear the torn-off flag.
+	bool bTornOff = false;
+	if (EnumHasAnyFlags(Flags, EFlags::IsTornOff))
+	{
+		check(!EnumHasAnyFlags(Flags, EFlags::HasRegistered));
+		BulkDataId = Other.BulkDataId;
+		bTornOff = true;
+	}
+	else
 	{
-		BulkDataId = FGuid::NewGuid();
+		Unregister();
+		if (EnumHasAnyFlags(Other.Flags, EFlags::IsTornOff))
+		{
+			BulkDataId = Other.BulkDataId;
+			bTornOff = true;
+		}
+		else if (!BulkDataId.IsValid() && Other.BulkDataId.IsValid())
+		{
+			BulkDataId = FGuid::NewGuid();
+		}
 	}
 
 	PayloadContentId = Other.PayloadContentId;
@@ -98,11 +145,60 @@ FVirtualizedUntypedBulkData& FVirtualizedUntypedBulkData::operator=(const FVirtu
 	PackagePath = Other.PackagePath;
 	PackageSegment = Other.PackageSegment;
 	Flags = Other.Flags;
-	
+	EnumRemoveFlags(Flags, EFlags::TransientFlags);
+
+	if (bTornOff)
+	{
+		EnumAddFlags(Flags, EFlags::IsTornOff);
+	}
+	else
+	{
+		Register(nullptr);
+	}
 	return *this;
 }
 
-void FVirtualizedUntypedBulkData::CreateFromBulkData(FUntypedBulkData& InBulkData, const FGuid& InGuid)
+FVirtualizedUntypedBulkData::~FVirtualizedUntypedBulkData()
+{
+	Unregister();
+}
+
+FVirtualizedUntypedBulkData::FVirtualizedUntypedBulkData(const FVirtualizedUntypedBulkData& Other, ETornOff)
+{
+	EnumAddFlags(Flags, EFlags::IsTornOff);
+	*this = Other; // We rely on operator= preserving the torn-off flag
+}
+
+void FVirtualizedUntypedBulkData::TearOff()
+{
+	Unregister();
+	EnumAddFlags(Flags, EFlags::IsTornOff);
+}
+
+void FVirtualizedUntypedBulkData::Register(UObject* Owner)
+{
+#if WITH_EDITOR
+	if (BulkDataId.IsValid() && PayloadSize > 0 && !EnumHasAnyFlags(Flags, EFlags::IsTornOff))
+	{
+		IBulkDataRegistry::Get().Register(Owner ? Owner->GetPackage() : nullptr, *this);
+		EnumAddFlags(Flags, EFlags::HasRegistered);
+	}
+#endif
+}
+
+void FVirtualizedUntypedBulkData::Unregister()
+{
+#if WITH_EDITOR
+	if (EnumHasAnyFlags(Flags, EFlags::HasRegistered))
+	{
+		check(!EnumHasAnyFlags(Flags, EFlags::IsTornOff));
+		IBulkDataRegistry::Get().OnExitMemory(*this);
+		EnumRemoveFlags(Flags, EFlags::HasRegistered);
+	}
+#endif
+}
+
+void FVirtualizedUntypedBulkData::CreateFromBulkData(FUntypedBulkData& InBulkData, const FGuid& InGuid, UObject* Owner)
 {
 	TRACE_CPUPROFILER_EVENT_SCOPE(FVirtualizedUntypedBulkData::CreateFromBulkData);
 	
@@ -156,6 +252,7 @@ void FVirtualizedUntypedBulkData::CreateFromBulkData(FUntypedBulkData& InBulkDat
 	{
 		EnumAddFlags(Flags, EFlags::LegacyKeyWasGuidDerived);
 	}
+	Register(Owner);
 }
 
 void FVirtualizedUntypedBulkData::Serialize(FArchive& Ar, UObject* Owner)
@@ -171,6 +268,11 @@ void FVirtualizedUntypedBulkData::Serialize(FArchive& Ar, UObject* Owner)
 
 		if (bNeedsTransaction)
 		{
+			if (Ar.IsLoading())
+			{
+				Unregister();
+			}
+
 			Ar << Flags;
 			Ar << BulkDataId;
 			Ar << PayloadContentId;
@@ -207,6 +309,8 @@ void FVirtualizedUntypedBulkData::Serialize(FArchive& Ar, UObject* Owner)
 				}
 				
 				Payload = CompressedPayload.Decompress();	
+
+				Register(Owner);
 			}
 		}
 	}
@@ -231,10 +335,18 @@ void FVirtualizedUntypedBulkData::Serialize(FArchive& Ar, UObject* Owner)
 				PushData(); // Note this can change various members if we are going from non-virtualized to virtualized
 			}
 		}
+		else
+		{
+			Unregister();
+		}
 
 		// Store the position in the archive of the flags in case we need to update it later
 		const int64 SavedFlagsPos = Ar.Tell();
 		Ar << Flags;
+		if (Ar.IsLoading())
+		{
+			EnumRemoveFlags(Flags, EFlags::TransientFlags);
+		}
 
 		// TODO: Can probably remove these checks before UE5 release
 		check(!Ar.IsSaving() || GetPayloadSize() == 0 || BulkDataId.IsValid()); // Sanity check to stop us saving out bad data
@@ -297,7 +409,7 @@ void FVirtualizedUntypedBulkData::Serialize(FArchive& Ar, UObject* Owner)
 
 				// The lambda is mutable so that PayloadToSerialize is not const (due to FArchive api not
 				// accepting const values)
-				auto SerializePayload = [this, OffsetPos, PayloadToSerialize, UpdatedFlags](FArchive& Ar) mutable
+				auto SerializePayload = [this, OffsetPos, PayloadToSerialize, UpdatedFlags, Owner](FArchive& Ar) mutable
 				{
 					checkf(Ar.IsCooking() == false, TEXT("FVirtualizedUntypedBulkData::Serialize should not be called during a cook"));
 
@@ -329,7 +441,7 @@ void FVirtualizedUntypedBulkData::Serialize(FArchive& Ar, UObject* Owner)
 							SidecarData.Payload = PayloadToSerialize;
 						}
 
-						auto OnSavePackage = [this, PayloadOffset, UpdatedFlags](const FPackagePath& InPackagePath, FObjectPostSaveContext ObjectSaveContext)
+						auto OnSavePackage = [this, PayloadOffset, UpdatedFlags, Owner](const FPackagePath& InPackagePath, FObjectPostSaveContext ObjectSaveContext)
 						{
 							if (!ObjectSaveContext.IsUpdatingLoadedPath())
 							{
@@ -345,6 +457,9 @@ void FVirtualizedUntypedBulkData::Serialize(FArchive& Ar, UObject* Owner)
 							{
 								this->Payload.Reset();
 							}
+							// Update our information in the registry
+							// TODO: Pass Owner into Register once the AssetRegistry has been fixed to use the updated PackageGuid from the save
+							Register(nullptr);
 						};
 
 						LinkerSave->PostSaveCallbacks.Add(MoveTemp(OnSavePackage));
@@ -413,10 +528,54 @@ void FVirtualizedUntypedBulkData::Serialize(FArchive& Ar, UObject* Owner)
 					Payload = CompressedPayload.Decompress();
 				}
 			}
+			Register(Owner);
 		}
 	}
 }
 
+void FVirtualizedUntypedBulkData::SerializeForRegistry(FArchive& Ar)
+{
+	if (Ar.IsSaving())
+	{
+		check(CanSaveForRegistry());
+		EFlags FlagsForSerialize = Flags;
+		EnumRemoveFlags(FlagsForSerialize, EFlags::TransientFlags);
+		Ar << FlagsForSerialize;
+	}
+	else
+	{
+		Ar << Flags;
+		EnumRemoveFlags(Flags, EFlags::TransientFlags);
+		EnumAddFlags(Flags, EFlags::IsTornOff);
+	}
+
+	Ar << BulkDataId;
+	Ar << PayloadContentId;
+	Ar << PayloadSize;
+	if (Ar.IsSaving())
+	{
+		FString PackageName = PackagePath.GetPackageName();
+		check(PackageName.IsEmpty() || PackageSegment == EPackageSegment::Header);
+		Ar << PackageName;
+	}
+	else
+	{
+		FString PackageName;
+		Ar << PackageName;
+		ensure(FPackagePath::TryFromPackageName(PackageName, PackagePath));
+		PackageSegment = EPackageSegment::Header;
+	}
+	Ar << OffsetInFile;
+}
+
+bool FVirtualizedUntypedBulkData::CanSaveForRegistry() const
+{
+	return BulkDataId.IsValid() && PayloadSize > 0 && !IsMemoryOnlyPayload()
+		&& EnumHasAnyFlags(Flags, EFlags::IsTornOff) && !EnumHasAnyFlags(Flags, EFlags::HasRegistered)
+		&& (PackagePath.IsEmpty() || PackageSegment == EPackageSegment::Header);
+}
+
+
 FCompressedBuffer FVirtualizedUntypedBulkData::LoadFromDisk() const
 {
 	TRACE_CPUPROFILER_EVENT_SCOPE(FVirtualizedUntypedBulkData::LoadFromDisk);
@@ -668,6 +827,9 @@ void FVirtualizedUntypedBulkData::PushData()
 			PackagePath.Empty();
 			PackageSegment = EPackageSegment::Header;
 			OffsetInFile = INDEX_NONE;
+
+			// Update our information in the registry
+			Register(nullptr);
 		}
 	}	
 }
@@ -718,9 +880,15 @@ bool FVirtualizedUntypedBulkData::CanUnloadData() const
 	return IsDataVirtualized() || PackagePath.IsEmpty() == false;
 }
 
+bool FVirtualizedUntypedBulkData::IsMemoryOnlyPayload() const
+{
+	return !Payload.IsNull() && !CanUnloadData();
+}
+
 void FVirtualizedUntypedBulkData::Reset()
 {
 	// Note that we do not reset the BulkDataId
+	Unregister();
 	PayloadContentId.Reset();
 	Payload.Reset();
 	PayloadSize = 0;
@@ -741,15 +909,8 @@ void FVirtualizedUntypedBulkData::UnloadData()
 
 FGuid FVirtualizedUntypedBulkData::GetIdentifier() const
 {
-	if (GetPayloadSize() > 0)
-	{
-		checkf(BulkDataId.IsValid(), TEXT("If bulkdata has a valid payload then it should have a valid BulkDataId"));
-		return BulkDataId;
-	}
-	else
-	{
-		return FGuid();
-	}
+	checkf(GetPayloadSize() == 0 || BulkDataId.IsValid(), TEXT("If bulkdata has a valid payload then it should have a valid BulkDataId"));
+	return BulkDataId;
 }
 
 FCompressedBuffer FVirtualizedUntypedBulkData::GetDataInternal() const
@@ -815,7 +976,8 @@ void FVirtualizedUntypedBulkData::UpdatePayload(FSharedBuffer InPayload, FName I
 	// Make sure that we own the memory in the shared buffer
 	Payload = InPayload.MakeOwned();
 	PayloadSize = (int64)Payload.GetSize();
-	
+	PayloadContentId = FPayloadId(Payload);
+
 	EnumRemoveFlags(Flags,	EFlags::IsVirtualized | 
 							EFlags::DisablePayloadCompression |
 							EFlags::ReferencesLegacyFile |
@@ -828,12 +990,18 @@ void FVirtualizedUntypedBulkData::UpdatePayload(FSharedBuffer InPayload, FName I
 	PackageSegment = EPackageSegment::Header;
 	OffsetInFile = INDEX_NONE;
 
-	if (!BulkDataId.IsValid() && PayloadSize > 0)
+	if (PayloadSize > 0)
 	{
-		BulkDataId = FGuid::NewGuid();
+		if (!BulkDataId.IsValid())
+		{
+			BulkDataId = FGuid::NewGuid();
+		}
+		Register(nullptr);
+	}
+	else
+	{
+		Unregister();
 	}
-
-	PayloadContentId = FPayloadId(Payload);
 }
 
 void FVirtualizedUntypedBulkData::SetCompressionFormat(FName InCompressionFormat)
@@ -852,6 +1020,11 @@ FCustomVersionContainer FVirtualizedUntypedBulkData::GetCustomVersions(FArchive&
 	return InlineArchive.GetCustomVersions();
 }
 
+void FVirtualizedUntypedBulkData::UpdatePayloadId()
+{
+	UpdateKeyIfNeeded();
+}
+
 void FVirtualizedUntypedBulkData::UpdateKeyIfNeeded()
 {
 	// If this was created from old BulkData then the key is generated from an older FGuid, we
diff --git a/Engine/Source/Runtime/CoreUObject/Public/Serialization/BulkDataRegistry.h b/Engine/Source/Runtime/CoreUObject/Public/Serialization/BulkDataRegistry.h
new file mode 100644
--- /dev/null
+++ b/Engine/Source/Runtime/CoreUObject/Public/Serialization/BulkDataRegistry.h
@@ -0,0 +1,87 @@
+// Copyright Epic Games, Inc. All Rights Reserved.
+
+#pragma once
+
+#if WITH_EDITOR
+
+#include "Async/Future.h"
+#include "Compression/CompressedBuffer.h"
+#include "Containers/ArrayView.h"
+#include "Delegates/Delegate.h"
+#include "HAL/Platform.h"
+#include "IO/IoHash.h"
+
+class IBulkDataRegistry;
+class UPackage;
+struct FGuid;
+struct FUntypedBulkData;
+namespace UE::Virtualization { class FVirtualizedUntypedBulkData; }
+
+DECLARE_DELEGATE_RetVal(IBulkDataRegistry*, FSetBulkDataRegistry);
+
+namespace UE::BulkDataRegistry
+{
+	/** Results of GetMeta call. */
+	struct FMetaData
+	{
+		/** True if data was found, else false. */
+		bool bValid;
+		/** IoHash of the uncompressed bytes of the data that will be returned from GetData. */
+		FIoHash RawHash;
+		/** Size of the uncompressed bytes of the data that will be returned from GetData. */
+		uint64 RawSize;
+	};
+
+	/** Results of GetData call. */
+	struct FData
+	{
+		/** True if data was found, else false. */
+		bool bValid;
+		/** The discovered data. Empty if data was not found. */
+		FCompressedBuffer Buffer;
+	};
+
+}
+
+/** Registers BulkDatas so that they can be referenced by guid during builds later in the editor process. */
+class IBulkDataRegistry
+{
+public:
+	/** The BulkDataRegistry can be configured off. Return whether it is enabled. A stub is used if not enabled. */
+	COREUOBJECT_API static bool IsEnabled();
+	/** Get the global BulkDataRegistry; always returns a valid interface, so long as Initialize has been called. */
+	COREUOBJECT_API static IBulkDataRegistry& Get();
+	
+	/** Register a BulkData with the registry. Its payload and metadata will be fetchable by its GetIdentifier. */
+	virtual void Register(UPackage* Owner, const UE::Virtualization::FVirtualizedUntypedBulkData& BulkData) = 0;
+	/** Report that a BulkData is leaving memory and its in-memory payload (if it had one) is no longer available. */
+	virtual void OnExitMemory(const UE::Virtualization::FVirtualizedUntypedBulkData& BulkData) = 0;
+
+	/** Return the metadata for the given registered BulkData; returns false if not registered. */
+	virtual TFuture<UE::BulkDataRegistry::FMetaData> GetMeta(const FGuid& BulkDataId) = 0;
+
+	/**
+	 * Return the (possibly compressed) payload for the given registered BulkData.
+	 * Returns an empty buffer if not registered.
+	 */
+	virtual TFuture<UE::BulkDataRegistry::FData> GetData(const FGuid& BulkDataId) = 0;
+
+	/** Set and intialize global IBulkDataRegistry; Get fatally fails before. */
+	COREUOBJECT_API static void Initialize();
+	/** Shutdown and deallocate global IBulkDataRegistry; Get fatally fails afterwards. */
+	COREUOBJECT_API static void Shutdown();
+	/** Subscribe to set the class for the global IBulkDataRegistry. */
+	COREUOBJECT_API static FSetBulkDataRegistry& GetSetBulkDataRegistryDelegate();
+
+protected:
+	virtual ~IBulkDataRegistry() {};
+};
+
+// Temporary interface for tunneling the EditorBuildInputResolver into CoreUObject.
+// In the future this will be implemented as part of the BuildAPI
+namespace UE::DerivedData { class IBuildInputResolver; }
+COREUOBJECT_API UE::DerivedData::IBuildInputResolver* GetGlobalBuildInputResolver();
+COREUOBJECT_API void SetGlobalBuildInputResolver(UE::DerivedData::IBuildInputResolver* InResolver);
+
+
+#endif // WITH_EDITOR
\ No newline at end of file
diff --git a/Engine/Source/Runtime/CoreUObject/Public/Virtualization/VirtualizedBulkData.h b/Engine/Source/Runtime/CoreUObject/Public/Virtualization/VirtualizedBulkData.h
--- a/Engine/Source/Runtime/CoreUObject/Public/Virtualization/VirtualizedBulkData.h
+++ b/Engine/Source/Runtime/CoreUObject/Public/Virtualization/VirtualizedBulkData.h
@@ -51,13 +51,13 @@ class COREUOBJECT_API FVirtualizedUntypedBulkData
 {
 public:
 	FVirtualizedUntypedBulkData() = default;
-	FVirtualizedUntypedBulkData(FVirtualizedUntypedBulkData&&) = default;
-	FVirtualizedUntypedBulkData& operator=(FVirtualizedUntypedBulkData&& ) = default;
+	FVirtualizedUntypedBulkData(FVirtualizedUntypedBulkData&& Other);
+	FVirtualizedUntypedBulkData& operator=(FVirtualizedUntypedBulkData&& Other);
 
 	FVirtualizedUntypedBulkData(const FVirtualizedUntypedBulkData& Other);
 	FVirtualizedUntypedBulkData& operator=(const FVirtualizedUntypedBulkData& Other);
 
-	~FVirtualizedUntypedBulkData() = default;
+	~FVirtualizedUntypedBulkData();
 
 	/** 
 	 * Convenience method to make it easier to convert from BulkData to FVirtualizedBulkData and sets the Guid 
@@ -66,7 +66,7 @@ public:
 	 * @param Guid		A guid associated with the bulkdata object which will be used to identify the payload.
 	 *					This MUST remain the same between sessions so that the payloads key remains consistent!
 	 */
-	void CreateFromBulkData(FUntypedBulkData& BulkData, const FGuid& Guid);
+	void CreateFromBulkData(FUntypedBulkData& BulkData, const FGuid& Guid, UObject* Owner);
 
 	/** 
 	 * Used to serialize the bulkdata to/from a FArchive
@@ -86,14 +86,7 @@ public:
 	/** Unloads the data (if possible) but leaves it in a state where the data can be reloaded */
 	void UnloadData();
 
-	/**
-	 * Returns a unique identifier for the object itself.
-	 * This should only return a valid FGuid as long as the object owns a valid payload.
-	 * If an object with a valid payload, has that payload removed then it should start
-	 * returning an invalid FGuid instead.
-	 * Should that object be given a new payload it should then return the original
-	 * identifier, there is no need to generate a new one.
-	 */
+	/** Returns a unique identifier for the object itself. */
 	FGuid GetIdentifier() const;
 
 	/** Returns an unique identifier for the content of the payload. */
@@ -173,6 +166,33 @@ public:
 	*/
 	FCustomVersionContainer GetCustomVersions(FArchive& InlineArchive);
 
+	/**
+	 * Set this BulkData into Torn-Off mode. It will no longer register with the BulkDataRegistry, even if
+	 * copied from another BulkData, and it will pass on this flag to any BulkData copied/moved from it.
+	 * Use Reset() to remove this state. Torn-off BulkDatas share the guid with the BulkData they copy from.
+	 */
+	void TearOff();
+
+	/** Make a torn-off copy of this bulk data. */
+	FVirtualizedUntypedBulkData CopyTornOff() const { return FVirtualizedUntypedBulkData(*this, ETornOff()); }
+
+	// Functions used by the BulkDataRegistry
+
+	/** Used to serialize the bulkdata to/from a limited cache system used by the BulkDataRegistry. */
+	void SerializeForRegistry(FArchive& Ar);
+	/** Return true if the bulkdata has a source location that persists between editor processes (package file or virtualization). */
+	bool CanSaveForRegistry() const;
+	/** Return whether the BulkData has legacy payload id that needs to be updated from loaded payload before it can be used in DDC. */
+	bool HasPlaceholderPayloadId() const { return EnumHasAnyFlags(Flags, EFlags::LegacyKeyWasGuidDerived); }
+	/** Return whether the BulkData is an in-memory payload without a persistent source location. */
+	bool IsMemoryOnlyPayload() const;
+	/** Load the payload and set the correct payload id, if the bulkdata has a PlaceholderPayloadId. */
+	void UpdatePayloadId();
+
+protected:
+	enum class ETornOff {};
+	FVirtualizedUntypedBulkData(const FVirtualizedUntypedBulkData& Other, ETornOff);
+
 private:
 	/** Flags used to store additional meta information about the bulk data */
 	enum class EFlags : uint32
@@ -192,6 +212,12 @@ private:
 		DisablePayloadCompression	= 1 << 4,
 		/** The legacy file being referenced derived its key from guid and it should be replaced with a key-from-hash when saved */
 		LegacyKeyWasGuidDerived		= 1 << 5,
+		/** The Guid has been registered with the BulkDataRegistry */
+		HasRegistered				= 1 << 6,
+		/** The BulkData object is a copy used only to represent the id and payload; it does not communicate with the BulkDataRegistry, and will point DDC jobs toward the original BulkData */
+		IsTornOff					= 1 << 7,
+
+		TransientFlags				= HasRegistered | IsTornOff,
 	};
 
 	/** Used to control what level of error reporting we return from some methods */
@@ -230,6 +256,9 @@ private:
 	bool HasPayloadSidecarFile() const { return EnumHasAnyFlags(Flags, EFlags::HasPayloadSidecarFile); }
 	bool IsReferencingOldBulkData() const { return EnumHasAnyFlags(Flags, EFlags::ReferencesLegacyFile); }
 
+	void Register(UObject* Owner);
+	void Unregister();
+
 	/** Unique identifier for the bulkdata object itself */
 	FGuid BulkDataId;
 
@@ -270,6 +299,14 @@ class TVirtualizedBulkData final : public FVirtualizedUntypedBulkData
 public:
 	TVirtualizedBulkData() = default;
 	~TVirtualizedBulkData() = default;
+
+	TVirtualizedBulkData<DataType> CopyTornOff() const
+	{
+		return TVirtualizedBulkData<DataType>(*this, ETornOff());
+	}
+
+protected:
+	TVirtualizedBulkData(const TVirtualizedBulkData<DataType>& Other, ETornOff) : FVirtualizedUntypedBulkData(Other, ETornOff()) {}
 };
 
 using FByteVirtualizedBulkData	= TVirtualizedBulkData<uint8>;
diff --git a/Engine/Source/Runtime/Engine/Classes/Engine/Texture.h b/Engine/Source/Runtime/Engine/Classes/Engine/Texture.h
--- a/Engine/Source/Runtime/Engine/Classes/Engine/Texture.h
+++ b/Engine/Source/Runtime/Engine/Classes/Engine/Texture.h
@@ -211,6 +211,9 @@ struct FTextureSource
 		ETextureSourceCompressionFormat NewSourceFormat
 	);
 
+	/** Make a copy with a torn-off BulkData that has the same Guid used for DDC as this->BulkData */
+	FTextureSource CopyTornOff() const;
+
 	/** PNG Compresses the source art if possible or tells the bulk data to zlib compress when it saves out to disk. */
 	ENGINE_API void Compress();
 
@@ -263,6 +266,7 @@ struct FTextureSource
 	ETextureSourceCompressionFormat GetSourceCompression() const;
 
 	/** Trivial accessors. These will only give values for Block0 so may not be correct for UDIM/multi-block textures, use GetBlock() for this case. */
+	FGuid GetPersistentId() const { return BulkData.GetIdentifier(); }
 	FORCEINLINE FGuid GetId() const { return Id; }
 	FORCEINLINE int32 GetSizeX() const { return SizeX; }
 	FORCEINLINE int32 GetSizeY() const { return SizeY; }
diff --git a/Engine/Source/Runtime/Engine/Private/Texture.cpp b/Engine/Source/Runtime/Engine/Private/Texture.cpp
--- a/Engine/Source/Runtime/Engine/Private/Texture.cpp
+++ b/Engine/Source/Runtime/Engine/Private/Texture.cpp
@@ -473,7 +473,8 @@ void UTexture::Serialize(FArchive& Ar)
 			FByteBulkData TempBulkData;
 			TempBulkData.Serialize(Ar, this);
 
-			Source.BulkData.CreateFromBulkData(TempBulkData, Source.GetId());
+			FGuid LegacyPersistentId = Source.GetId();
+			Source.BulkData.CreateFromBulkData(TempBulkData, LegacyPersistentId, this);
 		}
 		else
 		{
@@ -1228,6 +1229,16 @@ void FTextureSource::InitWithCompressedSourceData(
 	BulkData.UpdatePayload(FSharedBuffer::Clone(NewData.GetData(), NewData.Num()), CompressionName);
 }
 
+FTextureSource FTextureSource::CopyTornOff() const
+{
+	FTextureSource Result;
+	// Set the Torn off flag on Result.BulkData so that the copy constructor below will not set it
+	Result.BulkData.TearOff();
+	// Use the default copy constructor to copy all the fields without having to write them manually
+	Result = *this;
+	return Result;
+}
+
 void FTextureSource::Compress()
 {
 	checkf(LockState == ELockState::None, TEXT("Compress shouldn't be called in-between LockMip/UnlockMip"));
diff --git a/Engine/Source/Runtime/Engine/Private/TextureDerivedDataTask.cpp b/Engine/Source/Runtime/Engine/Private/TextureDerivedDataTask.cpp
--- a/Engine/Source/Runtime/Engine/Private/TextureDerivedDataTask.cpp
+++ b/Engine/Source/Runtime/Engine/Private/TextureDerivedDataTask.cpp
@@ -40,7 +40,9 @@
 #include "Math/NumericLimits.h"
 #include "Math/UnrealMathUtility.h"
 #include "Misc/FeedbackContext.h"
+#include "Misc/Optional.h"
 #include "ProfilingDebugging/CookStats.h"
+#include "Serialization/BulkDataRegistry.h"
 #include "TextureDerivedDataBuildUtils.h"
 #include "VT/VirtualTextureDataBuilder.h"
 
@@ -190,7 +192,7 @@ void FTextureSourceData::Init(UTexture& InTexture, const FTextureBuildSettings*
 	if (bAllowAsyncLoading && !InTexture.Source.IsBulkDataLoaded())
 	{
 		// Prepare the async source to be later able to load it from file if required.
-		AsyncSource = InTexture.Source; // This copies information required to make a safe IO load async.
+		AsyncSource = InTexture.Source.CopyTornOff(); // This copies information required to make a safe IO load async.
 	}
 
 	bValid = true;
@@ -442,10 +444,10 @@ void FTextureCacheDerivedDataWorker::BuildTexture(bool bReplaceExistingDDC)
 
 			DefinitionBuilder.AddConstant(TEXT("Settings"_SV),
 				SaveTextureBuildSettings(KeySuffix, Texture, BuildSettingsPerLayer[0], 0, NUM_INLINE_DERIVED_MIPS));
-			DefinitionBuilder.AddInputBulkData(TEXT("Source"_SV), Texture.Source.GetId());
+			DefinitionBuilder.AddInputBulkData(TEXT("Source"_SV), Texture.Source.GetPersistentId());
 			if (Texture.CompositeTexture)
 			{
-				DefinitionBuilder.AddInputBulkData(TEXT("CompositeSource"_SV), Texture.CompositeTexture->Source.GetId());
+				DefinitionBuilder.AddInputBulkData(TEXT("CompositeSource"_SV), Texture.CompositeTexture->Source.GetPersistentId());
 			}
 
 			class FTextureBuildInputResolver final : public IBuildInputResolver
@@ -458,7 +460,7 @@ void FTextureCacheDerivedDataWorker::BuildTexture(bool bReplaceExistingDDC)
 
 				const FCompressedBuffer& FindSource(FCompressedBuffer& Buffer, FTextureSource& Source, const FGuid& BulkDataId)
 				{
-					if (Source.GetId() != BulkDataId)
+					if (Source.GetPersistentId() != BulkDataId)
 					{
 						return FCompressedBuffer::Null;
 					}
@@ -536,8 +538,14 @@ void FTextureCacheDerivedDataWorker::BuildTexture(bool bReplaceExistingDDC)
 				FCompressedBuffer CompositeSourceBuffer;
 			};
 
-			FTextureBuildInputResolver InputResolver(Texture);
-			FBuildSession Session = Build.CreateSession(TexturePath, &InputResolver);
+			TOptional<FTextureBuildInputResolver> PlaceholderResolver;
+			UE::DerivedData::IBuildInputResolver* InputResolver = GetGlobalBuildInputResolver();
+			if (!InputResolver)
+			{
+				PlaceholderResolver.Emplace(Texture);
+				InputResolver = &PlaceholderResolver.GetValue();
+			}
+			FBuildSession Session = Build.CreateSession(TexturePath, InputResolver);
 			Session.Build(DefinitionBuilder.Build(), EBuildPolicy::Default, EPriority::Blocking,
 				[this, &TexturePath, bReplaceExistingDDC] (FBuildCompleteParams&& Params)
 				{
diff --git a/Engine/Source/Runtime/Launch/Private/LaunchEngineLoop.cpp b/Engine/Source/Runtime/Launch/Private/LaunchEngineLoop.cpp
--- a/Engine/Source/Runtime/Launch/Private/LaunchEngineLoop.cpp
+++ b/Engine/Source/Runtime/Launch/Private/LaunchEngineLoop.cpp
@@ -105,6 +105,7 @@
 	#include "Interfaces/IEditorStyleModule.h"
 	#include "PIEPreviewDeviceProfileSelectorModule.h"
 	#include "AssetCompilingManager.h"
+	#include "Serialization/BulkDataRegistry.h"
 	#include "ShaderCompiler.h"
 
 	#if PLATFORM_WINDOWS
@@ -3269,6 +3270,10 @@ int32 FEngineLoop::PreInitPostStartupScreen(const TCHAR* CmdLine)
 		// It has to be intialized after the AssetRegistryModule; the editor implementations of PackageResourceManager relies on it
 		IPackageResourceManager::Initialize();
 #endif
+#if WITH_EDITOR
+		// Initialize the BulkDataRegistry, which registers BulkData structs loaded from Packages for later building. It uses the same lifetime as IPackageResourceManager
+		IBulkDataRegistry::Initialize();
+#endif
 
 		FEmbeddedCommunication::ForceTick(5);
 
@@ -4504,6 +4509,10 @@ void FEngineLoop::Exit()
 	MALLOC_PROFILER( GMalloc->Exec(nullptr, TEXT("MPROF STOP"), *GLog);	);
 #endif // !ANDROID
 
+#if WITH_EDITOR
+	IBulkDataRegistry::Shutdown();
+#endif
+
 	// Stop the rendering thread.
 	StopRenderingThread();
 	
diff --git a/Engine/Source/Runtime/MeshDescription/Private/MeshDescription.cpp b/Engine/Source/Runtime/MeshDescription/Private/MeshDescription.cpp
--- a/Engine/Source/Runtime/MeshDescription/Private/MeshDescription.cpp
+++ b/Engine/Source/Runtime/MeshDescription/Private/MeshDescription.cpp
@@ -2046,7 +2046,7 @@ void FMeshDescriptionBulkData::Serialize( FArchive& Ar, UObject* Owner )
 	// Needs to be after the guid is serialized
 	if (bSerializedOldDataTypes)
 	{
-		BulkData.CreateFromBulkData(TempBulkData, Guid);
+		BulkData.CreateFromBulkData(TempBulkData, Guid, Owner);
 	}
 }
 
