Don't attempt to unload newly created worlds when performing a map load

This is an extension of the fix for UE-117449, which focused on the case where the map being loadead was the only available newly created world. This handles the case where there may be multiple newly created worlds available when performing a map load, as this case would still present an error dialog that these worlds could not be unloaded.

diff --git a/Engine/Source/Editor/UnrealEd/Private/EditorServer.cpp b/Engine/Source/Editor/UnrealEd/Private/EditorServer.cpp
--- a/Engine/Source/Editor/UnrealEd/Private/EditorServer.cpp
+++ b/Engine/Source/Editor/UnrealEd/Private/EditorServer.cpp
@@ -2570,8 +2570,9 @@ bool UEditorEngine::Map_Load(const TCHAR* Str, FOutputDevice& Ar)
 						TArray<UPackage*> WorldPackages;
 						for (TObjectIterator<UWorld> It; It; ++It)
 						{
-							UPackage* Package = Cast<UPackage>(It->GetOuter());
-							if (Package && Package != GetTransientPackage() && (!IsWorldValidForReuse(*It) || Package->GetPathName() != LongTempFname))
+							// Don't unload transient or newly created worlds, nor the world we're attempting to keep-alive
+							UPackage* Package = It->GetPackage();
+							if (Package != GetTransientPackage() && !Package->HasAnyPackageFlags(PKG_NewlyCreated) && *It != NewWorld)
 							{
 								WorldPackages.AddUnique(Package);
 							}
