IterationProfile: Specify Context and DataPoint directly from Statistics CSV file.

This data used to be copied exactly from the trace data, but to make things more concise and undestandable for the tableau dashboard we force specific Context and DataPoint strings from columns in the Statistics CSV. This gives us full control over how things are organized in the dashboard.

diff --git a/Engine/Source/Editor/UnrealEd/Private/Commandlets/SummarizeTraceCommandlet.cpp b/Engine/Source/Editor/UnrealEd/Private/Commandlets/SummarizeTraceCommandlet.cpp
--- a/Engine/Source/Editor/UnrealEd/Private/Commandlets/SummarizeTraceCommandlet.cpp
+++ b/Engine/Source/Editor/UnrealEd/Private/Commandlets/SummarizeTraceCommandlet.cpp
@@ -613,38 +613,48 @@ struct StatisticDefinition
 	StatisticDefinition()
 	{}
 
-	StatisticDefinition(const FString& InName, const FString& InStatistic, const FString& InUnit, const FString& InWarningThreshold, const FString& InErrorThreshold)
+	StatisticDefinition(const FString& InName, const FString& InStatistic,
+		const FString& InTelemetryContext, const FString& InTelemetryDataPoint, const FString& InTelemetryUnit,
+		const FString& InBaselineWarningThreshold, const FString& InBaselineErrorThreshold)
 		: Name(InName)
 		, Statistic(InStatistic)
-		, Unit(InUnit)
-		, WarningThreshold(InWarningThreshold)
-		, ErrorThreshold(InErrorThreshold)
+		, TelemetryContext(InTelemetryContext)
+		, TelemetryDataPoint(InTelemetryDataPoint)
+		, TelemetryUnit(InTelemetryUnit)
+		, BaselineWarningThreshold(InBaselineWarningThreshold)
+		, BaselineErrorThreshold(InBaselineErrorThreshold)
 	{}
 
 	StatisticDefinition(const StatisticDefinition& InStatistic)
 		: Name(InStatistic.Name)
 		, Statistic(InStatistic.Statistic)
-		, Unit(InStatistic.Unit)
-		, WarningThreshold(InStatistic.WarningThreshold)
-		, ErrorThreshold(InStatistic.ErrorThreshold)
+		, TelemetryContext(InStatistic.TelemetryContext)
+		, TelemetryDataPoint(InStatistic.TelemetryDataPoint)
+		, TelemetryUnit(InStatistic.TelemetryUnit)
+		, BaselineWarningThreshold(InStatistic.BaselineWarningThreshold)
+		, BaselineErrorThreshold(InStatistic.BaselineErrorThreshold)
 	{}
 
 	bool operator==(const StatisticDefinition& InStatistic) const
 	{
 		return Name == InStatistic.Name
 			&& Statistic == InStatistic.Statistic
-			&& Unit == InStatistic.Unit
-			&& WarningThreshold == InStatistic.WarningThreshold
-			&& ErrorThreshold == InStatistic.ErrorThreshold;
+			&& TelemetryContext == InStatistic.TelemetryContext
+			&& TelemetryDataPoint == InStatistic.TelemetryDataPoint
+			&& TelemetryUnit == InStatistic.TelemetryUnit
+			&& BaselineWarningThreshold == InStatistic.BaselineWarningThreshold
+			&& BaselineErrorThreshold == InStatistic.BaselineErrorThreshold;
 	}
 
 	static bool LoadFromCSV(const FString& FilePath, TMultiMap<FString, StatisticDefinition>& NameToDefinitionMap);
 
 	FString Name;
 	FString Statistic;
-	FString Unit;
-	FString WarningThreshold;
-	FString ErrorThreshold;
+	FString TelemetryContext;
+	FString TelemetryDataPoint;
+	FString TelemetryUnit;
+	FString BaselineWarningThreshold;
+	FString BaselineErrorThreshold;
 };
 
 bool StatisticDefinition::LoadFromCSV(const FString& FilePath, TMultiMap<FString, StatisticDefinition>& NameToDefinitionMap)
@@ -654,9 +664,11 @@ bool StatisticDefinition::LoadFromCSV(const FString& FilePath, TMultiMap<FString
 
 	int NameColumn = -1;
 	int StatisticColumn = -1;
-	int UnitColumn = -1;
-	int WarningThresholdColumn = -1;
-	int ErrorThresholdColumn = -1;
+	int TelemetryContextColumn = -1;
+	int TelemetryDataPointColumn = -1;
+	int TelemetryUnitColumn = -1;
+	int BaselineWarningThresholdColumn = -1;
+	int BaselineErrorThresholdColumn = -1;
 	struct Column
 	{
 		const TCHAR* Name = nullptr;
@@ -666,9 +678,11 @@ bool StatisticDefinition::LoadFromCSV(const FString& FilePath, TMultiMap<FString
 	{
 		{ TEXT("Name"), &NameColumn },
 		{ TEXT("Statistic"), &StatisticColumn },
-		{ TEXT("Unit"), &UnitColumn },
-		{ TEXT("WarningThreshold"), &WarningThresholdColumn },
-		{ TEXT("ErrorThreshold"), &ErrorThresholdColumn },
+		{ TEXT("TelemetryContext"), &TelemetryContextColumn },
+		{ TEXT("TelemetryDataPoint"), &TelemetryDataPointColumn },
+		{ TEXT("TelemetryUnit"), &TelemetryUnitColumn },
+		{ TEXT("BaselineWarningThreshold"), &BaselineWarningThresholdColumn },
+		{ TEXT("BaselineErrorThreshold"), &BaselineErrorThresholdColumn },
 	};
 
 	bool bValidColumns = true;
@@ -705,10 +719,12 @@ bool StatisticDefinition::LoadFromCSV(const FString& FilePath, TMultiMap<FString
 		{
 			const FString& Name(Fields[NameColumn]);
 			const FString& Statistic(Fields[StatisticColumn]);
-			const FString& Unit(Fields[UnitColumn]);
-			const FString& WarningThreshold(Fields[WarningThresholdColumn]);
-			const FString& ErrorThreshold(Fields[ErrorThresholdColumn]);
-			NameToDefinitionMap.AddUnique(Name, StatisticDefinition(Name, Statistic, Unit, WarningThreshold, ErrorThreshold));
+			const FString& TelemetryContext(Fields[TelemetryContextColumn]);
+			const FString& TelemetryDataPoint(Fields[TelemetryDataPointColumn]);
+			const FString& TelemetryUnit(Fields[TelemetryUnitColumn]);
+			const FString& BaselineWarningThreshold(Fields[BaselineWarningThresholdColumn]);
+			const FString& BaselineErrorThreshold(Fields[BaselineErrorThresholdColumn]);
+			NameToDefinitionMap.AddUnique(Name, StatisticDefinition(Name, Statistic, TelemetryContext, TelemetryDataPoint, TelemetryUnit, BaselineWarningThreshold, BaselineErrorThreshold));
 		}
 	}
 
@@ -724,24 +740,23 @@ struct TelemetryDefinition
 	TelemetryDefinition()
 	{}
 
-	TelemetryDefinition(const FString& InTestName, const FString& InContext, const FString& InDataPoint, const FString& InMeasurement,
-		const FString* InBaseline = nullptr,
-		const FString* InUnit = nullptr)
+	TelemetryDefinition(const FString& InTestName, const FString& InContext, const FString& InDataPoint, const FString& InUnit,
+		const FString& InMeasurement, const FString* InBaseline = nullptr)
 		: TestName(InTestName)
 		, Context(InContext)
 		, DataPoint(InDataPoint)
+		, Unit(InUnit)
 		, Measurement(InMeasurement)
 		, Baseline(InBaseline ? *InBaseline : FString ())
-		, Unit(InUnit ? *InUnit : FString ())
 	{}
 
 	TelemetryDefinition(const TelemetryDefinition& InStatistic)
 		: TestName(InStatistic.TestName)
 		, Context(InStatistic.Context)
 		, DataPoint(InStatistic.DataPoint)
+		, Unit(InStatistic.Unit)
 		, Measurement(InStatistic.Measurement)
 		, Baseline(InStatistic.Baseline)
-		, Unit(InStatistic.Unit)
 	{}
 
 	bool operator==(const TelemetryDefinition& InStatistic) const
@@ -761,9 +776,9 @@ struct TelemetryDefinition
 	FString TestName;
 	FString Context;
 	FString DataPoint;
+	FString Unit;
 	FString Measurement;
 	FString Baseline;
-	FString Unit;
 };
 
 bool TelemetryDefinition::LoadFromCSV(const FString& FilePath, TMap<TPair<FString, FString>, TelemetryDefinition>& ContextAndDataPointToDefinitionMap)
@@ -774,9 +789,9 @@ bool TelemetryDefinition::LoadFromCSV(const FString& FilePath, TMap<TPair<FStrin
 	int TestNameColumn = -1;
 	int ContextColumn = -1;
 	int DataPointColumn = -1;
+	int UnitColumn = -1;
 	int MeasurementColumn = -1;
 	int BaselineColumn = -1;
-	int UnitColumn = -1;
 	struct Column
 	{
 		const TCHAR* Name = nullptr;
@@ -788,9 +803,9 @@ bool TelemetryDefinition::LoadFromCSV(const FString& FilePath, TMap<TPair<FStrin
 		{ TEXT("TestName"), &TestNameColumn },
 		{ TEXT("Context"), &ContextColumn },
 		{ TEXT("DataPoint"), &DataPointColumn },
+		{ TEXT("Unit"), &UnitColumn },
 		{ TEXT("Measurement"), &MeasurementColumn },
 		{ TEXT("Baseline"), &BaselineColumn, false },
-		{ TEXT("Unit"), &UnitColumn, false },
 	};
 
 	bool bValidColumns = true;
@@ -828,6 +843,7 @@ bool TelemetryDefinition::LoadFromCSV(const FString& FilePath, TMap<TPair<FStrin
 			const FString& TestName(Fields[TestNameColumn]);
 			const FString& Context(Fields[ContextColumn]);
 			const FString& DataPoint(Fields[DataPointColumn]);
+			const FString& Unit(Fields[UnitColumn]);
 			const FString& Measurement(Fields[MeasurementColumn]);
 
 			FString Baseline;
@@ -836,13 +852,7 @@ bool TelemetryDefinition::LoadFromCSV(const FString& FilePath, TMap<TPair<FStrin
 				Baseline = Fields[BaselineColumn];
 			}
 
-			FString Unit;
-			if (UnitColumn != -1)
-			{
-				Unit = Fields[UnitColumn];
-			}
-
-			ContextAndDataPointToDefinitionMap.Add(TPair<FString, FString>(Context, DataPoint), TelemetryDefinition(TestName, Context, DataPoint, Measurement, &Baseline, &Unit));
+			ContextAndDataPointToDefinitionMap.Add(TPair<FString, FString>(Context, DataPoint), TelemetryDefinition(TestName, Context, DataPoint, Unit, Measurement, &Baseline));
 		}
 	}
 
@@ -1173,7 +1183,7 @@ int32 USummarizeTraceCommandlet::Main(const FString& CmdLineParams)
 			NameToDefinitionMap.MultiFind(Scope.Name, Statistics, true);
 			for (const StatisticDefinition& Statistic : Statistics)
 			{
-				TelemetryData.Add(TelemetryDefinition(TestName, Scope.Name, Statistic.Statistic, Scope.GetValue(Statistic.Statistic)));
+				TelemetryData.Add(TelemetryDefinition(TestName, Statistic.TelemetryContext, Statistic.TelemetryDataPoint, Statistic.TelemetryUnit, Scope.GetValue(Statistic.Statistic)));
 			}
 		}
 
@@ -1185,7 +1195,7 @@ int32 USummarizeTraceCommandlet::Main(const FString& CmdLineParams)
 			ensure(Statistics.Num() <= 1); // there should only be one, the counter value
 			for (const StatisticDefinition& Statistic : Statistics)
 			{
-				TelemetryData.Add(TelemetryDefinition(TestName, Counter.Value.Name, Statistic.Statistic, Counter.Value.GetValue()));
+				TelemetryData.Add(TelemetryDefinition(TestName, Statistic.TelemetryContext, Statistic.TelemetryDataPoint, Statistic.TelemetryUnit, Counter.Value.GetValue()));
 			}
 		}
 
@@ -1228,8 +1238,6 @@ int32 USummarizeTraceCommandlet::Main(const FString& CmdLineParams)
 				// do we still have the statistic definition in our current stats file? (if we don't that's fine, we don't care about it anymore)
 				if (RelatedStatistic)
 				{
-					Telemetry.Unit = RelatedStatistic->Unit;
-
 					// find the corresponding keyed telemetry item in the baseline telemetry file...
 					TelemetryDefinition* BaselineTelemetry = ContextAndDataPointToDefinitionMap.Find(TPair<FString, FString>(Telemetry.Context, Telemetry.DataPoint));
 					if (BaselineTelemetry)
@@ -1237,12 +1245,12 @@ int32 USummarizeTraceCommandlet::Main(const FString& CmdLineParams)
 						Telemetry.Baseline = BaselineTelemetry->Measurement;
 
 						// let's only report on statistics that have an assigned threshold, to keep things concise
-						if (!RelatedStatistic->WarningThreshold.IsEmpty() || !RelatedStatistic->ErrorThreshold.IsEmpty())
+						if (!RelatedStatistic->BaselineWarningThreshold.IsEmpty() || !RelatedStatistic->BaselineErrorThreshold.IsEmpty())
 						{
 							// verify that this telemetry measurement is within the allowed threshold as defined in the current stats file
-							if (TelemetryDefinition::MeasurementWithinThreshold(Telemetry.Measurement, BaselineTelemetry->Measurement, RelatedStatistic->WarningThreshold))
+							if (TelemetryDefinition::MeasurementWithinThreshold(Telemetry.Measurement, BaselineTelemetry->Measurement, RelatedStatistic->BaselineWarningThreshold))
 							{
-								FString SignFlippedWarningThreshold = TelemetryDefinition::SignFlipThreshold(RelatedStatistic->WarningThreshold);
+								FString SignFlippedWarningThreshold = TelemetryDefinition::SignFlipThreshold(RelatedStatistic->BaselineWarningThreshold);
 
 								// check if it's beyond the threshold the other way and needs lowering in the stats csv
 								if (!TelemetryDefinition::MeasurementWithinThreshold(Telemetry.Measurement, BaselineTelemetry->Measurement, SignFlippedWarningThreshold))
@@ -1255,30 +1263,30 @@ int32 USummarizeTraceCommandlet::Main(const FString& CmdLineParams)
 
 									UE_LOG(LogSummarizeTrace, Warning, TEXT("Telemetry %s,%s,%s,%s significantly within baseline value %s using warning threshold %s. Please submit a new baseline to %s or adjust the threshold in %s."),
 										*Telemetry.TestName, *Telemetry.Context, *Telemetry.DataPoint, *Telemetry.Measurement,
-										*BaselineTelemetry->Measurement, *RelatedStatistic->WarningThreshold,
+										*BaselineTelemetry->Measurement, *RelatedStatistic->BaselineWarningThreshold,
 										*BaselineRelPath, *StatisticsRelPath);
 								}
 								else // it's within tolerance, just report that it's ok
 								{
 									UE_LOG(LogSummarizeTrace, Verbose, TEXT("Telemetry %s,%s,%s,%s within baseline value %s using warning threshold %s"),
 										*Telemetry.TestName, *Telemetry.Context, *Telemetry.DataPoint, *Telemetry.Measurement,
-										*BaselineTelemetry->Measurement, *RelatedStatistic->WarningThreshold);
+										*BaselineTelemetry->Measurement, *RelatedStatistic->BaselineWarningThreshold);
 								}
 							}
 							else
 							{
 								// it's outside warning threshold, check if it's inside the error threshold to just issue a warning
-								if (TelemetryDefinition::MeasurementWithinThreshold(Telemetry.Measurement, BaselineTelemetry->Measurement, RelatedStatistic->ErrorThreshold))
+								if (TelemetryDefinition::MeasurementWithinThreshold(Telemetry.Measurement, BaselineTelemetry->Measurement, RelatedStatistic->BaselineErrorThreshold))
 								{
 									UE_LOG(LogSummarizeTrace, Warning, TEXT("Telemetry %s,%s,%s,%s beyond baseline value %s using warning threshold %s. This could be a performance regression!"),
 										*Telemetry.TestName, *Telemetry.Context, *Telemetry.DataPoint, *Telemetry.Measurement,
-										*BaselineTelemetry->Measurement, *RelatedStatistic->WarningThreshold);
+										*BaselineTelemetry->Measurement, *RelatedStatistic->BaselineWarningThreshold);
 								}
 								else // it's outside the error threshold, hard error
 								{
 									UE_LOG(LogSummarizeTrace, Error, TEXT("Telemetry %s,%s,%s,%s beyond baseline value %s using error threshold %s. This could be a performance regression!"),
 										*Telemetry.TestName, *Telemetry.Context, *Telemetry.DataPoint, *Telemetry.Measurement,
-										*BaselineTelemetry->Measurement, *RelatedStatistic->ErrorThreshold);
+										*BaselineTelemetry->Measurement, *RelatedStatistic->BaselineErrorThreshold);
 								}
 							}
 						}
@@ -1296,11 +1304,11 @@ int32 USummarizeTraceCommandlet::Main(const FString& CmdLineParams)
 		if (TelemetryCsvHandle)
 		{
 			// no newline, see row printfs
-			WriteUTF8(TelemetryCsvHandle, FString::Printf(TEXT("TestName,Context,DataPoint,Measurement,Baseline,Unit,")));
+			WriteUTF8(TelemetryCsvHandle, FString::Printf(TEXT("TestName,Context,DataPoint,Unit,Measurement,Baseline,")));
 			for (const TelemetryDefinition& Telemetry : TelemetryData)
 			{
 				// note newline is at the front of every data line to prevent final extraneous newline, per customary for csv
-				WriteUTF8(TelemetryCsvHandle, FString::Printf(TEXT("\n%s,%s,%s,%s,%s,%s,"), *Telemetry.TestName, *Telemetry.Context, *Telemetry.DataPoint, *Telemetry.Measurement, *Telemetry.Baseline, *Telemetry.Unit));
+				WriteUTF8(TelemetryCsvHandle, FString::Printf(TEXT("\n%s,%s,%s,%s,%s,%s,"), *Telemetry.TestName, *Telemetry.Context, *Telemetry.DataPoint, *Telemetry.Unit, *Telemetry.Measurement, *Telemetry.Baseline));
 			}
 
 			TelemetryCsvHandle->Flush();
