Fix: Trying to delete an unused blueprint variable warns that the variable is in use if another blueprint with a variable of the same name references this blueprint.
Additional Fix: if a variable was referenced by the GetClassDefaults node in a *different* blueprint, deleting that variable didn't warn the user

diff --git a/Engine/Source/Editor/UnrealEd/Private/Kismet2/BlueprintEditorUtils.cpp b/Engine/Source/Editor/UnrealEd/Private/Kismet2/BlueprintEditorUtils.cpp
--- a/Engine/Source/Editor/UnrealEd/Private/Kismet2/BlueprintEditorUtils.cpp
+++ b/Engine/Source/Editor/UnrealEd/Private/Kismet2/BlueprintEditorUtils.cpp
@@ -5827,12 +5827,15 @@ bool FBlueprintEditorUtils::IsVariableComponent(const FBPVariableDescription& Va
 	return false;
 }
 
-bool FBlueprintEditorUtils::IsVariableUsed(const UBlueprint* InBlueprint, const FName& Name, UEdGraph* LocalGraphScope/* = nullptr*/)
+bool FBlueprintEditorUtils::IsVariableUsed(const UBlueprint* VariableBlueprint, const FName& VariableName, UEdGraph* LocalGraphScope/* = nullptr*/)
 {
-	auto CheckSingleBlueprint = [&Name, LocalGraphScope](const UBlueprint* Blueprint) -> bool
+	FGuid VariableGuid;
+	UBlueprint::GetGuidFromClassByFieldName<FProperty>(VariableBlueprint->GeneratedClass, VariableName, VariableGuid);
+	
+	auto CheckSingleBlueprint = [VariableGuid, VariableName, VariableBlueprint, LocalGraphScope](const UBlueprint* BlueprintToSearch) -> bool
 	{
 		TArray<UEdGraph*> AllGraphs;
-		Blueprint->GetAllGraphs(AllGraphs);
+		BlueprintToSearch->GetAllGraphs(AllGraphs);
 		for (TArray<UEdGraph*>::TConstIterator it(AllGraphs); it; ++it)
 		{
 			const UEdGraph* CurrentGraph = *it;
@@ -5844,7 +5847,7 @@ bool FBlueprintEditorUtils::IsVariableUsed(const UBlueprint* InBlueprint, const
 
 				for (const UK2Node_Variable* CurrentNode : GraphNodes)
 				{
-					if (Name == CurrentNode->GetVarName())
+					if (VariableGuid == CurrentNode->VariableReference.GetMemberGuid())
 					{
 						return true;
 					}
@@ -5855,9 +5858,9 @@ bool FBlueprintEditorUtils::IsVariableUsed(const UBlueprint* InBlueprint, const
 				CurrentGraph->GetNodesOfClass(ClassDefaultsNodes);
 				for (const UK2Node_GetClassDefaults* ClassDefaultsNode : ClassDefaultsNodes)
 				{
-					if (ClassDefaultsNode->GetInputClass() == Blueprint->SkeletonGeneratedClass)
+					if (ClassDefaultsNode->GetInputClass() == VariableBlueprint->SkeletonGeneratedClass)
 					{
-						const UEdGraphPin* VarPin = ClassDefaultsNode->FindPin(Name);
+						const UEdGraphPin* VarPin = ClassDefaultsNode->FindPin(VariableName);
 						if (VarPin && VarPin->Direction == EGPD_Output && VarPin->LinkedTo.Num() > 0)
 						{
 							return true;
@@ -5870,7 +5873,7 @@ bool FBlueprintEditorUtils::IsVariableUsed(const UBlueprint* InBlueprint, const
 		return false;
 	};
 
-	if (CheckSingleBlueprint(InBlueprint))
+	if (CheckSingleBlueprint(VariableBlueprint))
 	{
 		return true;
 	}
@@ -5879,7 +5882,7 @@ bool FBlueprintEditorUtils::IsVariableUsed(const UBlueprint* InBlueprint, const
 	{
 		FAssetRegistryModule& AssetRegistryModule = FModuleManager::LoadModuleChecked<FAssetRegistryModule>(TEXT("AssetRegistry"));
 		FARFilter Filter;
-		AssetRegistryModule.Get().GetReferencers(InBlueprint->GetPackage()->GetFName(), Filter.PackageNames, UE::AssetRegistry::EDependencyCategory::Package, UE::AssetRegistry::EDependencyQuery::Hard);
+		AssetRegistryModule.Get().GetReferencers(VariableBlueprint->GetPackage()->GetFName(), Filter.PackageNames, UE::AssetRegistry::EDependencyCategory::Package, UE::AssetRegistry::EDependencyQuery::Hard);
 		if (Filter.PackageNames.Num() > 0)
 		{
 			GWarn->BeginSlowTask(LOCTEXT("LoadingReferencerAssets", "Loading Referencers..."), true);
diff --git a/Engine/Source/Editor/UnrealEd/Public/Kismet2/BlueprintEditorUtils.h b/Engine/Source/Editor/UnrealEd/Public/Kismet2/BlueprintEditorUtils.h
--- a/Engine/Source/Editor/UnrealEd/Public/Kismet2/BlueprintEditorUtils.h
+++ b/Engine/Source/Editor/UnrealEd/Public/Kismet2/BlueprintEditorUtils.h
@@ -1313,7 +1313,7 @@ public:
 	static bool IsVariableComponent(const FBPVariableDescription& Variable);
 
 	/** Indicates if the variable is used on any graphs in this Blueprint*/
-	static bool IsVariableUsed(const UBlueprint* Blueprint, const FName& Name, UEdGraph* LocalGraphScope = nullptr);
+	static bool IsVariableUsed(const UBlueprint* VariableBlueprint, const FName& VariableName, UEdGraph* LocalGraphScope = nullptr);
 
 	/** 
 	 * Copies the value from the passed in string into a property. ContainerMem points to the Struct or Class containing Property 
