BulkDataRegistry: Implement asynchronous behavior for GetData and GetMeta.
#rnx

diff --git a/Engine/Source/Editor/UnrealEd/Private/EditorDomain/BulkDataRegistryEditorDomain.cpp b/Engine/Source/Editor/UnrealEd/Private/EditorDomain/BulkDataRegistryEditorDomain.cpp
--- a/Engine/Source/Editor/UnrealEd/Private/EditorDomain/BulkDataRegistryEditorDomain.cpp
+++ b/Engine/Source/Editor/UnrealEd/Private/EditorDomain/BulkDataRegistryEditorDomain.cpp
@@ -50,6 +50,7 @@ public:
 
 FBulkDataRegistryEditorDomain::FBulkDataRegistryEditorDomain()
 {
+	SharedDataLock = new FTaskSharedDataLock();
 	// We piggyback on the BulkDataRegistry hook to set the pointer to tunnel in the pointer to the EditorBuildInputResolver as well
 	SetGlobalBuildInputResolver(&UE::DerivedData::FEditorBuildInputResolver::Get());
 }
@@ -58,13 +59,35 @@ FBulkDataRegistryEditorDomain::~FBulkDataRegistryEditorDomain()
 {
 	SetGlobalBuildInputResolver(nullptr);
 
+	TMap<FGuid, FUpdatingPayload> LocalUpdatingPayloads;
 	{
+		FWriteScopeLock SharedDataScopeLock(SharedDataLock->ActiveLock);
+		FWriteScopeLock RegistryScopeLock(RegistryLock);
 		FScopeLock PendingPackageScopeLock(&PendingPackageLock);
-		bAllowCacheUpdates = false;
+
+		// Disable all activity that might come in from other threads
+		bActive = false;
+		SharedDataLock->bActive = false;
+
+		// Take custody of UpdatingPayloads
+		Swap(LocalUpdatingPayloads, UpdatingPayloads);
+	}
+
+	// Since the AsyncTasks can no longer call the Requesters, we have to do it
+	for (TPair<FGuid, FUpdatingPayload>& Pair : LocalUpdatingPayloads)
+	{
+		for (TUniqueFunction<void(bool, const FCompressedBuffer&)>& Requester : Pair.Value.Requesters)
+		{
+			Requester(false, FCompressedBuffer());
+		}
 	}
-	PendingPackages.Empty();
+	LocalUpdatingPayloads.Empty();
 
-	FCoreDelegates::OnPostEngineInit.RemoveAll(this);
+	// Clear PendingPackages
+	{
+		FScopeLock PendingPackageScopeLock(&PendingPackageLock);
+		PendingPackages.Empty();
+	}
 }
 
 void FBulkDataRegistryEditorDomain::Register(UPackage* Owner, const UE::Virtualization::FVirtualizedUntypedBulkData& BulkData)
@@ -87,14 +110,16 @@ void FBulkDataRegistryEditorDomain::Register(UPackage* Owner, const UE::Virtuali
 
 	FName PackageNameToUpdate = CopyBulk.HasPlaceholderPayloadId() ? PackageName : NAME_None;
 	FWriteScopeLock RegistryScopeLock(RegistryLock);
-	Registry.Add(BulkData.GetIdentifier(), FBulkSource(MoveTemp(CopyBulk), PackageNameToUpdate));
+	check(bActive); // Registrations should not come in after we destruct
+	Registry.Add(BulkData.GetIdentifier(), FRegisteredBulk(MoveTemp(CopyBulk), PackageNameToUpdate));
 }
 
 void FBulkDataRegistryEditorDomain::OnExitMemory(const UE::Virtualization::FVirtualizedUntypedBulkData& BulkData)
 {
 	const FGuid& Key = BulkData.GetIdentifier();
 	FWriteScopeLock RegistryScopeLock(RegistryLock);
-	FBulkSource* Existing = Registry.Find(Key);
+	check(bActive); // Deregistrations should not come in after we destruct
+	FRegisteredBulk* Existing = Registry.Find(Key);
 	if (Existing)
 	{
 		if (Existing->BulkData.IsMemoryOnlyPayload())
@@ -106,129 +131,134 @@ void FBulkDataRegistryEditorDomain::OnExitMemory(const UE::Virtualization::FVirt
 
 TFuture<UE::BulkDataRegistry::FMetaData> FBulkDataRegistryEditorDomain::GetMeta(const FGuid& BulkDataId)
 {
-	TPromise< UE::BulkDataRegistry::FMetaData> Promise;
-	UE::BulkDataRegistry::FMetaData Result;
-	Result.bValid = false;
-	auto TryExistingCompleteResult = [&BulkDataId, this, &Result](FBulkSource*& OutExisting)
+	bool bIsWriteLock = false;
+	FRWScopeLock RegistryScopeLock(RegistryLock, SLT_ReadOnly);
+	for (;;)
 	{
-		OutExisting = Registry.Find(BulkDataId);
-		if (!OutExisting)
+		FRegisteredBulk* Existing = nullptr;
+		if (bActive)
 		{
-			Result.bValid = false;
-			Result.RawHash = FIoHash();
-			Result.RawSize = 0;
-			return true;
+			Existing = Registry.Find(BulkDataId);
 		}
-
-		if (!OutExisting->BulkData.HasPlaceholderPayloadId())
+		if (!Existing)
 		{
-			Result.bValid = true;
-			Result.RawHash = OutExisting->BulkData.GetPayloadId().GetIdentifier();
-			Result.RawSize = OutExisting->BulkData.GetPayloadSize();
-			return true;
+			TPromise<UE::BulkDataRegistry::FMetaData> Promise;
+			Promise.SetValue(UE::BulkDataRegistry::FMetaData{ false, FIoHash(), 0 });
+			return Promise.GetFuture();
 		}
 
-		return false;
-	};
-
-	TOptional<UE::Virtualization::FVirtualizedUntypedBulkData> CopyBulk;
-	FName PackageNameToUpdate;
-	{
-		FReadScopeLock RegistryScopeLock(RegistryLock);
-		FBulkSource* Existing = nullptr;
-		if (TryExistingCompleteResult(Existing))
+		UE::Virtualization::FVirtualizedUntypedBulkData& BulkData = Existing->BulkData;
+		if (!BulkData.HasPlaceholderPayloadId())
 		{
-			Promise.SetValue(MoveTemp(Result));
+			TPromise<UE::BulkDataRegistry::FMetaData> Promise;
+			Promise.SetValue(UE::BulkDataRegistry::FMetaData{ true, BulkData.GetPayloadId().GetIdentifier(), static_cast<uint64>(BulkData.GetPayloadSize()) });
 			return Promise.GetFuture();
 		}
 
-		CopyBulk.Emplace(Existing->BulkData);
-		PackageNameToUpdate = Existing->PackageNameToUpdate;
-	}
-
-	CopyBulk->UpdatePayloadId();
-	Result.bValid = true;
-	Result.RawHash= CopyBulk->GetPayloadId().GetIdentifier();
-	Result.RawSize = CopyBulk->GetPayloadSize();
-
-	{
-		FWriteScopeLock RegistryScopeLock(RegistryLock);
-		FBulkSource* Existing;
-		if (TryExistingCompleteResult(Existing))
+		if (!bIsWriteLock)
 		{
-			// Some other thread has beaten us to updating the RawHash
-			Promise.SetValue(MoveTemp(Result));
-			return Promise.GetFuture();
+			bIsWriteLock = true;
+			RegistryScopeLock.ReleaseReadOnlyLockAndAcquireWriteLock_USE_WITH_CAUTION();
+			continue;
 		}
 
-		Existing->BulkData = *CopyBulk;
-		Existing->PackageNameToUpdate = NAME_None;
-
-		AddTempLoadedPayload(CopyBulk->GetIdentifier(), CopyBulk->GetPayloadSize());
-		PruneTempLoadedPayloads();
-	}
-
-	if (!PackageNameToUpdate.IsNone())
-	{
-		AddPendingPackageBulkData(PackageNameToUpdate, *CopyBulk);
+		// The payload in the registry is missing its RawHash; start a thread to calculate it and subscribe our caller to the results
+		FUpdatingPayload& UpdatingPayload = UpdatingPayloads.FindOrAdd(BulkDataId);
+		if (!UpdatingPayload.AsyncTask)
+		{
+			UpdatingPayload.AsyncTask = new FAutoDeleteAsyncTask<FUpdatePayloadWorker>(this, BulkData);
+			UpdatingPayload.AsyncTask->StartBackgroundTask();
+		}
+		TPromise<UE::BulkDataRegistry::FMetaData> Promise;
+		TFuture<UE::BulkDataRegistry::FMetaData> Future = Promise.GetFuture();
+		UpdatingPayload.Requesters.Add([Promise = MoveTemp(Promise)](bool bValid, const FCompressedBuffer& Buffer) mutable
+			{
+				Promise.SetValue(UE::BulkDataRegistry::FMetaData{ bValid, Buffer.GetRawHash(), Buffer.GetRawSize() });
+			});
+		return Future;
 	}
-
-	Promise.SetValue(MoveTemp(Result));
-	return Promise.GetFuture();
 }
 
 TFuture<UE::BulkDataRegistry::FData> FBulkDataRegistryEditorDomain::GetData(const FGuid& BulkDataId)
 {
-	TPromise<UE::BulkDataRegistry::FData> Result;
 	TOptional<UE::Virtualization::FVirtualizedUntypedBulkData> CopyBulk;
-	FName PackageNameToUpdate;
-	bool bNeedsUpdate = false;
 	{
-		FReadScopeLock RegistryScopeLock(RegistryLock);
-		FBulkSource* Existing = Registry.Find(BulkDataId);
-		if (!Existing)
+		bool bIsWriteLock = false;
+		FRWScopeLock RegistryScopeLock(RegistryLock, SLT_ReadOnly);
+		for (;;)
 		{
-			Result.SetValue(UE::BulkDataRegistry::FData{ false, FCompressedBuffer() });
-			return Result.GetFuture();
-		}
-		CopyBulk.Emplace(Existing->BulkData);
-		PackageNameToUpdate = Existing->PackageNameToUpdate;
-		bNeedsUpdate = CopyBulk->HasPlaceholderPayloadId();
-	}
+			FRegisteredBulk* Existing = nullptr;
+			if (bActive)
+			{
+				Existing = Registry.Find(BulkDataId);
+			}
+			if (!Existing)
+			{
+				TPromise<UE::BulkDataRegistry::FData> Result;
+				Result.SetValue(UE::BulkDataRegistry::FData{ false, FCompressedBuffer() });
+				return Result.GetFuture();
+			}
 
-	if (bNeedsUpdate)
-	{
-		CopyBulk->UpdatePayloadId();
-	}
-	FCompressedBuffer Payload = CopyBulk->GetCompressedPayload().Get();
-	bool bUpdatedByGetPayload = true; // BULKDATAREGISTRY_TODO: Calculate this and return it from GetPayload
-	bool bChanged = bNeedsUpdate || bUpdatedByGetPayload;
-	CopyBulk->UnloadData();
+			if (!Existing->BulkData.HasPlaceholderPayloadId() && !Existing->bHasTempPayload)
+			{
+				// The contract of FVirtualizedUntypedBulkData does not guarantee that GetCompressedPayload() is a quick operation (it may load the data
+				// synchronously), so copy the BulkData into a temporary and call it outside the lock
+				CopyBulk.Emplace(Existing->BulkData);
+				break;
+			}
 
-	if (bChanged)
-	{
-		FWriteScopeLock RegistryScopeLock(RegistryLock);
-		FBulkSource* Existing = Registry.Find(BulkDataId);
-		if (Existing)
-		{
-			bNeedsUpdate = Existing->BulkData.HasPlaceholderPayloadId();
-			Existing->BulkData = *CopyBulk;
-			Existing->PackageNameToUpdate = NAME_None;
-		}
-	}
+			if (!bIsWriteLock)
+			{
+				bIsWriteLock = true;
+				RegistryScopeLock.ReleaseReadOnlyLockAndAcquireWriteLock_USE_WITH_CAUTION();
+				continue;
+			}
 
-	if (bNeedsUpdate && !PackageNameToUpdate.IsNone())
-	{
-		AddPendingPackageBulkData(PackageNameToUpdate, *CopyBulk);
+			if (!Existing->BulkData.HasPlaceholderPayloadId())
+			{
+				check(Existing->bHasTempPayload);
+				// We are the first GetData call after the BulkData previously loaded its Payload to calculate the RawHash.
+				// Sidenote, this means GetCompressedPayload will be fast.
+				// But we also have the responsibility to dump the data from memory since we have now consumed it.
+				// Make sure we copy the data pointer before dumping it from the Registry version!
+				CopyBulk.Emplace(Existing->BulkData);
+				Existing->BulkData.UnloadData();
+				Existing->bHasTempPayload = false;
+				break;
+			}
+
+			// The payload in the registry is missing its RawHash, and we calculate that on demand whenever the data is requested,
+			// which is now. Instead of only returning the data to our caller, we load the data and use it to update the RawHash
+			// in the registry and then return the data to our caller.
+			FUpdatingPayload& UpdatingPayload = UpdatingPayloads.FindOrAdd(BulkDataId);
+			if (!UpdatingPayload.AsyncTask)
+			{
+				UpdatingPayload.AsyncTask = new FAutoDeleteAsyncTask<FUpdatePayloadWorker>(this, Existing->BulkData);
+			}
+			TPromise<UE::BulkDataRegistry::FData> Promise;
+			TFuture<UE::BulkDataRegistry::FData> Future = Promise.GetFuture();
+			UpdatingPayload.Requesters.Add([Promise=MoveTemp(Promise)](bool bValid, const FCompressedBuffer& Buffer) mutable
+				{
+					Promise.SetValue(UE::BulkDataRegistry::FData{ bValid, Buffer });
+				});
+			return Future;
+		}
 	}
 
-	Result.SetValue(UE::BulkDataRegistry::FData{ true, MoveTemp(Payload) });
-	return Result.GetFuture();
+	// We are calling a function that returns a TFuture on the stack-local CopyBulk, which would cause a read-after-free if the asynchronous TFuture could
+	// read from the BulkData. However, the contract of FVirtualizedUntypedBulkData guarantees that the TFuture gets a copy of all data it needs and
+	// does not read from the BulkData after returning from GetCompressedPayload, so a read-after-free is not possible.
+	return CopyBulk->GetCompressedPayload().Next([](FCompressedBuffer Payload)
+		{
+			return UE::BulkDataRegistry::FData{ true, Payload };
+		});
 }
 
 void FBulkDataRegistryEditorDomain::TickCook(float DeltaTime, bool bTickComplete)
 {
 	bool bWaitForCooldown = !bTickComplete;
+	check(bActive); // Ticks should not come in after we destruct
+
 	PollPendingPackages(bWaitForCooldown);
 	{
 		FWriteScopeLock RegistryScopeLock(RegistryLock);
@@ -244,6 +274,12 @@ void FBulkDataRegistryEditorDomain::Tick(float DeltaTime)
 void FBulkDataRegistryEditorDomain::AddPendingPackageBulkData(FName PackageName, const UE::Virtualization::FVirtualizedUntypedBulkData& BulkData)
 {
 	FScopeLock PendingPackageScopeLock(&PendingPackageLock);
+	AddPendingPackageBulkDataInsideLock(PackageName, BulkData);
+}
+
+void FBulkDataRegistryEditorDomain::AddPendingPackageBulkDataInsideLock(FName PackageName, const UE::Virtualization::FVirtualizedUntypedBulkData& BulkData)
+{
+	check(bActive); // Registrations should not come in after we destruct, and AsyncTasks should check bActive before calling
 	TUniquePtr<FPendingPackage>& PendingPackage = PendingPackages.FindOrAdd(PackageName);
 	if (!PendingPackage.IsValid())
 	{
@@ -309,11 +345,12 @@ void FBulkDataRegistryEditorDomain::PruneTempLoadedPayloads()
 				|| TempLoadedPayloads[0].EndTime <= CurrentTime))
 		{
 			FTempLoadedPayload Payload = TempLoadedPayloads.PopFrontValue();
-			FBulkSource* Existing = Registry.Find(Payload.Guid);
+			FRegisteredBulk* Existing = Registry.Find(Payload.Guid);
 			if (Existing)
 			{
 				// UnloadData only unloads the in-memory data, and only if the BulkData can be reloaded from disk
 				Existing->BulkData.UnloadData();
+				Existing->bHasTempPayload = false;
 			}
 			TempLoadedPayloadsSize -= Payload.PayloadSize;
 		}
@@ -354,18 +391,18 @@ void FPendingPackage::OnCacheResults(FSharedBuffer Buffer)
 
 	// Add each CachedBulkData to the Registry. If the CachedBulkData exists, do not overwrite it, but do update the RawHash if it is missing it.
 	FWriteScopeLock RegistryScopeLock(Owner->RegistryLock);
-	if (!Owner->bAllowCacheUpdates)
+	if (!Owner->bActive)
 	{
 		return;
 	}
 	for (const UE::Virtualization::FVirtualizedUntypedBulkData& BulkData : CachedBulkDatas)
 	{
-		FBulkSource& TargetBulkSource = Owner->Registry.FindOrAdd(BulkData.GetIdentifier());
-		UE::Virtualization::FVirtualizedUntypedBulkData& TargetBulkData = TargetBulkSource.BulkData;
+		FRegisteredBulk& TargetRegisteredBulk = Owner->Registry.FindOrAdd(BulkData.GetIdentifier());
+		UE::Virtualization::FVirtualizedUntypedBulkData& TargetBulkData = TargetRegisteredBulk.BulkData;
 		if (!TargetBulkData.GetIdentifier().IsValid())
 		{
 			TargetBulkData = BulkData;
-			TargetBulkSource.PackageNameToUpdate = BulkData.HasPlaceholderPayloadId() ? PackageName : NAME_None;
+			TargetRegisteredBulk.PackageNameToUpdate = BulkData.HasPlaceholderPayloadId() ? PackageName : NAME_None;
 		}
 		else
 		{
@@ -376,7 +413,7 @@ void FPendingPackage::OnCacheResults(FSharedBuffer Buffer)
 				// && TargetBulkData.PayloadInfoMatches(BulkData)
 				TargetBulkData = BulkData;
 				// We now know that we don't need to update the package with the RawHash, so remove the PackageNameToUpdate
-				TargetBulkSource.PackageNameToUpdate = NAME_None;
+				TargetRegisteredBulk.PackageNameToUpdate = NAME_None;
 			}
 		}
 	}
@@ -391,10 +428,10 @@ void FPendingPackage::UpdateCache()
 {
 	CacheRequest.Wait();
 
-	// Remove any duplicates in the runtime BulkDatas; elements later in the list override were added after and override earlier elements
+	// Remove any duplicates in the runtime BulkDatas; elements later in the list were added after and override earlier elements
 	{
 		TSet<FGuid> BulkDataGuids;
-		for (int32 Index = BulkDataGuids.Num(); Index >= 0; --Index)
+		for (int32 Index = BulkDatas.Num()-1; Index >= 0; --Index)
 		{
 			bool bAlreadyExists;
 			BulkDataGuids.Add(BulkDatas[Index].GetIdentifier(), &bAlreadyExists);
@@ -452,13 +489,104 @@ void FPendingPackage::UpdateCache()
 
 		TArray<uint8> Bytes;
 		FMemoryWriter Writer(Bytes);
-		Serialize(Writer, BulkDatas);
+		Serialize(Writer, CachedBulkDatas);
 		// BULKDATAREGISTRY_TODO: Store the request in a list of pending writes so that we don't try to write a new entry for the package
 		// until the previous entry has finished storing, to avoid stale data overwriting new data
 		UE::EditorDomain::PutBulkDataList(PackageName, MakeSharedBufferFromArray(MoveTemp(Bytes)));
 	}
 }
 
+FUpdatePayloadWorker::FUpdatePayloadWorker(FBulkDataRegistryEditorDomain* InBulkDataRegistry,
+	const UE::Virtualization::FVirtualizedUntypedBulkData& InSourceBulk)
+	: BulkData(InSourceBulk)
+	, BulkDataRegistry(InBulkDataRegistry)
+{
+	SharedDataLock = InBulkDataRegistry->SharedDataLock;
+}
+
+void FUpdatePayloadWorker::DoWork()
+{
+	FUpdatingPayload LocalUpdatingPayload;
+	FCompressedBuffer Buffer;
+	bool bValid = true;
+	for (;;)
+	{
+		BulkData.UpdatePayloadId();
+		Buffer = BulkData.GetCompressedPayload().Get();
+
+		{
+			FReadScopeLock SharedDataScopeLock(SharedDataLock->ActiveLock);
+			if (!SharedDataLock->bActive)
+			{
+				// The BulkDataRegistry has destructed. Our list of requesters is on the BulkDataRegistry, so there's nothing we can do except exit
+				return;
+			}
+			FWriteScopeLock RegistryScopeLock(BulkDataRegistry->RegistryLock);
+
+			if (!BulkDataRegistry->UpdatingPayloads.RemoveAndCopyValue(BulkData.GetIdentifier(), LocalUpdatingPayload))
+			{
+				// The updating payload might not exist in the case of the Registry shutting down; it will clear the UpdatingPayloads to cancel our action
+				// Return canceled (which we treat the same as failed) to our requesters
+				bValid = false;
+				break;
+			}
+			check(BulkDataRegistry->bActive); // Only set to false at the same time as SharedDataLock->bActive
+
+			FRegisteredBulk* RegisteredBulk = BulkDataRegistry->Registry.Find(BulkData.GetIdentifier());
+			if (!RegisteredBulk)
+			{
+				// Some agent has deregistered the BulkData before we finished calculating its payload
+				// return failure to our requesters
+				bValid = false;
+				break;
+			}
+
+			// BULKDATAREGISTRY_TODO: Implement LocationMatches
+			if (false) // !BulkData.LocationMatches(BulkData));
+			{
+				// Some caller has assigned a new BulkData. We need to abandon the BulkData we just loaded and give our callers the
+				// information about the new one
+				check(RegisteredBulk->BulkData.GetIdentifier() == BulkData.GetIdentifier()); // The identifier in the BulkData should match the key for that BulkData in Registry
+				BulkData = RegisteredBulk->BulkData;
+				// Add our LocalUpdatingPayload back to UpdatingPayloads; we removed it because we thought we were done.
+				BulkDataRegistry->UpdatingPayloads.Add(BulkData.GetIdentifier(), MoveTemp(LocalUpdatingPayload));
+				continue;
+			}
+
+			// Store the new payload in the Registry's entry for the BulkData; new MetaData requests will no longer need to wait for it
+			RegisteredBulk->BulkData = BulkData;
+
+			// Mark that the next GetData call should remove the temporary payload
+			RegisteredBulk->bHasTempPayload = true;
+			BulkDataRegistry->AddTempLoadedPayload(BulkData.GetIdentifier(), BulkData.GetPayloadSize());
+			BulkDataRegistry->PruneTempLoadedPayloads();
+
+			// If the BulkData is marked to be stored in a package's BulkDataList, add an update task for the package.
+			// We have to enter the two locks at the same time, because it is only while inside the RegistryScopeLock that we are guaranteed our work has not been canceled.
+			if (!RegisteredBulk->PackageNameToUpdate.IsNone())
+			{
+				FScopeLock PendingPackageScopeLock(&BulkDataRegistry->PendingPackageLock);
+				BulkDataRegistry->AddPendingPackageBulkDataInsideLock(RegisteredBulk->PackageNameToUpdate, BulkData);
+				if (!BulkData.HasPlaceholderPayloadId())
+				{
+					// If the BulkData already has its PayloadId, then we no longer need to update it after adding it this time
+					RegisteredBulk->PackageNameToUpdate = NAME_None;
+				}
+			}
+			break;
+		}
+	}
+
+	if (!bValid)
+	{
+		Buffer.Reset();
+	}
+	for (TUniqueFunction<void(bool, const FCompressedBuffer&)>& Requester : LocalUpdatingPayload.Requesters)
+	{
+		Requester(bValid, Buffer);
+	}
+}
+
 void Serialize(FArchive& Ar, TArray<UE::Virtualization::FVirtualizedUntypedBulkData>& Datas)
 {
 	int32 Num = Datas.Num();
diff --git a/Engine/Source/Editor/UnrealEd/Private/EditorDomain/BulkDataRegistryEditorDomain.h b/Engine/Source/Editor/UnrealEd/Private/EditorDomain/BulkDataRegistryEditorDomain.h
--- a/Engine/Source/Editor/UnrealEd/Private/EditorDomain/BulkDataRegistryEditorDomain.h
+++ b/Engine/Source/Editor/UnrealEd/Private/EditorDomain/BulkDataRegistryEditorDomain.h
@@ -2,12 +2,15 @@
 
 #pragma once
 
+#include "Async/AsyncWork.h"
+#include "Async/Future.h"
 #include "Containers/Array.h"
 #include "Containers/Map.h"
 #include "Containers/RingBuffer.h"
 #include "DerivedDataRequest.h"
 #include "HAL/CriticalSection.h"
 #include "Serialization/BulkDataRegistry.h"
+#include "Templates/RefCounting.h"
 #include "TickableEditorObject.h"
 #include "Virtualization/VirtualizedBulkData.h"
 
@@ -21,16 +24,17 @@ namespace UE::BulkDataRegistry::Private
 class FBulkDataRegistryEditorDomain;
 
 /** Struct for storage of a BulkData in the registry, including the BulkData itself and data about cache status. */
-struct FBulkSource
+struct FRegisteredBulk
 {
-	FBulkSource() = default;
-	FBulkSource(UE::Virtualization::FVirtualizedUntypedBulkData&& InBulkData, FName InPackageNameToUpdate = NAME_None)
+	FRegisteredBulk() = default;
+	FRegisteredBulk(UE::Virtualization::FVirtualizedUntypedBulkData&& InBulkData, FName InPackageNameToUpdate = NAME_None)
 		:BulkData(MoveTemp(InBulkData)), PackageNameToUpdate(InPackageNameToUpdate)
 	{
 	}
 
 	UE::Virtualization::FVirtualizedUntypedBulkData BulkData;
 	FName PackageNameToUpdate;
+	bool bHasTempPayload = false;
 };
 
 /** Serialize an array of BulkDatas into or out of bytes saved/load from the registry's persistent cache. */
@@ -75,6 +79,42 @@ struct FTempLoadedPayload
 	double EndTime;
 };
 
+/** An Active flag and a lock around it for informing AutoDeleteAsyncTasks that their shared data is no longer available. */
+class FTaskSharedDataLock : public FThreadSafeRefCountedObject
+{
+public:
+	FRWLock ActiveLock;
+	bool bActive = true;
+};
+
+/** A worker that updates the PayloadId for a BulkData that is missing its RawHash. */
+class FUpdatePayloadWorker : public FNonAbandonableTask
+{
+public:
+	FUpdatePayloadWorker(FBulkDataRegistryEditorDomain* InBulkDataRegistry,
+		const UE::Virtualization::FVirtualizedUntypedBulkData& InSourceBulk);
+
+	void DoWork();
+	TStatId GetStatId() const
+	{
+		RETURN_QUICK_DECLARE_CYCLE_STAT(FUpdatePayloadWorker, STATGROUP_ThreadPoolAsyncTasks);
+	}
+
+private:
+	UE::Virtualization::FVirtualizedUntypedBulkData BulkData;
+	TRefCountPtr<FTaskSharedDataLock> SharedDataLock;
+	FBulkDataRegistryEditorDomain* BulkDataRegistry;
+};
+
+/** Data storage for the FUpdatePayloadWorker that is updated while in flight for additional requesters. */
+struct FUpdatingPayload
+{
+	/** Pointer to the task. Once set to non-null, is never modified. The task ensure this FUpdatingPayload is destroyed before the task destructs. */
+	FAutoDeleteAsyncTask<FUpdatePayloadWorker>* AsyncTask = nullptr;
+	TArray<TUniqueFunction<void(bool bValid, const FCompressedBuffer& Buffer)>> Requesters;
+};
+
+
 /** Implementation of a BulkDataRegistry that stores its persistent data in a DDC bucket. */
 class FBulkDataRegistryEditorDomain : public IBulkDataRegistry, public FTickableEditorObject, public FTickableCookObject
 {
@@ -96,18 +136,24 @@ public:
 
 private:
 	void AddPendingPackageBulkData(FName PackageName, const UE::Virtualization::FVirtualizedUntypedBulkData& BulkData);
+	void AddPendingPackageBulkDataInsideLock(FName PackageName, const UE::Virtualization::FVirtualizedUntypedBulkData& BulkData);
 	void PollPendingPackages(bool bWaitForCooldown);
 	void AddTempLoadedPayload(const FGuid& RegistryKey, uint64 PayloadSize);
 	void PruneTempLoadedPayloads();
 	friend class FPendingPackage;
+	friend class FUpdatePayloadWorker;
 
-	FCriticalSection PendingPackageLock;
+	// All three locks can  be held at the same time. They must always be entered in order: SharedDataLock, RegistryLock, PendingPackageLock.
+	TRefCountPtr<FTaskSharedDataLock> SharedDataLock;
 	FRWLock RegistryLock;
-	TMap<FGuid, FBulkSource> Registry;
+	FCriticalSection PendingPackageLock;
+
+	TMap<FGuid, FRegisteredBulk> Registry;
+	TMap<FGuid, FUpdatingPayload> UpdatingPayloads;
 	TMap<FName, TUniquePtr<FPendingPackage>> PendingPackages;
 	TRingBuffer<FTempLoadedPayload> TempLoadedPayloads;
 	uint64 TempLoadedPayloadsSize = 0;
-	bool bAllowCacheUpdates = true;
+	bool bActive = true;
 };
 
 } // namespace UE::BulkDataRegistry::Private
