OFPA: Fix Replace Actor not reusing same actor package

diff --git a/Engine/Source/Editor/UnrealEd/Private/EditorEngine.cpp b/Engine/Source/Editor/UnrealEd/Private/EditorEngine.cpp
--- a/Engine/Source/Editor/UnrealEd/Private/EditorEngine.cpp
+++ b/Engine/Source/Editor/UnrealEd/Private/EditorEngine.cpp
@@ -4999,12 +4999,22 @@ void UEditorEngine::ReplaceActors(UActorFactory* Factory, const FAssetData& Asse
 
 		const FName OldActorName = OldActor->GetFName();
 		const FName OldActorReplacedNamed = MakeUniqueObjectName(OldActor->GetOuter(), OldActor->GetClass(), *FString::Printf(TEXT("%s_REPLACED"), *OldActorName.ToString()));
-		OldActor->Rename(*OldActorReplacedNamed.ToString());
+		
+		FActorSpawnParameters SpawnParams;
+		SpawnParams.Name = OldActorName;
+		SpawnParams.bCreateActorPackage = false;
+		SpawnParams.OverridePackage = OldActor->GetExternalPackage();
+		SpawnParams.OverrideActorGuid = OldActor->GetActorGuid();
+				
+		// Don't go through AActor::Rename here because we aren't changing outers (the actor's level) and we also don't want to reset loaders
+		// if the actor is using an external package. We really just want to rename that actor out of the way so we can spawn the new one in
+		// the exact same package, keeping the package name intact.
+		OldActor->UObject::Rename(*OldActorReplacedNamed.ToString(), OldActor->GetOuter(), REN_DoNotDirty | REN_DontCreateRedirectors | REN_ForceNoResetLoaders);
 
 		const FTransform OldTransform = OldActor->ActorToWorld();
 
 		// create the actor
-		NewActor = Factory->CreateActor(Asset, Level, OldTransform);
+		NewActor = Factory->CreateActor(Asset, Level, OldTransform, SpawnParams);
 		// For blueprints, try to copy over properties
 		if (Factory->IsA(UActorFactoryBlueprint::StaticClass()))
 		{
@@ -5023,8 +5033,6 @@ void UEditorEngine::ReplaceActors(UActorFactory* Factory, const FAssetData& Asse
 
 		if (NewActor)
 		{
-			NewActor->Rename(*OldActorName.ToString());
-
 			// The new actor might not have a root component
 			USceneComponent* const NewActorRootComponent = NewActor->GetRootComponent();
 			if(NewActorRootComponent)
@@ -5114,7 +5122,7 @@ void UEditorEngine::ReplaceActors(UActorFactory* Factory, const FAssetData& Asse
 		else
 		{
 			// If creating the new Actor failed, put the old Actor's name back
-			OldActor->Rename(*OldActorName.ToString());
+			OldActor->UObject::Rename(*OldActorName.ToString(), OldActor->GetOuter(), REN_DoNotDirty | REN_DontCreateRedirectors | REN_ForceNoResetLoaders);
 		}
 	}
 
