WorldPartitionHLODsBuilder
* Added -ResumeBuild=, resume work at the provided index
* Added -DumpStats, outputs a csv file with lots of infos (tri count, texture size, disk size, etc)
* Always overwrite files in working dir

diff --git a/Engine/Source/Editor/UnrealEd/Private/WorldPartition/WorldPartitionHLODsBuilder.cpp b/Engine/Source/Editor/UnrealEd/Private/WorldPartition/WorldPartitionHLODsBuilder.cpp
--- a/Engine/Source/Editor/UnrealEd/Private/WorldPartition/WorldPartitionHLODsBuilder.cpp
+++ b/Engine/Source/Editor/UnrealEd/Private/WorldPartition/WorldPartitionHLODsBuilder.cpp
@@ -23,6 +23,15 @@
 #include "WorldPartition/WorldPartitionHelpers.h"
 #include "WorldPartition/HLOD/HLODActor.h"
 #include "WorldPartition/HLOD/HLODActorDesc.h"
+#include "WorldPartition/HLOD/HLODLayer.h"
+
+#include "Engine/StaticMesh.h"
+#include "Materials/Material.h"
+#include "Materials/MaterialInstance.h"
+
+#include "Components/InstancedStaticMeshComponent.h"
+
+#include "AssetCompilingManager.h"
 
 DEFINE_LOG_CATEGORY_STATIC(LogWorldPartitionHLODsBuilder, All, All);
 
@@ -154,10 +163,12 @@ UWorldPartitionHLODsBuilder::UWorldPartitionHLODsBuilder(const FObjectInitialize
 	bBuildHLODs = FParse::Param(FCommandLine::Get(), TEXT("BuildHLODs"));
 	bDeleteHLODs = FParse::Param(FCommandLine::Get(), TEXT("DeleteHLODs"));
 	bSubmitHLODs = FParse::Param(FCommandLine::Get(), TEXT("SubmitHLODs"));
+	bDumpStats = FParse::Param(FCommandLine::Get(), TEXT("DumpStats"));
 
-	bSingleBuildStep = !bSetupHLODs && !bBuildHLODs && !bSubmitHLODs && !bDeleteHLODs;
+	bSingleBuildStep = !bSetupHLODs && !bBuildHLODs && !bSubmitHLODs && !bDeleteHLODs && !bDumpStats;
 
 	bAutoSubmit = FParse::Param(FCommandLine::Get(), TEXT("AutoSubmit"));
+	bResumeBuild = FParse::Value(FCommandLine::Get(), TEXT("ResumeBuild="), ResumeBuildIndex);
 
 	bDistributedBuild = FParse::Param(FCommandLine::Get(), TEXT("DistributedBuild"));
 
@@ -180,7 +191,7 @@ bool UWorldPartitionHLODsBuilder::RequiresCommandletRendering() const
 {
 	// Commandlet requires rendering only for building HLODs
 	// Building will occur either if -BuildHLODs is provided or no explicit step arguments are provided
-	return bBuildHLODs || !(bSetupHLODs || bDeleteHLODs || bSubmitHLODs);
+	return bBuildHLODs || bSingleBuildStep;
 }
 
 bool UWorldPartitionHLODsBuilder::ValidateParams() const
@@ -303,6 +314,11 @@ bool UWorldPartitionHLODsBuilder::RunInternal(UWorld* World, const FBox& Bounds,
 		bRet = SubmitHLODActors();
 	}
 
+	if (bRet && bDumpStats)
+	{
+		bRet = DumpStats();
+	}
+
 	WorldPartition = nullptr;
 	delete SourceControlHelper;
 
@@ -426,16 +442,20 @@ bool UWorldPartitionHLODsBuilder::BuildHLODActors()
 	}
 
 	UE_LOG(LogWorldPartitionHLODsBuilder, Display, TEXT("#### Building %d HLOD actors ####"), HLODActorsToBuild.Num());
+	if (bResumeBuild)
+	{
+		UE_LOG(LogWorldPartitionHLODsBuilder, Display, TEXT("#### Resuming build at %d ####"), ResumeBuildIndex);
+	}
 
-	int32 CurrentActor = 0;
-	for (const FGuid& HLODActorGuid : HLODActorsToBuild)
+	for (int32 CurrentActor = ResumeBuildIndex; CurrentActor < HLODActorsToBuild.Num(); ++CurrentActor)
 	{
+		const FGuid& HLODActorGuid = HLODActorsToBuild[CurrentActor];
+
 		FWorldPartitionReference ActorRef(WorldPartition, HLODActorGuid);
 		FWorldPartitionActorDesc* ActorDesc = ActorRef.Get();
 
 		AWorldPartitionHLOD* HLODActor = CastChecked<AWorldPartitionHLOD>(ActorDesc->GetActor());
 
-		CurrentActor++;
 		UE_LOG(LogWorldPartitionHLODsBuilder, Display, TEXT("    [%d/%d] Building HLOD actor %s..."), CurrentActor, HLODActorsToBuild.Num(), *HLODActor->GetActorLabel());
 
 		HLODActor->BuildHLOD();
@@ -565,6 +585,150 @@ bool UWorldPartitionHLODsBuilder::SubmitHLODActors()
 	return bRet;
 }
 
+bool UWorldPartitionHLODsBuilder::DumpStats()
+{
+	struct FHLODStat
+	{
+		FString				Name;
+		FName				RuntimeGrid;
+		EActorGridPlacement GridPlacement = EActorGridPlacement::None;
+		uint64				GridLocationZ = 0;
+		uint64				GridLocationX = 0;
+		uint64				GridLocationY = 0;
+		FString				DataLayers;
+		EHLODLayerType		HLODType;
+
+		uint64				InstanceCount = 0;
+		uint64				NaniteTriangleCount = 0;
+		uint64				NaniteVertexCount = 0;
+		uint64				TriangleCount = 0;
+		uint64				VertexCount = 0;
+		uint64				UVChannelCount = 0;
+
+		uint32				BaseColorTextureSize = 0;
+		uint32				NormalTextureSize = 0;
+		uint32				MRSTextureSize = 0;
+		uint32				EmissiveTextureSize = 0;
+
+		uint64				MeshResourceSize = 0;
+		uint64				TexturesResourceSize = 0;
+
+		uint64				DiskSize = 0;
+
+		static FString GetHeaderCSVString()
+		{
+			return TEXT("Name, RuntimeGrid, GridPlacement, GridLocationZ, GridLocationX, GridLocationY, DataLayers, HLODType, InstanceCount, NaniteTriangleCount, NaniteVertexCount, TriangleCount, VertexCount, UVChannelCount, BaseColorTextureSize, NormalTextureSize, MRSTextureSize, EmissiveTextureSize, MeshResourceSize, TexturesResourceSize, DiskSize");
+		}
+
+		FString ToCSVString() const
+		{
+			return FString::Printf(TEXT("%s, %s, %s, %d, %d, %d, %s, %s, %d, %d, %d, %d, %d, %d, %d, %d, %d, %d, %d, %d, %d"), 
+				*Name, 
+				*RuntimeGrid.ToString(),
+				*StaticEnum<EActorGridPlacement>()->GetNameStringByValue((int64)GridPlacement),
+				GridLocationZ, GridLocationX, GridLocationY,
+				*DataLayers,
+				*StaticEnum<EHLODLayerType>()->GetNameStringByValue((int64)HLODType),
+				InstanceCount, NaniteTriangleCount, NaniteVertexCount, TriangleCount, VertexCount, UVChannelCount,
+				BaseColorTextureSize, NormalTextureSize, MRSTextureSize, EmissiveTextureSize, 
+				MeshResourceSize, TexturesResourceSize,
+				DiskSize);
+		}
+	};
+
+	TArray<FHLODStat> HLODStats;
+
+	for (UActorDescContainer::TIterator<AWorldPartitionHLOD> HLODIterator(WorldPartition); HLODIterator; ++HLODIterator)
+	{
+		FWorldPartitionReference HLODActorRef(WorldPartition, HLODIterator->GetGuid());
+		AWorldPartitionHLOD* HLODActor = CastChecked<AWorldPartitionHLOD>(HLODActorRef->GetActor());
+
+		FAssetCompilingManager::Get().FinishAllCompilation();
+
+		FHLODStat& HLODStat = HLODStats.Emplace_GetRef();
+
+		HLODStat.Name = HLODActor->GetActorLabel();
+		HLODStat.RuntimeGrid = HLODActor->GetRuntimeGrid();
+		HLODStat.GridPlacement = HLODActor->GetGridPlacement();
+		HLODActor->GetGridIndices(HLODStat.GridLocationX, HLODStat.GridLocationY, HLODStat.GridLocationZ);
+		HLODStat.DataLayers = *FString::JoinBy(HLODActor->GetDataLayerObjects(), TEXT("| "), [](const UDataLayer* DataLayer) { return DataLayer->GetDataLayerLabel().ToString(); });
+		HLODStat.HLODType = HLODActor->GetSubActorsHLODLayer()->GetLayerType();
+
+		HLODActor->ForEachComponent<UInstancedStaticMeshComponent>(false, [&](const UInstancedStaticMeshComponent* ISMC)
+		{
+			HLODStat.InstanceCount += ISMC->GetInstanceCount();
+		});
+
+		ForEachObjectWithPackage(HLODActor->GetPackage(), [&HLODStat](UObject* Object)
+		{
+			if (UStaticMesh* StaticMesh = Cast<UStaticMesh>(Object))
+			{
+				HLODStat.MeshResourceSize += StaticMesh->GetResourceSizeBytes(EResourceSizeMode::Exclusive);
+
+				HLODStat.TriangleCount += StaticMesh->GetNumTriangles(0);
+				HLODStat.VertexCount += StaticMesh->GetNumVertices(0);
+				HLODStat.UVChannelCount += StaticMesh->GetNumTexCoords(0);
+
+				HLODStat.NaniteTriangleCount += StaticMesh->GetNumNaniteTriangles();
+				HLODStat.NaniteVertexCount += StaticMesh->GetNumNaniteVertices();
+			}
+
+			if (UMaterialInstance* MaterialInstance = Cast<UMaterialInstance>(Object))
+			{
+				auto GetTexture = [MaterialInstance, &HLODStat](const FName TextureParamName)
+				{
+					UTexture* Texture = nullptr;
+					MaterialInstance->GetTextureParameterValue(TextureParamName, Texture, true);
+					HLODStat.TexturesResourceSize += Texture ? Texture->GetResourceSizeBytes(EResourceSizeMode::Exclusive) : 0;
+					return Texture;
+				};
+
+				UTexture* BaseColorTex = GetTexture("BaseColorTexture");
+				UTexture* NormalTex = GetTexture("NormalTexture");
+				UTexture* MetallicTex = GetTexture("MetallicTexture");
+				UTexture* RoughnessTex = GetTexture("RoughnessTexture");
+				UTexture* SpecularTex = GetTexture("SpecularTexture");
+				UTexture* EmissiveTex = GetTexture("EmissiveTexture");
+				UTexture* EmissiveHDRTex = GetTexture("EmissiveHDRTexture");
+				UTexture* PackedMRSTex = GetTexture("PackedTexture");
+				
+				HLODStat.BaseColorTextureSize = BaseColorTex ? BaseColorTex->GetSurfaceWidth() : 0;
+				HLODStat.NormalTextureSize = NormalTex ? NormalTex->GetSurfaceWidth() : 0;
+				HLODStat.MRSTextureSize = FMath::Max(HLODStat.MRSTextureSize, MetallicTex ? MetallicTex->GetSurfaceWidth() : 0);
+				HLODStat.MRSTextureSize = FMath::Max(HLODStat.MRSTextureSize, RoughnessTex ? RoughnessTex->GetSurfaceWidth() : 0);
+				HLODStat.MRSTextureSize = FMath::Max(HLODStat.MRSTextureSize, SpecularTex ? SpecularTex->GetSurfaceWidth() : 0);
+				HLODStat.MRSTextureSize = FMath::Max(HLODStat.MRSTextureSize, PackedMRSTex ? PackedMRSTex->GetSurfaceWidth() : 0);
+				HLODStat.EmissiveTextureSize = FMath::Max(HLODStat.EmissiveTextureSize, EmissiveTex ? EmissiveTex->GetSurfaceWidth() : 0);
+				HLODStat.EmissiveTextureSize = FMath::Max(HLODStat.EmissiveTextureSize, EmissiveHDRTex ? EmissiveHDRTex->GetSurfaceWidth() : 0);
+			}
+
+			return true;
+		}, false);
+
+		HLODStat.DiskSize = IFileManager::Get().FileSize(*HLODActor->GetPackage()->GetLoadedPath().GetLocalFullPath());
+	}
+
+	Algo::Sort(HLODStats, [](const FHLODStat& A, const FHLODStat& B) { return A.Name < B.Name; });
+
+
+	const FString Filename = FString::Printf(TEXT("%s/HLODStats/HLODStats-%s.csv"), *FPaths::ProjectLogDir(), *FDateTime::Now().ToString());
+	FArchive* CSVFile = IFileManager::Get().CreateFileWriter(*Filename, FILEWRITE_AllowRead);
+
+	FString CSVHeader = FHLODStat::GetHeaderCSVString();
+	FString Newline = TEXT("\n");
+
+	*CSVFile << CSVHeader << Newline;
+	for (const FHLODStat& HLODStat : HLODStats)
+	{
+		FString CSVString = HLODStat.ToCSVString();
+		*CSVFile << CSVString << Newline;
+	}
+
+	delete CSVFile;
+
+	return true;
+}
+
 bool UWorldPartitionHLODsBuilder::GetHLODActorsToBuild(TArray<FGuid>& HLODActorsToBuild) const
 {
 	bool bRet = true;
@@ -832,7 +996,9 @@ bool UWorldPartitionHLODsBuilder::CopyFilesToWorkingDir(const FString& TargetDir
 
 		if (FileAction != FileAction_Delete)
 		{
-			bool bRet = IFileManager::Get().Copy(*TargetFilename, *SourceFilename, false) == COPY_OK;
+			const bool bReplace = true;
+			const bool bEvenIfReadOnly = true;
+			bool bRet = IFileManager::Get().Copy(*TargetFilename, *SourceFilename, bReplace, bEvenIfReadOnly) == COPY_OK;
 			if (!bRet)
 			{
 				UE_LOG(LogWorldPartitionHLODsBuilder, Error, TEXT("Failed to copy file from \"%s\" to \"%s\""), *SourceFilename, *TargetFilename);
@@ -945,86 +1111,89 @@ bool UWorldPartitionHLODsBuilder::CopyFilesFromWorkingDir(const FString& SourceD
 		}
 	}
 
-	// Add
-	if (!FilesToAdd.IsEmpty())
-	{
-		bRet = CopyFromWorkingDir(FilesToAdd);
-		if (!bRet)
-		{
-			return false;
-		}
+	TArray<FString> ToAdd;
+	FilesToAdd.GetKeys(ToAdd);
 
-		TArray<FString> ToAdd;
-		FilesToAdd.GetKeys(ToAdd);
+	TArray<FString> ToEdit;
+	FilesToEdit.GetKeys(ToEdit);
 
-		if (ISourceControlModule::Get().IsEnabled())
-		{
-			bRet = USourceControlHelpers::MarkFilesForAdd(ToAdd);
+	// When resuming a build (after a crash for example) we don't need to perform any file operation as these modification were done in the first run.
+	if (!bResumeBuild)
+	{
+	    // Add
+	    if (!FilesToAdd.IsEmpty())
+	    {
+			bRet = CopyFromWorkingDir(FilesToAdd);
 			if (!bRet)
 			{
-				UE_LOG(LogWorldPartitionHLODsBuilder, Error, TEXT("Adding files to source control failed: %s"), *USourceControlHelpers::LastErrorMsg().ToString());
 				return false;
 			}
-		}
 
-		ModifiedFiles.Append(FHLODModifiedFiles::EFileOperation::FileAdded, ToAdd);
-	}
-
-	// Delete
-	if (!FilesToDelete.IsEmpty())
-	{
-		if (ISourceControlModule::Get().IsEnabled())
-		{
-			bRet = USourceControlHelpers::MarkFilesForDelete(FilesToDelete);
-			if (!bRet)
+			if (ISourceControlModule::Get().IsEnabled())
 			{
-				UE_LOG(LogWorldPartitionHLODsBuilder, Error, TEXT("Deleting files from source control failed: %s"), *USourceControlHelpers::LastErrorMsg().ToString());
-				return false;
+				bRet = USourceControlHelpers::MarkFilesForAdd(ToAdd);
+				if (!bRet)
+				{
+					UE_LOG(LogWorldPartitionHLODsBuilder, Error, TEXT("Adding files to source control failed: %s"), *USourceControlHelpers::LastErrorMsg().ToString());
+					return false;
+				}
 			}
-		}
-		else
-		{
-			for (const FString& FileToDelete : FilesToDelete)
+	    }
+    
+	    // Delete
+	    if (!FilesToDelete.IsEmpty())
+	    {
+			if (ISourceControlModule::Get().IsEnabled())
 			{
-				const bool bRequireExists = false;
-				const bool bEvenIfReadOnly = true;
-				bRet = IFileManager::Get().Delete(*FileToDelete, bRequireExists, bEvenIfReadOnly);
+				bRet = USourceControlHelpers::MarkFilesForDelete(FilesToDelete);
 				if (!bRet)
 				{
-					UE_LOG(LogWorldPartitionHLODsBuilder, Error, TEXT("Failed to delete file from disk: %s"), *USourceControlHelpers::LastErrorMsg().ToString());
+					UE_LOG(LogWorldPartitionHLODsBuilder, Error, TEXT("Deleting files from source control failed: %s"), *USourceControlHelpers::LastErrorMsg().ToString());
 					return false;
 				}
 			}
-		}
-
-		ModifiedFiles.Append(FHLODModifiedFiles::EFileOperation::FileDeleted, FilesToDelete);
-	}
-
-	// Edit
-	if (!FilesToEdit.IsEmpty())
-	{
-		TArray<FString> ToEdit;
-		FilesToEdit.GetKeys(ToEdit);
-
-		if (ISourceControlModule::Get().IsEnabled())
-		{
-			bRet = USourceControlHelpers::CheckOutFiles(ToEdit);
+			else
+			{
+				for (const FString& FileToDelete : FilesToDelete)
+				{
+					const bool bRequireExists = false;
+					const bool bEvenIfReadOnly = true;
+					bRet = IFileManager::Get().Delete(*FileToDelete, bRequireExists, bEvenIfReadOnly);
+					if (!bRet)
+					{
+						UE_LOG(LogWorldPartitionHLODsBuilder, Error, TEXT("Failed to delete file from disk: %s"), *USourceControlHelpers::LastErrorMsg().ToString());
+						return false;
+					}
+				}
+			}
+	    }
+    
+	    // Edit
+	    if (!FilesToEdit.IsEmpty())
+	    {
+		    if (ISourceControlModule::Get().IsEnabled())
+		    {
+				bRet = USourceControlHelpers::CheckOutFiles(ToEdit);
+				if (!bRet)
+				{
+					UE_LOG(LogWorldPartitionHLODsBuilder, Error, TEXT("Checking out files from source control failed: %s"), *USourceControlHelpers::LastErrorMsg().ToString());
+					return false;
+				}
+			}
+		
+			bRet = CopyFromWorkingDir(FilesToEdit);
 			if (!bRet)
 			{
-				UE_LOG(LogWorldPartitionHLODsBuilder, Error, TEXT("Checking out files from source control failed: %s"), *USourceControlHelpers::LastErrorMsg().ToString());
 				return false;
 			}
-		}
-
-		ModifiedFiles.Append(FHLODModifiedFiles::EFileOperation::FileEdited, ToEdit);
-
-		bRet = CopyFromWorkingDir(FilesToEdit);
-		if (!bRet)
-		{
-			return false;
-		}
+	    }
 	}
 
+	// Keep track of all modified files
+	ModifiedFiles.Append(FHLODModifiedFiles::EFileOperation::FileAdded, ToAdd);
+	ModifiedFiles.Append(FHLODModifiedFiles::EFileOperation::FileDeleted, FilesToDelete);
+	ModifiedFiles.Append(FHLODModifiedFiles::EFileOperation::FileEdited, ToEdit);
+
 	// Force a rescan of the updated files
 	IAssetRegistry& AssetRegistry = FModuleManager::LoadModuleChecked<FAssetRegistryModule>(TEXT("AssetRegistry")).Get();
 	AssetRegistry.ScanModifiedAssetFiles(ModifiedFiles.GetAllFiles());
diff --git a/Engine/Source/Editor/UnrealEd/Public/WorldPartition/WorldPartitionHLODsBuilder.h b/Engine/Source/Editor/UnrealEd/Public/WorldPartition/WorldPartitionHLODsBuilder.h
--- a/Engine/Source/Editor/UnrealEd/Public/WorldPartition/WorldPartitionHLODsBuilder.h
+++ b/Engine/Source/Editor/UnrealEd/Public/WorldPartition/WorldPartitionHLODsBuilder.h
@@ -78,6 +78,7 @@ protected:
 	bool BuildHLODActors();
 	bool DeleteHLODActors();
 	bool SubmitHLODActors();
+	bool DumpStats();
 
 	bool GenerateBuildManifest(TMap<FString, int32>& FilesToBuilderMap) const;
 	bool GetHLODActorsToBuild(TArray<FGuid>& HLODActorsToBuild) const;
@@ -97,12 +98,16 @@ private:
 	bool bBuildHLODs;
 	bool bDeleteHLODs;
 	bool bSubmitHLODs;
+	bool bDumpStats;
 	bool bSingleBuildStep;
 	bool bAutoSubmit;
 	bool bDistributedBuild;
 	FString BuildManifest;
 	int32 BuilderIdx;
 	int32 BuilderCount;
+	bool bResumeBuild;
+	int32 ResumeBuildIndex;
+
 
 	const FString DistributedBuildWorkingDir;
 	const FString DistributedBuildManifest;
