ExternalActorsCommandlet: remove unused code.

diff --git a/Engine/Source/Editor/UnrealEd/Private/Commandlets/ExternalActorsCommandlet.cpp b/Engine/Source/Editor/UnrealEd/Private/Commandlets/ExternalActorsCommandlet.cpp
--- a/Engine/Source/Editor/UnrealEd/Private/Commandlets/ExternalActorsCommandlet.cpp
+++ b/Engine/Source/Editor/UnrealEd/Private/Commandlets/ExternalActorsCommandlet.cpp
@@ -81,8 +81,6 @@ int32 UExternalActorsCommandlet::Main(const FString& Params)
 					AActor* PotentialMainPackageActor = nullptr;
 
 					const FString PackageName = FPackageName::FilenameToLongPackageName(*Filename);
-					const FWorldPartitionActorDesc* ActorDesc = nullptr;
-
 					if (UPackage* Package = LoadPackage(nullptr, *Filename, LOAD_None, nullptr, nullptr))
 					{
 						ForEachObjectWithPackage(Package, [&MainPackageActor, &PotentialMainPackageActor](UObject* Object)
