Fix focus loss searching for parent anim classes when creating anim blueprints

diff --git a/Engine/Source/Editor/UnrealEd/Private/Factories/AnimBlueprintFactory.cpp b/Engine/Source/Editor/UnrealEd/Private/Factories/AnimBlueprintFactory.cpp
--- a/Engine/Source/Editor/UnrealEd/Private/Factories/AnimBlueprintFactory.cpp
+++ b/Engine/Source/Editor/UnrealEd/Private/Factories/AnimBlueprintFactory.cpp
@@ -304,7 +304,8 @@ private:
 		};
 		AssetPickerConfig.bShowPathInColumnView = false;
 		AssetPickerConfig.bShowTypeInColumnView = false;
-
+		AssetPickerConfig.bFocusSearchBoxWhenOpened = false;
+		
 		SkeletonContainer->ClearChildren();
 		SkeletonContainer->AddSlot()
 		[
