Reduce the default size of spline component tangents

diff --git a/Engine/Source/Editor/UnrealEd/Classes/Settings/LevelEditorViewportSettings.h b/Engine/Source/Editor/UnrealEd/Classes/Settings/LevelEditorViewportSettings.h
--- a/Engine/Source/Editor/UnrealEd/Classes/Settings/LevelEditorViewportSettings.h
+++ b/Engine/Source/Editor/UnrealEd/Classes/Settings/LevelEditorViewportSettings.h
@@ -506,19 +506,19 @@ public:
 
 	/** The size adjustment to apply to selected spline points (in screen space units). */
 	UPROPERTY(EditAnywhere, config, Category = LookAndFeel, AdvancedDisplay, meta = (ClampMin = "-5.00", ClampMax = "20.00"))
-	float SelectedSplinePointSizeAdjustment;
+	float SelectedSplinePointSizeAdjustment = 0.0f;
 
 	/** The size adjustment to apply to spline line thickness which increases the spline's hit tolerance. */
 	UPROPERTY(EditAnywhere, config, Category = LookAndFeel, AdvancedDisplay, meta = (ClampMin = "0.00"))
-	float SplineLineThicknessAdjustment;
+	float SplineLineThicknessAdjustment = 0.0f;
 
 	/** The size adjustment to apply to spline tangent handle (in screen space units). */
 	UPROPERTY(EditAnywhere, config, Category = LookAndFeel, AdvancedDisplay, meta = (ClampMin = "-5.00", ClampMax = "20.00"))
-	float SplineTangentHandleSizeAdjustment;
+	float SplineTangentHandleSizeAdjustment = 0.0f;
 
 	/** The scale to apply to spline tangent lengths */
 	UPROPERTY(EditAnywhere, config, Category = LookAndFeel, AdvancedDisplay, meta = (ClampMin = "0.00"))
-	float SplineTangentScale;
+	float SplineTangentScale = 0.5f;
 	
 	UPROPERTY(config)
 	FVector2D LastInViewportMenuLocation;
diff --git a/Engine/Source/Editor/UnrealEd/Private/Settings/SettingsClasses.cpp b/Engine/Source/Editor/UnrealEd/Private/Settings/SettingsClasses.cpp
--- a/Engine/Source/Editor/UnrealEd/Private/Settings/SettingsClasses.cpp
+++ b/Engine/Source/Editor/UnrealEd/Private/Settings/SettingsClasses.cpp
@@ -902,10 +902,6 @@ ULevelEditorViewportSettings::ULevelEditorViewportSettings( const FObjectInitial
 	bLevelStreamingVolumePrevis = false;
 	BillboardScale = 1.0f;
 	TransformWidgetSizeAdjustment = 0.0f;
-	SelectedSplinePointSizeAdjustment = 0.0f;
-	SplineLineThicknessAdjustment = 0.0f;
-	SplineTangentHandleSizeAdjustment = 0.0f;
-	SplineTangentScale = 1.0f;
 	MeasuringToolUnits = MeasureUnits_Centimeters;
 	bAllowArcballRotate = false;
 	bAllowScreenRotate = false;
