Async EditorBuildInputResolver
#rnx

diff --git a/Engine/Source/Editor/UnrealEd/Private/EditorDomain/EditorBuildInputResolver.cpp b/Engine/Source/Editor/UnrealEd/Private/EditorDomain/EditorBuildInputResolver.cpp
--- a/Engine/Source/Editor/UnrealEd/Private/EditorDomain/EditorBuildInputResolver.cpp
+++ b/Engine/Source/Editor/UnrealEd/Private/EditorDomain/EditorBuildInputResolver.cpp
@@ -2,11 +2,114 @@
 
 #include "EditorBuildInputResolver.h"
 
+#include "Async/Future.h"
+#include <atomic>
+#include "Containers/ArrayView.h"
+#include "HAL/Event.h"
 #include "Serialization/BulkDataRegistry.h"
+#include "Templates/Tuple.h"
 
 namespace UE::DerivedData
 {
 
+/**
+ * Takes a collection of TFutures that calculate either ResolveInputMeta or ResolveInputData callback outputs,
+ * and calls the Resolve callback after all the futures complete
+ */
+template <class BulkPayloadType, class CallbackPayloadType, class CallbackType>
+class TResolveInputRequest : public IRequest, public FThreadSafeRefCountedObject
+{
+public:
+	typedef TResolveInputRequest<BulkPayloadType, CallbackPayloadType, CallbackType> ThisType;
+
+	TResolveInputRequest(const FBuildDefinition& InDefinition, CallbackType&& InOnResolved);
+
+	// IRequest interface
+	virtual void SetPriority(EPriority Priority) override
+	{
+	}
+	virtual void Cancel() override;
+	virtual void Wait() override;
+	virtual bool Poll() override;
+	virtual void AddRef() const override
+	{
+		FThreadSafeRefCountedObject::AddRef();
+	}
+	virtual void Release() const override
+	{
+		FThreadSafeRefCountedObject::Release();
+	}
+
+	/** Stores the full set of Payload Futures and completes the creation of this IRequest */
+	void SetPayloads(TArrayView<TTuple<FString, TFuture<BulkPayloadType>>> BulkPayloads, EStatus OtherStatus);
+
+private:
+	/** Call the OnResolve function if it has not been canceled. */
+	void CallOnResolved();
+
+	/**
+	 * Move BulkData's fields from the BulkDataRegistry's format into the BuildInputResolver format.
+	 * Implementation varies depending on whether the template type is for MetaData or Data.
+	*/
+	void AssignBulkPayload(CallbackPayloadType& OutPayload, BulkPayloadType& InPayload);
+	/**
+	 * Log that a BulkData has failed to resolve.
+	 * The text of the message varies depending on whether the template type is for MetaData or Data.
+	*/
+	void LogBulkDataError(FStringView Key);
+
+	/** Holds the key and TFuture handle for one of the Inputs being resolved. */
+	struct FPayloadData
+	{
+		FString Key;
+		TFuture<void> Future;
+	};
+
+private:
+	/**
+	 * The callback that was passed into the IBuildInputResolver function; called when all TFutures are complete.
+	 * Can only be read/written after construction while holding the CancelLock.
+	 */
+	CallbackType OnResolved;
+
+	// Data written during creation that is read-only during Cancel and CallOnResolved
+	FBuildDefinition Definition;
+#if DO_CHECK
+	bool bCreationComplete = false;
+#endif
+
+	/**
+	 * The array of data that is written by the TFutures and passed into the IBuildInputResolver callback function.
+	 * This array is allocated during creation before any Futures are created.
+	 * The elements are write-only, and only by the given Future, until RemainingInputCount is 0.
+	 * After RemainingInputCount reaches 0, it is read/writable by whichever thread received the 0.
+	 */
+	TArray<CallbackPayloadType> Payloads;
+	/**
+	 * Extra data needed by this class about each of the payloads.
+	 * This array is allocated during creation and is read-only afterwards.
+	 * It is readable only by the public api or by any thread after RemainingInputCount reaches 0.
+	 */
+	TArray<FPayloadData> PayloadDatas;
+
+	/**
+	 * Tracks whether all TFutures are complete; the last one to complete calls the Resolve callback.
+	 * Value starts at 1 and is decremented when creation completes, so that if all TFutures are already complete
+	 * during creation, the Resolve callback will be called when creation completes instead.
+	 */
+	std::atomic<int32> RemainingInputCount{ 1 };
+	/** Accumulated success result that is set to false if any of the TFutures fail. */
+	std::atomic<bool> bAllSucceeded{ true };
+	/** Atomic to allow Cancel or CallOnResolved to resolve the race to see which one will call the callback. */
+	std::atomic<bool> bCallbackCalled{ false };
+	/** Event for when the callback has finished executing; some functions must not return until it is done. */
+	FEventRef CallbackComplete;
+
+};
+typedef TResolveInputRequest<UE::BulkDataRegistry::FMetaData, FBuildInputMetaByKey, FOnBuildInputMetaResolved> FResolveInputMetaRequest;
+typedef TResolveInputRequest<UE::BulkDataRegistry::FData, FBuildInputDataByKey, FOnBuildInputDataResolved> FResolveInputDataRequest;
+
+
 FEditorBuildInputResolver& FEditorBuildInputResolver::Get()
 {
 	static FEditorBuildInputResolver Singleton;
@@ -23,109 +126,86 @@ FRequest FEditorBuildInputResolver::ResolveKey(const FBuildKey& Key, FOnBuildKey
 FRequest FEditorBuildInputResolver::ResolveInputMeta(const FBuildDefinition& Definition, EPriority Priority,
 	FOnBuildInputMetaResolved&& OnResolved)
 {
-	EStatus Status = EStatus::Ok;
-	TArray<FString> InputKeys;
-	TArray<FBuildInputMetaByKey> Inputs;
-
-	Definition.IterateInputBuilds([&Status, &Definition](FStringView Key, const FBuildPayloadKey& PayloadKey)
+	EStatus OtherStatus = EStatus::Ok;
+	Definition.IterateInputBuilds([&OtherStatus, &Definition](FStringView Key, const FBuildPayloadKey& PayloadKey)
 		{
 			UE_LOG(LogCore, Error, TEXT("FEditorBuildInputResolver::ResolveInputMeta: resolving InputBuilds is not yet implemented. Context=%.*s, Key=%.*s"),
 				Definition.GetName().Len(), Definition.GetName().GetData(), Key.Len(), Key.GetData());
-			Status = EStatus::Error;
+			OtherStatus = EStatus::Error;
 		});
 
 	/** Visits every input bulk data in order by key. */
-	Definition.IterateInputBulkData([&Status, &InputKeys, &Inputs, &Definition](FStringView Key, const FGuid& BulkDataId)
+	TArray<TTuple<FString, TFuture<UE::BulkDataRegistry::FMetaData>>, TInlineAllocator<4>> BulkPayloads;
+	Definition.IterateInputBulkData([&BulkPayloads](FStringView Key, const FGuid& BulkDataId)
 		{
-			FBuildInputMetaByKey& Input = Inputs.Emplace_GetRef();
-			// TODO: MakeAsync: Create an IRequest that has an array of TFutures and calls .Then to populate the Inputs
-			TFuture<UE::BulkDataRegistry::FMetaData> Future = IBulkDataRegistry::Get().GetMeta(BulkDataId);
-			const UE::BulkDataRegistry::FMetaData& Result = Future.Get();
-			if (Result.bValid)
-			{
-				Input.Key = InputKeys.Emplace_GetRef(Key);
-				Input.RawHash = Result.RawHash;
-				Input.RawSize = Result.RawSize;
-			}
-			else
-			{
-				Inputs.Pop();
-				UE_LOG(LogCore, Error, TEXT("Failed to ResolveInputMetaData. Context=%.*s, Key=%.*s"),
-					Definition.GetName().Len(), Definition.GetName().GetData(), Key.Len(), Key.GetData())
-				Status = EStatus::Error;
-			}
+			BulkPayloads.Add(MakeTuple(FString(Key), IBulkDataRegistry::Get().GetMeta(BulkDataId)));
 		});
 
-	Definition.IterateInputFiles([&Status, &Definition](FStringView Key, FStringView Path)
+	Definition.IterateInputFiles([&OtherStatus, &Definition](FStringView Key, FStringView Path)
 		{
 			UE_LOG(LogCore, Error, TEXT("FEditorBuildInputResolver::ResolveInputMeta: resolving InputFiles is not yet implemented. Context=%.*s, Key=%.*s"),
 				Definition.GetName().Len(), Definition.GetName().GetData(), Key.Len(), Key.GetData());
-			Status = EStatus::Error;
+			OtherStatus = EStatus::Error;
 		});
 
-	Definition.IterateInputHashes([&Status, &Definition](FStringView Key, const FIoHash& RawHash)
+	Definition.IterateInputHashes([&OtherStatus, &Definition](FStringView Key, const FIoHash& RawHash)
 		{
 			UE_LOG(LogCore, Error, TEXT("FEditorBuildInputResolver::ResolveInputMeta: resolving InputHashes is not yet implemented. Context=%.*s, Key=%.*s"),
 				Definition.GetName().Len(), Definition.GetName().GetData(), Key.Len(), Key.GetData());
-			Status = EStatus::Error;
+			OtherStatus = EStatus::Error;
 		});
 
-	OnResolved({ Inputs, Status });
-	return FRequest();
+	TRequest Request(new FResolveInputMetaRequest(Definition, MoveTemp(OnResolved)));
+	Request->SetPayloads(BulkPayloads, OtherStatus);
+	if (Priority == EPriority::Blocking)
+	{
+		Request->Wait();
+	}
+	return Request;
 }
 
 FRequest FEditorBuildInputResolver::ResolveInputData(const FBuildDefinition& Definition, EPriority Priority,
 	FOnBuildInputDataResolved&& OnResolved, FBuildInputFilter&& Filter)
 {
-	EStatus Status = EStatus::Ok;
-	TArray<FString> InputKeys;
-	TArray<FBuildInputDataByKey> Inputs;
-
-	Definition.IterateInputBuilds([&Status, &Definition](FStringView Key, const FBuildPayloadKey& PayloadKey)
+	EStatus OtherStatus = EStatus::Ok;
+	Definition.IterateInputBuilds([&OtherStatus, &Definition](FStringView Key, const FBuildPayloadKey& PayloadKey)
 		{
 			UE_LOG(LogCore, Error, TEXT("FEditorBuildInputResolver::ResolveInputData resolving InputBuilds is not yet implemented. Context=%.*s, Key=%.*s"),
 				Definition.GetName().Len(), Definition.GetName().GetData(), Key.Len(), Key.GetData());
-			Status = EStatus::Error;
+			OtherStatus = EStatus::Error;
 		});
 
 	/** Visits every input bulk data in order by key. */
-	Definition.IterateInputBulkData([&Filter, &Status, &InputKeys, &Inputs, &Definition](FStringView Key, const FGuid& BulkDataId)
+	TArray<TTuple<FString, TFuture<UE::BulkDataRegistry::FData>>, TInlineAllocator<4>> BulkPayloads;
+	Definition.IterateInputBulkData([&BulkPayloads, &Filter](FStringView Key, const FGuid& BulkDataId)
 		{
 			if (!Filter || Filter(Key))
 			{
-				// TODO: MakeAsync: Create an IRequest that has an array of TFutures and calls .Then to populate the Inputs
-				TFuture<UE::BulkDataRegistry::FData> Future = IBulkDataRegistry::Get().GetData(BulkDataId);
-				const UE::BulkDataRegistry::FData& Result = Future.Get();
-				if (Result.bValid)
-				{
-					InputKeys.Emplace(Key);
-					Inputs.Add({ InputKeys.Last(), Result.Buffer });
-				}
-				else
-				{
-					UE_LOG(LogCore, Error, TEXT("Failed to ResolveInputData. Context=%.*s, Key=%.*s"),
-						Definition.GetName().Len(), Definition.GetName().GetData(), Key.Len(), Key.GetData())
-					Status = EStatus::Error;
-				}
+				BulkPayloads.Add(MakeTuple(FString(Key), IBulkDataRegistry::Get().GetData(BulkDataId)));
 			}
 		});
 
-	Definition.IterateInputFiles([&Status, &Definition](FStringView Key, FStringView Path)
+	Definition.IterateInputFiles([&OtherStatus, &Definition](FStringView Key, FStringView Path)
 		{
 			UE_LOG(LogCore, Error, TEXT("FEditorBuildInputResolver::ResolveInputData: resolving InputFiles is not yet implemented. Context=%.*s, Key=%.*s"),
 				Definition.GetName().Len(), Definition.GetName().GetData(), Key.Len(), Key.GetData());
-			Status = EStatus::Error;
+			OtherStatus = EStatus::Error;
 		});
 
-	Definition.IterateInputHashes([&Status, &Definition](FStringView Key, const FIoHash& RawHash)
+	Definition.IterateInputHashes([&OtherStatus, &Definition](FStringView Key, const FIoHash& RawHash)
 		{
 			UE_LOG(LogCore, Error, TEXT("FEditorBuildInputResolver::ResolveInputData: resolving InputHashes is not yet implemented. Context=%.*s, Key=%.*s"),
 				Definition.GetName().Len(), Definition.GetName().GetData(), Key.Len(), Key.GetData());
-			Status = EStatus::Error;
+			OtherStatus = EStatus::Error;
 		});
 
-	OnResolved({ Inputs, Status });
-	return FRequest();
+	TRequest Request(new FResolveInputDataRequest(Definition, MoveTemp(OnResolved)));
+	Request->SetPayloads(BulkPayloads, OtherStatus);
+	if (Priority == EPriority::Blocking)
+	{
+		Request->Wait();
+	}
+	return Request;
 }
 
 FRequest FEditorBuildInputResolver::ResolveInputData(const FBuildAction& Action, EPriority Priority,
@@ -136,4 +216,161 @@ FRequest FEditorBuildInputResolver::ResolveInputData(const FBuildAction& Action,
 	return FRequest();
 }
 
+
+template <class BulkPayloadType, class CallbackPayloadType, class CallbackType>
+TResolveInputRequest<BulkPayloadType, CallbackPayloadType, CallbackType>::TResolveInputRequest(
+	const FBuildDefinition& InDefinition, CallbackType&& InOnResolved)
+	: OnResolved(MoveTemp(InOnResolved))
+	, Definition(InDefinition)
+	, CallbackComplete(EEventMode::ManualReset)
+{
+}
+
+template <class BulkPayloadType, class CallbackPayloadType, class CallbackType>
+void TResolveInputRequest<BulkPayloadType, CallbackPayloadType, CallbackType>::Cancel()
+{
+	check(bCreationComplete);
+	if (bCallbackCalled.exchange(true) == false)
+	{
+		TArray<CallbackPayloadType> EmptyPayloads;
+		OnResolved({ EmptyPayloads, EStatus::Canceled });
+		CallbackComplete->Trigger();
+	}
+	else
+	{
+		CallbackComplete->Wait();
+	}
+}
+
+template <class BulkPayloadType, class CallbackPayloadType, class CallbackType>
+void TResolveInputRequest<BulkPayloadType, CallbackPayloadType, CallbackType>::Wait()
+{
+	check(bCreationComplete);
+	if (!bCallbackCalled)
+	{
+		for (FPayloadData& PayloadData : PayloadDatas)
+		{
+			PayloadData.Future.Wait();
+		}
+	}
+	CallbackComplete->Wait();
+}
+
+template <class BulkPayloadType, class CallbackPayloadType, class CallbackType>
+bool TResolveInputRequest<BulkPayloadType, CallbackPayloadType, CallbackType>::Poll()
+{
+	check(bCreationComplete);
+	if (!bCallbackCalled)
+	{
+		for (FPayloadData& PayloadData : PayloadDatas)
+		{
+			if (!PayloadData.Future.IsReady())
+			{
+				return false;
+			}
+		}
+	}
+	return CallbackComplete->Wait(0, /* bIgnoreThreadIdleStats */ true);
+}
+
+template <class BulkPayloadType, class CallbackPayloadType, class CallbackType>
+void TResolveInputRequest<BulkPayloadType, CallbackPayloadType, CallbackType>::SetPayloads(
+	TArrayView<TTuple<FString, TFuture<BulkPayloadType>>> InBulkPayloads, EStatus OtherStatus)
+{
+	RemainingInputCount += InBulkPayloads.Num();
+
+	// Preallocate the Payloads array so that it will not be reallocated while we are adding futures to it
+	// This allows us to have each future write to its assigned element, potentially from another thread,
+	// without the need for synchronization.
+	PayloadDatas.Empty(InBulkPayloads.Num());
+	Payloads.Empty(InBulkPayloads.Num());
+	for (TTuple<FString, TFuture<BulkPayloadType>>& InPair : InBulkPayloads)
+	{
+		// Initialize the elements in Payloads and PayloadDatas before creating the Future that writes to the entry.
+		FPayloadData& PayloadData = PayloadDatas.Emplace_GetRef();
+		CallbackPayloadType& Payload = Payloads.Emplace_GetRef();
+		PayloadData.Key = MoveTemp(InPair.template Get<0>());
+		// Payload.Key is a view of the FString in PayloadData; this is safe because PayloadDatas will not reallocate.
+		Payload.Key = PayloadData.Key;
+
+		// The callback we define may be called and write the Payload element before we finish assigning the Future
+		// and it may do so either from this thread or another thread.
+		// The lambda has a TRefCountPtr to *this; if the IBuildInputResolver drops its reference to *this
+		// before the final TFuture completes, *this will be deleted when the lambda completes.
+		PayloadData.Future = InPair.template Get<1>().Next(
+			[pThis = TRefCountPtr<ThisType>(this), &Payload](BulkPayloadType&& Result)
+			{
+				if (Result.bValid)
+				{
+					pThis->AssignBulkPayload(Payload, Result);
+				}
+				else
+				{
+					pThis->LogBulkDataError(Payload.Key);
+					pThis->bAllSucceeded = false;
+				}
+				if (--pThis->RemainingInputCount == 0)
+				{
+					pThis->CallOnResolved();
+				}
+			});
+	}
+
+#if DO_CHECK
+	check(!bCreationComplete);
+	bCreationComplete = true;
+#endif
+	if (--RemainingInputCount == 0)
+	{
+		CallOnResolved();
+	}
+}
+
+template <class BulkPayloadType, class CallbackPayloadType, class CallbackType>
+void TResolveInputRequest<BulkPayloadType, CallbackPayloadType, CallbackType>::CallOnResolved()
+{
+	if (bCallbackCalled.exchange(true) == false)
+	{
+		if (OnResolved)
+		{
+			OnResolved({ Payloads, EStatus::Ok });
+		}
+		else
+		{
+			TArray<CallbackPayloadType> EmptyPayloads;
+			OnResolved({ EmptyPayloads, EStatus::Error });
+		}
+		CallbackComplete->Trigger();
+	}
+}
+
+template<>
+void FResolveInputMetaRequest::AssignBulkPayload(FBuildInputMetaByKey& OutPayload, UE::BulkDataRegistry::FMetaData& InPayload)
+{
+	OutPayload.RawHash = InPayload.RawHash;
+	OutPayload.RawSize = InPayload.RawSize;
+}
+
+template<>
+void FResolveInputDataRequest::AssignBulkPayload(FBuildInputDataByKey& OutPayload, UE::BulkDataRegistry::FData& InPayload)
+{
+	OutPayload.Data = MoveTemp(InPayload.Buffer);
+}
+
+template<>
+void FResolveInputMetaRequest::LogBulkDataError(FStringView Key)
+{
+	UE_LOG(LogCore, Error, TEXT("Failed to resolve metadata for bulkdata input '%.*s' for build of '%.*s' by %.*s."),
+		Key.Len(), Key.GetData(), Definition.GetName().Len(), Definition.GetName().GetData(),
+		Definition.GetFunction().Len(), Definition.GetFunction().GetData());
+}
+
+template<>
+void FResolveInputDataRequest::LogBulkDataError(FStringView Key)
+{
+	UE_LOG(LogCore, Error, TEXT("Failed to resolve data for bulkdata input '%.*s' for build of '%.*s' by %.*s."),
+		Key.Len(), Key.GetData(), Definition.GetName().Len(), Definition.GetName().GetData(),
+		Definition.GetFunction().Len(), Definition.GetFunction().GetData());
+}
+
 } // namespace UE::DerivedData
