PR #8188: Fix unchecked pointer in RemoveLevelsFromWorld (Contributed by [REDACTED_PERSON])

diff --git a/Engine/Source/Editor/UnrealEd/Private/EditorLevelUtils.cpp b/Engine/Source/Editor/UnrealEd/Private/EditorLevelUtils.cpp
--- a/Engine/Source/Editor/UnrealEd/Private/EditorLevelUtils.cpp
+++ b/Engine/Source/Editor/UnrealEd/Private/EditorLevelUtils.cpp
@@ -769,7 +769,7 @@ bool UEditorLevelUtils::RemoveLevelsFromWorld(TArray<ULevel*> InLevels, bool bCl
 	}
 
 	FText TransResetText(LOCTEXT("RemoveLevelTransReset", "Removing Levels from World"));
-	if (bResetTransBuffer)
+	if (bResetTransBuffer && GEditor->Trans)
 	{
 		GEditor->Trans->Reset(TransResetText);
 	}
