World Partition HLOD builder - Run the setup phase with -AllowCommandletRendering

diff --git a/Engine/Build/Graph/Examples/BuildWorldPartitionHLODs.xml b/Engine/Build/Graph/Examples/BuildWorldPartitionHLODs.xml
--- a/Engine/Build/Graph/Examples/BuildWorldPartitionHLODs.xml
+++ b/Engine/Build/Graph/Examples/BuildWorldPartitionHLODs.xml
@@ -35,7 +35,7 @@
 	
 	<Agent Name="HLOD Setup" Type="Win64">
 		<Node Name="HLOD Setup" Requires="$(HLODSetupDependencies)">
-			<Commandlet Name="WorldPartitionBuilderCommandlet" Project="$(ProjectName)" Arguments="$(HLODCommonBuilderArgs) $(HLODBuilderSCC_None) -SetupHLODs -BuilderCount=$(BuilderCount)"/>
+			<Commandlet Name="WorldPartitionBuilderCommandlet" Project="$(ProjectName)" Arguments="$(HLODCommonBuilderArgs) $(HLODBuilderSCC_None) -AllowCommandletRendering -SetupHLODs -BuilderCount=$(BuilderCount)"/>
 			<GatherBuildProductsFromFile BuildProductsFile="$(RootDir)/HLODTemp/BuildProducts.txt"/>
 		</Node>
 	</Agent>
diff --git a/Engine/Source/Editor/UnrealEd/Private/WorldPartition/WorldPartitionHLODsBuilder.cpp b/Engine/Source/Editor/UnrealEd/Private/WorldPartition/WorldPartitionHLODsBuilder.cpp
--- a/Engine/Source/Editor/UnrealEd/Private/WorldPartition/WorldPartitionHLODsBuilder.cpp
+++ b/Engine/Source/Editor/UnrealEd/Private/WorldPartition/WorldPartitionHLODsBuilder.cpp
@@ -191,7 +191,8 @@ bool UWorldPartitionHLODsBuilder::RequiresCommandletRendering() const
 {
 	// Commandlet requires rendering only for building HLODs
 	// Building will occur either if -BuildHLODs is provided or no explicit step arguments are provided
-	return bBuildHLODs || bSingleBuildStep;
+	// Currently required with -SetupHLODs as we want to build texture streaming data for HLODs on save.
+	return bSetupHLODs || bBuildHLODs || bSingleBuildStep;
 }
 
 bool UWorldPartitionHLODsBuilder::ValidateParams() const
@@ -458,7 +459,12 @@ bool UWorldPartitionHLODsBuilder::BuildHLODActors()
 
 		UE_LOG(LogWorldPartitionHLODsBuilder, Display, TEXT("    [%d/%d] Building HLOD actor %s..."), CurrentActor, HLODActorsToBuild.Num(), *HLODActor->GetActorLabel());
 
-		HLODActor->BuildHLOD();
+		if (HLODActor->GetSubActorsHLODLayer()->GetLayerType() != EHLODLayerType::Instancing)
+		{
+			continue;
+		}
+
+		HLODActor->BuildHLOD(true);
 
 		UPackage* ActorPackage = HLODActor->GetPackage();
 		if (ActorPackage->IsDirty())
