Added support for specifying a log category in UE_SCOPED_TIMER.

diff --git a/Engine/Source/Editor/UnrealEd/Private/Commandlets/WorldPartitionBuilderCommandlet.cpp b/Engine/Source/Editor/UnrealEd/Private/Commandlets/WorldPartitionBuilderCommandlet.cpp
--- a/Engine/Source/Editor/UnrealEd/Private/Commandlets/WorldPartitionBuilderCommandlet.cpp
+++ b/Engine/Source/Editor/UnrealEd/Private/Commandlets/WorldPartitionBuilderCommandlet.cpp
@@ -23,7 +23,7 @@ int32 UWorldPartitionBuilderCommandlet::Main(const FString& Params)
 {
 	TRACE_CPUPROFILER_EVENT_SCOPE(UWorldPartitionBuilderCommandlet::Main);
 
-	UE_SCOPED_TIMER(TEXT("Execution"), LogWorldPartitionBuilderCommandlet);
+	UE_SCOPED_TIMER(TEXT("Execution"), LogWorldPartitionBuilderCommandlet, Display);
 
 	TArray<FString> Tokens, Switches;
 	ParseCommandLine(*Params, Tokens, Switches);
diff --git a/Engine/Source/Editor/UnrealEd/Private/Commandlets/WorldPartitionConvertCommandlet.cpp b/Engine/Source/Editor/UnrealEd/Private/Commandlets/WorldPartitionConvertCommandlet.cpp
--- a/Engine/Source/Editor/UnrealEd/Private/Commandlets/WorldPartitionConvertCommandlet.cpp
+++ b/Engine/Source/Editor/UnrealEd/Private/Commandlets/WorldPartitionConvertCommandlet.cpp
@@ -423,7 +423,7 @@ void UWorldPartitionConvertCommandlet::FixupSoftObjectPaths(UPackage* OuterPacka
 {
 	TRACE_CPUPROFILER_EVENT_SCOPE(UWorldPartitionConvertCommandlet::FixupSoftObjectPaths);
 
-	UE_SCOPED_TIMER(TEXT("FixupSoftObjectPaths"), LogWorldPartitionConvertCommandlet);
+	UE_SCOPED_TIMER(TEXT("FixupSoftObjectPaths"), LogWorldPartitionConvertCommandlet, Display);
 
 	struct FSoftPathFixupSerializer : public FArchiveUObject
 	{
@@ -695,7 +695,7 @@ void UWorldPartitionConvertCommandlet::CreateWorldMiniMapTexture(UWorld* World)
 {
 	TRACE_CPUPROFILER_EVENT_SCOPE(UWorldPartitionConvertCommandlet::CreateWorldMiniMapTexture);
 
-	UE_SCOPED_TIMER(TEXT("CreateWorldMiniMapTexture"), LogWorldPartitionConvertCommandlet);
+	UE_SCOPED_TIMER(TEXT("CreateWorldMiniMapTexture"), LogWorldPartitionConvertCommandlet, Display);
 
 	AWorldPartitionMiniMap* WorldMiniMap = FWorldPartitionMiniMapHelper::GetWorldPartitionMiniMap(World, /*bCreateNewMiniMap*/true);
 	if (!WorldMiniMap)
@@ -711,7 +711,7 @@ int32 UWorldPartitionConvertCommandlet::Main(const FString& Params)
 {
 	TRACE_CPUPROFILER_EVENT_SCOPE(UWorldPartitionConvertCommandlet::Main);
 
-	UE_SCOPED_TIMER(TEXT("Conversion"), LogWorldPartitionConvertCommandlet);
+	UE_SCOPED_TIMER(TEXT("Conversion"), LogWorldPartitionConvertCommandlet, Display);
 
 	TArray<FString> Tokens, Switches;
 	ParseCommandLine(*Params, Tokens, Switches);
@@ -783,7 +783,7 @@ int32 UWorldPartitionConvertCommandlet::Main(const FString& Params)
 	// Delete existing result from running the commandlet, even if not using the suffix mode to cleanup previous conversion
 	if (!bReportOnly)
 	{
-		UE_SCOPED_TIMER(TEXT("Deleting existing conversion results"), LogWorldPartitionConvertCommandlet);
+		UE_SCOPED_TIMER(TEXT("Deleting existing conversion results"), LogWorldPartitionConvertCommandlet, Display);
 
 		FString OldLevelName = Tokens[0] + ConversionSuffix;
 		FString ExternalActorsPath = ULevel::GetExternalActorsPath(OldLevelName);
@@ -1031,7 +1031,7 @@ int32 UWorldPartitionConvertCommandlet::Main(const FString& Params)
 		// do loop after as it may modify Level->Actors
 		if (IFAs.Num())
 		{
-			UE_SCOPED_TIMER(TEXT("PartitionFoliage"), LogWorldPartitionConvertCommandlet);
+			UE_SCOPED_TIMER(TEXT("PartitionFoliage"), LogWorldPartitionConvertCommandlet, Display);
 
 			for (AInstancedFoliageActor* IFA : IFAs)
 			{
@@ -1041,7 +1041,7 @@ int32 UWorldPartitionConvertCommandlet::Main(const FString& Params)
 
 		if (LandscapeInfos.Num())
 		{
-			UE_SCOPED_TIMER(TEXT("PartitionLandscape"), LogWorldPartitionConvertCommandlet);
+			UE_SCOPED_TIMER(TEXT("PartitionLandscape"), LogWorldPartitionConvertCommandlet, Display);
 
 			for (ULandscapeInfo* LandscapeInfo : LandscapeInfos)
 			{
diff --git a/Engine/Source/Editor/UnrealEd/Private/WorldPartition/WorldPartitionHLODsBuilder.cpp b/Engine/Source/Editor/UnrealEd/Private/WorldPartition/WorldPartitionHLODsBuilder.cpp
--- a/Engine/Source/Editor/UnrealEd/Private/WorldPartition/WorldPartitionHLODsBuilder.cpp
+++ b/Engine/Source/Editor/UnrealEd/Private/WorldPartition/WorldPartitionHLODsBuilder.cpp
@@ -344,7 +344,7 @@ bool UWorldPartitionHLODsBuilder::SetupHLODActors(bool bCreateOnly)
 		UE_LOG(LogWorldPartitionHLODsBuilder, Display, TEXT("#### World HLOD actors ####"));
 
 		int32 NumActors = 0;
-		for (UActorDescContainer::TIterator<AWorldPartitionHLOD> HLODIterator(WorldPartition); HLODIterator; ++HLODIterator)
+		for (FActorDescList::TIterator<AWorldPartitionHLOD> HLODIterator(WorldPartition); HLODIterator; ++HLODIterator)
 		{
 			FWorldPartitionActorDesc* HLODActorDesc = *HLODIterator;
 			FString PackageName = HLODActorDesc->GetActorPackage().ToString();
@@ -523,7 +523,7 @@ bool UWorldPartitionHLODsBuilder::DeleteHLODActors()
 	UE_LOG(LogWorldPartitionHLODsBuilder, Display, TEXT("#### Deleting HLOD actors ####"));
 
 	TArray<FString> PackagesToDelete;
-	for (UActorDescContainer::TIterator<AWorldPartitionHLOD> HLODIterator(WorldPartition); HLODIterator; ++HLODIterator)
+	for (FActorDescList::TIterator<AWorldPartitionHLOD> HLODIterator(WorldPartition); HLODIterator; ++HLODIterator)
 	{
 		FWorldPartitionActorDesc* HLODActorDesc = *HLODIterator;
 		FString PackageName = HLODActorDesc->GetActorPackage().ToString();
@@ -787,7 +787,7 @@ TArray<TArray<FGuid>> UWorldPartitionHLODsBuilder::GetHLODWorldloads(int32 NumWo
 {
 	// Build a mapping of 1 HLOD[Level] -> N HLOD[Level - 1]
 	TMap<FGuid, TArray<FGuid>>	HLODParenting;
-	for (UActorDescContainer::TIterator<AWorldPartitionHLOD> HLODIterator(WorldPartition); HLODIterator; ++HLODIterator)
+	for (FActorDescList::TIterator<AWorldPartitionHLOD> HLODIterator(WorldPartition); HLODIterator; ++HLODIterator)
 	{
 		TArray<FGuid>& ChildHLODs = HLODParenting.Add(HLODIterator->GetGuid());
 
diff --git a/Engine/Source/Editor/UnrealEd/Private/WorldPartition/WorldPartitionMiniMapBuilder.cpp b/Engine/Source/Editor/UnrealEd/Private/WorldPartition/WorldPartitionMiniMapBuilder.cpp
--- a/Engine/Source/Editor/UnrealEd/Private/WorldPartition/WorldPartitionMiniMapBuilder.cpp
+++ b/Engine/Source/Editor/UnrealEd/Private/WorldPartition/WorldPartitionMiniMapBuilder.cpp
@@ -46,12 +46,12 @@ bool UWorldPartitionMiniMapBuilder::RunInternal(UWorld* World, const FBox& Bound
 		UWorldPartition* WorldPartition = World->GetWorldPartition();
 
 		TSet<FGuid> AllSubActors;
-		for (UActorDescContainer::TIterator<AWorldPartitionHLOD> HLODIterator(WorldPartition); HLODIterator; ++HLODIterator)
+		for (FActorDescList::TIterator<AWorldPartitionHLOD> HLODIterator(WorldPartition); HLODIterator; ++HLODIterator)
 		{
 			AllSubActors.Append(HLODIterator->GetSubActors());
 		}
 
-		for (UActorDescContainer::TIterator<AWorldPartitionHLOD> HLODIterator(WorldPartition); HLODIterator; ++HLODIterator)
+		for (FActorDescList::TIterator<AWorldPartitionHLOD> HLODIterator(WorldPartition); HLODIterator; ++HLODIterator)
 		{
 			// Only include top level HLODs - If an HLOD actor isn't included as a subactor it means it is top level
 			if (!AllSubActors.Contains(HLODIterator->GetGuid()))
diff --git a/Engine/Source/Runtime/Core/Public/ProfilingDebugging/ScopedTimers.h b/Engine/Source/Runtime/Core/Public/ProfilingDebugging/ScopedTimers.h
--- a/Engine/Source/Runtime/Core/Public/ProfilingDebugging/ScopedTimers.h
+++ b/Engine/Source/Runtime/Core/Public/ProfilingDebugging/ScopedTimers.h
@@ -203,23 +203,24 @@ public:
  * doesn't have to call Start() and Stop() manually) using a custom
  * output function.
  */
+template<class F>
 class FScopedDurationTimeCustomLogger
 {
 public:
-	explicit FScopedDurationTimeCustomLogger(const TCHAR* InMsg, const FLogCategoryBase& InLogCategory)
-		: Msg(InMsg)
-		, LogCategory(InLogCategory)
+	explicit FScopedDurationTimeCustomLogger(const TCHAR* InTitle, F InFunc)
+		: Title(InTitle)
+		, Func(InFunc)
 		, Accumulator(0.0)
 		, Timer(Accumulator)
 	{
-		UE_LOG_REF(LogCategory, Display, TEXT("%s"), *Msg);
+		Func(*FString::Printf(TEXT("%s started..."), InTitle));
 		Timer.Start();
 	}
 
 	~FScopedDurationTimeCustomLogger()
 	{
 		Timer.Stop();
-		UE_LOG_REF(LogCategory, Display, TEXT("%s took %s"), *Msg, *SecondsToString(Accumulator));
+		Func(*FString::Printf(TEXT("%s took %s"), *Title, *SecondsToString(Accumulator)));
 	}
 
 private:
@@ -257,15 +258,15 @@ private:
 		return FString::Printf(TEXT("%02dm %02ds"), Minutes, Seconds);
 	}
 
-	FString Msg;
-	const FLogCategoryBase& LogCategory;
+	FString Title;
+	F Func;
 	double Accumulator;
 	FDurationTimer Timer;
 };
 
 #if NO_LOGGING
-#define UE_SCOPED_TIMER(Title, Category)
+#define UE_SCOPED_TIMER(Title, Category, Verbosity)
 #else
-#define UE_SCOPED_TIMER(Title, Category) \
-	FScopedDurationTimeCustomLogger BODY_MACRO_COMBINE(Scoped,Timer,_,__LINE__)(Title, Category)
+#define UE_SCOPED_TIMER(Title, Category, Verbosity) \
+	FScopedDurationTimeCustomLogger BODY_MACRO_COMBINE(Scoped,Timer,_,__LINE__)(Title, [](const TCHAR* Msg) { UE_LOG(Category, Verbosity, TEXT("%s"), Msg); })
 #endif
diff --git a/Engine/Source/Runtime/Engine/Private/WorldPartition/RuntimeSpatialHash/RuntimeSpatialHashGridHelper.cpp b/Engine/Source/Runtime/Engine/Private/WorldPartition/RuntimeSpatialHash/RuntimeSpatialHashGridHelper.cpp
--- a/Engine/Source/Runtime/Engine/Private/WorldPartition/RuntimeSpatialHash/RuntimeSpatialHashGridHelper.cpp
+++ b/Engine/Source/Runtime/Engine/Private/WorldPartition/RuntimeSpatialHash/RuntimeSpatialHashGridHelper.cpp
@@ -98,7 +98,7 @@ int32 FSquare2DGridHelper::ForEachIntersectingCells(const FSphere& InSphere, TFu
 #if WITH_EDITOR
 void FSquare2DGridHelper::ValidateSingleActorReferer()
 {
-	UE_SCOPED_TIMER(TEXT("ValidateSingleActorReferer"), LogWorldPartition);
+	UE_SCOPED_TIMER(TEXT("ValidateSingleActorReferer"), LogWorldPartition, Display);
 
 	TSet<FActorInstance> ActorUsage;
 	for (int32 Level = 0; Level < Levels.Num() - 1; Level++)
@@ -131,7 +131,7 @@ FSquare2DGridHelper GetGridHelper(const FBox& WorldBounds, int32 GridCellSize)
 FSquare2DGridHelper GetPartitionedActors(const UWorldPartition* WorldPartition, const FBox& WorldBounds, const FSpatialHashRuntimeGrid& Grid, const TArray<const FActorClusterInstance*>& GridActors)
 {
 	TRACE_CPUPROFILER_EVENT_SCOPE(GetPartitionedActors);
-	UE_SCOPED_TIMER(TEXT("GetPartitionedActors"), LogWorldPartition);
+	UE_SCOPED_TIMER(TEXT("GetPartitionedActors"), LogWorldPartition, Display);
 
 	//
 	// Create the hierarchical grids for the game
diff --git a/Engine/Source/Runtime/Engine/Private/WorldPartition/RuntimeSpatialHash/RuntimeSpatialHashHLOD.cpp b/Engine/Source/Runtime/Engine/Private/WorldPartition/RuntimeSpatialHash/RuntimeSpatialHashHLOD.cpp
--- a/Engine/Source/Runtime/Engine/Private/WorldPartition/RuntimeSpatialHash/RuntimeSpatialHashHLOD.cpp
+++ b/Engine/Source/Runtime/Engine/Private/WorldPartition/RuntimeSpatialHash/RuntimeSpatialHashHLOD.cpp
@@ -258,7 +258,7 @@ static void GatherHLODGrids(UWorldPartition* WorldPartition, TMap<FName, FSpatia
 	// Gather up all HLODLayer referenced by the actors
 	TSet<UHLODLayer*> HLODLayers;
 
-	for (UActorDescContainer::TIterator<> ActorDescIterator(WorldPartition); ActorDescIterator; ++ActorDescIterator)
+	for (FActorDescList::TIterator<> ActorDescIterator(WorldPartition); ActorDescIterator; ++ActorDescIterator)
 	{
 		const FWorldPartitionActorDesc& ActorDesc = **ActorDescIterator;
 		if (!ActorDesc.GetActorClass()->IsChildOf<AWorldPartitionHLOD>())
@@ -495,7 +495,7 @@ bool UWorldPartitionRuntimeSpatialHash::GenerateHLOD(ISourceControlHelper* Sourc
 
 	TSet<FWorldPartitionHandle> InvalidHLODActors;
 
-	for (UActorDescContainer::TIterator<AWorldPartitionHLOD> HLODIterator(WorldPartition); HLODIterator; ++HLODIterator)
+	for (FActorDescList::TIterator<AWorldPartitionHLOD> HLODIterator(WorldPartition); HLODIterator; ++HLODIterator)
 	{
 		uint64 CellHash = HLODIterator->GetCellHash();
 
diff --git a/Engine/Source/Runtime/Engine/Private/WorldPartition/WorldPartition.cpp b/Engine/Source/Runtime/Engine/Private/WorldPartition/WorldPartition.cpp
--- a/Engine/Source/Runtime/Engine/Private/WorldPartition/WorldPartition.cpp
+++ b/Engine/Source/Runtime/Engine/Private/WorldPartition/WorldPartition.cpp
@@ -335,7 +335,7 @@ FName UWorldPartition::GetWorldPartitionEditorName()
 
 void UWorldPartition::Initialize(UWorld* InWorld, const FTransform& InTransform)
 {
-	UE_SCOPED_TIMER(TEXT("WorldPartition initialize"), LogWorldPartition);
+	UE_SCOPED_TIMER(TEXT("WorldPartition initialize"), LogWorldPartition, Display);
 	TRACE_CPUPROFILER_EVENT_SCOPE(UWorldPartition::Initialize);
 	
 	check(!World || (World == InWorld));
@@ -409,7 +409,7 @@ void UWorldPartition::Initialize(UWorld* InWorld, const FTransform& InTransform)
 			ReplaceTo = DestWorldName + TEXT(".") + DestWorldName;
 		}
 
-		for (UActorDescContainer::TIterator<> ActorDescIterator(this); ActorDescIterator; ++ActorDescIterator)
+		for (FActorDescList::TIterator<> ActorDescIterator(this); ActorDescIterator; ++ActorDescIterator)
 		{
 			if (bIsInstanced)
 			{
diff --git a/Engine/Source/Runtime/Engine/Private/WorldPartition/WorldPartitionRuntimeSpatialHash.cpp b/Engine/Source/Runtime/Engine/Private/WorldPartition/WorldPartitionRuntimeSpatialHash.cpp
--- a/Engine/Source/Runtime/Engine/Private/WorldPartition/WorldPartitionRuntimeSpatialHash.cpp
+++ b/Engine/Source/Runtime/Engine/Private/WorldPartition/WorldPartitionRuntimeSpatialHash.cpp
@@ -632,7 +632,7 @@ bool UWorldPartitionRuntimeSpatialHash::GenerateStreaming(EWorldPartitionStreami
 	TRACE_CPUPROFILER_EVENT_SCOPE(UWorldPartitionRuntimeSpatialHash::GenerateStreaming);
 	UWorldPartition* WorldPartition = GetOuterUWorldPartition();
 
-	UE_SCOPED_TIMER(TEXT("GenerateStreaming"), LogWorldPartition);
+	UE_SCOPED_TIMER(TEXT("GenerateStreaming"), LogWorldPartition, Display);
 	
 	if (!Grids.Num())
 	{
