Add limiting heuristics and metrics to the DDC2 remote executor along with some additional logging around uploaded blobs.

diff --git a/Engine/Source/Editor/UnrealEd/Private/Commandlets/CookCommandlet.cpp b/Engine/Source/Editor/UnrealEd/Private/Commandlets/CookCommandlet.cpp
--- a/Engine/Source/Editor/UnrealEd/Private/Commandlets/CookCommandlet.cpp
+++ b/Engine/Source/Editor/UnrealEd/Private/Commandlets/CookCommandlet.cpp
@@ -43,6 +43,7 @@
 #include "AssetRegistryModule.h"
 #include "StudioAnalytics.h"
 #include "Cooker/CookProfiling.h"
+#include "DerivedDataBuildRemoteExecutor.h"
 
 DEFINE_LOG_CATEGORY_STATIC(LogCookCommandlet, Log, All);
 
@@ -351,6 +352,9 @@ namespace DetailedCookStats
 				UE_LOG(LogCookCommandlet, Display, TEXT("%-16s=%10s"), *Attr.Key, *Attr.Value);
 			}
 		}
+
+		DumpDerivedDataBuildRemoteExecutorStats();
+
 		if (DDCResourceUsageStats.Num() > 0)
 		{
 			// sort the list
diff --git a/Engine/Source/Editor/UnrealEd/Private/DerivedDataBuildRemoteExecutor.cpp b/Engine/Source/Editor/UnrealEd/Private/DerivedDataBuildRemoteExecutor.cpp
--- a/Engine/Source/Editor/UnrealEd/Private/DerivedDataBuildRemoteExecutor.cpp
+++ b/Engine/Source/Editor/UnrealEd/Private/DerivedDataBuildRemoteExecutor.cpp
@@ -13,6 +13,7 @@
 #include "IRemoteExecutor.h"
 #include "Messages.h"
 #include "Misc/CommandLine.h"
+#include "Misc/ConfigCacheIni.h"
 #include "Misc/Optional.h"
 #include "Misc/PathViews.h"
 #include "Modules/ModuleManager.h"
@@ -39,13 +40,17 @@ class FRemoteBuildWorkerExecutor final: public IBuildWorkerExecutor
 {
 public:
 	FRemoteBuildWorkerExecutor()
-	: InstanceName(TEXT("dev.test.reapi"))
-	, Salt(TEXT("807c6a49-0657-49f3-b498-fd457213c0a7"))
+	: Salt(TEXT("807c6a49-0657-49f3-b498-fd457213c0a7"))
 	, BaseDirectoryPath(TEXT("Engine/Binaries/Win64"))
 	, RemoteExecutor(nullptr)
 	, ContentAddressableStorage(nullptr)
 	, bEnabled(false)
 	{
+		check(IsInGameThread()); // initialization from the main thread is expected to allow config reading for the limiting heuristics
+		check(GConfig && GConfig->IsReadyForUse());
+
+		GConfig->GetString(TEXT("DerivedDataBuildRemoteExecutor"), TEXT("InstanceName"), InstanceName, GEngineIni);
+
 		const FName RemoteExecutionFeatureName(TEXT("RemoteExecution"));
 		IModularFeatures& ModularFeatures = IModularFeatures::Get();
 
@@ -62,7 +67,7 @@ public:
 			}
 		}
 
-		bEnabled = (RemoteExecutor != nullptr) && (ContentAddressableStorage != nullptr);
+		bEnabled = (RemoteExecutor != nullptr) && (ContentAddressableStorage != nullptr) && !InstanceName.IsEmpty();
 		if (bEnabled)
 		{
 			ModularFeatures.RegisterModularFeature(IBuildWorkerExecutor::GetFeatureName(), this);
@@ -84,6 +89,7 @@ public:
 		//Symlink,
 	};
 
+
 	enum class EFileType
 	{
 		Worker,
@@ -91,6 +97,21 @@ public:
 		BuildAction,
 	};
 
+	static const TCHAR* LexToString(EFileType FileType)
+	{
+		switch (FileType)
+		{
+			case EFileType::Worker:
+				return TEXT("Worker");
+			case EFileType::Input:
+				return TEXT("Input");
+			case EFileType::BuildAction:
+				return TEXT("BuildAction");
+			default:
+				return TEXT("Unknown");
+		}
+	}
+
 	struct FVariantIndex
 	{
 		ENodeType NodeType;
@@ -415,15 +436,19 @@ public:
 		{
 			if (MissingItem == State.ActionDigest)
 			{
+				Stats.TotalActionBlobsUploaded.AddBlob(State.ActionContentBytes.Num());
 				FBatchUpdateBlobsRequest::FRequest& NewRequest = State.BatchUpdateBlobsRequest.Requests.AddDefaulted_GetRef();
 				NewRequest.Digest = State.ActionDigest;
 				NewRequest.Data = MoveTemp(State.ActionContentBytes);
+				UE_LOG(LogDerivedDataBuildRemoteExecutor, Verbose, TEXT("Uploading action (hash: %s) of upload size %d."), *::LexToString(NewRequest.Digest.Hash), NewRequest.Data.Num());
 			}
 			else if (MissingItem == State.Action.CommandDigest)
 			{
+				Stats.TotalCommandBlobsUploaded.AddBlob(State.CommandContentBytes.Num());
 				FBatchUpdateBlobsRequest::FRequest& NewRequest = State.BatchUpdateBlobsRequest.Requests.AddDefaulted_GetRef();
 				NewRequest.Digest = State.Action.CommandDigest;
 				NewRequest.Data = MoveTemp(State.CommandContentBytes);
+				UE_LOG(LogDerivedDataBuildRemoteExecutor, Verbose, TEXT("Uploading command (hash: %s) of upload size %d."), *::LexToString(NewRequest.Digest.Hash), NewRequest.Data.Num());
 			}
 			else
 			{
@@ -433,11 +458,14 @@ public:
 				switch (VariantIndex.NodeType)
 				{
 				case ENodeType::Directory:
+					Stats.TotalDirectoryBlobsUploaded.AddBlob(State.Directories[VariantIndex.Index].ContentBytes.Num());
 					NewRequest.Digest = State.Directories[VariantIndex.Index].Digest.GetValue();
 					NewRequest.Data = MoveTemp(State.Directories[VariantIndex.Index].ContentBytes);
+					UE_LOG(LogDerivedDataBuildRemoteExecutor, Verbose, TEXT("Uploading directory '%s' (hash: %s) of upload size %d."), *FString(State.Directories[VariantIndex.Index].Name), *::LexToString(NewRequest.Digest.Hash), NewRequest.Data.Num());
 					break;
 				case ENodeType::File:
 					{
+						Stats.TotalFileBlobsUploaded.AddBlob(State.Files[VariantIndex.Index].ContentBytes.GetSize());
 						NewRequest.Digest.Hash = FIoHash::HashBuffer(State.Files[VariantIndex.Index].ContentBytes);
 						NewRequest.Digest.SizeBytes = State.Files[VariantIndex.Index].ContentBytes.GetSize();
 						FCompositeBuffer& FileBuffer = State.Files[VariantIndex.Index].ContentBytes;
@@ -448,6 +476,7 @@ public:
 							NewRequest.Data.Append((const char *)Segment.GetData(), Segment.GetSize());
 						}
 						FileBuffer.Reset();
+						UE_LOG(LogDerivedDataBuildRemoteExecutor, Verbose, TEXT("Uploading file '%s' (hash: %s, type: %s) of upload size %d."), *State.Files[VariantIndex.Index].File.Name, *::LexToString(NewRequest.Digest.Hash), *LexToString(State.Files[VariantIndex.Index].Type), NewRequest.Data.Num());
 					}
 					break;
 				default:
@@ -473,7 +502,7 @@ public:
 				if (!FoundResponse->Status.Ok())
 				{
 					FStringView ActionName = State.BuildAction.GetName();
-					UE_LOG(LogDerivedDataBuildRemoteExecutor, Warning, TEXT("Remote execution system error: data for action '%.*s' could not be uploaded (hash: %s, size: %u)"), ActionName.Len(), ActionName.GetData(), *LexToString(RequestedUpload.Digest.Hash), RequestedUpload.Digest.SizeBytes);
+					UE_LOG(LogDerivedDataBuildRemoteExecutor, Log, TEXT("Remote execution system error: data for action '%.*s' could not be uploaded (hash: %s, size: %u)"), ActionName.Len(), ActionName.GetData(), *::LexToString(RequestedUpload.Digest.Hash), RequestedUpload.Digest.SizeBytes);
 					bSuccess = false;
 				}
 			}
@@ -533,6 +562,7 @@ public:
 				UE_LOG(LogDerivedDataBuildRemoteExecutor, Warning, TEXT("Remote execution system error: output not downloaded!"));
 				return FOptionalBuildOutput();
 			}
+			Stats.TotalFileBlobsDownloaded.AddBlob(OutputReadResponse.Data.Num());
 
 			if (OutputReadResponse.Digest == BuildOutputDigest)
 			{
@@ -589,30 +619,6 @@ public:
 		return OutputBuilder.Build();
 	}
 
-	bool IsRemoteExecutionAttemptAllowed(const FBuildAction& Action)
-	{
-		uint64 TotalInputSize = 0;
-		Action.IterateInputs([&TotalInputSize] (FStringView Key, const FIoHash& RawHash, uint64 RawSize)
-			{
-				TotalInputSize += RawSize;
-			});
-
-		// Only bother sending large input data for remote execution right now.
-		if (TotalInputSize < 10*1024*1024)
-		{
-			return false;
-		}
-
-		static std::atomic<uint32> NumRemoteBuilds {0};
-		// Limit to 10 total remote builds for now until the remote instance can handle a high quantity better.
-		if (NumRemoteBuilds.fetch_add(1, std::memory_order_relaxed) > 8)
-		{
-			return false;
-		}
-
-		return true;
-	}
-
 	FRequest BuildAction(
 		const FBuildAction& Action,
 		const FOptionalBuildInputs& Inputs,
@@ -621,14 +627,10 @@ public:
 		EPriority Priority,
 		FOnBuildWorkerActionComplete&& OnComplete) final
 	{
-		if (!IsRemoteExecutionAttemptAllowed(Action))
-		{
-			OnComplete({Action.GetKey(), {}, {}, EStatus::Error});
-			return FRequest();
-		}
-
 		TArray<FString> MissingInputs;
 		TArray<FStringView> MissingInputViews;
+		uint64 TotalInputSize = 0;
+		uint64 TotalMissingInputSize = 0;
 		{
 			// TODO: This block forces resolution of inputs before we attempt to determine which
 			//		 inputs need to be uploaded.  This is required because we can't refer to inputs
@@ -636,15 +638,23 @@ public:
 			//		 CompressedSize.  Once the remote execution API allows us to represent inputs with RawHash/
 			//		 RawSize, this block can be removed and we can find missing CAS inputs without having resolved
 			//		 the inputs first.
-			Action.IterateInputs([&MissingInputs, &MissingInputViews, &Inputs] (FStringView Key, const FIoHash& RawHash, uint64 RawSize)
+			Action.IterateInputs([&MissingInputs, &MissingInputViews, &Inputs, &TotalInputSize, &TotalMissingInputSize] (FStringView Key, const FIoHash& RawHash, uint64 RawSize)
 				{
 					if (Inputs.IsNull() || Inputs.Get().FindInput(Key).IsNull())
 					{
 						MissingInputs.Emplace(Key);
 						MissingInputViews.Add(MissingInputs.Last());
+						TotalMissingInputSize += RawSize;
 					}
+					TotalInputSize += RawSize;
 				});
 
+			if (!LimitingHeuristics.PassesPreResolveRequirements(TotalInputSize, TotalMissingInputSize))
+			{
+				OnComplete({Action.GetKey(), {}, {}, EStatus::Error});
+				return FRequest();
+			}
+
 			if (!MissingInputViews.IsEmpty())
 			{
 				OnComplete({Action.GetKey(), {}, MissingInputViews, EStatus::Ok});
@@ -663,11 +673,21 @@ public:
 			return FRequest();
 		}
 
+		if (!LimitingHeuristics.TryStartNewBuild(Stats))
+		{
+			OnComplete({Action.GetKey(), {}, {}, EStatus::Error});
+			return FRequest();
+		}
+		ON_SCOPE_EXIT
+		{
+			LimitingHeuristics.FinishBuild(Stats);
+		};
+
 		LoadMissingWorkerFileBlobs(State);
 		if (!State.FindMissingBlobsResponse.MissingBlobDigests.IsEmpty())
 		{
 			UploadMissingBlobs(State);
-			UE_LOG(LogDerivedDataBuildRemoteExecutor, Display, TEXT("Uploaded %d data blobs for remote execution."), State.FindMissingBlobsResponse.MissingBlobDigests.Num());
+			UE_LOG(LogDerivedDataBuildRemoteExecutor, Verbose, TEXT("Uploaded %d data blobs for remote execution."), State.BatchUpdateBlobsRequest.Requests.Num());
 			if (!ValidateUploadSuccess(State))
 			{
 				OnComplete({Action.GetKey(), {}, {}, EStatus::Error});
@@ -681,6 +701,10 @@ public:
 
 			EStatus BuildStatus = EStatus::Error;
 			FOptionalBuildOutput BuildOutput = ComposeBuildOutput(State, BuildStatus);
+			if (BuildStatus == EStatus::Ok)
+			{
+				Stats.TotalSuccessfulRemoteBuilds.fetch_add(1, std::memory_order_relaxed);
+			}
 			OnComplete({State.BuildAction.GetKey(), MoveTemp(BuildOutput), {}, BuildStatus});
 		}
 		else
@@ -697,8 +721,156 @@ public:
 		static constexpr FStringView HostPlatforms[]{TEXT("Win64"_SV)};
 		return HostPlatforms;
 	}
+
+	void DumpStats()
+	{
+		if (Stats.TotalRemoteBuilds.load() == 0)
+		{
+			return;
+		}
+
+		Stats.Dump();
+	}
+
 private:
-	const FStringView InstanceName;
+	struct FStats
+	{
+		std::atomic<uint64> TotalRemoteBuilds{0};
+		std::atomic<uint32> InFlightRemoteBuilds{0};
+
+		std::atomic<uint64> TotalSuccessfulRemoteBuilds{0};
+
+		struct FBlobStat
+		{
+			std::atomic<uint64> Quantity{0};
+			std::atomic<uint64> Bytes{0};
+
+			void AddBlob(uint64 InBytes)
+			{
+				Quantity.fetch_add(1, std::memory_order_relaxed);
+				Bytes.fetch_add(InBytes, std::memory_order_relaxed);
+			}
+		};
+		FBlobStat TotalActionBlobsUploaded;
+		FBlobStat TotalCommandBlobsUploaded;
+		FBlobStat TotalDirectoryBlobsUploaded;
+		FBlobStat TotalFileBlobsUploaded;
+		FBlobStat TotalFileBlobsDownloaded;
+
+		void Dump()
+		{
+			UE_LOG(LogDerivedDataBuildRemoteExecutor, Display, TEXT(""));
+			UE_LOG(LogDerivedDataBuildRemoteExecutor, Display, TEXT("DDC Remote Execution Stats"));
+			UE_LOG(LogDerivedDataBuildRemoteExecutor, Display, TEXT("=========================="));
+			UE_LOG(LogDerivedDataBuildRemoteExecutor, Display, TEXT("%-35s=%10") UINT64_FMT, TEXT("Total remote builds"), TotalRemoteBuilds.load());
+			UE_LOG(LogDerivedDataBuildRemoteExecutor, Display, TEXT("%-35s=%10") UINT64_FMT, TEXT("Successful remote builds"), TotalSuccessfulRemoteBuilds.load());
+			UE_LOG(LogDerivedDataBuildRemoteExecutor, Display, TEXT("%-35s=%10") UINT64_FMT, TEXT("Uploaded actions (quantity)"), TotalActionBlobsUploaded.Quantity.load());
+			UE_LOG(LogDerivedDataBuildRemoteExecutor, Display, TEXT("%-35s=%10") UINT64_FMT, TEXT("Uploaded actions (KB)"), TotalActionBlobsUploaded.Bytes.load()/1024);
+			UE_LOG(LogDerivedDataBuildRemoteExecutor, Display, TEXT("%-35s=%10") UINT64_FMT, TEXT("Uploaded commands (quantity)"), TotalCommandBlobsUploaded.Quantity.load());
+			UE_LOG(LogDerivedDataBuildRemoteExecutor, Display, TEXT("%-35s=%10") UINT64_FMT, TEXT("Uploaded commands (KB)"), TotalCommandBlobsUploaded.Bytes.load()/1024);
+			UE_LOG(LogDerivedDataBuildRemoteExecutor, Display, TEXT("%-35s=%10") UINT64_FMT, TEXT("Uploaded directories (quantity)"), TotalDirectoryBlobsUploaded.Quantity.load());
+			UE_LOG(LogDerivedDataBuildRemoteExecutor, Display, TEXT("%-35s=%10") UINT64_FMT, TEXT("Uploaded directories (KB)"), TotalDirectoryBlobsUploaded.Bytes.load()/1024);
+			UE_LOG(LogDerivedDataBuildRemoteExecutor, Display, TEXT("%-35s=%10") UINT64_FMT, TEXT("Uploaded files (quantity)"), TotalFileBlobsUploaded.Quantity.load());
+			UE_LOG(LogDerivedDataBuildRemoteExecutor, Display, TEXT("%-35s=%10") UINT64_FMT, TEXT("Uploaded files (KB)"), TotalFileBlobsUploaded.Bytes.load()/1024);
+			UE_LOG(LogDerivedDataBuildRemoteExecutor, Display, TEXT("%-35s=%10") UINT64_FMT, TEXT("Downloaded files (quantity)"), TotalFileBlobsDownloaded.Quantity.load());
+			UE_LOG(LogDerivedDataBuildRemoteExecutor, Display, TEXT("%-35s=%10") UINT64_FMT, TEXT("Downloaded files (KB)"), TotalFileBlobsDownloaded.Bytes.load()/1024);
+		}
+	};
+
+	// Temporary heuristics until a scheduler makes higher level decisions about how to limit remote execution of builds
+	class FLimitingHeuristics
+	{
+	public:
+		FLimitingHeuristics()
+		{
+			check(IsInGameThread()); // initialization from the main thread is expected to allow config reading for the limiting heuristics
+			check(GConfig && GConfig->IsReadyForUse());
+			const TCHAR* Section = TEXT("DerivedDataBuildRemoteExecutor.LimitingHeuristics");
+			GConfig->GetBool(Section, TEXT("bEnableLimits"), bEnableLimits, GEngineIni);
+
+			int32 SignedMaxTotalRemoteBuilds{MAX_int32};
+			GConfig->GetInt(Section, TEXT("MaxTotalRemoteBuilds"), SignedMaxTotalRemoteBuilds, GEngineIni);
+			if ((SignedMaxTotalRemoteBuilds >= 0) && (SignedMaxTotalRemoteBuilds < MAX_int32))
+			{
+				MaxTotalRemoteBuilds = (uint64)SignedMaxTotalRemoteBuilds;
+			}
+
+			int32 SignedMaxInFlightRemoteBuilds{MAX_int32};
+			GConfig->GetInt(Section, TEXT("MaxInFlightRemoteBuilds"), SignedMaxInFlightRemoteBuilds, GEngineIni);
+			if ((SignedMaxInFlightRemoteBuilds >= 0) && (SignedMaxInFlightRemoteBuilds < MAX_int32))
+			{
+				MaxInFlightRemoteBuilds = (uint32)SignedMaxInFlightRemoteBuilds;
+			}
+
+			int32 SignedMinInputSizeForRemoteBuilds{0};
+			GConfig->GetInt(Section, TEXT("MinInputSizeForRemoteBuilds"), SignedMinInputSizeForRemoteBuilds, GEngineIni);
+			if ((SignedMinInputSizeForRemoteBuilds >= 0) && (SignedMinInputSizeForRemoteBuilds < MAX_int32))
+			{
+				MinInputSizeForRemoteBuilds = (uint64)SignedMinInputSizeForRemoteBuilds;
+			}
+
+			int32 SignedMaxMissingInputSizeForRemoteBuilds{MAX_int32};
+			GConfig->GetInt(Section, TEXT("MaxMissingInputSizeForRemoteBuilds"), SignedMaxMissingInputSizeForRemoteBuilds, GEngineIni);
+			if ((SignedMaxMissingInputSizeForRemoteBuilds >= 0) && (SignedMaxMissingInputSizeForRemoteBuilds < MAX_int32))
+			{
+				MaxMissingInputSizeForRemoteBuilds = (uint64)SignedMaxMissingInputSizeForRemoteBuilds;
+			}
+		}
+
+		bool PassesPreResolveRequirements(uint64 InputSize, uint64 MissingInputSize)
+		{
+			if (!bEnableLimits)
+			{
+				return true;
+			}
+
+			if (InputSize < MinInputSizeForRemoteBuilds)
+			{
+				return false;
+			}
+
+			if (MissingInputSize > MaxMissingInputSizeForRemoteBuilds)
+			{
+				return false;
+			}
+
+			return true;
+		}
+
+		bool TryStartNewBuild(FStats& InStats)
+		{
+			if ((InStats.TotalRemoteBuilds.fetch_add(1, std::memory_order_relaxed) >= MaxTotalRemoteBuilds) && bEnableLimits)
+			{
+				InStats.TotalRemoteBuilds.fetch_sub(1, std::memory_order_relaxed);
+				return false;
+			}
+
+			if ((InStats.InFlightRemoteBuilds.fetch_add(1, std::memory_order_relaxed) >= MaxInFlightRemoteBuilds) && bEnableLimits)
+			{
+				InStats.TotalRemoteBuilds.fetch_sub(1, std::memory_order_relaxed);
+				InStats.InFlightRemoteBuilds.fetch_sub(1, std::memory_order_relaxed);
+				return false;
+			}
+
+			return true;
+		}
+
+		void FinishBuild(FStats& InStats)
+		{
+			InStats.InFlightRemoteBuilds.fetch_sub(1, std::memory_order_relaxed);
+		}
+
+	private:
+		uint64 MaxTotalRemoteBuilds{MAX_uint64};
+		uint32 MaxInFlightRemoteBuilds{MAX_uint32};
+		uint64 MinInputSizeForRemoteBuilds{0};
+		uint64 MaxMissingInputSizeForRemoteBuilds{MAX_uint64};
+		bool bEnableLimits{false};
+	};
+
+	FStats Stats;
+	FLimitingHeuristics LimitingHeuristics;
+	FString InstanceName;
 	const FStringView Salt;
 	const FString BaseDirectoryPath;
 	IRemoteExecutor* RemoteExecutor;
@@ -708,7 +880,22 @@ private:
 
 } // namespace UE::DerivedData
 
+TOptional<UE::DerivedData::FRemoteBuildWorkerExecutor> GRemoteBuildWorkerExecutor;
+
 void InitDerivedDataBuildRemoteExecutor()
 {
-	static UE::DerivedData::FRemoteBuildWorkerExecutor GRemoteBuildWorkerExecutor;
+	if (!GRemoteBuildWorkerExecutor.IsSet())
+	{
+		GRemoteBuildWorkerExecutor.Emplace();
+	}
+}
+
+void DumpDerivedDataBuildRemoteExecutorStats()
+{
+	static bool bHasRun = false;
+	if (GRemoteBuildWorkerExecutor.IsSet() && !bHasRun)
+	{
+		bHasRun = true;
+		GRemoteBuildWorkerExecutor->DumpStats();
+	}
 }
diff --git a/Engine/Source/Editor/UnrealEd/Private/DerivedDataBuildRemoteExecutor.h b/Engine/Source/Editor/UnrealEd/Private/DerivedDataBuildRemoteExecutor.h
--- a/Engine/Source/Editor/UnrealEd/Private/DerivedDataBuildRemoteExecutor.h
+++ b/Engine/Source/Editor/UnrealEd/Private/DerivedDataBuildRemoteExecutor.h
@@ -3,3 +3,4 @@
 #pragma once
 
 void InitDerivedDataBuildRemoteExecutor();
+void DumpDerivedDataBuildRemoteExecutorStats();
