PR #8171: [UE5] Stop attempting to load nonexistent TimeManagementEditor module (Contributed by [REDACTED_PERSON])

diff --git a/Engine/Source/Editor/UnrealEd/Private/EditorEngine.cpp b/Engine/Source/Editor/UnrealEd/Private/EditorEngine.cpp
--- a/Engine/Source/Editor/UnrealEd/Private/EditorEngine.cpp
+++ b/Engine/Source/Editor/UnrealEd/Private/EditorEngine.cpp
@@ -1058,7 +1058,6 @@ void UEditorEngine::Init(IEngineLoop* InEngineLoop)
 			TEXT("MergeActors"),
 			TEXT("InputBindingEditor"),
 			TEXT("AudioEditor"),
-			TEXT("TimeManagementEditor"),
 			TEXT("EditorInteractiveToolsFramework"),
 			TEXT("TraceInsights"),
 			TEXT("StaticMeshEditor"),
