Add a configurable global timeout for all DDC2 remote execution work.

diff --git a/Engine/Source/Editor/UnrealEd/Private/DerivedDataBuildRemoteExecutor.cpp b/Engine/Source/Editor/UnrealEd/Private/DerivedDataBuildRemoteExecutor.cpp
--- a/Engine/Source/Editor/UnrealEd/Private/DerivedDataBuildRemoteExecutor.cpp
+++ b/Engine/Source/Editor/UnrealEd/Private/DerivedDataBuildRemoteExecutor.cpp
@@ -40,7 +40,8 @@ class FRemoteBuildWorkerExecutor final: public IBuildWorkerExecutor
 {
 public:
 	FRemoteBuildWorkerExecutor()
-	: Salt(TEXT("807c6a49-0657-49f3-b498-fd457213c0a7"))
+	: GlobalExecutionTimeoutSeconds(-1)
+	, Salt(TEXT("807c6a49-0657-49f3-b498-fd457213c0a7"))
 	, BaseDirectoryPath(TEXT("Engine/Binaries/Win64"))
 	, RemoteExecutor(nullptr)
 	, ContentAddressableStorage(nullptr)
@@ -50,6 +51,7 @@ public:
 		check(GConfig && GConfig->IsReadyForUse());
 
 		GConfig->GetString(TEXT("DerivedDataBuildRemoteExecutor"), TEXT("InstanceName"), InstanceName, GEngineIni);
+		GConfig->GetInt(TEXT("DerivedDataBuildRemoteExecutor"), TEXT("GlobalExecutionTimeoutSeconds"), GlobalExecutionTimeoutSeconds, GEngineIni);
 
 		const FName RemoteExecutionFeatureName(TEXT("RemoteExecution"));
 		IModularFeatures& ModularFeatures = IModularFeatures::Get();
@@ -515,7 +517,12 @@ public:
 		State.ExecuteRequest.InstanceName = InstanceName;
 		State.ExecuteRequest.ActionDigest = State.ActionDigest;
 		State.ExecuteRequest.SkipCacheLookup = true;
-		return RemoteExecutor->GetExecution()->Execute(State.ExecuteRequest, State.ExecuteResponse);
+		int64 TimeoutMs = 0;
+		if (GlobalExecutionTimeoutSeconds > 0)
+		{
+			TimeoutMs = GlobalExecutionTimeoutSeconds * 1000LL;
+		}
+		return RemoteExecutor->GetExecution()->Execute(State.ExecuteRequest, State.ExecuteResponse, TimeoutMs);
 	}
 
 	FStatus DownloadResults(FRemoteExecutionState& State)
@@ -871,6 +878,7 @@ private:
 	FStats Stats;
 	FLimitingHeuristics LimitingHeuristics;
 	FString InstanceName;
+	int GlobalExecutionTimeoutSeconds;
 	const FStringView Salt;
 	const FString BaseDirectoryPath;
 	IRemoteExecutor* RemoteExecutor;
