BulkDataRegistry: Store PayloadIds in separate cache key from BulkDataList, because updating the BulkDataList cache value with new data breaks the rule that cache values are immutable.
#rnx

diff --git a/Engine/Source/Editor/UnrealEd/Private/EditorDomain/BulkDataRegistryEditorDomain.cpp b/Engine/Source/Editor/UnrealEd/Private/EditorDomain/BulkDataRegistryEditorDomain.cpp
--- a/Engine/Source/Editor/UnrealEd/Private/EditorDomain/BulkDataRegistryEditorDomain.cpp
+++ b/Engine/Source/Editor/UnrealEd/Private/EditorDomain/BulkDataRegistryEditorDomain.cpp
@@ -18,6 +18,7 @@
 #include "Templates/RefCounting.h"
 #include "TickableEditorObject.h"
 #include "UObject/Package.h"
+#include "UObject/UObjectGlobals.h"
 #include "UObject/UObjectIterator.h"
 
 namespace UE::BulkDataRegistry::Private
@@ -25,11 +26,13 @@ namespace UE::BulkDataRegistry::Private
 
 namespace Constants
 {
-	constexpr double UpdateHashesInitialWaitDuration = 1.0;
 	constexpr uint64 TempLoadedPayloadsSizeBudget = 1024 * 1024 * 100;
 	constexpr double TempLoadedPayloadsDuration = 60.;
-	constexpr double PendingPackageCooldown = 60.;
+}
 
+TConstArrayView<uint8> MakeArrayView(FSharedBuffer Buffer)
+{
+	return TConstArrayView<uint8>(reinterpret_cast<const uint8*>(Buffer.GetData()), Buffer.GetSize());
 }
 
 /** Add a hook to the BulkDataRegistry's startup delegate to use the EditorDomain as the BulkDataRegistry */
@@ -75,7 +78,7 @@ FBulkDataRegistryEditorDomain::~FBulkDataRegistryEditorDomain()
 		Swap(LocalUpdatingPayloads, UpdatingPayloads);
 	}
 
-	// Since the AsyncTasks can no longer call the Requesters, we have to do it
+	// Since the UpdatingPayloads AsyncTasks can no longer access their Requesters, we have to call those callbacks
 	for (TPair<FGuid, FUpdatingPayload>& Pair : LocalUpdatingPayloads)
 	{
 		for (TUniqueFunction<void(bool, const FCompressedBuffer&)>& Requester : Pair.Value.Requesters)
@@ -86,9 +89,58 @@ FBulkDataRegistryEditorDomain::~FBulkDataRegistryEditorDomain()
 	LocalUpdatingPayloads.Empty();
 
 	// Clear PendingPackages
+	TMap<FName, TUniquePtr<FPendingPackage>> LocalPendingPackages;
 	{
 		FScopeLock PendingPackageScopeLock(&PendingPackageLock);
-		PendingPackages.Empty();
+		// Take custody of PendingPackages
+		Swap(LocalPendingPackages, PendingPackages);
+	}
+	for (TPair<FName, TUniquePtr<FPendingPackage>>& PendingPackagePair : LocalPendingPackages)
+	{
+		PendingPackagePair.Value->Cancel();
+	}
+	LocalPendingPackages.Empty();
+
+	// Clear PendingPackagePayloadIds. We have to take custody of PendingPayloadIds after calling
+	// Cancel from all FPendingPackage, as the FPendingPackages have callbacks that may write to PendingPayloadIds
+	TMap<FGuid, TRefCountPtr<FPendingPayloadId>> LocalPendingPayloadIds;
+	{
+		FWriteScopeLock RegistryScopeLock(RegistryLock);
+		// Take custody of PendingPayloadIds
+		Swap(LocalPendingPayloadIds, PendingPayloadIds);
+	}
+	for (TPair<FGuid, TRefCountPtr<FPendingPayloadId>>& PendingPayloadIdPair : LocalPendingPayloadIds)
+	{
+		PendingPayloadIdPair.Value->Cancel();
+	}
+}
+
+void FBulkDataRegistryEditorDomain::OnEndLoadPackage(TConstArrayView<UPackage*> LoadedPackages)
+{
+	FScopeLock PendingPackageScopeLock(&PendingPackageLock);
+
+	TArray<TUniquePtr<FPendingPackage>> PackagesToWrite;
+	for (TMap<FName, TUniquePtr<FPendingPackage>>::TIterator Iter(PendingPackages); Iter; ++Iter)
+	{
+		bool bShouldRemove;
+		bool bShouldWriteCache;
+		Iter.Value()->OnEndLoad(bShouldRemove, bShouldWriteCache);
+		// We currently do not hold a lock when writing the cache, so we require that it not be used
+		// by any other thread if we need to write the cache
+		check(!bShouldWriteCache || bShouldRemove);
+		if (bShouldRemove)
+		{
+			if (bShouldWriteCache)
+			{
+				PackagesToWrite.Add(MoveTemp(Iter.Value()));
+			}
+			Iter.RemoveCurrent();
+		}
+	}
+
+	for (TUniquePtr<FPendingPackage>& Package : PackagesToWrite)
+	{
+		Package->WriteCache();
 	}
 }
 
@@ -103,6 +155,7 @@ void FBulkDataRegistryEditorDomain::Register(UPackage* Owner, const UE::Virtuali
 	UE::Virtualization::FVirtualizedUntypedBulkData CopyBulk(BulkData.CopyTornOff());
 	if (Owner
 		&& Owner->GetFileSize() // We only record the BulkDataList for disk packages
+		&& !Owner->bHasBeenEndLoaded // We only record BulkDats that are loaded before the package finishes loading
 		&& CopyBulk.CanSaveForRegistry()
 		)
 	{
@@ -110,14 +163,9 @@ void FBulkDataRegistryEditorDomain::Register(UPackage* Owner, const UE::Virtuali
 		AddPendingPackageBulkData(PackageName, CopyBulk);
 	}
 
-	FName PackageNameToUpdate = CopyBulk.HasPlaceholderPayloadId() ? PackageName : NAME_None;
 	FWriteScopeLock RegistryScopeLock(RegistryLock);
 	check(bActive); // Registrations should not come in after we destruct
-	Registry.Add(BulkData.GetIdentifier(), FRegisteredBulk(MoveTemp(CopyBulk), PackageNameToUpdate));
-}
-
-void FBulkDataRegistryEditorDomain::OnEndLoadPackage(TConstArrayView<UPackage*> LoadedPackages)
-{
+	Registry.Add(BulkData.GetIdentifier(), FRegisteredBulk(MoveTemp(CopyBulk), PackageName));
 }
 
 void FBulkDataRegistryEditorDomain::OnExitMemory(const UE::Virtualization::FVirtualizedUntypedBulkData& BulkData)
@@ -265,7 +313,6 @@ void FBulkDataRegistryEditorDomain::TickCook(float DeltaTime, bool bTickComplete
 	bool bWaitForCooldown = !bTickComplete;
 	check(bActive); // Ticks should not come in after we destruct
 
-	PollPendingPackages(bWaitForCooldown);
 	{
 		FWriteScopeLock RegistryScopeLock(RegistryLock);
 		PruneTempLoadedPayloads();
@@ -280,57 +327,17 @@ void FBulkDataRegistryEditorDomain::Tick(float DeltaTime)
 void FBulkDataRegistryEditorDomain::AddPendingPackageBulkData(FName PackageName, const UE::Virtualization::FVirtualizedUntypedBulkData& BulkData)
 {
 	FScopeLock PendingPackageScopeLock(&PendingPackageLock);
-	AddPendingPackageBulkDataInsideLock(PackageName, BulkData);
-}
-
-void FBulkDataRegistryEditorDomain::AddPendingPackageBulkDataInsideLock(FName PackageName, const UE::Virtualization::FVirtualizedUntypedBulkData& BulkData)
-{
 	check(bActive); // Registrations should not come in after we destruct, and AsyncTasks should check bActive before calling
 	TUniquePtr<FPendingPackage>& PendingPackage = PendingPackages.FindOrAdd(PackageName);
 	if (!PendingPackage.IsValid())
 	{
 		PendingPackage = MakeUnique<FPendingPackage>(PackageName, this);
 	}
-	PendingPackage->AddBulkData(BulkData);
-}
-
-void FBulkDataRegistryEditorDomain::PollPendingPackages(bool bWaitForCooldown)
-{
-	// BULKDATAREGISTRY_TODO: Reduce array iteration time; tick is called frequently in the editor
-	TArray<TPair<FName, TUniquePtr<FPendingPackage>>> PoppedPackages;
-	{
-		FScopeLock PendingPackageScopeLock(&PendingPackageLock);
-		if (PendingPackages.Num() == 0)
-		{
-			return;
-		}
-
-		double CreateTimeToPush = 0.0;
-		if (bWaitForCooldown)
-		{
-			CreateTimeToPush = FPlatformTime::Seconds() - Constants::PendingPackageCooldown;
-		}
-		for (TPair<FName, TUniquePtr<FPendingPackage>>& Pair : PendingPackages)
-		{
-			FPendingPackage* PendingPackage = Pair.Value.Get();
-			if (PendingPackage->IsReadFinished() && (!bWaitForCooldown || PendingPackage->GetCreateTime() <= CreateTimeToPush))
-			{
-				PoppedPackages.Emplace(Pair.Key, MoveTemp(Pair.Value));
-			}
-		}
-		for (TPair<FName, TUniquePtr<FPendingPackage>>& Pair : PoppedPackages)
-		{
-			PendingPackages.Remove(Pair.Get<0>());
-		}
-	}
-
-	TArray<UE::Virtualization::FVirtualizedUntypedBulkData> BulkDatas;
-	for (TPair<FName, TUniquePtr<FPendingPackage>>& Pair : PoppedPackages)
+	if (!PendingPackage->IsLoadInProgress())
 	{
-		FName PackageName = Pair.Get<0>();
-		TUniquePtr<FPendingPackage>& PendingPackage = Pair.Get<1>();
-		PendingPackage->UpdateCache();
+		return;
 	}
+	PendingPackage->AddBulkData(BulkData);
 }
 
 void FBulkDataRegistryEditorDomain::AddTempLoadedPayload(const FGuid& RegistryKey, uint64 PayloadSize)
@@ -363,81 +370,269 @@ void FBulkDataRegistryEditorDomain::PruneTempLoadedPayloads()
 	}
 }
 
-TConstArrayView<uint8> MakeArrayView(FSharedBuffer Buffer)
+void FBulkDataRegistryEditorDomain::WritePayloadIdToCache(FName PackageName, const UE::Virtualization::FVirtualizedUntypedBulkData& BulkData) const
 {
-	return TConstArrayView<uint8>(reinterpret_cast<const uint8*>(Buffer.GetData()), Buffer.GetSize());
+	check(!PackageName.IsNone());
+	TArray<uint8> Bytes;
+	FMemoryWriter Writer(Bytes);
+	const_cast<UE::Virtualization::FVirtualizedUntypedBulkData&>(BulkData).SerializeForRegistry(Writer);
+	UE::EditorDomain::PutBulkDataPayloadId(PackageName, BulkData.GetIdentifier(), MakeSharedBufferFromArray(MoveTemp(Bytes)));
+}
+
+void FBulkDataRegistryEditorDomain::ReadPayloadIdsFromCache(FName PackageName, TArray<TRefCountPtr<FPendingPayloadId>>&& OldPendings,
+	TArray<TPair<TRefCountPtr<FPendingPayloadId>, UE::DerivedData::FRequest>>&& NewPendings)
+{
+	// Cancel any old requests for the Guids in NewPendings; we are about to overwrite them
+	// This cancelation has to occur outside of any lock, since the task may be in progress and
+	// and to enter the lock and Cancel will wait on it
+	for (TRefCountPtr<FPendingPayloadId>& OldPending : OldPendings)
+	{
+		OldPending->GetRequest().Cancel();
+	}
+	OldPendings.Empty();
+	for (TPair<TRefCountPtr<FPendingPayloadId>, UE::DerivedData::FRequest>& NewPair : NewPendings)
+	{
+		TRefCountPtr<FPendingPayloadId>& NewPending = NewPair.Get<0>();
+		// Creation of the Request has to occur outside of any lock, because the request
+		// may execute immediately on this thread and need to enter the lock; our locks are non-reentrant
+		UE::DerivedData::FRequest& Request = NewPair.Get<1>();
+		Request = UE::EditorDomain::GetBulkDataPayloadId(PackageName, NewPending->GetBulkDataId(),
+			[this, PackageName, NewPending](FSharedBuffer Buffer)
+		{
+			if (Buffer.IsNull())
+			{
+				return;
+			}
+			FMemoryReaderView Reader(MakeArrayView(Buffer));
+			UE::Virtualization::FVirtualizedUntypedBulkData CachedBulkData;
+			CachedBulkData.SerializeForRegistry(Reader);
+			const FGuid& BulkDataId = NewPending->GetBulkDataId();
+			if (Reader.IsError() || CachedBulkData.GetIdentifier() != BulkDataId)
+			{
+				UE_LOG(LogEditorDomain, Warning, TEXT("Corrupt cache data for BulkDataPayloadId %s."), WriteToString<192>(PackageName, TEXT("/"), BulkDataId).ToString());
+				return;
+			}
+
+			FWriteScopeLock RegistryScopeLock(RegistryLock);
+			if (!bActive)
+			{
+				return;
+			}
+			TRefCountPtr<FPendingPayloadId> ExistingPending;
+			if (!PendingPayloadIds.RemoveAndCopyValue(BulkDataId, ExistingPending))
+			{
+				return;
+			}
+			check(ExistingPending->GetBulkDataId() == BulkDataId);
+			if (ExistingPending != NewPending)
+			{
+				// We removed ExistingPending because we thought it was equal to NewPending, but it's not, so put it back
+				PendingPayloadIds.Add(BulkDataId, MoveTemp(ExistingPending));
+				return;
+			}
+
+			FRegisteredBulk* ExistingRegisteredBulk = Registry.Find(BulkDataId);
+			if (!ExistingRegisteredBulk)
+			{
+				return;
+			}
+
+			UE::Virtualization::FVirtualizedUntypedBulkData& ExistingBulkData = ExistingRegisteredBulk->BulkData;
+			check(ExistingBulkData.GetIdentifier() == BulkDataId);
+			if (ExistingBulkData.HasPlaceholderPayloadId())
+			{
+				// BULKDATAREGISTRY_TODO: Implement LocationMatches
+				//&& CachedBulkData.LocationMatches(ExistingBulkData);
+
+				ExistingBulkData = CachedBulkData;
+			}
+		});
+	}
+
+	// Assign the Requests we just created to the data for each guid in the map,
+	// which has to be edited only within the lock
+	// If for any reason (race condition, shutting down) we can not assign a Request,
+	// we have to cancel the request before returning, to make sure the callback does not hold a pointer
+	// to this that could become dangling.
+	{
+		FWriteScopeLock RegistryScopeLock(RegistryLock);
+		if (!bActive)
+		{
+			for (TPair<TRefCountPtr<FPendingPayloadId>, UE::DerivedData::FRequest>& NewPair : NewPendings)
+			{
+				OldPendings.Add(MoveTemp(NewPair.Get<0>()));
+			}
+		}
+		else
+		{
+			for (TPair<TRefCountPtr<FPendingPayloadId>, UE::DerivedData::FRequest>& NewPair : NewPendings)
+			{
+				TRefCountPtr<FPendingPayloadId>& NewPending = NewPair.Get<0>();
+				UE::DerivedData::FRequest& Request = NewPair.Get<1>();
+
+				TRefCountPtr<FPendingPayloadId>* ExistingPending = PendingPayloadIds.Find(NewPending->GetBulkDataId());
+				if (!ExistingPending || *ExistingPending != NewPending)
+				{
+					OldPendings.Add(MoveTemp(NewPending));
+				}
+				else
+				{
+					check(NewPending->GetRequest().IsNull());
+					NewPending->GetRequest() = Request;
+				}
+			}
+		}
+	}
+	NewPendings.Empty();
+	for (TRefCountPtr<FPendingPayloadId>& OldPending : OldPendings)
+	{
+		OldPending->GetRequest().Cancel();
+	}
+	OldPendings.Empty();
 }
 
 FPendingPackage::FPendingPackage(FName InPackageName, FBulkDataRegistryEditorDomain* InOwner)
 	: PackageName(InPackageName)
 	, Owner(InOwner)
-	, CreateTime(FPlatformTime::Seconds())
 {
-	CacheRequest = UE::EditorDomain::GetBulkDataList(PackageName, [this](FSharedBuffer Buffer)
-		{ OnCacheResults(Buffer); });
+	PendingOperations = Flag_EndLoad | Flag_BulkDataListResults;
+
+	BulkDataListCacheRequest = UE::EditorDomain::GetBulkDataList(PackageName, [this](FSharedBuffer Buffer)
+		{ OnBulkDataListResults(Buffer); });
 }
 
-FPendingPackage::~FPendingPackage()
+void FPendingPackage::Cancel()
 {
-	CacheRequest.Cancel();
+	// Called from outside Owner->PendingPackagesLock, so OnBulkDataList can complete on other thread while we wait
+	// Called after removing this from Owner->PendingPackages under a previous cover of the lock
+	// If OnBulkDataList is running on other thread its attempt to remove from PendingPackages will be a noop
+	if (!BulkDataListCacheRequest->Poll())
+	{
+		// Optimization: prevent WriteCache from running at all if we reach here first
+		PendingOperations.fetch_or(Flag_Canceled);
+		BulkDataListCacheRequest.Cancel();
+	}
 }
 
-void FPendingPackage::OnCacheResults(FSharedBuffer Buffer)
+void FPendingPackage::OnEndLoad(bool& bOutShouldRemove, bool& bOutShouldWriteCache)
 {
-	if (Buffer.IsNull())
+	// Called from within Owner->PendingPackageLock
+	bLoadInProgress = false;
+	if (PendingOperations.fetch_and(~Flag_EndLoad) == Flag_EndLoad)
 	{
-		return;
+		bOutShouldWriteCache = true;
+		bOutShouldRemove = true;
+	}
+	bOutShouldWriteCache = false;
+	bOutShouldRemove = false;
+}
+
+void FPendingPackage::OnBulkDataListResults(FSharedBuffer Buffer)
+{
+	if (!Buffer.IsNull())
+	{
+		FMemoryReaderView Reader(MakeArrayView(Buffer));
+		Serialize(Reader, CachedBulkDatas);
+		if (Reader.IsError())
+		{
+			CachedBulkDatas.Empty();
+		}
 	}
-	FMemoryReaderView Reader(MakeArrayView(Buffer));
-	Serialize(Reader, CachedBulkDatas);
-	if (Reader.IsError())
+
+	ReadCache();
+
+	if (PendingOperations.fetch_and(~Flag_BulkDataListResults) == Flag_BulkDataListResults)
 	{
-		CachedBulkDatas.Empty();
+		WriteCache();
+
+		FScopeLock PendingPackageScopeLock(&Owner->PendingPackageLock);
+		Owner->PendingPackages.Remove(PackageName);
+		// *this has been deleted and can no longer be accessed
 	}
+}
 
-	// Add each CachedBulkData to the Registry. If the CachedBulkData exists, do not overwrite it, but do update the RawHash if it is missing it.
-	FWriteScopeLock RegistryScopeLock(Owner->RegistryLock);
-	if (!Owner->bActive)
+void FPendingPackage::ReadCache()
+{
+	if (CachedBulkDatas.Num() == 0)
 	{
 		return;
 	}
-	for (const UE::Virtualization::FVirtualizedUntypedBulkData& BulkData : CachedBulkDatas)
+
+	TArray<TRefCountPtr<FPendingPayloadId>> OldPendings;
+	TArray<TPair<TRefCountPtr<FPendingPayloadId>, UE::DerivedData::FRequest>> NewPendings;
+
+	// Add each CachedBulkData to the Registry, updating RawHash if it is missing.
+	// For every BulkData in this package in the Registry after the CachedBulkData has been added, 
+	// if the RawHash is missing from the CachedBulkData as well, queue a read of its RawHash
+	// from the separate PlaceholderPayloadId BulkTablePayloadId cache bucket.
 	{
-		FRegisteredBulk& TargetRegisteredBulk = Owner->Registry.FindOrAdd(BulkData.GetIdentifier());
-		UE::Virtualization::FVirtualizedUntypedBulkData& TargetBulkData = TargetRegisteredBulk.BulkData;
-		if (!TargetBulkData.GetIdentifier().IsValid())
+		FWriteScopeLock RegistryScopeLock(Owner->RegistryLock);
+		if (!Owner->bActive)
 		{
-			TargetBulkData = BulkData;
-			TargetRegisteredBulk.PackageNameToUpdate = BulkData.HasPlaceholderPayloadId() ? PackageName : NAME_None;
+			return;
 		}
-		else
+		for (const UE::Virtualization::FVirtualizedUntypedBulkData& BulkData : CachedBulkDatas)
 		{
-			check(TargetBulkData.GetIdentifier() == BulkData.GetIdentifier());
-			if (!BulkData.HasPlaceholderPayloadId() && TargetBulkData.HasPlaceholderPayloadId())
+			FGuid BulkDataId = BulkData.GetIdentifier();
+			FRegisteredBulk& TargetRegisteredBulk = Owner->Registry.FindOrAdd(BulkData.GetIdentifier());
+
+			bool bCachedLocationMatches = true;
+			UE::Virtualization::FVirtualizedUntypedBulkData& TargetBulkData = TargetRegisteredBulk.BulkData;
+			if (!TargetBulkData.GetIdentifier().IsValid())
 			{
-				// BULKDATAREGISTRY_TODO: Verify the payloads are the same
-				// && TargetBulkData.PayloadInfoMatches(BulkData)
 				TargetBulkData = BulkData;
-				// We now know that we don't need to update the package with the RawHash, so remove the PackageNameToUpdate
-				TargetRegisteredBulk.PackageNameToUpdate = NAME_None;
+				TargetRegisteredBulk.PackageName = PackageName;
+			}
+			else
+			{
+				check(TargetBulkData.GetIdentifier() == BulkDataId);
+				// BULKDATAREGISTRY_TODO: Implement LocationMatches
+				// bCachedLocationMatches = BulkData.LocationMatches(TargetBulkData);
+				if (bCachedLocationMatches && !BulkData.HasPlaceholderPayloadId() && TargetBulkData.HasPlaceholderPayloadId())
+				{
+					TargetBulkData = BulkData;
+					TargetRegisteredBulk.PackageName = PackageName;
+				}
+			}
+
+			if (bCachedLocationMatches && TargetBulkData.HasPlaceholderPayloadId())
+			{
+				NewPendings.Emplace(MakeTuple(new FPendingPayloadId(BulkDataId), UE::DerivedData::FRequest()));
+			}
+		}
+
+		if (NewPendings.Num())
+		{
+			for (TPair<TRefCountPtr<FPendingPayloadId>, UE::DerivedData::FRequest>& NewPair : NewPendings)
+			{
+				TRefCountPtr<FPendingPayloadId>& NewPending = NewPair.Get<0>();
+				TRefCountPtr<FPendingPayloadId>& ExistingPending = Owner->PendingPayloadIds.FindOrAdd(NewPending->GetBulkDataId());
+				if (ExistingPending.IsValid())
+				{
+					OldPendings.Add(MoveTemp(ExistingPending));
+				}
+				ExistingPending = NewPending;
 			}
 		}
 	}
-}
 
-bool FPendingPackage::IsReadFinished() const
-{
-	return CacheRequest.Poll();
+	Owner->ReadPayloadIdsFromCache(PackageName, MoveTemp(OldPendings), MoveTemp(NewPendings));
 }
 
-void FPendingPackage::UpdateCache()
+void FPendingPackage::WriteCache()
 {
-	CacheRequest.Wait();
+	// If the BulkDataList cache read found some existing results, then exit; cache results are deterministic so there
+	// is no need to write the list to the cache again.
+	if (CachedBulkDatas.Num() > 0)
+	{
+		return;
+	}
 
-	// Remove any duplicates in the runtime BulkDatas; elements later in the list were added after and override earlier elements
+	check(BulkDatas.Num() > 0); // We should have >= 1 bulkdatas, or we would not have created the FPendingPackage
+	// Remove any duplicates in the runtime BulkDatas; elements later in the list override earlier elements
 	{
 		TSet<FGuid> BulkDataGuids;
-		for (int32 Index = BulkDatas.Num()-1; Index >= 0; --Index)
+		for (int32 Index = BulkDatas.Num() - 1; Index >= 0; --Index)
 		{
 			bool bAlreadyExists;
 			BulkDataGuids.Add(BulkDatas[Index].GetIdentifier(), &bAlreadyExists);
@@ -448,58 +643,16 @@ void FPendingPackage::UpdateCache()
 		}
 	}
 
-	// If the current list of BulkDatas has any elements not in the cache, update the cache
-	TMap<FGuid, UE::Virtualization::FVirtualizedUntypedBulkData*> GuidToCachedBulkData;
-	for (UE::Virtualization::FVirtualizedUntypedBulkData& BulkData : CachedBulkDatas)
-	{
-		GuidToCachedBulkData.Add(BulkData.GetIdentifier(), &BulkData);
-	}
-
-	// Update any existing BulkDatas and add on any new ones
-	bool bUpdatedExisting = false;
-	TArray<UE::Virtualization::FVirtualizedUntypedBulkData*> BulkDatasToAdd;
-	for (UE::Virtualization::FVirtualizedUntypedBulkData& BulkData : BulkDatas)
-	{
-		UE::Virtualization::FVirtualizedUntypedBulkData** CachedBulkData = GuidToCachedBulkData.Find(BulkData.GetIdentifier());
-		if (CachedBulkData)
-		{
-			UE::Virtualization::FVirtualizedUntypedBulkData& Cached = **CachedBulkData;
-			if (Cached.HasPlaceholderPayloadId() && !BulkData.HasPlaceholderPayloadId())
-				// BULKDATAREGISTRY_TODO: Also update if the payloads are different
-				// || !TargetBulkData.PayloadInfoMatches(BulkData)
-			{
-				bUpdatedExisting = true;
-				Cached = BulkData;
-			}
-		}
-		else
-		{
-			BulkDatasToAdd.Add(&BulkData);
-		}
-	}
-
-	if (bUpdatedExisting || BulkDatasToAdd.Num() > 0)
-	{
-		// Add on all the new ones
-		GuidToCachedBulkData.Reset(); // This map is now invalidated since we are modifying the array it has pointers into
-		for (UE::Virtualization::FVirtualizedUntypedBulkData* BulkData : BulkDatasToAdd)
+	// Sort the list by guid, to avoid indeterminism in the list
+	BulkDatas.Sort([](const UE::Virtualization::FVirtualizedUntypedBulkData& A, const UE::Virtualization::FVirtualizedUntypedBulkData& B)
 		{
-			CachedBulkDatas.Add(*BulkData);
-		}
-
-		// Sort the list by guid, to avoid indeterminism in the list
-		CachedBulkDatas.Sort([](const UE::Virtualization::FVirtualizedUntypedBulkData& A, const UE::Virtualization::FVirtualizedUntypedBulkData& B)
-			{
-				return A.GetIdentifier() < B.GetIdentifier();
-			});
+			return A.GetIdentifier() < B.GetIdentifier();
+		});
 
-		TArray<uint8> Bytes;
-		FMemoryWriter Writer(Bytes);
-		Serialize(Writer, CachedBulkDatas);
-		// BULKDATAREGISTRY_TODO: Store the request in a list of pending writes so that we don't try to write a new entry for the package
-		// until the previous entry has finished storing, to avoid stale data overwriting new data
-		UE::EditorDomain::PutBulkDataList(PackageName, MakeSharedBufferFromArray(MoveTemp(Bytes)));
-	}
+	TArray<uint8> Bytes;
+	FMemoryWriter Writer(Bytes);
+	Serialize(Writer, BulkDatas);
+	UE::EditorDomain::PutBulkDataList(PackageName, MakeSharedBufferFromArray(MoveTemp(Bytes)));
 }
 
 FUpdatePayloadWorker::FUpdatePayloadWorker(FBulkDataRegistryEditorDomain* InBulkDataRegistry,
@@ -567,17 +720,9 @@ void FUpdatePayloadWorker::DoWork()
 			BulkDataRegistry->AddTempLoadedPayload(BulkData.GetIdentifier(), BulkData.GetPayloadSize());
 			BulkDataRegistry->PruneTempLoadedPayloads();
 
-			// If the BulkData is marked to be stored in a package's BulkDataList, add an update task for the package.
-			// We have to enter the two locks at the same time, because it is only while inside the RegistryScopeLock that we are guaranteed our work has not been canceled.
-			if (!RegisteredBulk->PackageNameToUpdate.IsNone())
+			if (!RegisteredBulk->PackageName.IsNone())
 			{
-				FScopeLock PendingPackageScopeLock(&BulkDataRegistry->PendingPackageLock);
-				BulkDataRegistry->AddPendingPackageBulkDataInsideLock(RegisteredBulk->PackageNameToUpdate, BulkData);
-				if (!BulkData.HasPlaceholderPayloadId())
-				{
-					// If the BulkData already has its PayloadId, then we no longer need to update it after adding it this time
-					RegisteredBulk->PackageNameToUpdate = NAME_None;
-				}
+				BulkDataRegistry->WritePayloadIdToCache(RegisteredBulk->PackageName, BulkData);
 			}
 			break;
 		}
@@ -593,6 +738,11 @@ void FUpdatePayloadWorker::DoWork()
 	}
 }
 
+void FPendingPayloadId::Cancel()
+{
+	Request.Cancel();
+}
+
 void Serialize(FArchive& Ar, TArray<UE::Virtualization::FVirtualizedUntypedBulkData>& Datas)
 {
 	int32 Num = Datas.Num();
diff --git a/Engine/Source/Editor/UnrealEd/Private/EditorDomain/BulkDataRegistryEditorDomain.h b/Engine/Source/Editor/UnrealEd/Private/EditorDomain/BulkDataRegistryEditorDomain.h
--- a/Engine/Source/Editor/UnrealEd/Private/EditorDomain/BulkDataRegistryEditorDomain.h
+++ b/Engine/Source/Editor/UnrealEd/Private/EditorDomain/BulkDataRegistryEditorDomain.h
@@ -4,6 +4,7 @@
 
 #include "Async/AsyncWork.h"
 #include "Async/Future.h"
+#include <atomic>
 #include "Containers/Array.h"
 #include "Containers/Map.h"
 #include "Containers/RingBuffer.h"
@@ -27,13 +28,13 @@ class FBulkDataRegistryEditorDomain;
 struct FRegisteredBulk
 {
 	FRegisteredBulk() = default;
-	FRegisteredBulk(UE::Virtualization::FVirtualizedUntypedBulkData&& InBulkData, FName InPackageNameToUpdate = NAME_None)
-		:BulkData(MoveTemp(InBulkData)), PackageNameToUpdate(InPackageNameToUpdate)
+	FRegisteredBulk(UE::Virtualization::FVirtualizedUntypedBulkData&& InBulkData, FName InPackageName = NAME_None)
+		:BulkData(MoveTemp(InBulkData)), PackageName(InPackageName)
 	{
 	}
 
 	UE::Virtualization::FVirtualizedUntypedBulkData BulkData;
-	FName PackageNameToUpdate;
+	FName PackageName;
 	bool bHasTempPayload = false;
 };
 
@@ -45,30 +46,48 @@ class FPendingPackage
 {
 public:
 	FPendingPackage(FName PackageName, FBulkDataRegistryEditorDomain* InOwner);
-	FPendingPackage(FPendingPackage&& Other) = default;
+	FPendingPackage(FPendingPackage&& Other) = delete;
 	FPendingPackage(const FPendingPackage& Other) = delete;
-	~FPendingPackage();
 
-	bool IsReadFinished() const;
-	void UpdateCache();
-	double GetCreateTime() const
-	{
-		return CreateTime;
-	}
+	void Cancel();
 	void AddBulkData(const UE::Virtualization::FVirtualizedUntypedBulkData& BulkData)
 	{
 		BulkDatas.Add(BulkData);
 	}
+	void OnEndLoad(bool& bOutShouldRemove, bool& bOutShouldWriteCache);
+	bool IsLoadInProgress() const
+	{
+		return bLoadInProgress;
+	}
+	void WriteCache();
 
 private:
-	void OnCacheResults(FSharedBuffer Buffer);
+	void OnBulkDataListResults(FSharedBuffer Buffer);
+	void ReadCache();
 
+	enum EFlags : int32
+	{
+		Flag_EndLoad = 1 << 0,
+		Flag_BulkDataListResults = 1 << 1,
+		Flag_Canceled = 1 << 2,
+	};
 	FName PackageName;
 	TArray<UE::Virtualization::FVirtualizedUntypedBulkData> BulkDatas;
 	TArray<UE::Virtualization::FVirtualizedUntypedBulkData> CachedBulkDatas;
-	UE::DerivedData::FRequest CacheRequest;
+	TArray<UE::DerivedData::FRequest> PayloadIdCacheRequests;
+	UE::DerivedData::FRequest BulkDataListCacheRequest;
 	FBulkDataRegistryEditorDomain* Owner;
-	double CreateTime;
+	/**
+	 * When PendingOperations reaches zero, we can remove the FPendingPackage.
+	 */
+	std::atomic<int32> PendingOperations;
+	/**
+	 * True until the LoadPackage for *this is complete. The FPendingPackage may last longer
+	 * than the initial load period while it waits for a CacheRead, but BulkDatas can only
+	 * be written to its list during its initial load, to avoid non-deterministic changes
+	 * to the list when operations occur on the UPackage in the editor.
+	 */
+	bool bLoadInProgress = true;
 };
 
 /** Data about a BulkData that has loaded its payload for TryGetMeta and should drop it after GetData or a timeout. */
@@ -114,6 +133,32 @@ struct FUpdatingPayload
 	TArray<TUniqueFunction<void(bool bValid, const FCompressedBuffer& Buffer)>> Requesters;
 };
 
+/** Data storage for a BulkData that is loading its PayloadId from the cache. */
+class FPendingPayloadId : public FThreadSafeRefCountedObject
+{
+public:
+	FPendingPayloadId(const FGuid& InBulkDataId)
+		:BulkDataId(InBulkDataId)
+	{
+	}
+	FPendingPayloadId(FPendingPayloadId&& Other) = delete;
+	FPendingPayloadId(const FPendingPayloadId& Other) = delete;
+
+	void Cancel();
+
+	UE::DerivedData::FRequest& GetRequest()
+	{
+		return Request;
+	}
+	const FGuid& GetBulkDataId()
+	{
+		return BulkDataId;
+	}
+
+private:
+	FGuid BulkDataId;
+	UE::DerivedData::FRequest Request;
+};
 
 /** Implementation of a BulkDataRegistry that stores its persistent data in a DDC bucket. */
 class FBulkDataRegistryEditorDomain : public IBulkDataRegistry, public FTickableEditorObject, public FTickableCookObject
@@ -136,15 +181,18 @@ public:
 
 private:
 	void AddPendingPackageBulkData(FName PackageName, const UE::Virtualization::FVirtualizedUntypedBulkData& BulkData);
-	void AddPendingPackageBulkDataInsideLock(FName PackageName, const UE::Virtualization::FVirtualizedUntypedBulkData& BulkData);
-	void OnEndLoadPackage(TConstArrayView<UPackage*> LoadedPackages);
 	void PollPendingPackages(bool bWaitForCooldown);
 	void AddTempLoadedPayload(const FGuid& RegistryKey, uint64 PayloadSize);
 	void PruneTempLoadedPayloads();
+	void OnEndLoadPackage(TConstArrayView<UPackage*> LoadedPackages);
+	void WritePayloadIdToCache(FName PackageName, const UE::Virtualization::FVirtualizedUntypedBulkData& BulkData) const;
+	void ReadPayloadIdsFromCache(FName PackageName, TArray<TRefCountPtr<FPendingPayloadId>>&& OldPendings,
+		TArray<TPair<TRefCountPtr<FPendingPayloadId>, UE::DerivedData::FRequest>>&& NewPendings);
+
 	friend class FPendingPackage;
 	friend class FUpdatePayloadWorker;
 
-	// All three locks can  be held at the same time. They must always be entered in order: SharedDataLock, RegistryLock, PendingPackageLock.
+	// All locks can be held at the same time. They must always be entered in order: SharedDataLock, RegistryLock, PendingPackageLock
 	TRefCountPtr<FTaskSharedDataLock> SharedDataLock;
 	FRWLock RegistryLock;
 	FCriticalSection PendingPackageLock;
@@ -152,6 +200,7 @@ private:
 	TMap<FGuid, FRegisteredBulk> Registry;
 	TMap<FGuid, FUpdatingPayload> UpdatingPayloads;
 	TMap<FName, TUniquePtr<FPendingPackage>> PendingPackages;
+	TMap<FGuid, TRefCountPtr<FPendingPayloadId>> PendingPayloadIds;
 	TRingBuffer<FTempLoadedPayload> TempLoadedPayloads;
 	uint64 TempLoadedPayloadsSize = 0;
 	bool bActive = true;
diff --git a/Engine/Source/Editor/UnrealEd/Private/EditorDomain/EditorDomainUtils.cpp b/Engine/Source/Editor/UnrealEd/Private/EditorDomain/EditorDomainUtils.cpp
--- a/Engine/Source/Editor/UnrealEd/Private/EditorDomain/EditorDomainUtils.cpp
+++ b/Engine/Source/Editor/UnrealEd/Private/EditorDomain/EditorDomainUtils.cpp
@@ -36,14 +36,14 @@ FClassDigestMap& GetClassDigests()
 
 // Change to a new guid when EditorDomain needs to be invalidated
 const TCHAR* EditorDomainVersion = TEXT("C217EB656E9B4C04816D3DC0E21901F6");
-// Identifier of the CacheBucket for EditorDomainPackages
+// Identifier of the CacheBuckets for EditorDomain tables
 const TCHAR* EditorDomainPackageBucketName = TEXT("EditorDomainPackage");
 const TCHAR* EditorDomainBulkDataListBucketName = TEXT("EditorDomainBulkDataList");
+const TCHAR* EditorDomainBulkDataPayloadIdBucketName = TEXT("EditorDomainBulkDataPayloadId");
 
-EPackageDigestResult GetPackageDigest(FPackageDigest& OutDigest, FString& OutErrorMessage,
+EPackageDigestResult AppendPackageDigest(FCbWriter& Writer, FString& OutErrorMessage,
 	const FAssetPackageData& PackageData, FName PackageName)
 {
-	FCbWriter Writer;
 	int32 CurrentFileVersionUE = GPackageFileUEVersion;
 	int32 CurrentFileVersionLicenseeUE = GPackageFileLicenseeUEVersion;
 	Writer << EditorDomainVersion;
@@ -101,7 +101,6 @@ EPackageDigestResult GetPackageDigest(FPackageDigest& OutDigest, FString& OutErr
 			}
 		}
 	}
-	OutDigest = Writer.Save().GetRangeHash();
 	return EPackageDigestResult::Success;
 }
 
@@ -168,6 +167,15 @@ void PrecacheClassDigests(TConstArrayView<FName> ClassNames)
 
 EPackageDigestResult GetPackageDigest(IAssetRegistry& AssetRegistry, FName PackageName,
 	FPackageDigest& OutPackageDigest, FString& OutErrorMessage)
+{
+	FCbWriter Builder;
+	EPackageDigestResult Result = AppendPackageDigest(AssetRegistry, PackageName, Builder, OutErrorMessage);
+	OutPackageDigest = Builder.Save().GetRangeHash();
+	return Result;
+}
+
+EPackageDigestResult AppendPackageDigest(IAssetRegistry& AssetRegistry, FName PackageName,
+	FCbWriter& Builder, FString& OutErrorMessage)
 {
 	AssetRegistry.WaitForPackage(PackageName.ToString());
 	TOptional<FAssetPackageData> PackageData = AssetRegistry.GetAssetPackageDataCopy(PackageName);
@@ -177,7 +185,7 @@ EPackageDigestResult GetPackageDigest(IAssetRegistry& AssetRegistry, FName Packa
 			*PackageName.ToString());
 		return EPackageDigestResult::FileDoesNotExist;
 	}
-	return GetPackageDigest(OutPackageDigest, OutErrorMessage, *PackageData, PackageName);
+	return AppendPackageDigest(Builder, OutErrorMessage, *PackageData, PackageName);
 }
 
 UE::DerivedData::FCacheKey GetEditorDomainPackageKey(const FPackageDigest& PackageDigest)
@@ -194,6 +202,14 @@ UE::DerivedData::FCacheKey GetBulkDataListKey(const FPackageDigest& PackageDiges
 	return UE::DerivedData::FCacheKey{ EditorDomainBulkDataListBucket, PackageDigest };
 }
 
+UE::DerivedData::FCacheKey GetBulkDataPayloadIdKey(const FIoHash& PackageAndGuidDigest)
+{
+	static UE::DerivedData::FCacheBucket EditorDomainBulkDataPayloadIdBucket =
+		GetDerivedDataCacheRef().CreateBucket(EditorDomainBulkDataPayloadIdBucketName);
+
+	return UE::DerivedData::FCacheKey{ EditorDomainBulkDataPayloadIdBucket, PackageAndGuidDigest };
+}
+
 UE::DerivedData::FRequest RequestEditorDomainPackage(const FPackagePath& PackagePath,
 	const FPackageDigest& PackageDigest, UE::DerivedData::ECachePolicy SkipFlags, UE::DerivedData::EPriority CachePriority,
 	UE::DerivedData::FOnCacheGetComplete&& Callback)
@@ -336,5 +352,66 @@ void PutBulkDataList(FName PackageName, FSharedBuffer Buffer)
 	Cache.Put(MakeArrayView(&Record, 1), WriteToString<128>(PackageName));
 }
 
+FIoHash GetPackageAndGuidDigest(FCbWriter& Builder, const FGuid& BulkDataId)
+{
+	Builder << BulkDataId;
+	return Builder.Save().GetRangeHash();
+}
+
+UE::DerivedData::FRequest GetBulkDataPayloadId(FName PackageName, const FGuid& BulkDataId,
+	TUniqueFunction<void(FSharedBuffer Buffer)>&& Callback)
+{
+	UE::DerivedData::ICache& Cache = GetDerivedDataCacheRef();
+
+	FString ErrorMessage;
+	FCbWriter Builder;
+	EPackageDigestResult FindHashResult = AppendPackageDigest(*IAssetRegistry::Get(), PackageName, Builder, ErrorMessage);
+	switch (FindHashResult)
+	{
+	case EPackageDigestResult::Success:
+		break;
+	default:
+	{
+		Callback(FSharedBuffer());
+		return UE::DerivedData::FRequest();
+	}
+	}
+	FIoHash PackageAndGuidDigest = GetPackageAndGuidDigest(Builder, BulkDataId);
+
+	return Cache.Get({ GetBulkDataPayloadIdKey(PackageAndGuidDigest) },
+		WriteToString<192>(PackageName, TEXT("/"), BulkDataId), UE::DerivedData::ECachePolicy::Default, UE::DerivedData::EPriority::Low,
+		[InnerCallback = MoveTemp(Callback)](UE::DerivedData::FCacheGetCompleteParams&& Params)
+	{
+		bool bOk = Params.Status == UE::DerivedData::EStatus::Ok;
+		InnerCallback(bOk ? Params.Record.GetValue() : FSharedBuffer());
+	});
+}
+
+void PutBulkDataPayloadId(FName PackageName, const FGuid& BulkDataId, FSharedBuffer Buffer)
+{
+	UE::DerivedData::ICache& Cache = GetDerivedDataCacheRef();
+
+	FString ErrorMessage;
+	FCbWriter Builder;
+	EPackageDigestResult FindHashResult = AppendPackageDigest(*IAssetRegistry::Get(), PackageName, Builder, ErrorMessage);
+	switch (FindHashResult)
+	{
+	case EPackageDigestResult::Success:
+		break;
+	default:
+	{
+		return;
+	}
+	}
+	FIoHash PackageAndGuidDigest = GetPackageAndGuidDigest(Builder, BulkDataId);
+
+	UE::DerivedData::FCacheRecordBuilder RecordBuilder = Cache.CreateRecord(GetBulkDataPayloadIdKey(PackageAndGuidDigest));
+	RecordBuilder.SetValue(Buffer);
+	UE::DerivedData::FCacheRecord Record = RecordBuilder.Build();
+	Cache.Put(MakeArrayView(&Record, 1), WriteToString<128>(PackageName));
+}
+
+
+
 }
 }
diff --git a/Engine/Source/Editor/UnrealEd/Private/EditorDomain/EditorDomainUtils.h b/Engine/Source/Editor/UnrealEd/Private/EditorDomain/EditorDomainUtils.h
--- a/Engine/Source/Editor/UnrealEd/Private/EditorDomain/EditorDomainUtils.h
+++ b/Engine/Source/Editor/UnrealEd/Private/EditorDomain/EditorDomainUtils.h
@@ -13,6 +13,7 @@
 DECLARE_LOG_CATEGORY_EXTERN(LogEditorDomainSave, Log, All);
 
 class FAssetPackageData;
+class FCbWriter;
 class FPackagePath;
 class FSharedBuffer;
 class IAssetRegistry;
@@ -49,15 +50,13 @@ struct FClassDigestMap
  */
 EPackageDigestResult GetPackageDigest(IAssetRegistry& AssetRegistry, FName PackageName,
 	FPackageDigest& OutPackageDigest, FString& OutErrorMessage);
+/** Appends the fields to calculate the packagedigest; call Builder.Save().GetRangeHash() to get digest. */
+EPackageDigestResult AppendPackageDigest(IAssetRegistry& AssetRegistry, FName PackageName,
+	FCbWriter& Builder, FString& OutErrorMessage);
+
 /** For any ClassNames not already in ClassDigests, look up their UStruct and add them. */
 void PrecacheClassDigests(TConstArrayView<FName> ClassNames);
 
-/** Get the cachekey for the EditorDomainPackage for the given PackageDigest. */
-UE::DerivedData::FCacheKey GetEditorDomainPackageKey(const FPackageDigest& PackageDigest);
-
-/** Get the cachekey for the EditorDomainBulkDataList for the given PackageDigest. */
-UE::DerivedData::FCacheKey GetBulkDataListKey(const FPackageDigest& PackageDigest);
-
 /** Get the CacheRequest for the given package from the EditorDomain cache bucket. */
 UE::DerivedData::FRequest RequestEditorDomainPackage(const FPackagePath& PackagePath,
 	const FPackageDigest& PackageDigest, UE::DerivedData::ECachePolicy SkipFlags,
@@ -72,6 +71,11 @@ UE::DerivedData::FRequest GetBulkDataList(FName PackageName, TUniqueFunction<voi
 /** Write the data for the BulkDataList of the given package to the cache. */
 void PutBulkDataList(FName PackageName, FSharedBuffer Buffer);
 
+UE::DerivedData::FRequest GetBulkDataPayloadId(FName PackageName, const FGuid& BulkDataId,
+	TUniqueFunction<void(FSharedBuffer Buffer)> && Callback);
+
+void PutBulkDataPayloadId(FName PackageName, const FGuid& BulkDataId, FSharedBuffer Buffer);
+
 /** Accessor for the global ClassDigest map shared by systems needing to calculate PackageDigests. */
 FClassDigestMap& GetClassDigests();
 
