Consider the dirty state of BuiltData when saving the main Map package

diff --git a/Engine/Source/Editor/UnrealEd/Private/FileHelpers.cpp b/Engine/Source/Editor/UnrealEd/Private/FileHelpers.cpp
--- a/Engine/Source/Editor/UnrealEd/Private/FileHelpers.cpp
+++ b/Engine/Source/Editor/UnrealEd/Private/FileHelpers.cpp
@@ -778,9 +778,8 @@ static bool SaveWorld(UWorld* World,
 		if (bSuccess && !bAutosaving)
 		{
 			// Also save MapBuildData packages when saving the current level and save external packages if the world was duplicated
-			bCheckDirty = bCheckDirty || bPIESaving;
 			const bool bSaveExternal = DuplicatedWorld != nullptr || bPackageNeedsRename;
-			if (!FEditorFileUtils::SaveMapDataPackages(SaveWorld, bCheckDirty, bSaveExternal))
+			if (!FEditorFileUtils::SaveMapDataPackages(SaveWorld, /*bCheckDirty*/true, bSaveExternal))
 			{
 				FMessageDialog::Open(EAppMsgType::Ok, NSLOCTEXT("UnrealEd", "Error_FailedToSaveExternalActorPackages", "Failed to save map data packages"));
 				bSuccess = false;
