EditorDomain: Fix assertion "It is invalid to save when nothing has been written" in GetPackageDigest when returning an error.

diff --git a/Engine/Source/Editor/UnrealEd/Private/EditorDomain/EditorDomainUtils.cpp b/Engine/Source/Editor/UnrealEd/Private/EditorDomain/EditorDomainUtils.cpp
--- a/Engine/Source/Editor/UnrealEd/Private/EditorDomain/EditorDomainUtils.cpp
+++ b/Engine/Source/Editor/UnrealEd/Private/EditorDomain/EditorDomainUtils.cpp
@@ -170,7 +170,10 @@ EPackageDigestResult GetPackageDigest(IAssetRegistry& AssetRegistry, FName Packa
 {
 	FCbWriter Builder;
 	EPackageDigestResult Result = AppendPackageDigest(AssetRegistry, PackageName, Builder, OutErrorMessage);
-	OutPackageDigest = Builder.Save().GetRangeHash();
+	if (Result == EPackageDigestResult::Success)
+	{
+		OutPackageDigest = Builder.Save().GetRangeHash();
+	}
 	return Result;
 }
 
