Use raw ISMInstanceManager pointers, rather than TScriptInterface

TScriptInterface is needed for serialization, but these aren't currently serialized anywhere. If serialization is needed, then an interface pointer can be cast to a UObject for use with TScriptInterface.

diff --git a/Engine/Source/Editor/UnrealEd/Private/Elements/SMInstance/SMInstanceElementEditorWorldInterface.cpp b/Engine/Source/Editor/UnrealEd/Private/Elements/SMInstance/SMInstanceElementEditorWorldInterface.cpp
--- a/Engine/Source/Editor/UnrealEd/Private/Elements/SMInstance/SMInstanceElementEditorWorldInterface.cpp
+++ b/Engine/Source/Editor/UnrealEd/Private/Elements/SMInstance/SMInstanceElementEditorWorldInterface.cpp
@@ -34,7 +34,7 @@ bool USMInstanceElementEditorWorldInterface::DeleteElements(TArrayView<const FTy
 	if (SMInstancesToDelete.Num() > 0)
 	{
 		// Batch by the ISM manager
-		TMap<TScriptInterface<ISMInstanceManager>, TArray<FSMInstanceId>> BatchedInstancesToDelete;
+		TMap<ISMInstanceManager*, TArray<FSMInstanceId>> BatchedInstancesToDelete;
 		for (const FSMInstanceManager& SMInstance : SMInstancesToDelete)
 		{
 			if (SMInstance.CanDeleteSMInstance())
@@ -47,7 +47,7 @@ bool USMInstanceElementEditorWorldInterface::DeleteElements(TArrayView<const FTy
 		bool bDidDelete = false;
 
 		FTypedElementListLegacySyncScopedBatch LegacySyncBatch(*InSelectionSet->GetElementList());
-		for (TTuple<TScriptInterface<ISMInstanceManager>, TArray<FSMInstanceId>>& BatchedInstancesToDeletePair : BatchedInstancesToDelete)
+		for (TTuple<ISMInstanceManager*, TArray<FSMInstanceId>>& BatchedInstancesToDeletePair : BatchedInstancesToDelete)
 		{
 			for (const FSMInstanceId& InstanceId : BatchedInstancesToDeletePair.Value)
 			{
@@ -81,7 +81,7 @@ FTypedElementHandle USMInstanceElementEditorWorldInterface::DuplicateElement(con
 
 				if (!bOffsetIsZero)
 				{
-					const TScriptInterface<ISMInstanceManager>& InstanceManager = SMInstance.GetInstanceManager();
+					ISMInstanceManager* InstanceManager = SMInstance.GetInstanceManager();
 
 					FTransform NewInstanceTransform = FTransform::Identity;
 					InstanceManager->GetSMInstanceTransform(NewInstanceId, NewInstanceTransform, /*bWorldSpace*/false);
@@ -104,7 +104,7 @@ void USMInstanceElementEditorWorldInterface::DuplicateElements(TArrayView<const
 	if (SMInstancesToDuplicate.Num() > 0)
 	{
 		// Batch by the ISM manager
-		TMap<TScriptInterface<ISMInstanceManager>, TArray<FSMInstanceId>> BatchedInstancesToDuplicate;
+		TMap<ISMInstanceManager*, TArray<FSMInstanceId>> BatchedInstancesToDuplicate;
 		for (const FSMInstanceManager& SMInstance : SMInstancesToDuplicate)
 		{
 			if (SMInstance.CanDuplicateSMInstance())
@@ -114,7 +114,7 @@ void USMInstanceElementEditorWorldInterface::DuplicateElements(TArrayView<const
 			}
 		}
 
-		for (const TTuple<TScriptInterface<ISMInstanceManager>, TArray<FSMInstanceId>>& BatchedInstancesToDuplicatePair : BatchedInstancesToDuplicate)
+		for (const TTuple<ISMInstanceManager*, TArray<FSMInstanceId>>& BatchedInstancesToDuplicatePair : BatchedInstancesToDuplicate)
 		{
 			TArray<FSMInstanceId> NewInstanceIds;
 			if (BatchedInstancesToDuplicatePair.Key->DuplicateSMInstances(BatchedInstancesToDuplicatePair.Value, NewInstanceIds))
diff --git a/Engine/Source/Runtime/Engine/Private/Elements/SMInstance/SMInstanceElementData.cpp b/Engine/Source/Runtime/Engine/Private/Elements/SMInstance/SMInstanceElementData.cpp
--- a/Engine/Source/Runtime/Engine/Private/Elements/SMInstance/SMInstanceElementData.cpp
+++ b/Engine/Source/Runtime/Engine/Private/Elements/SMInstance/SMInstanceElementData.cpp
@@ -22,7 +22,7 @@ bool SMInstanceElementsEnabled()
 	return GEnableSMInstanceElements != 0;
 }
 
-TScriptInterface<ISMInstanceManager> GetSMInstanceManager(const FSMInstanceId& InstanceId)
+ISMInstanceManager* GetSMInstanceManager(const FSMInstanceId& InstanceId)
 {
 	if (!InstanceId)
 	{
@@ -40,7 +40,7 @@ TScriptInterface<ISMInstanceManager> GetSMInstanceManager(const FSMInstanceId& I
 		// If the owner actor is an instance manager, then just use that
 		if (ISMInstanceManager* InstanceManager = Cast<ISMInstanceManager>(OwnerActor))
 		{
-			return OwnerActor;
+			return InstanceManager;
 		}
 	}
 
diff --git a/Engine/Source/Runtime/Engine/Private/ISMPartition/ISMPartitionActor.cpp b/Engine/Source/Runtime/Engine/Private/ISMPartition/ISMPartitionActor.cpp
--- a/Engine/Source/Runtime/Engine/Private/ISMPartition/ISMPartitionActor.cpp
+++ b/Engine/Source/Runtime/Engine/Private/ISMPartition/ISMPartitionActor.cpp
@@ -651,7 +651,7 @@ void AISMPartitionActor::GetClientComponents(const FISMClientHandle& Handle, TAr
 
 #endif
 
-TScriptInterface<ISMInstanceManager> AISMPartitionActor::GetSMInstanceManager(const FSMInstanceId& InstanceId)
+ISMInstanceManager* AISMPartitionActor::GetSMInstanceManager(const FSMInstanceId& InstanceId)
 {
 	// Disable per-instance editing for now...
 	return nullptr;
diff --git a/Engine/Source/Runtime/Engine/Public/Elements/SMInstance/SMInstanceElementData.h b/Engine/Source/Runtime/Engine/Public/Elements/SMInstance/SMInstanceElementData.h
--- a/Engine/Source/Runtime/Engine/Public/Elements/SMInstance/SMInstanceElementData.h
+++ b/Engine/Source/Runtime/Engine/Public/Elements/SMInstance/SMInstanceElementData.h
@@ -30,7 +30,7 @@ ENGINE_API bool SMInstanceElementsEnabled();
  * Get the static mesh instance manager for the given instance.
  * @return The static mesh instance manager, or null if this instance cannot be managed.
  */
-ENGINE_API TScriptInterface<ISMInstanceManager> GetSMInstanceManager(const FSMInstanceId& InstanceId);
+ENGINE_API ISMInstanceManager* GetSMInstanceManager(const FSMInstanceId& InstanceId);
 
 /**
  * Attempt to get the static mesh instance ID from the given element handle.
diff --git a/Engine/Source/Runtime/Engine/Public/Elements/SMInstance/SMInstanceManager.h b/Engine/Source/Runtime/Engine/Public/Elements/SMInstance/SMInstanceManager.h
--- a/Engine/Source/Runtime/Engine/Public/Elements/SMInstance/SMInstanceManager.h
+++ b/Engine/Source/Runtime/Engine/Public/Elements/SMInstance/SMInstanceManager.h
@@ -125,7 +125,7 @@ public:
 	 * Attempt to get the instance manager associated with the given static mesh instance, if any.
 	 * @return The instance manager, or null if there is no instance manager associated with the given instance.
 	 */
-	virtual TScriptInterface<ISMInstanceManager> GetSMInstanceManager(const FSMInstanceId& InstanceId) = 0;
+	virtual ISMInstanceManager* GetSMInstanceManager(const FSMInstanceId& InstanceId) = 0;
 };
 
 /**
@@ -136,7 +136,7 @@ struct ENGINE_API FSMInstanceManager
 public:
 	FSMInstanceManager() = default;
 
-	FSMInstanceManager(const FSMInstanceId& InInstanceId, const TScriptInterface<ISMInstanceManager>& InInstanceManager)
+	FSMInstanceManager(const FSMInstanceId& InInstanceId, ISMInstanceManager* InInstanceManager)
 		: InstanceId(InInstanceId)
 		, InstanceManager(InInstanceManager)
 	{
@@ -165,7 +165,7 @@ public:
 	}
 
 	const FSMInstanceId& GetInstanceId() const { return InstanceId; }
-	const TScriptInterface<ISMInstanceManager>& GetInstanceManager() const { return InstanceManager; }
+	ISMInstanceManager* GetInstanceManager() const { return InstanceManager; }
 
 	UInstancedStaticMeshComponent* GetISMComponent() const { return InstanceId.ISMComponent; }
 	int32 GetISMInstanceIndex() const { return InstanceId.InstanceIndex; }
@@ -195,5 +195,5 @@ public:
 
 private:
 	FSMInstanceId InstanceId;
-	TScriptInterface<ISMInstanceManager> InstanceManager;
+	ISMInstanceManager* InstanceManager;
 };
diff --git a/Engine/Source/Runtime/Engine/Public/ISMPartition/ISMPartitionActor.h b/Engine/Source/Runtime/Engine/Public/ISMPartition/ISMPartitionActor.h
--- a/Engine/Source/Runtime/Engine/Public/ISMPartition/ISMPartitionActor.h
+++ b/Engine/Source/Runtime/Engine/Public/ISMPartition/ISMPartitionActor.h
@@ -68,7 +68,7 @@ private:
 
 protected:
 	//~ ISMInstanceManagerProvider interface
-	virtual TScriptInterface<ISMInstanceManager> GetSMInstanceManager(const FSMInstanceId& InstanceId) override;
+	virtual ISMInstanceManager* GetSMInstanceManager(const FSMInstanceId& InstanceId) override;
 
 private:
 #if WITH_EDITORONLY_DATA
diff --git a/Engine/Source/Runtime/Foliage/Private/InstancedFoliage.cpp b/Engine/Source/Runtime/Foliage/Private/InstancedFoliage.cpp
--- a/Engine/Source/Runtime/Foliage/Private/InstancedFoliage.cpp
+++ b/Engine/Source/Runtime/Foliage/Private/InstancedFoliage.cpp
@@ -5157,7 +5157,7 @@ float AInstancedFoliageActor::InternalTakeRadialDamage(float Damage, struct FRad
 	return Damage;
 }
 
-TScriptInterface<ISMInstanceManager> AInstancedFoliageActor::GetSMInstanceManager(const FSMInstanceId& InstanceId)
+ISMInstanceManager* AInstancedFoliageActor::GetSMInstanceManager(const FSMInstanceId& InstanceId)
 {
 #if WITH_EDITOR
 	if (const FFoliageInfo* FoliageInfo = GetFoliageInfoForSMInstance(InstanceId))
diff --git a/Engine/Source/Runtime/Foliage/Public/InstancedFoliageActor.h b/Engine/Source/Runtime/Foliage/Public/InstancedFoliageActor.h
--- a/Engine/Source/Runtime/Foliage/Public/InstancedFoliageActor.h
+++ b/Engine/Source/Runtime/Foliage/Public/InstancedFoliageActor.h
@@ -58,7 +58,7 @@ public:
 
 protected:
 	//~ ISMInstanceManagerProvider interface
-	virtual TScriptInterface<ISMInstanceManager> GetSMInstanceManager(const FSMInstanceId& InstanceId) override;
+	virtual ISMInstanceManager* GetSMInstanceManager(const FSMInstanceId& InstanceId) override;
 
 	//~ ISMInstanceManager interface
 	virtual bool CanEditSMInstance(const FSMInstanceId& InstanceId) const override;
