Cooking: Refactor SaveCookedPackage to reduce complexity in advance of adding TargetDomainDependencies into it.

This should have no behavior change except for callouts below.
- BugFix: Fixed World pointer passed into GEditor->Save always being null
- BugFix: Fixed incorrect SavePackageContext being used if requested save platforms is not the same as the cooker's overall session platforms.
- BugFix: PackageStoreWriter should see an error code if UAssetManager::VerifyCanCookPackage fails.
- Fatal error if -DiffOnly is used with -zenstore or -iostore; previously zenstore and iostore were silently ignored if it was set
- Log an error if bReferencedByEditorOnlyData is different between saved platforms.
- A missing package filename is impossible; removed the handling for that case and added a check.

#rnx

diff --git a/Engine/Source/Editor/UnrealEd/Classes/CookOnTheSide/CookOnTheFlyServer.h b/Engine/Source/Editor/UnrealEd/Classes/CookOnTheSide/CookOnTheFlyServer.h
--- a/Engine/Source/Editor/UnrealEd/Classes/CookOnTheSide/CookOnTheFlyServer.h
+++ b/Engine/Source/Editor/UnrealEd/Classes/CookOnTheSide/CookOnTheFlyServer.h
@@ -20,14 +20,15 @@
 
 class FAssetRegistryGenerator;
 class FAsyncIODelete;
-struct FPackageNameCache;
-struct FPropertyChangedEvent;
+class FDiffModeCookServerUtils;
 class FReferenceCollector;
 class FSavePackageContext;
-class ITargetPlatform;
 class IAssetRegistry;
-class IPlugin;
 class IPackageStoreWriter;
+class IPlugin;
+class ITargetPlatform;
+struct FPackageNameCache;
+struct FPropertyChangedEvent;
 
 enum class ECookInitializationFlags
 {
@@ -107,26 +108,24 @@ enum class ECookTickFlags : uint8
 };
 ENUM_CLASS_FLAGS(ECookTickFlags);
 
-namespace UE
-{
-namespace Cook
+namespace UE::Cook
 {
-	namespace Private
-	{
-		class FRegisteredCookPackageSplitter;
-	}
-	struct FCookerTimer;
 	class FExternalRequests;
+	class FSaveCookedPackageContext;
+	class ICookOnTheFlyRequestManager;
+	struct FCookerTimer;
 	struct FPackageData;
-	struct FGeneratorPackage;
 	struct FPackageDatas;
 	struct FPackageTracker;
 	struct FPendingCookedPlatformData;
+	struct FGeneratorPackage;
 	struct FPlatformManager;
-	struct FTickStackData;
 	struct FPopulatePackageContext;
-	class ICookOnTheFlyRequestManager;
+	struct FTickStackData;
 }
+namespace UE::Cook::Private
+{
+	class FRegisteredCookPackageSplitter;
 }
 
 UCLASS()
@@ -134,8 +133,8 @@ class UNREALED_API UCookOnTheFlyServer : public UObject, public FTickableEditorO
 {
 	GENERATED_BODY()
 
-		UCookOnTheFlyServer(const FObjectInitializer& ObjectInitializer = FObjectInitializer::Get());
-		UCookOnTheFlyServer(FVTableHelper& Helper); // Declare the FVTableHelper constructor manually so that we can forward-declare-only TUniquePtrs in the header without getting compile error in generated cpp
+	UCookOnTheFlyServer(const FObjectInitializer& ObjectInitializer = FObjectInitializer::Get());
+	UCookOnTheFlyServer(FVTableHelper& Helper); // Declare the FVTableHelper constructor manually so that we can forward-declare-only TUniquePtrs in the header without getting compile error in generated cpp
 
 public:
 	struct FCookByTheBookOptions;
@@ -1020,18 +1019,9 @@ private:
 		return (CookFlags & InCookFlags) != ECookInitializationFlags::None;
 	}
 
-	/**
-	*	Cook (save) the given package
-	*
-	*	@param	PackageData			The PackageData to cook/save
-	*	@param	SaveFlags			The flags to pass to the SavePackage function
-	*	@param	bOutWasUpToDate		Upon return, if true then the cooked package was cached (up to date)
-	*	@param  TargetPlatforms		Only cook for target platforms which are included in this array
-	*
-	*	@return	ESavePackageResult::Success if packages was cooked
-	*/
-	void SaveCookedPackage(UE::Cook::FPackageData& PackageData, uint32 SaveFlags, const TArrayView<const ITargetPlatform* const>& TargetPlatforms, TArray<FSavePackageResultStruct>& SavePackageResults);
-
+	/** Cook (save) a package and process the results */
+	void SaveCookedPackage(UE::Cook::FSaveCookedPackageContext& Context);
+	friend class UE::Cook::FSaveCookedPackageContext;
 
 	/**
 	*  Save the global shader map
@@ -1114,6 +1104,7 @@ private:
 	uint32 FullLoadAndSave(uint32& CookedPackageCount);
 
 	IPackageStoreWriter* GetPackageStoreWriter(const FName& PlatformName) const;
+	FSavePackageContext* GetSavePackageContext(const ITargetPlatform* TargetPlatform) const;
 
 	uint32		StatLoadedPackageCount = 0;
 	uint32		StatSavedPackageCount = 0;
@@ -1144,6 +1135,8 @@ private:
 	/** Objects that were collected during the single-threaded PreGarbageCollect callback and that should be reported as referenced in CookerAddReferencedObjects. */
 	TArray<UObject*> GCKeepObjects;
 	UE::Cook::FPackageData* SavingPackageData = nullptr;
+	/** Helper struct for running cooking in diagnostic modes */
+	TUniquePtr<FDiffModeCookServerUtils> DiffModeHelper;
 
 	// temporary -- should eliminate the need for this. Only required right now because FullLoadAndSave 
 	// accesses maps directly
diff --git a/Engine/Source/Editor/UnrealEd/Private/CookOnTheFlyServer.cpp b/Engine/Source/Editor/UnrealEd/Private/CookOnTheFlyServer.cpp
--- a/Engine/Source/Editor/UnrealEd/Private/CookOnTheFlyServer.cpp
+++ b/Engine/Source/Editor/UnrealEd/Private/CookOnTheFlyServer.cpp
@@ -707,6 +707,7 @@ UCookOnTheFlyServer::UCookOnTheFlyServer(const FObjectInitializer& ObjectInitial
 	PlatformManager = MakeUnique<UE::Cook::FPlatformManager>();
 	ExternalRequests = MakeUnique<UE::Cook::FExternalRequests>();
 	PackageTracker = MakeUnique<UE::Cook::FPackageTracker>(*PackageDatas.Get());
+	DiffModeHelper = MakeUnique<FDiffModeCookServerUtils>();
 	bSaveAsyncAllowed = true;
 	FString Temp;
 	const TCHAR* CommandLine = FCommandLine::Get();
@@ -1948,7 +1949,7 @@ void UCookOnTheFlyServer::ProcessRequest(UE::Cook::FPackageData& PackageData, UE
 	if (PackageData.HasAllCookedPlatforms(PackageData.GetRequestedPlatforms(), true /* bIncludeFailed */))
 	{
 #if DEBUG_COOKONTHEFLY
-		UE_LOG(LogCook, Display, TEXT("Package for platform already cooked %s, discarding request"), *OutToBuild.GetFilename().ToString());
+		UE_LOG(LogCook, Display, TEXT("Package for platform already cooked %s, discarding request"), *PackageData.GetFileName().ToString());
 #endif
 		PackageData.SendToState(EPackageState::Idle, ESendFlags::QueueAdd);
 		return;
@@ -2803,7 +2804,7 @@ bool UCookOnTheFlyServer::BeginPrepareSave(UE::Cook::FPackageData& PackageData,
 	UE_SCOPED_HIERARCHICAL_COOKTIMER_AND_DURATION(BeginPrepareSave, DetailedCookStats::TickCookOnTheSideBeginPrepareSaveTimeSec);
 
 #if DEBUG_COOKONTHEFLY 
-	UE_LOG(LogCook, Display, TEXT("Caching objects for package %s"), *Package->GetFName().ToString());
+	UE_LOG(LogCook, Display, TEXT("Caching objects for package %s"), *PackageData.GetPackageName().ToString());
 #endif
 	UPackage* Package = PackageData.GetPackage();
 	check(Package && Package->IsFullyLoaded());
@@ -3226,9 +3227,52 @@ void UCookOnTheFlyServer::ProcessUnsolicitedPackages()
 	}
 }
 
-UE_TRACE_EVENT_BEGIN(UE_CUSTOM_COOKTIMER_LOG, SaveCookedPackage, NoSync)
-	UE_TRACE_EVENT_FIELD(UE::Trace::WideString, PackageName)
-UE_TRACE_EVENT_END()
+namespace UE::Cook
+{
+
+/** Local parameters and helper functions used by SaveCookedPackage */
+class FSaveCookedPackageContext
+{
+private: // Used only by UCookOnTheFlyServer, which has private access
+
+	FSaveCookedPackageContext(UCookOnTheFlyServer& InCOTFS, UE::Cook::FPackageData& InPackageData,
+		TArrayView<const ITargetPlatform*> InPlatformsForPackage, UE::Cook::FTickStackData& StackData);
+
+	void SetupPackage();
+	void SetupPlatform(const ITargetPlatform* InTargetPlatform);
+	void FinishPlatform();
+	void FinishPackage();
+
+	// General Package Data
+	UCookOnTheFlyServer& COTFS;
+	UE::Cook::FPackageData& PackageData;
+	TArrayView<const ITargetPlatform*> PlatformsForPackage;
+	FTickStackData& StackData;
+	UPackage* Package;
+	const FString PackageName;
+	FString Filename;
+	uint32 SaveFlags = 0;
+	bool bReferencedOnlyByEditorOnlyData = false;
+	bool bHasFirstPlatformResults = false;
+
+	// General Package Data that is delay-loaded the first time we save a platform
+	UWorld* World = nullptr;
+	EObjectFlags FlagsToCook = RF_Public;
+	bool bHasDelayLoaded = false;
+	bool bContainsMap = false;
+
+	// Per-platform data, only valid in PerPlatform callbacks
+	const ITargetPlatform* TargetPlatform = nullptr;
+	FSavePackageContext* SavePackageContext = nullptr;
+	IPackageStoreWriter* PackageStoreWriter = nullptr;
+	FString PlatFilename;
+	FSavePackageResultStruct SavePackageResult;
+	bool bPlatformSetupSuccessful = false;
+	bool bEndianSwap = false;
+
+	friend class ::UCookOnTheFlyServer;
+};
+}
 
 void UCookOnTheFlyServer::PumpSaves(UE::Cook::FTickStackData& StackData, uint32 DesiredQueueLength, int32& OutNumPushed, bool& bOutBusy)
 {
@@ -3425,138 +3469,10 @@ void UCookOnTheFlyServer::PumpSaves(UE::Cook::FTickStackData& StackData, uint32
 			}
 		}
 
-		TArray<bool> SucceededSavePackage;
-		TArray<FSavePackageResultStruct> SavePackageResults;
-		{
-			TStringBuilder<512> PackageDataFileName;
-			if (UE_TRACE_CHANNELEXPR_IS_ENABLED(CookChannel))
-			{
-				PackageData.GetFileName().ToString(PackageDataFileName);
-			}
-			UE_SCOPED_HIERARCHICAL_CUSTOM_COOKTIMER_AND_DURATION(SaveCookedPackage, DetailedCookStats::TickCookOnTheSideSaveCookedPackageTimeSec)
-				UE_ADD_CUSTOM_COOKTIMER_META(SaveCookedPackage, PackageName, *PackageDataFileName);
-
-			uint32 SaveFlags = SAVE_KeepGUID | (bSaveAsyncAllowed ? SAVE_Async : SAVE_None) | (IsCookFlagSet(ECookInitializationFlags::Unversioned) ? SAVE_Unversioned : 0);
-
-			bool KeepEditorOnlyPackages = false;
-			// removing editor only packages only works when cooking in commandlet and non iterative cooking
-			// also doesn't work in multiprocess cooking
-			KeepEditorOnlyPackages = !(IsCookByTheBookMode() && !IsCookingInEditor());
-			KeepEditorOnlyPackages |= IsCookFlagSet(ECookInitializationFlags::Iterative);
-			SaveFlags |= KeepEditorOnlyPackages ? SAVE_KeepEditorOnlyCookedPackages : SAVE_None;
-			SaveFlags |= CookByTheBookOptions ? SAVE_ComputeHash : SAVE_None;
-
-			GOutputCookingWarnings = IsCookFlagSet(ECookInitializationFlags::OutputVerboseCookerWarnings);
-
-			{
-				// SaveCookedPackage can CollectGarbage, so we need to store the currently-unqueued PackageData in a separate variable that we register for garbage collection
-				check(SavingPackageData == nullptr);
-				SavingPackageData = &PackageData;
-				try
-				{
-					SaveCookedPackage(PackageData, SaveFlags, PlatformsForPackage, SavePackageResults);
-				}
-				catch (std::exception&)
-				{
-					FString TargetPlatforms;
-					for (const ITargetPlatform* Platform : PlatformsForPackage)
-					{
-						TargetPlatforms += FString::Printf(TEXT("%s, "), *Platform->PlatformName());
-					}
-					UE_LOG(LogCook, Warning, TEXT("Tried to save package %s for target platforms %s but threw an exception"), *Package->GetName(), *TargetPlatforms);
-					SavePackageResults.Empty(PlatformsForPackage.Num());
-					for (int n = 0; n < PlatformsForPackage.Num(); ++n)
-					{
-						SavePackageResults.Add(ESavePackageResult::Error);
-					}
-				}
-				SavingPackageData = nullptr;
-			}
-
-			GOutputCookingWarnings = false;
-			check(PlatformsForPackage.Num() == SavePackageResults.Num());
-			for (int iResultIndex = 0; iResultIndex < SavePackageResults.Num(); iResultIndex++)
-			{
-				FSavePackageResultStruct& SavePackageResult = SavePackageResults[iResultIndex];
-
-				// Update asset registry
-				if (CookByTheBookOptions)
-				{
-					FAssetRegistryGenerator* Generator = PlatformManager->GetPlatformData(PlatformsForPackage[iResultIndex])->RegistryGenerator.Get();
-					UpdateAssetRegistryPackageData(Generator, *Package, SavePackageResult);
-				}
-
-				if (SavePackageResult.IsSuccessful())
-				{
-					SucceededSavePackage.Add(true);
-					// Update flags used to determine garbage collection.
-					if (Package->ContainsMap())
-					{
-						StackData.ResultFlags |= COSR_CookedMap;
-					}
-					else
-					{
-						++StackData.CookedPackageCount;
-						StackData.ResultFlags |= COSR_CookedPackage;
-					}
-				}
-				else
-				{
-					SucceededSavePackage.Add(false);
-				}
-			}
-			check(SavePackageResults.Num() == SucceededSavePackage.Num());
-			StackData.Timer.SavedPackage();
-		}
+		FSaveCookedPackageContext Context(*this, PackageData, PlatformsForPackage, StackData);
+		SaveCookedPackage(Context);
 
 		ReleaseCookedPlatformData(PackageData, true /* bCompletedSave */);
-
-		FName FileName = PackageData.GetFileName();
-
-		// We always want to mark package as processed unless it wasn't saved because it was referenced by editor-only data
-		// in which case we may still need to save it later when new content loads it through non editor-only references
-		if (!FileName.IsNone())
-		{
-			// mark the package as cooked
-			bool bWasReferencedOnlyByEditorOnlyData = false;
-			for (const FSavePackageResultStruct& SavePackageResult : SavePackageResults)
-			{
-				if (SavePackageResult == ESavePackageResult::ReferencedOnlyByEditorOnlyData)
-				{
-					bWasReferencedOnlyByEditorOnlyData = true;
-					// if this is the case all of the platforms should be referenced only by editor only data
-				}
-			}
-			if (!bWasReferencedOnlyByEditorOnlyData)
-			{
-				PackageData.AddCookedPlatforms(PackageData.GetRequestedPlatforms(), SucceededSavePackage);
-
-				if ((CurrentCookMode == ECookMode::CookOnTheFly) && !PackageData.GetIsUrgent())
-				{
-					// this is an unsolicited package
-					if (FPaths::FileExists(FileName.ToString()))
-					{
-						PackageTracker->UnsolicitedCookedPackages.AddCookedPackage(FFilePlatformRequest(FileName, PlatformsForPackage));
-
-#if DEBUG_COOKONTHEFLY
-						UE_LOG(LogCook, Display, TEXT("UnsolicitedCookedPackages: %s"), *FileName.ToString());
-#endif
-					}
-				}
-			}
-			else
-			{
-				PackageTracker->UncookedEditorOnlyPackages.AddUnique(Package->GetFName());
-			}
-		}
-		else
-		{
-			for (const bool bSucceededSavePackage : SucceededSavePackage)
-			{
-				check(bSucceededSavePackage == false);
-			}
-		}
-
 		PackageData.SendToState(EPackageState::Idle, ESendFlags::QueueAdd);
 		++OutNumPushed;
 	}
@@ -4315,26 +4231,35 @@ public:
 
 private:
 	/** Misc / common settings */
-	bool bDiffEnabled;
-	ELinkerDiffMode LinkerDiffMode;
+	bool bInitialized = false;
+	bool bDiffEnabled = false;
+	ELinkerDiffMode LinkerDiffMode = LDM_None;
 	FString PackageFilter;
 
 	/** DumpObjList settings */
-	bool bDumpObjList;
+	bool bDumpObjList = false;
 	FString DumpObjListParams;
 
 	/** DumpObjects settings */
-	bool bDumpObjects;
-	bool bDumpObjectsSorted;
+	bool bDumpObjects = false;
+	bool bDumpObjectsSorted = false;
+
+	/* LinkerDiff settings */
+	IConsoleVariable* EnableNewSave = nullptr;
+	int32 PreviousCvarValue = 0;
 
 public:
 
 
-	FDiffModeCookServerUtils()
+	void Initialize(TConstArrayView<FSavePackageContext*> SavePackageContexts)
 	{
+		if (bInitialized)
+		{
+			return;
+		}
+
 		bDiffEnabled = FParse::Param(FCommandLine::Get(), TEXT("DIFFONLY"));
 
-		LinkerDiffMode = LDM_None;
 		FString DiffMode;
 		if (FParse::Value(FCommandLine::Get(), TEXT("-LINKERDIFF="), DiffMode))
 		{
@@ -4348,11 +4273,33 @@ public:
 			}
 		}
 
-		bDumpObjList = false;
-		bDumpObjects = false;
-		bDumpObjectsSorted = false;
-
 		ParseCmds();
+		InitializePlatformSettings(SavePackageContexts);
+		bInitialized = true;
+	}
+
+	void InitializePlatformSettings(TConstArrayView<FSavePackageContext*> SavePackageContexts)
+	{
+		if (IsRunningCookDiff() || IsRunningCookLinkerDiff())
+		{
+			bool bCompatible = true;
+			for (FSavePackageContext* SavePackageContext : SavePackageContexts)
+			{
+				if (SavePackageContext && SavePackageContext->PackageStoreWriter)
+				{
+					bCompatible = false;
+					break;
+				}
+			}
+			if (!bCompatible)
+			{
+				const TCHAR* CommandLineArg = IsRunningCookDiff() ? TEXT("-DIFFONLY") : TEXT("-LINKERDIFF");
+				UE_LOG(LogCook, Fatal, TEXT("%s was enabled, but PackageStoreWriter is also enabled and PackageStoreWriter does not support %s."),
+					CommandLineArg, CommandLineArg);
+				bDiffEnabled = false;
+				LinkerDiffMode = LDM_None;
+			}
+		}
 	}
 
 	bool IsRunningCookDiff() const
@@ -4376,6 +4323,39 @@ public:
 		ConditionallyDumpObjects(InPackage);
 	}
 
+	void LinkerDiffSetup()
+	{
+		// Resave the package again using the other save algorithm
+		if (!EnableNewSave)
+		{
+			EnableNewSave = IConsoleManager::Get().FindConsoleVariable(TEXT("SavePackage.EnableNewSave"));
+		}
+		PreviousCvarValue = EnableNewSave->GetInt();
+
+		// If the linker diff is comparing the two save algo, switch the cvar to the other algo before saving again,
+		// Otherwise the linker diff mode is tracking if the save is consistent across multiple save
+		if (LinkerDiffMode == FDiffModeCookServerUtils::LDM_Algo)
+		{
+			// see CVarEnablePackageNewSave definition in SavePackageUtilities.cpp for value meaning
+			EnableNewSave->Set(PreviousCvarValue & 1 ? 0 : 1);
+		}
+	}
+	void LinkerDiffFinish(FSavePackageResultStruct& FirstResult, FSavePackageResultStruct& SecondResult)
+	{
+		if (LinkerDiffMode == FDiffModeCookServerUtils::LDM_Algo)
+		{
+			EnableNewSave->Set(PreviousCvarValue);
+		}
+
+		if (FirstResult.LinkerSave && SecondResult.LinkerSave)
+		{
+			FLinkerDiff LinkerDiff = FLinkerDiff::CompareLinkers(FirstResult.LinkerSave.Get(), SecondResult.LinkerSave.Get());
+			LinkerDiff.PrintDiff(*GWarn);
+		}
+		FirstResult.LinkerSave.Reset();
+		SecondResult.LinkerSave.Reset();
+	}
+
 private:
 
 	void RemoveParam(FString& InOutParams, const TCHAR* InParamToRemove)
@@ -4503,260 +4483,342 @@ private:
 	}
 };
 
-void UCookOnTheFlyServer::SaveCookedPackage(UE::Cook::FPackageData& PackageData, uint32 SaveFlags, const TArrayView<const ITargetPlatform* const>& TargetPlatforms, TArray<FSavePackageResultStruct>& SavePackageResults)
-{
-	check( SavePackageResults.Num() == 0);
-	check( bIsSavingPackage == false );
-	bIsSavingPackage = true;
-
-	UPackage* Package = PackageData.GetPackage();
-	check(Package && Package->IsFullyLoaded());
-	const FString PackageName(Package->GetName());
-	check(Package->GetPathName().Equals(Package->GetName())); // We should only be saving outermost packages, so the path name should be the same as the package name
-	FString Filename(PackageData.GetFileName().ToString());
+UE_TRACE_EVENT_BEGIN(UE_CUSTOM_COOKTIMER_LOG, SaveCookedPackage, NoSync)
+	UE_TRACE_EVENT_FIELD(UE::Trace::WideString, PackageName)
+UE_TRACE_EVENT_END()
 
-	// Also request any localized variants of this package
-	if (IsCookByTheBookMode() && !CookByTheBookOptions->bSkipSoftReferences && !FPackageName::IsLocalizedPackage(PackageName))
-	{
-		const TArray<FName>* LocalizedVariants = CookByTheBookOptions->SourceToLocalizedPackageVariants.Find(Package->GetFName());
-		if (LocalizedVariants)
-		{
-			for (const FName& LocalizedPackageName : *LocalizedVariants)
-			{
-				UE::Cook::FPackageData* LocalizedPackageData = PackageDatas->TryAddPackageDataByPackageName(LocalizedPackageName);
-				if (LocalizedPackageData)
-				{
-					bool bIsUrgent = false;
-					LocalizedPackageData->UpdateRequestData(PackageData.GetRequestedPlatforms(), bIsUrgent, UE::Cook::FCompletionCallback());
-				}
-			}
-		}
-	}
+void UCookOnTheFlyServer::SaveCookedPackage(UE::Cook::FSaveCookedPackageContext& Context)
+{
+	using namespace UE::Cook;
 
-	if (Filename.Len() != 0 )
-	{
-		if (Package->HasAnyPackageFlags(PKG_ReloadingForCooker))
-		{
-			UE_LOG(LogCook, Warning, TEXT("Package %s marked as reloading for cook by was requested to save"), *Package->GetName());
-			UE_LOG(LogCook, Fatal, TEXT("Package %s marked as reloading for cook by was requested to save"), *Package->GetName());
-		}
+	UE_SCOPED_HIERARCHICAL_CUSTOM_COOKTIMER_AND_DURATION(SaveCookedPackage, DetailedCookStats::TickCookOnTheSideSaveCookedPackageTimeSec)
+		UE_ADD_CUSTOM_COOKTIMER_META(SaveCookedPackage, PackageName, *WriteToString<256>(Context.PackageData.GetFileName()));
 
-		// Use SandboxFile to do path conversion to properly handle sandbox paths (outside of standard paths in particular).
-		Filename = ConvertToFullSandboxPath(*Filename, true);
+	UPackage* Package = Context.Package;
+	uint32 OriginalPackageFlags = Package->GetPackageFlags();
 
-		uint32 OriginalPackageFlags = Package->GetPackageFlags();
-		UWorld* World = nullptr;
-		EObjectFlags FlagsToCook = RF_Public;
+	Context.SetupPackage();
 
-		ITargetPlatformManagerModule& TPM = GetTargetPlatformManagerRef();
+	ITargetPlatformManagerModule& TPM = GetTargetPlatformManagerRef();
+	TGuardValue<bool> ScopedOutputCookerWarnings(GOutputCookingWarnings, IsCookFlagSet(ECookInitializationFlags::OutputVerboseCookerWarnings));
+	// SavePackage can CollectGarbage, so we need to store the currently-unqueued PackageData in a separate variable that we register for garbage collection
+	TGuardValue<UE::Cook::FPackageData*> ScopedSavingPackageData(SavingPackageData, &Context.PackageData);
+	TGuardValue<bool> ScopedIsSavingPackage(bIsSavingPackage, true);
+	// For legacy reasons we set GIsCookerLoadingPackage == true during save. Some classes use it to conditionally execute cook operations in both save and load
+	TGuardValue<bool> ScopedIsCookerLoadingPackage(GIsCookerLoadingPackage, true);
+	ON_SCOPE_EXIT{ Package->SetPackageFlagsTo(OriginalPackageFlags); };
 
-		static FDiffModeCookServerUtils DiffModeHelper;
-		if (DiffModeHelper.IsRunningCookLinkerDiff())
+	for (const ITargetPlatform* TargetPlatform : Context.PlatformsForPackage)
+	{
+		Context.SetupPlatform(TargetPlatform);
+		if (Context.bPlatformSetupSuccessful)
 		{
-			SaveFlags |= SAVE_CompareLinker;
-		}
+			UE_SCOPED_HIERARCHICAL_COOKTIMER(GEditorSavePackage);
 
-		for (int32 PlatformIndex = 0; PlatformIndex < TargetPlatforms.Num(); ++PlatformIndex)
-		{
-			SavePackageResults.Add(FSavePackageResultStruct(ESavePackageResult::Success));
-			const ITargetPlatform* Target = TargetPlatforms[PlatformIndex];
-			FString PlatFilename = Filename.Replace(TEXT("[Platform]"), *Target->PlatformName());
+			try
+			{
+				if (DiffModeHelper->IsRunningCookDiff())
+				{
+					check(!Context.PackageStoreWriter); // IsRunningCookDiff should have been disabled by DiffModeHelper if we have any PackageStoreWriter
+					DiffModeHelper->ProcessPackage(Package);
 
-			FSavePackageResultStruct& Result = SavePackageResults[PlatformIndex];
+					// When looking for deterministic cook issues, first serialize the package to memory and do a simple diff with the existing package
+					uint32 DiffSaveFlags = Context.SaveFlags | SAVE_DiffOnly;
+					FArchiveDiffMap DiffMap;
+					Context.SavePackageResult = GEditor->Save(Package, Context.World, Context.FlagsToCook, *Context.PlatFilename,
+						GError, nullptr, Context.bEndianSwap, false, DiffSaveFlags, TargetPlatform,
+						FDateTime::MinValue(), false, &DiffMap, Context.SavePackageContext);
+					if (Context.SavePackageResult == ESavePackageResult::DifferentContent)
+					{
+						// If the simple memory diff was not identical, collect callstacks for all Serialize calls and dump differences to log
+						DiffSaveFlags = Context.SaveFlags | SAVE_DiffCallstack;
+						Context.SavePackageResult = GEditor->Save(Package, Context.World, Context.FlagsToCook, *Context.PlatFilename,
+							GError, nullptr, Context.bEndianSwap, false, DiffSaveFlags, TargetPlatform,
+							FDateTime::MinValue(), false, &DiffMap, Context.SavePackageContext);
+					}
+				}
+				else if (DiffModeHelper->IsRunningCookLinkerDiff())
+				{
+					check(!Context.PackageStoreWriter); // IsRunningCookLinkerDiff should have been disabled by DiffModeHelper if we have any PackageStoreWriter
 
-			bool bCookPackage = true;
+					Context.SavePackageResult = GEditor->Save(Package, Context.World, Context.FlagsToCook, *Context.PlatFilename,
+						GError, nullptr, Context.bEndianSwap, false, Context.SaveFlags, TargetPlatform,
+						FDateTime::MinValue(), false, /*DiffMap*/ nullptr, Context.SavePackageContext);
 
-			// don't save Editor resources from the Engine if the target doesn't have editoronly data
-			if (IsCookFlagSet(ECookInitializationFlags::SkipEditorContent) &&
-				(PackageName.StartsWith(TEXT("/Engine/Editor")) || PackageName.StartsWith(TEXT("/Engine/VREditor"))) &&
-				!Target->HasEditorOnlyData())
-			{
-				Result = ESavePackageResult::ContainsEditorOnlyData;
-				bCookPackage = false;
+					// Resave the package again using the other save algorithm
+					DiffModeHelper->LinkerDiffSetup();
+					FSavePackageResultStruct NewResult;
+					NewResult = GEditor->Save(Package, Context.World, Context.FlagsToCook, *Context.PlatFilename,
+						GError, nullptr, Context.bEndianSwap, false, Context.SaveFlags | SAVE_DiffOnly, TargetPlatform,
+						FDateTime::MinValue(), false, /*DiffMap*/ nullptr, Context.SavePackageContext);
+
+					DiffModeHelper->LinkerDiffFinish(Context.SavePackageResult, NewResult);
+				}
+				else
+				{
+					Context.SavePackageResult = GEditor->Save(Package, Context.World, Context.FlagsToCook, *Context.PlatFilename,
+						GError, nullptr, Context.bEndianSwap, false, Context.SaveFlags, TargetPlatform,
+						FDateTime::MinValue(), false, /*DiffMap*/ nullptr, Context.SavePackageContext);
+				}
 			}
-			// Check whether or not game-specific behaviour should prevent this package from being cooked for the target platform
-			else if (UAssetManager::IsValid() && !UAssetManager::Get().ShouldCookForPlatform(Package, Target))
+			catch (std::exception&)
 			{
-				Result = ESavePackageResult::ContainsEditorOnlyData;
-				bCookPackage = false;
-				UE_LOG(LogCook, Display, TEXT("Excluding %s -> %s"), *Package->GetName(), *PlatFilename);
+				UE_LOG(LogCook, Warning, TEXT("Tried to save package %s for target platform %s but threw an exception"),
+					*Package->GetName(), *TargetPlatform->PlatformName());
+				Context.SavePackageResult = ESavePackageResult::Error;
 			}
-			// check if this package is unsupported for the target platform (typically plugin content)
-			else 
+
+			// If package was actually saved check with asset manager to make sure it wasn't excluded for being a
+			// development or never cook package. But skip sending the warnings from this check if it was editor-only.
+			if (Context.SavePackageResult == ESavePackageResult::Success && UAssetManager::IsValid())
 			{
-				TSet<FName>* NeverCookPackages = PackageTracker->PlatformSpecificNeverCookPackages.Find(Target);
-				if (NeverCookPackages && NeverCookPackages->Find(FName(*PackageName)))
+				UE_SCOPED_HIERARCHICAL_COOKTIMER(VerifyCanCookPackage);
+				if (!UAssetManager::Get().VerifyCanCookPackage(Package->GetFName()))
 				{
-					Result = ESavePackageResult::ContainsEditorOnlyData;
-					bCookPackage = false;
-					UE_LOG(LogCook, Display, TEXT("Excluding %s -> %s"), *Package->GetName(), *PlatFilename);
+					Context.SavePackageResult = ESavePackageResult::Error;
 				}
 			}
 
-			if (bCookPackage == true)
-			{
-				// look for a world object in the package (if there is one, there's a map)
-				if (UWorld::FindWorldInPackage(Package))
-				{
-					FlagsToCook = RF_NoFlags;
-				}
+			++this->StatSavedPackageCount;
+		}
 
-				UE_CLOG(GCookProgressDisplay & (int32)ECookProgressDisplayMode::PackageNames, LogCook, Display, TEXT("Cooking %s -> %s"), *Package->GetName(), *PlatFilename);
+		Context.FinishPlatform();
+	}
 
-				bool bSwap = (!Target->IsLittleEndian()) ^ (!PLATFORM_LITTLE_ENDIAN);
+	Context.FinishPackage();
+	Context.StackData.Timer.SavedPackage();
+}
 
+namespace UE::Cook
+{
+FSaveCookedPackageContext::FSaveCookedPackageContext(UCookOnTheFlyServer& InCOTFS, UE::Cook::FPackageData& InPackageData,
+	TArrayView<const ITargetPlatform*> InPlatformsForPackage, UE::Cook::FTickStackData& InStackData)
+	: COTFS(InCOTFS)
+	, PackageData(InPackageData)
+	, PlatformsForPackage(InPlatformsForPackage)
+	, StackData(InStackData)
+	, Package(PackageData.GetPackage())
+	, PackageName(Package ? Package->GetName() : FString())
+	, Filename(PackageData.GetFileName().ToString())
+{
+}
 
-				if (!Target->HasEditorOnlyData())
-				{
-					Package->SetPackageFlags(PKG_FilterEditorOnly);
-				}
-				else
-				{
-					Package->ClearPackageFlags(PKG_FilterEditorOnly);
-				}
+void FSaveCookedPackageContext::SetupPackage()
+{
+	COTFS.DiffModeHelper->Initialize(COTFS.SavePackageContexts);
+	check(Package && Package->IsFullyLoaded()); // PackageData should not be in the save state if Package is not fully loaded
+	check(Package->GetPathName().Equals(PackageName)); // We should only be saving outermost packages, so the path name should be the same as the package name
+	check(!Filename.IsEmpty()); // PackageData guarantees FileName is non-empty; if not found it should never make it into save state
+	if (Package->HasAnyPackageFlags(PKG_ReloadingForCooker))
+	{
+		UE_LOG(LogCook, Warning, TEXT("Package %s marked as reloading for cook by was requested to save"), *PackageName);
+		UE_LOG(LogCook, Fatal, TEXT("Package %s marked as reloading for cook by was requested to save"), *PackageName);
+	}
 
-				if (World)
-				{
-					// Fixup legacy lightmaps before saving
-					// This should be done after loading, but Core loads UWorlds with LoadObject so there's no opportunity to handle this fixup on load
-					World->PersistentLevel->HandleLegacyMapBuildData();
-				}
+	SaveFlags = SAVE_KeepGUID | (COTFS.bSaveAsyncAllowed ? SAVE_Async : SAVE_None)
+		| (COTFS.IsCookFlagSet(ECookInitializationFlags::Unversioned) ? SAVE_Unversioned : 0);
 
-				const FString FullFilename = FPaths::ConvertRelativePathToFull(PlatFilename);
-				if (FullFilename.Len() >= FPlatformMisc::GetMaxPathLength())
-				{
-					LogCookerMessage(FString::Printf(TEXT("Couldn't save package, filename is too long (%d >= %d): %s"), FullFilename.Len(), FPlatformMisc::GetMaxPathLength(), *PlatFilename), EMessageSeverity::Error);
-					Result = ESavePackageResult::Error;
-				}
-				else
-				{
-					UE_SCOPED_HIERARCHICAL_COOKTIMER(GEditorSavePackage);
-					GIsCookerLoadingPackage = true;
+	// removing editor only packages only works when cooking in commandlet and non iterative cooking
+	// also doesn't work in multiprocess cooking
+	bool bKeepEditorOnlyPackages = !(COTFS.IsCookByTheBookMode() && !COTFS.IsCookingInEditor());
+	bKeepEditorOnlyPackages |= COTFS.IsCookFlagSet(ECookInitializationFlags::Iterative);
+	SaveFlags |= bKeepEditorOnlyPackages ? SAVE_KeepEditorOnlyCookedPackages : SAVE_None;
+	SaveFlags |= COTFS.CookByTheBookOptions ? SAVE_ComputeHash : SAVE_None;
+	if (COTFS.DiffModeHelper->IsRunningCookLinkerDiff())
+	{
+		SaveFlags |= SAVE_CompareLinker;
+	}
 
-					if (DiffModeHelper.IsRunningCookDiff())
-					{
-						FSavePackageContext* const SavePackageContext = SavePackageContexts.Num() > 0 ? SavePackageContexts[PlatformIndex] : nullptr;
+	// Use SandboxFile to do path conversion to properly handle sandbox paths (outside of standard paths in particular).
+	Filename = COTFS.ConvertToFullSandboxPath(*Filename, true);
+}
 
-						DiffModeHelper.ProcessPackage(Package);
+void FSaveCookedPackageContext::SetupPlatform(const ITargetPlatform* InTargetPlatform)
+{
+	TargetPlatform = InTargetPlatform;
+	PlatFilename = Filename.Replace(TEXT("[Platform]"), *TargetPlatform->PlatformName());
+	bPlatformSetupSuccessful = false;
 
-						// When looking for deterministic cook issues, first serialize the package to memory and do a simple diff with the existing package
-						uint32 DiffSaveFlags = SaveFlags | SAVE_DiffOnly;
-						FArchiveDiffMap DiffMap;
-						Result = GEditor->Save(Package, World, FlagsToCook, *PlatFilename, GError, NULL, bSwap, false, DiffSaveFlags, Target, FDateTime::MinValue(), false, &DiffMap, SavePackageContext);
-						if (Result == ESavePackageResult::DifferentContent)
-						{
-							// If the simple memory diff was not identical, collect callstacks for all Serialize calls and dump differences to log
-							DiffSaveFlags = SaveFlags | SAVE_DiffCallstack;
-							Result = GEditor->Save(Package, World, FlagsToCook, *PlatFilename, GError, NULL, bSwap, false, DiffSaveFlags, Target, FDateTime::MinValue(), false, &DiffMap, SavePackageContext);
-						}
-					}
-					else
-					{
-						FSavePackageContext* const SavePackageContext = SavePackageContexts.Num() > 0 ? SavePackageContexts[PlatformIndex] : nullptr;
+	// don't save Editor resources from the Engine if the target doesn't have editoronly data
+	if (COTFS.IsCookFlagSet(ECookInitializationFlags::SkipEditorContent) &&
+		(PackageName.StartsWith(TEXT("/Engine/Editor")) || PackageName.StartsWith(TEXT("/Engine/VREditor"))) &&
+		!TargetPlatform->HasEditorOnlyData())
+	{
+		SavePackageResult = ESavePackageResult::ContainsEditorOnlyData;
+		return;
+	}
+	// Check whether or not game-specific behaviour should prevent this package from being cooked for the target platform
+	else if (UAssetManager::IsValid() && !UAssetManager::Get().ShouldCookForPlatform(Package, TargetPlatform))
+	{
+		SavePackageResult = ESavePackageResult::ContainsEditorOnlyData;
+		UE_LOG(LogCook, Display, TEXT("Excluding %s -> %s"), *PackageName, *PlatFilename);
+		return;
+	}
+	// check if this package is unsupported for the target platform (typically plugin content)
+	else
+	{
+		TSet<FName>* NeverCookPackages = COTFS.PackageTracker->PlatformSpecificNeverCookPackages.Find(TargetPlatform);
+		if (NeverCookPackages && NeverCookPackages->Find(Package->GetFName()))
+		{
+			SavePackageResult = ESavePackageResult::ContainsEditorOnlyData;
+			UE_LOG(LogCook, Display, TEXT("Excluding %s -> %s"), *PackageName, *PlatFilename);
+			return;
+		}
+	}
 
-						IPackageStoreWriter* PackageStoreWriter = SavePackageContext ? SavePackageContext->PackageStoreWriter : nullptr;
-						if (PackageStoreWriter)
-						{
-							IPackageStoreWriter::FBeginPackageInfo Info;
-							Info.PackageName = Package->GetFName();
-							
-							PackageStoreWriter->BeginPackage(Info);
-						}
+	const FString FullFilename = FPaths::ConvertRelativePathToFull(PlatFilename);
+	if (FullFilename.Len() >= FPlatformMisc::GetMaxPathLength())
+	{
+		LogCookerMessage(FString::Printf(TEXT("Couldn't save package, filename is too long (%d >= %d): %s"),
+			FullFilename.Len(), FPlatformMisc::GetMaxPathLength(), *FullFilename), EMessageSeverity::Error);
+		SavePackageResult = ESavePackageResult::Error;
+		return;
+	}
 
-						Result = GEditor->Save(	Package, World, FlagsToCook, *PlatFilename, 
-												GError, nullptr, bSwap, false, SaveFlags, Target,
-												FDateTime::MinValue(), false, /*DiffMap*/ nullptr, 
-												SavePackageContext);
+	if (!bHasDelayLoaded)
+	{
+		// look for a world object in the package (if there is one, there's a map)
+		World = UWorld::FindWorldInPackage(Package);
+		if (World)
+		{
+			FlagsToCook = RF_NoFlags;
+		}
+		bContainsMap = Package->ContainsMap();
+		bHasDelayLoaded = true;
+	}
 
-						if (PackageStoreWriter)
-						{
-							FAssetRegistryGenerator* Generator = PlatformManager->GetPlatformData(Target)->RegistryGenerator.Get();
-							FAssetPackageData* AssetPackageData = Generator->GetAssetPackageData(Package->GetFName());
-							check(AssetPackageData);
-
-							IPackageStoreWriter::FCommitPackageInfo Info;
-							Info.bSucceeded		= Result.IsSuccessful();
-							Info.PackageName	= Package->GetFName();
-							PRAGMA_DISABLE_DEPRECATION_WARNINGS
-							Info.PackageGuid	= AssetPackageData->PackageGuid;
-							PRAGMA_ENABLE_DEPRECATION_WARNINGS
-
-							PackageStoreWriter->CommitPackage(Info);
-						}
-					}
+	UE_CLOG(GCookProgressDisplay & (int32)ECookProgressDisplayMode::PackageNames, LogCook, Display, TEXT("Cooking %s -> %s"),
+		*PackageName, *PlatFilename);
 
-					// if running linker diff, resave the package again using the other save algorithm
-					if (DiffModeHelper.IsRunningCookLinkerDiff())
-					{
-						static IConsoleVariable* EnableNewSave = IConsoleManager::Get().FindConsoleVariable(TEXT("SavePackage.EnableNewSave"));
-						int32 PreviousCvarValue = EnableNewSave->GetInt();
+	bEndianSwap = (!TargetPlatform->IsLittleEndian()) ^ (!PLATFORM_LITTLE_ENDIAN);
 
-						// If the linker diff is comparing the two save algo, switch the cvar to the other algo before saving again,
-						// Otherwise the linker diff mode is tracking if the save is consistent across multiple save
-						bool bAlgoLinkerDiff = DiffModeHelper.GetLinkerDiffMode() == FDiffModeCookServerUtils::LDM_Algo;
-						if (bAlgoLinkerDiff)
-						{
-							// see CVarEnablePackageNewSave definition in SavePackageUtilities.cpp for value meaning
-							EnableNewSave->Set(PreviousCvarValue & 1 ? 0 : 1);
-						}
+	if (!TargetPlatform->HasEditorOnlyData())
+	{
+		Package->SetPackageFlags(PKG_FilterEditorOnly);
+	}
+	else
+	{
+		Package->ClearPackageFlags(PKG_FilterEditorOnly);
+	}
 
-						FSavePackageResultStruct NewResult = GEditor->Save(Package, World, FlagsToCook, *PlatFilename,
-							GError, nullptr, bSwap, false, SaveFlags|SAVE_DiffOnly, Target,
-							FDateTime::MinValue(), false, /*DiffMap*/ nullptr,
-							nullptr);
+	if (World)
+	{
+		// Fixup legacy lightmaps before saving
+		// This should be done after loading, but Core loads UWorlds with LoadObject so there's no opportunity to handle this fixup on load
+		World->PersistentLevel->HandleLegacyMapBuildData();
+	}
 
-						if (bAlgoLinkerDiff)
-						{
-							EnableNewSave->Set(PreviousCvarValue);
-						}
+	SavePackageContext = COTFS.GetSavePackageContext(TargetPlatform);
+	PackageStoreWriter = SavePackageContext ? SavePackageContext->PackageStoreWriter : nullptr;
+	if (PackageStoreWriter)
+	{
+		IPackageStoreWriter::FBeginPackageInfo Info;
+		Info.PackageName = Package->GetFName();
 
-						if (Result.LinkerSave && NewResult.LinkerSave)
-						{
-							FLinkerDiff LinkerDiff = FLinkerDiff::CompareLinkers(Result.LinkerSave.Get(), NewResult.LinkerSave.Get());
-							LinkerDiff.PrintDiff(*GWarn);
-						}
-						Result.LinkerSave.Reset();
-						NewResult.LinkerSave.Reset();
-					}
+		PackageStoreWriter->BeginPackage(Info);
+	}
 
-					GIsCookerLoadingPackage = false;
-					++this->StatSavedPackageCount;
+	// Indicate Setup was successful
+	bPlatformSetupSuccessful = true;
+	SavePackageResult = ESavePackageResult::Success;
+}
 
-					// If package was actually saved check with asset manager to make sure it wasn't excluded for being a development or never cook package. We do this after Editor Only filtering
-					if (Result == ESavePackageResult::Success && UAssetManager::IsValid())
-					{
-						UE_SCOPED_HIERARCHICAL_COOKTIMER(VerifyCanCookPackage);
-						if (!UAssetManager::Get().VerifyCanCookPackage(Package->GetFName()))
-						{
-							Result = ESavePackageResult::Error;
-						}
-					}
-				}
-			}
-		}
+void FSaveCookedPackageContext::FinishPlatform()
+{
+	bool bSuccessful = SavePackageResult.IsSuccessful();
+	bool bLocalReferencedOnlyByEditorOnlyData = SavePackageResult == ESavePackageResult::ReferencedOnlyByEditorOnlyData;
+
+	if (bPlatformSetupSuccessful && PackageStoreWriter)
+	{
+		FAssetRegistryGenerator* Generator = COTFS.PlatformManager->GetPlatformData(TargetPlatform)->RegistryGenerator.Get();
+		FAssetPackageData* AssetPackageData = Generator->GetAssetPackageData(Package->GetFName());
+		check(AssetPackageData);
+
+		IPackageStoreWriter::FCommitPackageInfo Info;
+		Info.bSucceeded = bSuccessful;
+		Info.PackageName = Package->GetFName();
+		PRAGMA_DISABLE_DEPRECATION_WARNINGS
+		Info.PackageGuid = AssetPackageData->PackageGuid;
+		PRAGMA_ENABLE_DEPRECATION_WARNINGS
 
-		Package->SetPackageFlagsTo(OriginalPackageFlags);
+		PackageStoreWriter->CommitPackage(Info);
 	}
-	else
+
+	// Update asset registry
+	if (COTFS.CookByTheBookOptions)
+	{
+		FAssetRegistryGenerator* Generator = COTFS.PlatformManager->GetPlatformData(TargetPlatform)->RegistryGenerator.Get();
+		COTFS.UpdateAssetRegistryPackageData(Generator, *Package, SavePackageResult);
+	}
+
+	// In success or failure we want to mark the package as cooked, unless the failure was due to being referenced only
+	// by editor-only data; we may need to save it later when new content loads it through non editor-only references.
+	if (!bLocalReferencedOnlyByEditorOnlyData)
 	{
-		for (int32 PlatformIndex = 0; PlatformIndex < TargetPlatforms.Num(); ++PlatformIndex)
+		PackageData.AddCookedPlatform(TargetPlatform, bSuccessful);
+	}
+
+	// Update flags used to determine garbage collection.
+	if (bSuccessful)
+	{
+		if (bContainsMap)
 		{
-			SavePackageResults.Add(FSavePackageResultStruct(ESavePackageResult::MissingFile));
+			StackData.ResultFlags |= UCookOnTheFlyServer::COSR_CookedMap;
+		}
+		else
+		{
+			++StackData.CookedPackageCount;
+			StackData.ResultFlags |= UCookOnTheFlyServer::COSR_CookedPackage;
 		}
 	}
 
-	// Don't resolve, just add to request list as needed, do this after save to catch save permutations
-	TSet<FName> SoftObjectPackages;
-	if (!IsCookByTheBookMode() || !CookByTheBookOptions->bSkipSoftReferences)
+	// Accumulate results for SaveCookedPackage_Finish
+	if (bHasFirstPlatformResults && bLocalReferencedOnlyByEditorOnlyData != bReferencedOnlyByEditorOnlyData)
 	{
-		GRedirectCollector.ProcessSoftObjectPathPackageList(Package->GetFName(), false, SoftObjectPackages);
+		UE_LOG(LogCook, Error, TEXT("Package %s had different values for IsReferencedOnlyByEditorOnlyData from multiple platforms. ")
+			TEXT("Treading all platforms as IsReferencedOnlyByEditorOnlyData = true; this will cause the package to be ignored on the platforms that need it."),
+			*Package->GetName());
+		bReferencedOnlyByEditorOnlyData = true;
+	}
+	else
+	{
+		bReferencedOnlyByEditorOnlyData = bLocalReferencedOnlyByEditorOnlyData;
+	}
+	bHasFirstPlatformResults = true;
+}
+
+void FSaveCookedPackageContext::FinishPackage()
+{
+	// Add soft references discovered from the package
+	if (!COTFS.IsCookByTheBookMode() || !COTFS.CookByTheBookOptions->bSkipSoftReferences)
+	{
+		// Also request any localized variants of this package
+		if (COTFS.IsCookByTheBookMode() && !FPackageName::IsLocalizedPackage(Package->GetName()))
+		{
+			const TArray<FName>* LocalizedVariants = COTFS.CookByTheBookOptions->SourceToLocalizedPackageVariants.Find(Package->GetFName());
+			if (LocalizedVariants)
+			{
+				for (const FName& LocalizedPackageName : *LocalizedVariants)
+				{
+					UE::Cook::FPackageData* LocalizedPackageData = COTFS.PackageDatas->TryAddPackageDataByPackageName(LocalizedPackageName);
+					if (LocalizedPackageData)
+					{
+						bool bIsUrgent = false;
+						LocalizedPackageData->UpdateRequestData(PackageData.GetRequestedPlatforms(), bIsUrgent, UE::Cook::FCompletionCallback());
+					}
+				}
+			}
+		}
 
+		// Add SoftObjectPaths to the cook. This has to be done after the package save to catch any SoftObjectPaths that were added during save.
+		TSet<FName> SoftObjectPackages;
+		GRedirectCollector.ProcessSoftObjectPathPackageList(Package->GetFName(), false, SoftObjectPackages);
 		for (FName SoftObjectPackage : SoftObjectPackages)
 		{
 			TMap<FName, FName> RedirectedPaths;
 
 			// If this is a redirector, extract destination from asset registry
-			if (ContainsRedirector(SoftObjectPackage, RedirectedPaths))
+			if (COTFS.ContainsRedirector(SoftObjectPackage, RedirectedPaths))
 			{
 				for (TPair<FName, FName>& RedirectedPath : RedirectedPaths)
 				{
@@ -4764,10 +4826,9 @@ void UCookOnTheFlyServer::SaveCookedPackage(UE::Cook::FPackageData& PackageData,
 				}
 			}
 
-			// Verify package actually exists
-			if (IsCookByTheBookMode())
+			if (COTFS.IsCookByTheBookMode())
 			{
-				UE::Cook::FPackageData* SoftObjectPackageData = PackageDatas->TryAddPackageDataByPackageName(SoftObjectPackage);
+				UE::Cook::FPackageData* SoftObjectPackageData = COTFS.PackageDatas->TryAddPackageDataByPackageName(SoftObjectPackage);
 				if (SoftObjectPackageData)
 				{
 					bool bIsUrgent = false;
@@ -4777,11 +4838,29 @@ void UCookOnTheFlyServer::SaveCookedPackage(UE::Cook::FPackageData& PackageData,
 		}
 	}
 
-	check(bIsSavingPackage == true);
-	bIsSavingPackage = false;
+	if (!bReferencedOnlyByEditorOnlyData)
+	{
+		if ((COTFS.CurrentCookMode == ECookMode::CookOnTheFly) && !PackageData.GetIsUrgent())
+		{
+			// this is an unsolicited package
+			if (FPaths::FileExists(Filename))
+			{
+				COTFS.PackageTracker->UnsolicitedCookedPackages.AddCookedPackage(FFilePlatformRequest(PackageData.GetFileName(), PlatformsForPackage));
 
+#if DEBUG_COOKONTHEFLY
+				UE_LOG(LogCook, Display, TEXT("UnsolicitedCookedPackages: %s"), *Filename);
+#endif
+			}
+		}
+	}
+	else
+	{
+		COTFS.PackageTracker->UncookedEditorOnlyPackages.AddUnique(Package->GetFName());
+	}
 }
 
+} // namespace UE::Cook
+
 void UCookOnTheFlyServer::Initialize( ECookMode::Type DesiredCookMode, ECookInitializationFlags InCookFlags, const FString &InOutputDirectoryOverride )
 {
 	TRACE_CPUPROFILER_EVENT_SCOPE(UCookOnTheFlyServer::Initialize);
@@ -9275,4 +9354,17 @@ IPackageStoreWriter* UCookOnTheFlyServer::GetPackageStoreWriter(const FName& Pla
 	return nullptr;
 }
 
+FSavePackageContext* UCookOnTheFlyServer::GetSavePackageContext(const ITargetPlatform* TargetPlatform) const
+{
+	for (FSavePackageContext* Context : SavePackageContexts)
+	{
+		if (Context->TargetPlatform == TargetPlatform)
+		{
+			return Context;
+		}
+	}
+	return nullptr;
+}
+
+
 #undef LOCTEXT_NAMESPACE
diff --git a/Engine/Source/Editor/UnrealEd/Private/Cooker/CookPackageData.cpp b/Engine/Source/Editor/UnrealEd/Private/Cooker/CookPackageData.cpp
--- a/Engine/Source/Editor/UnrealEd/Private/Cooker/CookPackageData.cpp
+++ b/Engine/Source/Editor/UnrealEd/Private/Cooker/CookPackageData.cpp
@@ -178,56 +178,32 @@ namespace Cook
 	void FPackageData::AddCookedPlatforms(const TArrayView<const ITargetPlatform* const>& New, const TArrayView<const bool>& Succeeded)
 	{
 		check(New.Num() == Succeeded.Num());
-		int32 Num = New.Num();
-		if (Num == 0)
+		for (int32 n = 0; n < New.Num(); ++n)
 		{
-			return;
+			AddCookedPlatform(New[n], Succeeded[n]);
 		}
+	}
 
-		const ITargetPlatform* const* NewData = New.GetData();
-		const bool* SucceededData = Succeeded.GetData();
-		for (int32 n = 0; n < Num; ++n)
+	void FPackageData::AddCookedPlatforms(const TArrayView<const ITargetPlatform* const>& New, bool bSucceeded)
+	{
+		for (const ITargetPlatform* TargetPlatform : New)
 		{
-			const ITargetPlatform* TargetPlatform = NewData[n];
-			bool bSucceeded = SucceededData[n];
-
-			int32 ExistingIndex = CookedPlatforms.IndexOfByKey(TargetPlatform);
-			if (ExistingIndex != INDEX_NONE)
-			{
-				CookSucceeded[ExistingIndex] = bSucceeded;
-			}
-			else
-			{
-				CookedPlatforms.Add(TargetPlatform);
-				CookSucceeded.Add(bSucceeded);
-			}
+			AddCookedPlatform(TargetPlatform, bSucceeded);
 		}
-
-		PackageDatas.GetMonitor().OnCookedPlatformAdded(*this);
 	}
 
-	void FPackageData::AddCookedPlatforms(const TArrayView<const ITargetPlatform* const>& New, bool bSucceeded)
+	void FPackageData::AddCookedPlatform(const ITargetPlatform* TargetPlatform, bool bSucceeded)
 	{
-		int32 Num = New.Num();
-		if (Num == 0)
+		int32 ExistingIndex = CookedPlatforms.IndexOfByKey(TargetPlatform);
+		if (ExistingIndex != INDEX_NONE)
 		{
-			return;
+			CookSucceeded[ExistingIndex] = bSucceeded;
 		}
-
-		for (const ITargetPlatform* TargetPlatform : New)
+		else
 		{
-			int32 ExistingIndex = CookedPlatforms.IndexOfByKey(TargetPlatform);
-			if (ExistingIndex != INDEX_NONE)
-			{
-				CookSucceeded[ExistingIndex] = bSucceeded;
-			}
-			else
-			{
-				CookedPlatforms.Add(TargetPlatform);
-				CookSucceeded.Add(bSucceeded);
-			}
+			CookedPlatforms.Add(TargetPlatform);
+			CookSucceeded.Add(bSucceeded);
 		}
-
 		PackageDatas.GetMonitor().OnCookedPlatformAdded(*this);
 	}
 
diff --git a/Engine/Source/Editor/UnrealEd/Private/Cooker/CookPackageData.h b/Engine/Source/Editor/UnrealEd/Private/Cooker/CookPackageData.h
--- a/Engine/Source/Editor/UnrealEd/Private/Cooker/CookPackageData.h
+++ b/Engine/Source/Editor/UnrealEd/Private/Cooker/CookPackageData.h
@@ -137,6 +137,8 @@ namespace Cook
 		void AddCookedPlatforms(const TArrayView<const ITargetPlatform* const>& New, const TArrayView<const bool>& Succeeded);
 		/** Add each element of New to CookedPlatforms if it is not already present, succeeded set to bSucceeded. If a platform is already present in CookedPlatforms, its success flag is overwritten. */
 		void AddCookedPlatforms(const TArrayView<const ITargetPlatform* const>& New, bool bSucceeded);
+		/** Add TargetPlatform to CookedPlatforms if it is not already present, succeeded set to bSucceeded. If TargetPlatform is already present in CookedPlatforms, its success flag is overwritten. */
+		void AddCookedPlatform(const ITargetPlatform* TargetPlatform, bool bSucceeded);
 		/** Remove the given Platform and its succeeded flag from CookedPlatforms if it exists. */
 		void RemoveCookedPlatform(const ITargetPlatform* Platform);
 		/** Remove each element of Platforms and its succeeded flags from CookedPlatforms if it exists. */
