Remove code in UMaterialGraphNode_Base that associated additional data with pins via a TMap
Instead, add 'SourceIndex' directly to UEdGraphPin and remove the map
Too many edge cases where something would modify pins directly, and the map would get out of sync with actual pins

diff --git a/Engine/Source/Editor/GraphEditor/Private/MaterialNodes/SGraphNodeMaterialBase.cpp b/Engine/Source/Editor/GraphEditor/Private/MaterialNodes/SGraphNodeMaterialBase.cpp
--- a/Engine/Source/Editor/GraphEditor/Private/MaterialNodes/SGraphNodeMaterialBase.cpp
+++ b/Engine/Source/Editor/GraphEditor/Private/MaterialNodes/SGraphNodeMaterialBase.cpp
@@ -14,6 +14,7 @@
 #include "CanvasItem.h"
 #include "CanvasTypes.h"
 #include "MaterialGraph/MaterialGraphNode.h"
+#include "MaterialGraph/MaterialGraphSchema.h"
 #include "GraphEditorSettings.h"
 #include "SGraphPanel.h"
 #include "TutorialMetaData.h"
@@ -278,10 +279,9 @@ void SGraphNodeMaterialBase::CreatePinWidgets()
 		
 			if ((bIsABreakAttrNode && CurPin->Direction == EGPD_Output) || (bIsAMakeAttrNode && CurPin->Direction == EGPD_Input))
 			{
-				const FMaterialGraphPinInfo& PinInfo = MaterialNode->GetPinInfo(CurPin);
-				if (PinInfo.PinType == EMaterialGraphPinType::Data)
+				if (CurPin->PinType.PinCategory != UMaterialGraphSchema::PC_Exec)
 				{
-					bPinDesiresToBeHidden |= !MaterialGraph->MaterialInputs[PinInfo.Index].IsVisiblePin(MaterialGraph->Material, true);
+					bPinDesiresToBeHidden |= !MaterialGraph->MaterialInputs[CurPin->SourceIndex].IsVisiblePin(MaterialGraph->Material, true);
 				}
 			}
 		}
diff --git a/Engine/Source/Editor/GraphEditor/Private/MaterialNodes/SGraphNodeMaterialResult.cpp b/Engine/Source/Editor/GraphEditor/Private/MaterialNodes/SGraphNodeMaterialResult.cpp
--- a/Engine/Source/Editor/GraphEditor/Private/MaterialNodes/SGraphNodeMaterialResult.cpp
+++ b/Engine/Source/Editor/GraphEditor/Private/MaterialNodes/SGraphNodeMaterialResult.cpp
@@ -3,6 +3,7 @@
 #include "MaterialNodes/SGraphNodeMaterialResult.h"
 #include "MaterialGraph/MaterialGraph.h"
 #include "MaterialGraph/MaterialGraphNode_Root.h"
+#include "MaterialGraph/MaterialGraphSchema.h"
 #include "SGraphPanel.h"
 #include "TutorialMetaData.h"
 #include "Widgets/SToolTip.h"
@@ -33,16 +34,14 @@ void SGraphNodeMaterialResult::CreatePinWidgets()
 		bHideNoConnectionPins = OwnerGraphPanelPtr.Pin()->GetPinVisibility() == SGraphEditor::Pin_HideNoConnection;
 	}
 
-	for(const auto& It : MaterialGraphNode->PinInfoMap)
+	for (UEdGraphPin* CurPin : MaterialGraphNode->Pins)
 	{
-		UEdGraphPin* CurPin = It.Key;
-		const FMaterialGraphPinInfo& PinInfo = It.Value;
 		const bool bPinHasConections = CurPin->LinkedTo.Num() > 0;
 
 		bool bPinDesiresToBeHidden = bHideNoConnectionPins && !bPinHasConections;
-		if (PinInfo.PinType == EMaterialGraphPinType::Data)
+		if (CurPin->PinType.PinCategory != UMaterialGraphSchema::PC_Exec)
 		{
-			if (!MaterialGraph->MaterialInputs[PinInfo.Index].IsVisiblePin(MaterialGraph->Material))
+			if (!MaterialGraph->MaterialInputs[CurPin->SourceIndex].IsVisiblePin(MaterialGraph->Material))
 			{
 				bPinDesiresToBeHidden = true;
 			}
@@ -53,9 +52,9 @@ void SGraphNodeMaterialResult::CreatePinWidgets()
 			TSharedPtr<SGraphPin> NewPin = CreatePinWidget(CurPin);
 			check(NewPin.IsValid());
 
-			if (PinInfo.PinType == EMaterialGraphPinType::Data)
+			if (CurPin->PinType.PinCategory != UMaterialGraphSchema::PC_Exec)
 			{
-				TSharedPtr<SToolTip> ToolTipWidget = IDocumentation::Get()->CreateToolTip(MaterialGraph->MaterialInputs[PinInfo.Index].GetToolTip(), nullptr, FString(TEXT("")), FString(TEXT("")));
+				TSharedPtr<SToolTip> ToolTipWidget = IDocumentation::Get()->CreateToolTip(MaterialGraph->MaterialInputs[CurPin->SourceIndex].GetToolTip(), nullptr, FString(TEXT("")), FString(TEXT("")));
 				NewPin->SetToolTip(ToolTipWidget.ToSharedRef());
 			}
 
diff --git a/Engine/Source/Editor/MaterialEditor/Private/MaterialEditor.cpp b/Engine/Source/Editor/MaterialEditor/Private/MaterialEditor.cpp
--- a/Engine/Source/Editor/MaterialEditor/Private/MaterialEditor.cpp
+++ b/Engine/Source/Editor/MaterialEditor/Private/MaterialEditor.cpp
@@ -3334,19 +3334,22 @@ bool FMaterialEditor::IsToggleHideUnrelatedNodesChecked() const
 
 void FMaterialEditor::CollectDownstreamNodes(UMaterialGraphNode* CurrentNode, TArray<UMaterialGraphNode*>& CollectedNodes)
 {
-	for (auto& OutputPin : CurrentNode->OutputPins)
+	for (UEdGraphPin* OutputPin : CurrentNode->Pins)
 	{
-		for (auto& Link : OutputPin->LinkedTo)
+		if (OutputPin->Direction == EGPD_Output)
 		{
-			UMaterialGraphNode* LinkedNode = Cast<UMaterialGraphNode>(Link->GetOwningNode());
-			if (LinkedNode && !CollectedNodes.Contains(LinkedNode))
+			for (auto& Link : OutputPin->LinkedTo)
 			{
-				CollectedNodes.Add(LinkedNode);
-				CollectDownstreamNodes( LinkedNode, CollectedNodes );
-
-				if (bFocusWholeChain)
+				UMaterialGraphNode* LinkedNode = Cast<UMaterialGraphNode>(Link->GetOwningNode());
+				if (LinkedNode && !CollectedNodes.Contains(LinkedNode))
 				{
-					CollectUpstreamNodes( LinkedNode, CollectedNodes );
+					CollectedNodes.Add(LinkedNode);
+					CollectDownstreamNodes(LinkedNode, CollectedNodes);
+
+					if (bFocusWholeChain)
+					{
+						CollectUpstreamNodes(LinkedNode, CollectedNodes);
+					}
 				}
 			}
 		}
@@ -3355,15 +3358,18 @@ void FMaterialEditor::CollectDownstreamNodes(UMaterialGraphNode* CurrentNode, TA
 
 void FMaterialEditor::CollectUpstreamNodes(UMaterialGraphNode* CurrentNode, TArray<UMaterialGraphNode*>& CollectedNodes)
 {
-	for (auto& InputPin : CurrentNode->InputPins)
+	for (UEdGraphPin* InputPin : CurrentNode->Pins)
 	{
-		for (auto& Link : InputPin->LinkedTo)
+		if (InputPin->Direction == EGPD_Input)
 		{
-			UMaterialGraphNode* LinkedNode = Cast<UMaterialGraphNode>(Link->GetOwningNode());
-			if (LinkedNode && !CollectedNodes.Contains(LinkedNode))
+			for (auto& Link : InputPin->LinkedTo)
 			{
-				CollectedNodes.Add(LinkedNode);
-				CollectUpstreamNodes( LinkedNode, CollectedNodes );
+				UMaterialGraphNode* LinkedNode = Cast<UMaterialGraphNode>(Link->GetOwningNode());
+				if (LinkedNode && !CollectedNodes.Contains(LinkedNode))
+				{
+					CollectedNodes.Add(LinkedNode);
+					CollectUpstreamNodes(LinkedNode, CollectedNodes);
+				}
 			}
 		}
 	}
@@ -4100,20 +4106,23 @@ void FMaterialEditor::OnSelectDownstreamNodes()
 	while (NodesToCheck.Num() > 0)
 	{
 		UMaterialGraphNode* CurrentNode = NodesToCheck.Last();
-		for (int32 Index = 0; Index < CurrentNode->OutputPins.Num(); ++Index)
+		for (UEdGraphPin* Pin : CurrentNode->Pins)
 		{
-			for (int32 LinkIndex = 0; LinkIndex < CurrentNode->OutputPins[Index]->LinkedTo.Num(); ++LinkIndex)
+			if (Pin->Direction == EGPD_Output)
 			{
-				UMaterialGraphNode* LinkedNode = Cast<UMaterialGraphNode>(CurrentNode->OutputPins[Index]->LinkedTo[LinkIndex]->GetOwningNode());
-				if (LinkedNode)
+				for (int32 LinkIndex = 0; LinkIndex < Pin->LinkedTo.Num(); ++LinkIndex)
 				{
-					int32 FoundIndex = -1;
-					CheckedNodes.Find(LinkedNode, FoundIndex);
-
-					if (FoundIndex < 0)
+					UMaterialGraphNode* LinkedNode = Cast<UMaterialGraphNode>(Pin->LinkedTo[LinkIndex]->GetOwningNode());
+					if (LinkedNode)
 					{
-						NodesToSelect.Add(LinkedNode);
-						NodesToCheck.Add(LinkedNode);
+						int32 FoundIndex = -1;
+						CheckedNodes.Find(LinkedNode, FoundIndex);
+
+						if (FoundIndex < 0)
+						{
+							NodesToSelect.Add(LinkedNode);
+							NodesToCheck.Add(LinkedNode);
+						}
 					}
 				}
 			}
@@ -4153,20 +4162,23 @@ void FMaterialEditor::OnSelectUpstreamNodes()
 	while (NodesToCheck.Num() > 0)
 	{
 		UMaterialGraphNode* CurrentNode = NodesToCheck.Last();
-		for (int32 Index = 0; Index < CurrentNode->InputPins.Num(); ++Index)
+		for (UEdGraphPin* Pin : CurrentNode->Pins)
 		{
-			for (int32 LinkIndex = 0; LinkIndex < CurrentNode->InputPins[Index]->LinkedTo.Num(); ++LinkIndex)
+			if (Pin->Direction == EGPD_Input)
 			{
-				UMaterialGraphNode* LinkedNode = Cast<UMaterialGraphNode>(CurrentNode->InputPins[Index]->LinkedTo[LinkIndex]->GetOwningNode());
-				if (LinkedNode)
+				for (int32 LinkIndex = 0; LinkIndex < Pin->LinkedTo.Num(); ++LinkIndex)
 				{
-					int32 FoundIndex = -1;
-					CheckedNodes.Find(LinkedNode, FoundIndex);
-
-					if (FoundIndex < 0)
+					UMaterialGraphNode* LinkedNode = Cast<UMaterialGraphNode>(Pin->LinkedTo[LinkIndex]->GetOwningNode());
+					if (LinkedNode)
 					{
-						NodesToSelect.Add(LinkedNode);
-						NodesToCheck.Add(LinkedNode);
+						int32 FoundIndex = -1;
+						CheckedNodes.Find(LinkedNode, FoundIndex);
+
+						if (FoundIndex < 0)
+						{
+							NodesToSelect.Add(LinkedNode);
+							NodesToCheck.Add(LinkedNode);
+						}
 					}
 				}
 			}
diff --git a/Engine/Source/Editor/UnrealEd/Classes/MaterialGraph/MaterialGraphNode.h b/Engine/Source/Editor/UnrealEd/Classes/MaterialGraph/MaterialGraphNode.h
--- a/Engine/Source/Editor/UnrealEd/Classes/MaterialGraph/MaterialGraphNode.h
+++ b/Engine/Source/Editor/UnrealEd/Classes/MaterialGraph/MaterialGraphNode.h
@@ -71,7 +71,7 @@ public:
 	virtual void GetPinHoverText(const UEdGraphPin& Pin, FString& HoverTextOut) const override;
 	virtual FString GetDocumentationExcerptName() const override;
 	virtual bool CanUserDeleteNode() const override;
-	virtual uint32 GetPinMaterialType(const UEdGraphPin* Pin, const FMaterialGraphPinInfo& PinInfo) const override;
+	virtual uint32 GetPinMaterialType(const UEdGraphPin* Pin) const override;
 	//~ End UEdGraphNode Interface.
 
 	//~ Begin UMaterialGraphNode_Base Interface
diff --git a/Engine/Source/Editor/UnrealEd/Classes/MaterialGraph/MaterialGraphNode_Base.h b/Engine/Source/Editor/UnrealEd/Classes/MaterialGraph/MaterialGraphNode_Base.h
--- a/Engine/Source/Editor/UnrealEd/Classes/MaterialGraph/MaterialGraphNode_Base.h
+++ b/Engine/Source/Editor/UnrealEd/Classes/MaterialGraph/MaterialGraphNode_Base.h
@@ -10,34 +10,11 @@
 class UEdGraphPin;
 class UEdGraphSchema;
 
-enum class EMaterialGraphPinType
-{
-	Data,
-	Exec,
-};
-
-struct FMaterialGraphPinInfo
-{
-	EMaterialGraphPinType PinType;
-	int32 Index; // index into the expression's list of inputs/outputs (exec inputs/outpus are indexed separately)
-};
-
 UCLASS(MinimalAPI)
 class UMaterialGraphNode_Base : public UEdGraphNode
 {
 	GENERATED_UCLASS_BODY()
 
-	/** Contains additional information about material graph pins, avoid adding material-specific data to base pin type */
-	TMap<UEdGraphPin*, FMaterialGraphPinInfo> PinInfoMap;
-
-	/** Lists of pins that match up with the underlying UMaterialExpression's (non-exec) inputs and outputs */
-	TArray<UEdGraphPin*> InputPins;
-	TArray<UEdGraphPin*> OutputPins;
-
-	/** The exec input pin */
-	UEdGraphPin* ExecInputPin = nullptr;
-	TArray<UEdGraphPin*> ExecOutputPins;
-
 	/** Create all of the input pins required */
 	virtual void CreateInputPins() {};
 	/** Create all of the output pins required */
@@ -45,18 +22,16 @@ class UMaterialGraphNode_Base : public UEdGraphNode
 	/** Is this the undeletable root node */
 	virtual bool IsRootNode() const {return false;}
 	/** Get a single Input Pin via its index */
-	class UEdGraphPin* GetInputPin(int32 InputIndex) const { return InputPins[InputIndex]; }
-	/** Get a single Output Pin via its index (trusting the caller to provide a valid index) */
-	class UEdGraphPin* GetOutputPin(int32 OutputIndex) const { return OutputPins[OutputIndex]; }
-	/** Get a single Output Pin via its index (can return null if an invalid index is provided) */
-	class UEdGraphPin* TryGetOutputPin(int32 OutputIndex) const { return (OutputIndex >= 0 && OutputIndex < OutputPins.Num()) ? OutputPins[OutputIndex] : nullptr; }
+	UNREALED_API class UEdGraphPin* GetInputPin(int32 InputIndex) const;
+	/** Get a single Output Pin via its index */
+	UNREALED_API class UEdGraphPin* GetOutputPin(int32 OutputIndex) const;
 	/** Gets the exec input pin */
-	class UEdGraphPin* GetExecInputPin() const { return ExecInputPin; }
+	UNREALED_API class UEdGraphPin* GetExecInputPin() const;
+	/** Get a single exec Output Pin via its index */
+	UNREALED_API class UEdGraphPin* GetExecOutputPin(int32 OutputIndex) const;
 	/** Replace a given node with this one, changing all pin links */
 	UNREALED_API void ReplaceNode(UMaterialGraphNode_Base* OldNode);
 
-	UNREALED_API const FMaterialGraphPinInfo& GetPinInfo(const class UEdGraphPin* Pin) const;
-
 	/** Get the Material value type of an input pin */
 	uint32 GetInputType(const UEdGraphPin* InputPin) const;
 
@@ -85,9 +60,7 @@ class UMaterialGraphNode_Base : public UEdGraphNode
 protected:
 	void ModifyAndCopyPersistentPinData(UEdGraphPin& TargetPin, const UEdGraphPin& SourcePin) const;
 
-	void RegisterPin(UEdGraphPin* Pin, EMaterialGraphPinType Type, int32 Index);
-
-	virtual uint32 GetPinMaterialType(const UEdGraphPin* Pin, const FMaterialGraphPinInfo& PinInfo) const;
+	virtual uint32 GetPinMaterialType(const UEdGraphPin* Pin) const;
 
 	void EmptyPins();
 };
diff --git a/Engine/Source/Editor/UnrealEd/Classes/MaterialGraph/MaterialGraphNode_Root.h b/Engine/Source/Editor/UnrealEd/Classes/MaterialGraph/MaterialGraphNode_Root.h
--- a/Engine/Source/Editor/UnrealEd/Classes/MaterialGraph/MaterialGraphNode_Root.h
+++ b/Engine/Source/Editor/UnrealEd/Classes/MaterialGraph/MaterialGraphNode_Root.h
@@ -28,7 +28,7 @@ class UMaterialGraphNode_Root : public UMaterialGraphNode_Base
 	//~ End UEdGraphNode Interface.
 
 	//~ Begin UMaterialGraphNode_Base Interface
-	virtual uint32 GetPinMaterialType(const UEdGraphPin* Pin, const FMaterialGraphPinInfo& PinInfo) const override;
+	virtual uint32 GetPinMaterialType(const UEdGraphPin* Pin) const override;
 	virtual void CreateInputPins() override;
 	virtual bool IsRootNode() const override {return true;}
 	//~ End UMaterialGraphNode_Base Interface
diff --git a/Engine/Source/Editor/UnrealEd/Private/MaterialGraph.cpp b/Engine/Source/Editor/UnrealEd/Private/MaterialGraph.cpp
--- a/Engine/Source/Editor/UnrealEd/Private/MaterialGraph.cpp
+++ b/Engine/Source/Editor/UnrealEd/Private/MaterialGraph.cpp
@@ -9,6 +9,7 @@
 #include "MaterialGraph/MaterialGraphNode_PinBase.h"
 #include "MaterialGraph/MaterialGraphNode.h"
 #include "MaterialGraph/MaterialGraphNode_Root.h"
+#include "MaterialGraph/MaterialGraphSchema.h"
 
 #include "Materials/MaterialExpressionComment.h"
 #include "Materials/MaterialExpressionComposite.h"
@@ -336,25 +337,23 @@ void UMaterialGraph::LinkGraphNodesFromMaterial()
 		TArray<FExpressionExecOutputEntry> ExecOutputs;
 		Expression->GetExecOutputs(ExecOutputs);
 
-		for (const auto& It : MaterialGraphNode->PinInfoMap)
+		for (UEdGraphPin* Pin : MaterialGraphNode->Pins)
 		{
-			UEdGraphPin* Pin = It.Key;
-			const FMaterialGraphPinInfo& PinInfo = It.Value;
-			if (Pin->Direction == EGPD_Input && PinInfo.PinType == EMaterialGraphPinType::Data)
+			if (Pin->Direction == EGPD_Input && Pin->PinType.PinCategory != UMaterialGraphSchema::PC_Exec)
 			{
-				if (ExpressionInputs[PinInfo.Index]->Expression)
+				if (ExpressionInputs[Pin->SourceIndex]->Expression)
 				{
 					// Unclear why this is null sometimes outside of composite reroute, but this is safer than crashing
-					if (UMaterialGraphNode* GraphNode = Cast<UMaterialGraphNode>(ExpressionInputs[PinInfo.Index]->Expression->GraphNode))
+					if (UMaterialGraphNode* GraphNode = Cast<UMaterialGraphNode>(ExpressionInputs[Pin->SourceIndex]->Expression->GraphNode))
 					{
 						// if GraphNode is a material function call for a missing material function, it may not have any output pins
-						UEdGraphPin* OutputPin = GraphNode->TryGetOutputPin(GetValidOutputIndex(ExpressionInputs[PinInfo.Index]));
+						UEdGraphPin* OutputPin = GraphNode->GetOutputPin(GetValidOutputIndex(ExpressionInputs[Pin->SourceIndex]));
 						if (LIKELY(OutputPin))
 						{
 							Pin->MakeLinkTo(OutputPin);
 						}
 					}
-					else if (UMaterialExpressionReroute* CompositeReroute = Cast<UMaterialExpressionReroute>(ExpressionInputs[PinInfo.Index]->Expression))
+					else if (UMaterialExpressionReroute* CompositeReroute = Cast<UMaterialExpressionReroute>(ExpressionInputs[Pin->SourceIndex]->Expression))
 					{
 						// This is an unseen composite reroute expression, find the actual expression output to connect to.
 						UMaterialExpressionComposite* OwningComposite = CastChecked<UMaterialExpressionComposite>(CompositeReroute->SubgraphExpression);
@@ -375,9 +374,9 @@ void UMaterialGraph::LinkGraphNodesFromMaterial()
 					}
 				}
 			}
-			else if (Pin->Direction == EGPD_Output && PinInfo.PinType == EMaterialGraphPinType::Exec)
+			else if (Pin->Direction == EGPD_Output && Pin->PinType.PinCategory == UMaterialGraphSchema::PC_Exec)
 			{
-				FExpressionExecOutput* ExecOutput = ExecOutputs[PinInfo.Index].Output;
+				FExpressionExecOutput* ExecOutput = ExecOutputs[Pin->SourceIndex].Output;
 				UMaterialExpression* ConnectedExpression = ExecOutput->GetExpression();
 				if (ConnectedExpression)
 				{
@@ -410,31 +409,27 @@ void UMaterialGraph::LinkMaterialExpressionsFromGraph() const
 			Material->Modify();
 			Material->EditorX = RootNode->NodePosX;
 			Material->EditorY = RootNode->NodePosY;
-			check(RootNode->InputPins.Num() == MaterialInputs.Num());
-			for (int32 PinIndex = 0; PinIndex < RootNode->InputPins.Num() && PinIndex < MaterialInputs.Num(); ++PinIndex)
+			
+			for (UEdGraphPin* Pin : RootNode->Pins)
 			{
-				FExpressionInput& MaterialInput = MaterialInputs[PinIndex].GetExpressionInput(Material);
-
-				if (RootNode->InputPins[PinIndex]->LinkedTo.Num() > 0)
+				if (Pin->Direction == EGPD_Input && Pin->PinType.PinCategory != UMaterialGraphSchema::PC_Exec)
 				{
-					UMaterialGraphNode* ConnectedNode = CastChecked<UMaterialGraphNode>(RootNode->InputPins[PinIndex]->LinkedTo[0]->GetOwningNode());
-					// Work out the index of the connected pin
-					for (int32 OutPinIndex = 0; OutPinIndex < ConnectedNode->OutputPins.Num(); ++OutPinIndex)
+					FExpressionInput& MaterialInput = MaterialInputs[Pin->SourceIndex].GetExpressionInput(Material);
+					if (Pin->LinkedTo.Num() > 0)
 					{
-						if (ConnectedNode->OutputPins[OutPinIndex] == RootNode->InputPins[PinIndex]->LinkedTo[0])
+						UEdGraphPin* ConnectedPin = Pin->LinkedTo[0];
+						UMaterialGraphNode* ConnectedNode = CastChecked<UMaterialGraphNode>(ConnectedPin->GetOwningNode());
+
+						if (!ConnectedNode->MaterialExpression->IsExpressionConnected(&MaterialInput, ConnectedPin->SourceIndex))
 						{
-							if (!ConnectedNode->MaterialExpression->IsExpressionConnected(&MaterialInput, OutPinIndex))
-							{
-								ConnectedNode->MaterialExpression->Modify();
-								MaterialInput.Connect(OutPinIndex, ConnectedNode->MaterialExpression);
-							}
-							break;
+							ConnectedNode->MaterialExpression->Modify();
+							MaterialInput.Connect(ConnectedPin->SourceIndex, ConnectedNode->MaterialExpression);
 						}
 					}
-				}
-				else if (MaterialInput.Expression)
-				{
-					MaterialInput.Expression = NULL;
+					else if (MaterialInput.Expression)
+					{
+						MaterialInput.Expression = NULL;
+					}
 				}
 			}
 		}
@@ -467,36 +462,27 @@ void UMaterialGraph::LinkMaterialExpressionsFromGraph() const
 					TArray<FExpressionExecOutputEntry> ExecOutputs;
 					Expression->GetExecOutputs(ExecOutputs);
 
-					for (const auto& It : GraphNode->PinInfoMap)
+					for (UEdGraphPin* Pin : GraphNode->Pins)
 					{
-						UEdGraphPin* Pin = It.Key;
-						const FMaterialGraphPinInfo& PinInfo = It.Value;
-						if (Pin->Direction == EGPD_Input && PinInfo.PinType == EMaterialGraphPinType::Data)
+						if (Pin->Direction == EGPD_Input && Pin->PinType.PinCategory != UMaterialGraphSchema::PC_Exec)
 						{
 							// Wire up non-execution input pins
-							FExpressionInput* ExpressionInput = ExpressionInputs[PinInfo.Index];
+							FExpressionInput* ExpressionInput = ExpressionInputs[Pin->SourceIndex];
 							if (Pin->LinkedTo.Num() > 0)
 							{
-								UMaterialGraphNode* ConnectedNode = CastChecked<UMaterialGraphNode>(Pin->LinkedTo[0]->GetOwningNode());
+								UEdGraphPin* ConnectedPin = Pin->LinkedTo[0];
+								UMaterialGraphNode* ConnectedNode = CastChecked<UMaterialGraphNode>(ConnectedPin->GetOwningNode());
 
-								// Work out the index of the connected pin
-								for (int32 OutPinIndex = 0; OutPinIndex < ConnectedNode->OutputPins.Num(); ++OutPinIndex)
+								if (ExpressionInput && !ConnectedNode->MaterialExpression->IsExpressionConnected(ExpressionInput, ConnectedPin->SourceIndex))
 								{
-									if (ConnectedNode->OutputPins[OutPinIndex] == Pin->LinkedTo[0])
+									if (!bModifiedExpression)
 									{
-										if (ExpressionInput && !ConnectedNode->MaterialExpression->IsExpressionConnected(ExpressionInput, OutPinIndex))
-										{
-											if (!bModifiedExpression)
-											{
-												bModifiedExpression = true;
-												Expression->Modify();
-											}
-
-											ConnectedNode->MaterialExpression->Modify();
-											ExpressionInput->Connect(OutPinIndex, ConnectedNode->MaterialExpression);
-										}
-										break;
+										bModifiedExpression = true;
+										Expression->Modify();
 									}
+
+									ConnectedNode->MaterialExpression->Modify();
+									ExpressionInput->Connect(ConnectedPin->SourceIndex, ConnectedNode->MaterialExpression);
 								}
 							}
 							else if (ExpressionInput && ExpressionInput->Expression)
@@ -509,10 +495,10 @@ void UMaterialGraph::LinkMaterialExpressionsFromGraph() const
 								ExpressionInput->Expression = NULL;
 							}
 						}
-						else if (Pin->Direction == EGPD_Output && PinInfo.PinType == EMaterialGraphPinType::Exec)
+						else if (Pin->Direction == EGPD_Output && Pin->PinType.PinCategory == UMaterialGraphSchema::PC_Exec)
 						{
 							// Wire up execution output pins
-							FExpressionExecOutput* ExpressionOutput = ExecOutputs[PinInfo.Index].Output;
+							FExpressionExecOutput* ExpressionOutput = ExecOutputs[Pin->SourceIndex].Output;
 
 							if (Pin->LinkedTo.Num() > 0)
 							{
@@ -592,13 +578,13 @@ void UMaterialGraph::LinkMaterialExpressionsFromGraph() const
 
 bool UMaterialGraph::IsInputActive(UEdGraphPin* GraphPin) const
 {
-	if (Material && RootNode)
+	if (Material &&
+		RootNode && // No inputs without a root node
+		GraphPin->Direction == EGPD_Input &&
+		GraphPin->PinType.PinCategory != UMaterialGraphSchema::PC_Exec)
 	{
-		int32 Index = INDEX_NONE;
-		if (RootNode->InputPins.Find(GraphPin, Index))
-		{
-			return Material->IsPropertyActiveInEditor(MaterialInputs[Index].GetProperty());
-		}
+		
+		return Material->IsPropertyActiveInEditor(MaterialInputs[GraphPin->SourceIndex].GetProperty());
 	}
 	return true;
 }
@@ -611,12 +597,15 @@ void UMaterialGraph::GetUnusedExpressions(TArray<UEdGraphNode*>& UnusedNodes) co
 
 	if (RootNode)
 	{
-		for (int32 Index = 0; Index < RootNode->InputPins.Num(); ++Index)
+		for (UEdGraphPin* Pin : RootNode->Pins)
 		{
-			if (MaterialInputs[Index].IsVisiblePin(Material)
-				&& RootNode->InputPins[Index]->LinkedTo.Num() > 0 && RootNode->InputPins[Index]->LinkedTo[0])
+			if (Pin->Direction == EGPD_Input &&
+				Pin->PinType.PinCategory != UMaterialGraphSchema::PC_Exec &&
+				MaterialInputs[Pin->SourceIndex].IsVisiblePin(Material) &&
+				Pin->LinkedTo.Num() > 0 &&
+				Pin->LinkedTo[0])
 			{
-				NodesToCheck.Push(RootNode->InputPins[Index]->LinkedTo[0]->GetOwningNode());
+				NodesToCheck.Push(Pin->LinkedTo[0]->GetOwningNode());
 			}
 		}
 
@@ -665,11 +654,15 @@ void UMaterialGraph::GetUnusedExpressions(TArray<UEdGraphNode*>& UnusedNodes) co
 				UsedNodes.Add(GraphNode);
 
 				// Iterate over the expression's inputs and add them to the pending stack.
-				for (int32 Index = 0; Index < GraphNode->InputPins.Num(); ++Index)
+				for (UEdGraphPin* Pin : RootNode->Pins)
 				{
-					if (GraphNode->InputPins[Index]->LinkedTo.Num() > 0 && GraphNode->InputPins[Index]->LinkedTo[0])
+					if (Pin->Direction == EGPD_Input &&
+						Pin->PinType.PinCategory != UMaterialGraphSchema::PC_Exec &&
+						MaterialInputs[Pin->SourceIndex].IsVisiblePin(Material) &&
+						Pin->LinkedTo.Num() > 0 &&
+						Pin->LinkedTo[0])
 					{
-						NodesToCheck.Push(GraphNode->InputPins[Index]->LinkedTo[0]->GetOwningNode());
+						NodesToCheck.Push(Pin->LinkedTo[0]->GetOwningNode());
 					}
 				}
 
diff --git a/Engine/Source/Editor/UnrealEd/Private/MaterialGraphNode.cpp b/Engine/Source/Editor/UnrealEd/Private/MaterialGraphNode.cpp
--- a/Engine/Source/Editor/UnrealEd/Private/MaterialGraphNode.cpp
+++ b/Engine/Source/Editor/UnrealEd/Private/MaterialGraphNode.cpp
@@ -605,24 +605,20 @@ FName UMaterialGraphNode::GetShortenPinName(const FName PinName)
 	return InputName;
 }
 
-uint32 UMaterialGraphNode::GetPinMaterialType(const UEdGraphPin* Pin, const FMaterialGraphPinInfo& PinInfo) const
+uint32 UMaterialGraphNode::GetPinMaterialType(const UEdGraphPin* Pin) const
 {
-	switch (PinInfo.PinType)
+	if (Pin->PinType.PinCategory == UMaterialGraphSchema::PC_Exec)
 	{
-	case EMaterialGraphPinType::Exec:
 		return MCT_Execution;
-	case EMaterialGraphPinType::Data:
-		if (Pin->Direction == EGPD_Input)
-		{
-			return MaterialExpression->GetInputType(PinInfo.Index);
-		}
-		else
-		{
-			return MaterialExpression->GetOutputType(PinInfo.Index);
-		}
-	default:
-		checkNoEntry();
-		return 0u;
+	}
+	
+	if (Pin->Direction == EGPD_Input)
+	{
+		return MaterialExpression->GetInputType(Pin->SourceIndex);
+	}
+	else
+	{
+		return MaterialExpression->GetOutputType(Pin->SourceIndex);
 	}
 }
 
@@ -631,11 +627,10 @@ void UMaterialGraphNode::CreateInputPins()
 	if (MaterialExpression->HasExecInput())
 	{
 		UEdGraphPin* NewPin = CreatePin(EGPD_Input, UMaterialGraphSchema::PC_Exec, NAME_None, NAME_None);
+		NewPin->SourceIndex = 0;
 		// Makes sure pin has a name for lookup purposes but user will never see it
 		NewPin->PinName = CreateUniquePinName(TEXT("Input"));
 		NewPin->PinFriendlyName = SpaceText;
-
-		RegisterPin(NewPin, EMaterialGraphPinType::Exec, 0);
 	}
 
 	const TArray<FExpressionInput*> ExpressionInputs = MaterialExpression->GetInputs();
@@ -656,14 +651,13 @@ void UMaterialGraphNode::CreateInputPins()
 		}
 
 		UEdGraphPin* NewPin = CreatePin(EGPD_Input, PinCategory, InputName);
+		NewPin->SourceIndex = Index;
 		if (NewPin->PinName.IsNone())
 		{
 			// Makes sure pin has a name for lookup purposes but user will never see it
 			NewPin->PinName = CreateUniquePinName(TEXT("Input"));
 			NewPin->PinFriendlyName = SpaceText;
 		}
-
-		RegisterPin(NewPin, EMaterialGraphPinType::Data, Index);
 	}
 }
 
@@ -709,14 +703,13 @@ void UMaterialGraphNode::CreateOutputPins()
 		}
 
 		UEdGraphPin* NewPin = CreatePin(EGPD_Output, PinCategory, PinSubCategory, OutputName);
+		NewPin->SourceIndex = Index;
 		if (NewPin->PinName.IsNone())
 		{
 			// Makes sure pin has a name for lookup purposes but user will never see it
 			NewPin->PinName = CreateUniquePinName(TEXT("Output"));
 			NewPin->PinFriendlyName = SpaceText;
 		}
-
-		RegisterPin(NewPin, EMaterialGraphPinType::Data, Index);
 	}
 
 	TArray<FExpressionExecOutputEntry> ExecOutputs;
@@ -732,14 +725,13 @@ void UMaterialGraphNode::CreateOutputPins()
 		}
 
 		UEdGraphPin* NewPin = CreatePin(EGPD_Output, UMaterialGraphSchema::PC_Exec, NAME_None, OutputName);
+		NewPin->SourceIndex = Index;
 		if (NewPin->PinName.IsNone())
 		{
 			// Makes sure pin has a name for lookup purposes but user will never see it
 			NewPin->PinName = CreateUniquePinName(TEXT("Output"));
 			NewPin->PinFriendlyName = SpaceText;
 		}
-
-		RegisterPin(NewPin, EMaterialGraphPinType::Exec, Index);
 	}
 }
 
diff --git a/Engine/Source/Editor/UnrealEd/Private/MaterialGraphNode_Base.cpp b/Engine/Source/Editor/UnrealEd/Private/MaterialGraphNode_Base.cpp
--- a/Engine/Source/Editor/UnrealEd/Private/MaterialGraphNode_Base.cpp
+++ b/Engine/Source/Editor/UnrealEd/Private/MaterialGraphNode_Base.cpp
@@ -17,35 +17,79 @@ UMaterialGraphNode_Base::UMaterialGraphNode_Base(const FObjectInitializer& Objec
 {
 }
 
-const FMaterialGraphPinInfo& UMaterialGraphNode_Base::GetPinInfo(const class UEdGraphPin* Pin) const
+UEdGraphPin* UMaterialGraphNode_Base::GetInputPin(int32 InputIndex) const
 {
-	const FMaterialGraphPinInfo* PinInfo = PinInfoMap.Find(Pin);
-	checkf(PinInfo, TEXT("Missing info for pin %s, missing call to RegisterPin()?"), *Pin->GetName());
-	return *PinInfo;
+	for (UEdGraphPin* Pin : Pins)
+	{
+		if (Pin->PinType.PinCategory != UMaterialGraphSchema::PC_Exec &&
+			Pin->Direction == EGPD_Input &&
+			Pin->SourceIndex == InputIndex)
+		{
+			return Pin;
+		}
+	}
+	return nullptr;
+}
+
+UEdGraphPin* UMaterialGraphNode_Base::GetOutputPin(int32 OutputIndex) const
+{
+	for (UEdGraphPin* Pin : Pins)
+	{
+		if (Pin->PinType.PinCategory != UMaterialGraphSchema::PC_Exec &&
+			Pin->Direction == EGPD_Output &&
+			Pin->SourceIndex == OutputIndex)
+		{
+			return Pin;
+		}
+	}
+	return nullptr;
+}
+
+UEdGraphPin* UMaterialGraphNode_Base::GetExecInputPin() const
+{
+	for (UEdGraphPin* Pin : Pins)
+	{
+		if (Pin->PinType.PinCategory == UMaterialGraphSchema::PC_Exec &&
+			Pin->Direction == EGPD_Input)
+		{
+			check(Pin->SourceIndex == 0);
+			return Pin;
+		}
+	}
+	return nullptr;
+}
+
+UEdGraphPin* UMaterialGraphNode_Base::GetExecOutputPin(int32 OutputIndex) const
+{
+	for (UEdGraphPin* Pin : Pins)
+	{
+		if (Pin->PinType.PinCategory == UMaterialGraphSchema::PC_Exec &&
+			Pin->Direction == EGPD_Output &&
+			Pin->SourceIndex == OutputIndex)
+		{
+			return Pin;
+		}
+	}
+	return nullptr;
 }
 
 uint32 UMaterialGraphNode_Base::GetOutputType(const UEdGraphPin* OutputPin) const
 {
-	return GetPinMaterialType(OutputPin, GetPinInfo(OutputPin));
+	return GetPinMaterialType(OutputPin);
 }
 
 uint32 UMaterialGraphNode_Base::GetInputType(const UEdGraphPin* InputPin) const
 {
-	return GetPinMaterialType(InputPin, GetPinInfo(InputPin));
+	return GetPinMaterialType(InputPin);
 }
 
-uint32 UMaterialGraphNode_Base::GetPinMaterialType(const UEdGraphPin* Pin, const FMaterialGraphPinInfo& PinInfo) const
+uint32 UMaterialGraphNode_Base::GetPinMaterialType(const UEdGraphPin* Pin) const
 {
-	switch (PinInfo.PinType)
+	if (Pin->PinType.PinCategory == UMaterialGraphSchema::PC_Exec)
 	{
-	case EMaterialGraphPinType::Data:
-		return MCT_Unknown;
-	case EMaterialGraphPinType::Exec:
 		return MCT_Execution;
-	default:
-		checkNoEntry();
-		return 0u;
 	}
+	return MCT_Unknown;
 }
 
 void UMaterialGraphNode_Base::ReplaceNode(UMaterialGraphNode_Base* OldNode)
@@ -53,50 +97,57 @@ void UMaterialGraphNode_Base::ReplaceNode(UMaterialGraphNode_Base* OldNode)
 	check(OldNode);
 	check(OldNode != this);
 
-	// Copy Inputs from old node
-	for (int32 PinIndex = 0; PinIndex < OldNode->InputPins.Num(); PinIndex++)
+	for (int32 PinIndex = 0; PinIndex < OldNode->Pins.Num(); ++PinIndex)
 	{
-		if (PinIndex < InputPins.Num())
+		UEdGraphPin* OldPin = OldNode->Pins[PinIndex];
+		if (OldPin->Direction == EGPD_Input)
 		{
-			ModifyAndCopyPersistentPinData(*InputPins[PinIndex], *OldNode->InputPins[PinIndex]);
-		}
-	}
+			UEdGraphPin* NewPin = nullptr;
+			if(OldPin->PinType.PinCategory == UMaterialGraphSchema::PC_Exec) NewPin = GetExecInputPin();
+			else NewPin = GetInputPin(OldPin->SourceIndex);
 
-	// Copy Outputs from old node
-	for (int32 PinIndex = 0; PinIndex < OldNode->OutputPins.Num(); PinIndex++)
-	{
-		// Try to find an equivalent output in this node
-		int32 FoundPinIndex = -1;
+			if (NewPin)
+			{
+				ModifyAndCopyPersistentPinData(*NewPin, *OldPin);
+			}
+		}
+		else
 		{
-			// First check names
-			for (int32 NewPinIndex = 0; NewPinIndex < OutputPins.Num(); NewPinIndex++)
+			// Try to find an equivalent output in this node
+			int32 FoundPinIndex = -1;
 			{
-				if (OldNode->OutputPins[PinIndex]->PinName == OutputPins[NewPinIndex]->PinName)
+				// First check names
+				for (int32 NewPinIndex = 0; NewPinIndex < Pins.Num(); NewPinIndex++)
 				{
-					FoundPinIndex = NewPinIndex;
-					break;
+					UEdGraphPin* NewPin = Pins[NewPinIndex];
+					if (NewPin->Direction == EGPD_Output &&
+						NewPin->PinType.PinCategory == OldPin->PinType.PinCategory &&
+						NewPin->PinName == OldPin->PinName)
+					{
+						FoundPinIndex = NewPinIndex;
+						break;
+					}
 				}
 			}
-		}
-		if (FoundPinIndex == -1)
-		{
-			// Now check types
-			for (int32 NewPinIndex = 0; NewPinIndex < OutputPins.Num(); NewPinIndex++)
+			if (FoundPinIndex == -1)
 			{
-				if (OldNode->OutputPins[PinIndex]->PinType == OutputPins[NewPinIndex]->PinType)
+				// Now check types
+				for (int32 NewPinIndex = 0; NewPinIndex < Pins.Num(); NewPinIndex++)
 				{
-					FoundPinIndex = NewPinIndex;
-					break;
+					UEdGraphPin* NewPin = Pins[NewPinIndex];
+					if (NewPin->Direction == EGPD_Output &&
+						NewPin->PinType == OldPin->PinType)
+					{
+						FoundPinIndex = NewPinIndex;
+						break;
+					}
 				}
 			}
-		}
 
-		// If we can't find an equivalent output in this node, just use the first
-		// The user will have to fix up any issues from the mismatch
-		FoundPinIndex = FMath::Max(FoundPinIndex, 0);
-		if (FoundPinIndex < OutputPins.Num())
-		{
-			ModifyAndCopyPersistentPinData(*OutputPins[FoundPinIndex], *OldNode->OutputPins[PinIndex]);
+			// If we can't find an equivalent output in this node, just use the first
+			// The user will have to fix up any issues from the mismatch
+			FoundPinIndex = FMath::Max(FoundPinIndex, 0);
+			ModifyAndCopyPersistentPinData(*Pins[FoundPinIndex], *OldPin);
 		}
 	}
 
@@ -145,18 +196,12 @@ void UMaterialGraphNode_Base::InsertNewNode(UEdGraphPin* FromPin, UEdGraphPin* N
 void UMaterialGraphNode_Base::AllocateDefaultPins()
 {
 	check(Pins.Num() == 0);
-	check(InputPins.Num() == 0);
-	check(OutputPins.Num() == 0);
-	check(PinInfoMap.Num() == 0);
-
 	CreateInputPins();
 	CreateOutputPins();
 }
 
 void UMaterialGraphNode_Base::PostPasteNode()
 {
-	check(PinInfoMap.Num() == 0)
-
 	int32 NumInputDataPins = 0;
 	int32 NumOutputDataPins = 0;
 	int32 NumInputExecPins = 0;
@@ -164,75 +209,23 @@ void UMaterialGraphNode_Base::PostPasteNode()
 	for (UEdGraphPin* Pin : Pins)
 	{
 		int32 Index = INDEX_NONE;
-		EMaterialGraphPinType Type;
-
 		if (Pin->PinType.PinCategory == UMaterialGraphSchema::PC_Exec)
 		{
-			Type = EMaterialGraphPinType::Exec;
 			if (Pin->Direction == EGPD_Input) Index = NumInputExecPins++;
 			else Index = NumOutputExecPins++;
 		}
 		else
 		{
-			Type = EMaterialGraphPinType::Data;
 			if (Pin->Direction == EGPD_Input) Index = NumInputDataPins++;
 			else Index = NumOutputDataPins++;
 		}
-		RegisterPin(Pin, Type, Index);
-	}
-}
-
-void UMaterialGraphNode_Base::RegisterPin(UEdGraphPin* Pin, EMaterialGraphPinType Type, int32 Index)
-{
-	FMaterialGraphPinInfo& PinInfo = PinInfoMap.FindOrAdd(Pin);
-	PinInfo.PinType = Type;
-	PinInfo.Index = Index;
-
-	if (Type == EMaterialGraphPinType::Exec)
-	{
-		switch (Pin->Direction)
-		{
-		case EGPD_Input:
-			checkf(ExecInputPin == nullptr, TEXT("Only 1 exec input pin allowed"));
-			check(Index == 0);
-			ExecInputPin = Pin;
-			break;
-		case EGPD_Output: verify(ExecOutputPins.Add(Pin) == Index); break;
-		default: checkNoEntry(); break;
-		}
-	}
-	else
-	{
-		switch (Pin->Direction)
-		{
-		case EGPD_Input: verify(InputPins.Add(Pin) == Index); break;
-		case EGPD_Output: verify(OutputPins.Add(Pin) == Index); break;
-		default: checkNoEntry(); break;
-		}
+		Pin->SourceIndex = Index;
 	}
 }
 
 void UMaterialGraphNode_Base::EmptyPins()
 {
 	Pins.Reset();
-	PinInfoMap.Reset();
-	InputPins.Reset();
-	OutputPins.Reset();
-	ExecOutputPins.Reset();
-	ExecInputPin = nullptr;
-}
-
-static void TransferPinArray(const TArray<UEdGraphPin*>& NewPins, const TArray<UEdGraphPin*>& OldPins)
-{
-	const int32 Num = FMath::Min(NewPins.Num(), OldPins.Num());
-	for (int32 i = 0; i < Num; ++i)
-	{
-		UEdGraphPin* OldPin = OldPins[i];
-		UEdGraphPin* NewPin = NewPins[i];
-		ensure(OldPin->Direction == NewPin->Direction);
-		ensure(OldPin->PinType.PinCategory == NewPin->PinType.PinCategory);
-		NewPin->MovePersistentDataFromOldPin(*OldPin);
-	}
 }
 
 void UMaterialGraphNode_Base::ReconstructNode()
@@ -257,23 +250,30 @@ void UMaterialGraphNode_Base::ReconstructNode()
 
 	// Move the existing pins to a saved array
 	TArray<UEdGraphPin*> OldPins = MoveTemp(Pins);
-	TArray<UEdGraphPin*> OldInputPins = MoveTemp(InputPins);
-	TArray<UEdGraphPin*> OldOutputPins = MoveTemp(OutputPins);
-	TArray<UEdGraphPin*> OldExecOutputPins = MoveTemp(ExecOutputPins);
-	UEdGraphPin* OldExecInputPin = ExecInputPin;
 
 	EmptyPins();
 
 	// Recreate the new pins
 	AllocateDefaultPins();
 
-	// Transfer data to new pins
-	TransferPinArray(InputPins, OldInputPins);
-	TransferPinArray(OutputPins, OldOutputPins);
-	TransferPinArray(ExecOutputPins, OldExecOutputPins);
-	if (OldExecInputPin && ExecInputPin)
+	for (int32 OldPinIndex = 0; OldPinIndex < OldPins.Num(); ++OldPinIndex)
 	{
-		ExecInputPin->MovePersistentDataFromOldPin(*OldExecInputPin);
+		UEdGraphPin* OldPin = OldPins[OldPinIndex];
+		UEdGraphPin* NewPin = nullptr;
+		if (OldPin->PinType.PinCategory == UMaterialGraphSchema::PC_Exec)
+		{
+			if (OldPin->Direction == EGPD_Input) NewPin = GetExecInputPin();
+			else NewPin = GetExecOutputPin(OldPin->SourceIndex);
+		}
+		else
+		{
+			if (OldPin->Direction == EGPD_Input) NewPin = GetInputPin(OldPin->SourceIndex);
+			else NewPin = GetOutputPin(OldPin->SourceIndex);
+		}
+		if (NewPin)
+		{
+			NewPin->MovePersistentDataFromOldPin(*OldPin);
+		}
 	}
 
 	for (UEdGraphPin* OldPin : OldPins)
@@ -291,39 +291,14 @@ void UMaterialGraphNode_Base::RemovePinAt(const int32 PinIndex, const EEdGraphPi
 	UEdGraphPin* Pin = GetPinWithDirectionAt(PinIndex, PinDirection);
 	check(Pin);
 
-	FMaterialGraphPinInfo PinInfo;
-	verify(PinInfoMap.RemoveAndCopyValue(Pin, PinInfo));
-	switch (PinInfo.PinType)
-	{
-	case EMaterialGraphPinType::Data:
-		switch (Pin->Direction)
-		{
-		case EGPD_Input: InputPins.RemoveAt(PinInfo.Index); break;
-		case EGPD_Output: OutputPins.RemoveAt(PinInfo.Index); break;
-		default: checkNoEntry(); break;
-		}
-		break;
-	case EMaterialGraphPinType::Exec:
-		switch (Pin->Direction)
-		{
-		case EGPD_Input: check(ExecInputPin == Pin);  ExecInputPin = nullptr; break;
-		case EGPD_Output: ExecOutputPins.RemoveAt(PinInfo.Index); break;
-		default: checkNoEntry(); break;
-		}
-		break;
-	default:
-		checkNoEntry();
-		break;
-	}
-
 	// Shift down indices to account for the pin we removed
-	for (auto& It : PinInfoMap)
+	for (UEdGraphPin* CheckPin : Pins)
 	{
-		if (It.Value.PinType == PinInfo.PinType &&
-			It.Value.Index > PinInfo.Index &&
-			It.Key->Direction == PinDirection)
+		if (CheckPin->PinType.PinCategory == Pin->PinType.PinCategory &&
+			CheckPin->Direction == Pin->Direction &&
+			CheckPin->SourceIndex > Pin->SourceIndex)
 		{
-			It.Value.Index--;
+			CheckPin->SourceIndex--;
 		}
 	}
 
diff --git a/Engine/Source/Editor/UnrealEd/Private/MaterialGraphNode_Knot.cpp b/Engine/Source/Editor/UnrealEd/Private/MaterialGraphNode_Knot.cpp
--- a/Engine/Source/Editor/UnrealEd/Private/MaterialGraphNode_Knot.cpp
+++ b/Engine/Source/Editor/UnrealEd/Private/MaterialGraphNode_Knot.cpp
@@ -26,11 +26,11 @@ void UMaterialGraphNode_Knot::AllocateDefaultPins()
 	const FName OutputPinName(TEXT("OutputPin"));
 
 	UEdGraphPin* MyInputPin = CreatePin(EGPD_Input, PC_Wildcard, InputPinName);
+	MyInputPin->SourceIndex = 0;
 	MyInputPin->bDefaultValueIsIgnored = true;
-	RegisterPin(MyInputPin, EMaterialGraphPinType::Data, 0);
 
 	UEdGraphPin* MyOutputPin = CreatePin(EGPD_Output, PC_Wildcard, OutputPinName);
-	RegisterPin(MyOutputPin, EMaterialGraphPinType::Data, 0);
+	MyOutputPin->SourceIndex = 0;
 }
 
 FText UMaterialGraphNode_Knot::GetTooltipText() const
diff --git a/Engine/Source/Editor/UnrealEd/Private/MaterialGraphNode_Root.cpp b/Engine/Source/Editor/UnrealEd/Private/MaterialGraphNode_Root.cpp
--- a/Engine/Source/Editor/UnrealEd/Private/MaterialGraphNode_Root.cpp
+++ b/Engine/Source/Editor/UnrealEd/Private/MaterialGraphNode_Root.cpp
@@ -45,30 +45,31 @@ void UMaterialGraphNode_Root::PostPlacedNewNode()
 	}
 }
 
-uint32 UMaterialGraphNode_Root::GetPinMaterialType(const UEdGraphPin* Pin, const FMaterialGraphPinInfo& PinInfo) const
+uint32 UMaterialGraphNode_Root::GetPinMaterialType(const UEdGraphPin* Pin) const
 {
-	if (PinInfo.PinType == EMaterialGraphPinType::Data)
+	if (Pin->PinType.PinCategory == UMaterialGraphSchema::PC_Exec)
 	{
-		const UMaterialGraph* MaterialGraph = CastChecked<UMaterialGraph>(GetGraph());
-		const FMaterialInputInfo& MaterialInput = MaterialGraph->MaterialInputs[PinInfo.Index];
-		EMaterialProperty Property = MaterialInput.GetProperty();
+		return MCT_Execution;
+	}
 
-		uint32 MaterialType = 0u;
-		if (Property == MP_MaterialAttributes)
-		{
-			MaterialType = MCT_MaterialAttributes;
-		}
-		else if (Property == MP_FrontMaterial)
-		{
-			MaterialType = MCT_Strata;
-		}
-		else
-		{
-			MaterialType = FMaterialAttributeDefinitionMap::GetValueType(Property);
-		}
-		return MaterialType;
+	const UMaterialGraph* MaterialGraph = CastChecked<UMaterialGraph>(GetGraph());
+	const FMaterialInputInfo& MaterialInput = MaterialGraph->MaterialInputs[Pin->SourceIndex];
+	EMaterialProperty Property = MaterialInput.GetProperty();
+
+	uint32 MaterialType = 0u;
+	if (Property == MP_MaterialAttributes)
+	{
+		MaterialType = MCT_MaterialAttributes;
+	}
+	else if (Property == MP_FrontMaterial)
+	{
+		MaterialType = MCT_Strata;
+	}
+	else
+	{
+		MaterialType = FMaterialAttributeDefinitionMap::GetValueType(Property);
 	}
-	return MCT_Execution;
+	return MaterialType;
 }
 
 void UMaterialGraphNode_Root::CreateInputPins()
@@ -79,10 +80,10 @@ void UMaterialGraphNode_Root::CreateInputPins()
 	{
 		// Create the execution pin
 		UEdGraphPin* NewPin = CreatePin(EGPD_Input, UMaterialGraphSchema::PC_Exec, NAME_None, NAME_None);
+		NewPin->SourceIndex = 0;
 		// Makes sure pin has a name for lookup purposes but user will never see it
 		NewPin->PinName = CreateUniquePinName(TEXT("Input"));
 		NewPin->PinFriendlyName = LOCTEXT("Space", " ");
-		RegisterPin(NewPin, EMaterialGraphPinType::Exec, 0);
 	}
 
 	for (int32 Index = 0; Index < MaterialGraph->MaterialInputs.Num(); ++Index)
@@ -91,7 +92,7 @@ void UMaterialGraphNode_Root::CreateInputPins()
 		EMaterialProperty Property = MaterialInput.GetProperty();
 	
 		UEdGraphPin* InputPin = CreatePin(EGPD_Input, UMaterialGraphSchema::PC_MaterialInput, *FString::Printf(TEXT("%d"), (int32)Property), *MaterialInput.GetName().ToString());
-		RegisterPin(InputPin, EMaterialGraphPinType::Data, Index);
+		InputPin->SourceIndex = Index;
 	}
 
 }
diff --git a/Engine/Source/Editor/UnrealEd/Public/MaterialGraphNode_Knot.h b/Engine/Source/Editor/UnrealEd/Public/MaterialGraphNode_Knot.h
--- a/Engine/Source/Editor/UnrealEd/Public/MaterialGraphNode_Knot.h
+++ b/Engine/Source/Editor/UnrealEd/Public/MaterialGraphNode_Knot.h
@@ -26,11 +26,11 @@ public:
 	
 	UEdGraphPin* GetInputPin() const
 	{
-		return InputPins[0];
+		return UMaterialGraphNode::GetInputPin(0);
 	}
 
 	UEdGraphPin* GetOutputPin() const
 	{
-		return OutputPins[0];
+		return UMaterialGraphNode::GetOutputPin(0);
 	}
 };
diff --git a/Engine/Source/Runtime/Core/Public/UObject/UE5MainStreamObjectVersion.h b/Engine/Source/Runtime/Core/Public/UObject/UE5MainStreamObjectVersion.h
--- a/Engine/Source/Runtime/Core/Public/UObject/UE5MainStreamObjectVersion.h
+++ b/Engine/Source/Runtime/Core/Public/UObject/UE5MainStreamObjectVersion.h
@@ -160,6 +160,9 @@ struct CORE_API FUE5MainStreamObjectVersion
 		// No longer bake MassToLocal transform into recorded transform data in GeometryCollection caching
 		GeometryCollectionCacheRemovesMassToLocal,
 
+		// UEdGraphPin serializes SourceIndex
+		EdGraphPinSourceIndex,
+
 		// -----<new versions can be added above this line>-------------------------------------------------
 		VersionPlusOne,
 		LatestVersion = VersionPlusOne - 1
diff --git a/Engine/Source/Runtime/Engine/Classes/EdGraph/EdGraphPin.h b/Engine/Source/Runtime/Engine/Classes/EdGraph/EdGraphPin.h
--- a/Engine/Source/Runtime/Engine/Classes/EdGraph/EdGraphPin.h
+++ b/Engine/Source/Runtime/Engine/Classes/EdGraph/EdGraphPin.h
@@ -302,6 +302,9 @@ public:
 	/** Name of this pin. */
 	FName PinName;
 
+	/** Index of the pin in the source data structure represented by the graph, INDEX_NONE if not used */
+	int32 SourceIndex;
+
 	/** Direction of flow of this pin (input or output) */
 	TEnumAsByte<enum EEdGraphPinDirection> Direction;
 
@@ -496,6 +499,7 @@ public:
 		PinType.ResetToDefaults();
 
 		PinName = NAME_None;
+		SourceIndex = INDEX_NONE;
 #if WITH_EDITORONLY_DATA
 		PinFriendlyName = FText::GetEmpty();
 #endif // WITH_EDITORONLY_DATA
diff --git a/Engine/Source/Runtime/Engine/Private/EdGraph/EdGraphPin.cpp b/Engine/Source/Runtime/Engine/Private/EdGraph/EdGraphPin.cpp
--- a/Engine/Source/Runtime/Engine/Private/EdGraph/EdGraphPin.cpp
+++ b/Engine/Source/Runtime/Engine/Private/EdGraph/EdGraphPin.cpp
@@ -5,6 +5,7 @@
 #include "UObject/BlueprintsObjectVersion.h"
 #include "UObject/FrameworkObjectVersion.h"
 #include "UObject/ReleaseObjectVersion.h"
+#include "UObject/UE5MainStreamObjectVersion.h"
 #include "UObject/UnrealType.h"
 #include "UObject/TextProperty.h"
 #include "EdGraph/EdGraph.h"
@@ -1361,6 +1362,7 @@ UEdGraphPin::UEdGraphPin(UEdGraphNode* InOwningNode, const FGuid& PinIdGuid)
 	: OwningNode(InOwningNode)
 	, PinId(PinIdGuid)
 	, PinName()
+	, SourceIndex(INDEX_NONE)
 	, Direction(EGPD_Input)
 #if WITH_EDITORONLY_DATA
 	, bHidden(false)
@@ -1617,6 +1619,7 @@ void UEdGraphPin::DestroyImpl(bool bClearLinks)
 bool UEdGraphPin::Serialize(FArchive& Ar)
 {
 	Ar.UsingCustomVersion(FFrameworkObjectVersion::GUID);
+	Ar.UsingCustomVersion(FUE5MainStreamObjectVersion::GUID);
 
 	// These properties are in every pin and are unlikely to be removed, so they are native serialized for speed.
 	Ar << OwningNode;
@@ -1640,6 +1643,11 @@ bool UEdGraphPin::Serialize(FArchive& Ar)
 	}
 #endif
 
+	if (Ar.CustomVer(FUE5MainStreamObjectVersion::GUID) >= FUE5MainStreamObjectVersion::EdGraphPinSourceIndex)
+	{
+		Ar << SourceIndex;
+	}
+
 	Ar << PinToolTip;
 	Ar << Direction;
 	PinType.Serialize(Ar);
diff --git a/Engine/Source/Runtime/Engine/Private/Materials/MaterialExpressions.cpp b/Engine/Source/Runtime/Engine/Private/Materials/MaterialExpressions.cpp
--- a/Engine/Source/Runtime/Engine/Private/Materials/MaterialExpressions.cpp
+++ b/Engine/Source/Runtime/Engine/Private/Materials/MaterialExpressions.cpp
@@ -1490,11 +1490,13 @@ bool UMaterialExpression::HasConnectedOutputs() const
 		UMaterialGraphNode* MatGraphNode = Cast<UMaterialGraphNode>(GraphNode);
 		if (MatGraphNode)
 		{
-			for (UEdGraphPin* Pin : MatGraphNode->OutputPins)
+			for (UEdGraphPin* Pin : MatGraphNode->Pins)
 			{
-				if (Pin->LinkedTo.Num() > 0)
+				if (Pin->Direction == EGPD_Output &&
+					Pin->LinkedTo.Num() > 0)
 				{
 					bIsConnected = true;
+					break;
 				}
 			}
 		}
