Fix crash in mode manager due to dereferencing null pointer.

diff --git a/Engine/Source/Editor/UnrealEd/Private/EditorModeManager.cpp b/Engine/Source/Editor/UnrealEd/Private/EditorModeManager.cpp
--- a/Engine/Source/Editor/UnrealEd/Private/EditorModeManager.cpp
+++ b/Engine/Source/Editor/UnrealEd/Private/EditorModeManager.cpp
@@ -180,7 +180,11 @@ USelection* FEditorModeTools::GetSelectedComponents() const
 
 UTypedElementSelectionSet* FEditorModeTools::GetEditorSelectionSet() const
 {
-	return GetSelectedActors()->GetElementSelectionSet();
+	if (USelection* SelectedActorsSet = GetSelectedActors())
+	{
+		return SelectedActorsSet->GetElementSelectionSet();
+	}
+	return nullptr;
 }
 
 void FEditorModeTools::StoreSelection(FName SelectionStoreKey, bool bClearSelection)
