remove debug pragma_disable_optimization

diff --git a/Engine/Source/Editor/UnrealEd/Private/LevelEditorViewport.cpp b/Engine/Source/Editor/UnrealEd/Private/LevelEditorViewport.cpp
--- a/Engine/Source/Editor/UnrealEd/Private/LevelEditorViewport.cpp
+++ b/Engine/Source/Editor/UnrealEd/Private/LevelEditorViewport.cpp
@@ -99,8 +99,6 @@
 DEFINE_LOG_CATEGORY(LogEditorViewport);
 
 #define LOCTEXT_NAMESPACE "LevelEditorViewportClient"
-PRAGMA_DISABLE_OPTIMIZATION
-
 const FLevelViewportActorLock FLevelViewportActorLock::None(nullptr);
 
 TArray< TWeakObjectPtr< AActor > > FLevelEditorViewportClient::DropPreviewActors;
@@ -4983,5 +4981,4 @@ bool FLevelEditorViewportClient::GetPivotForOrbit(FVector& Pivot) const
 	return false;
 }
 
-#undef LOCTEXT_NAMESPACE
-PRAGMA_ENABLE_OPTIMIZATION
\ No newline at end of file
+#undef LOCTEXT_NAMESPACE
\ No newline at end of file
