Some small improvements to SVirtualizationStatusIndicator

#rnx

* SVirtualizationStatusIndicator
- Added the total time taken for each activity to the widget's tooltip and reformatted the output to be clearer.
- The initial tool tip should just give a brief highlevel overview of the activity taking place so groups all backends together across all threads.

* VirtualizationManager
- Profiling data for Mirage is now stored as a struct per activity type, making it easier to edit in the future.
- We now store the profiling data in the base format and let the different display methods format it instead.

diff --git a/Engine/Source/Editor/UnrealEd/Private/Virtualization/SVirtualizationStatusIndicator.cpp b/Engine/Source/Editor/UnrealEd/Private/Virtualization/SVirtualizationStatusIndicator.cpp
--- a/Engine/Source/Editor/UnrealEd/Private/Virtualization/SVirtualizationStatusIndicator.cpp
+++ b/Engine/Source/Editor/UnrealEd/Private/Virtualization/SVirtualizationStatusIndicator.cpp
@@ -16,174 +16,177 @@
 namespace UE::Virtualization
 {
 
-/** The amount of time (s) that the indication arrows will take to fade in/out */
-static constexpr float ArrowFadeTime = 0.5f;
-
-/** 
-  * The tool tip that will show when the mouse is hovered over SVirtualizationStatusIndicator. 
-  * 
-  * Currently this will display the overall amount of data that content virtualization has 
-  * transfered to and from the backends.
-  */
-class SVirtualizationToolTip : public SCompoundWidget
-{
-	SLATE_BEGIN_ARGS(SVirtualizationToolTip) {}
-	SLATE_END_ARGS()
-
-	void Construct(const FArguments& InArgs);
-
-private:
-	FText GetToolTipText() const;
-};
-
-/** 
- * A widget used to show the status of content virtualization.
- * 
- * A downwards pointing arrow will be displayed when ever a payload is pulled from a backend.
- * An upwards pointing arrow will be displayed when ever a payload is pushed to a backend.
- * This gives the user immediately visual feedback as to when the system is in action.
- * Additional data can be found when hovering the mouse of the widget.
- */
-class SVirtualizationStatusIndicator : public SCompoundWidget
-{
-	SLATE_BEGIN_ARGS(SVirtualizationStatusIndicator) {}
-	SLATE_END_ARGS()
+	/** The amount of time (s) that the indication arrows will take to fade in/out */
+	static constexpr float ArrowFadeTime = 0.5f;
+
+	/** 
+	* The tool tip that will show when the mouse is hovered over SVirtualizationStatusIndicator. 
+	* 
+	* Currently this will display the overall amount of data that content virtualization has 
+	* transfered to and from the backends.
+	*/
+	class SVirtualizationToolTip : public SCompoundWidget
+	{
+		SLATE_BEGIN_ARGS(SVirtualizationToolTip) {}
+		SLATE_END_ARGS()
+
+			void Construct(const FArguments& InArgs);
+
+	private:
+		FText GetToolTipText() const;
+	};
+
+	/** 
+	* A widget used to show the status of content virtualization.
+	* 
+	* A downwards pointing arrow will be displayed when ever a payload is pulled from a backend.
+	* An upwards pointing arrow will be displayed when ever a payload is pushed to a backend.
+	* This gives the user immediately visual feedback as to when the system is in action.
+	* Additional data can be found when hovering the mouse of the widget.
+	*/
+	class SVirtualizationStatusIndicator : public SCompoundWidget
+	{
+		SLATE_BEGIN_ARGS(SVirtualizationStatusIndicator) {}
+		SLATE_END_ARGS()
 
-	void Construct(const FArguments& InArgs);
+			void Construct(const FArguments& InArgs);
 
-private:
+	private:
 
-	EActiveTimerReturnType UpdateLastDataAccess(double InCurrentTime, float InDeltaTime);
+		EActiveTimerReturnType UpdateLastDataAccess(double InCurrentTime, float InDeltaTime);
 
-	int64 PayloadsPulled = 0;
-	int64 PayloadsPushed = 0;
+		int64 PayloadsPulled = 0;
+		int64 PayloadsPushed = 0;
 
-	FCurveSequence FadePullArrow;
-	FCurveSequence FadePushArrow;
-};
+		FCurveSequence FadePullArrow;
+		FCurveSequence FadePushArrow;
+	};
 
-void SVirtualizationToolTip::Construct(const FArguments& InArgs)
-{
-	this->ChildSlot
-	[
-		SNew(SVerticalBox)
-
-		+ SVerticalBox::Slot()
-		.AutoHeight()
-		[
-			SNew(STextBlock)
-			.Text(this, &SVirtualizationToolTip::GetToolTipText)
-		]
-	];
-}
-
-FText SVirtualizationToolTip::GetToolTipText() const
-{
-	using namespace UE::Virtualization;
-	
-	FPayloadActivityInfo Info = FVirtualizationManager::Get().GetPayloadActivityInfo();
-	TStringBuilder<1024> Builder;
-	
-	Builder << TEXT("Total Payload Data Pushed: ") << Info.TotalSizePushed << TEXT(" MB\n");
-	Builder << TEXT("Total Payload Data Pulled: ") << Info.TotalSizePulled << TEXT(" MB");
-	
-	return FText::FromString(Builder.ToString());
-}
-
-void SVirtualizationStatusIndicator::Construct(const FArguments& InArgs)
-{
-	using namespace UE::Virtualization;
+	void SVirtualizationToolTip::Construct(const FArguments& InArgs)
+	{
+		this->ChildSlot
+			[
+				SNew(SVerticalBox)
+
+				+ SVerticalBox::Slot()
+			.AutoHeight()
+			[
+				SNew(STextBlock)
+				.Text(this, &SVirtualizationToolTip::GetToolTipText)
+			]
+			];
+	}
 
-	FPayloadActivityInfo Info = FVirtualizationManager::Get().GetPayloadActivityInfo();
-	PayloadsPulled = Info.PayloadsPulled;
-	PayloadsPushed = Info.PayloadsPushed;
-	
-	FadePushArrow = FCurveSequence(0.0f, ArrowFadeTime, ECurveEaseFunction::Linear);
-	FadePullArrow = FCurveSequence(0.0f, ArrowFadeTime, ECurveEaseFunction::Linear);
+	FText SVirtualizationToolTip::GetToolTipText() const
+	{
+		using namespace UE::Virtualization;
 
-	const FText MirageLabel = FText::FromString(TEXT("Mirage"));
+		FPayloadActivityInfo Info = FVirtualizationManager::Get().GetPayloadActivityInfo();
+		TStringBuilder<1024> Builder;
 
-	this->ChildSlot
-	[
-		SNew(SHorizontalBox)
-		.ToolTip(SNew(SToolTip)[SNew(SVirtualizationToolTip)])
+		Builder << TEXT("Pushed: ") << Info.Push.TotalBytes / (int64)(1024 * 1024) << TEXT(" MB in ");
+		Builder.Appendf(TEXT("%.3f Seconds\n"), Info.Push.CyclesSpent * FPlatformTime::GetSecondsPerCycle());
 
-		+ SHorizontalBox::Slot()
-		.AutoWidth()
-		.VAlign(VAlign_Center)
-		.Padding(0, 0, 3, 0)
-		[
-			SNew(SOverlay)
+		Builder << TEXT("Pulled: ") << Info.Pull.TotalBytes / (int64)(1024 * 1024) << TEXT(" MB in ");
+		Builder.Appendf(TEXT("%.3f Seconds"), Info.Pull.CyclesSpent  * FPlatformTime::GetSecondsPerCycle());
+
+		return FText::FromString(Builder.ToString());
+	}
+
+	void SVirtualizationStatusIndicator::Construct(const FArguments& InArgs)
+	{
+		using namespace UE::Virtualization;
+
+		FPayloadActivityInfo Info = FVirtualizationManager::Get().GetPayloadActivityInfo();
+		PayloadsPulled = Info.Pull.PayloadCount;
+		PayloadsPushed = Info.Push.PayloadCount;
 
-			+ SOverlay::Slot()
+		FadePushArrow = FCurveSequence(0.0f, ArrowFadeTime, ECurveEaseFunction::Linear);
+		FadePullArrow = FCurveSequence(0.0f, ArrowFadeTime, ECurveEaseFunction::Linear);
+
+		const FText MirageLabel = FText::FromString(TEXT("Mirage"));
+
+		this->ChildSlot
+			[
+				SNew(SHorizontalBox)
+				.ToolTip(SNew(SToolTip)[SNew(SVirtualizationToolTip)])
+
+			+ SHorizontalBox::Slot()
+			.AutoWidth()
+			.VAlign(VAlign_Center)
+			.Padding(0, 0, 3, 0)
+			[
+				SNew(SOverlay)
+
+				+ SOverlay::Slot()
 			.HAlign(HAlign_Center)
 			.VAlign(VAlign_Top)
 			.Padding(0, 0, 4, 4)
 			[
 				SNew(SImage)
 				.Image(FAppStyle::Get().GetBrush("Icons.ArrowUp"))
-				.ColorAndOpacity_Lambda([this] { return FLinearColor::Green.CopyWithNewOpacity(FadePushArrow.GetLerp()); })
+			.ColorAndOpacity_Lambda([this] { return FLinearColor::Green.CopyWithNewOpacity(FadePushArrow.GetLerp()); })
 			]
 
-			+ SOverlay::Slot()
+		+ SOverlay::Slot()
 			.HAlign(HAlign_Center)
 			.VAlign(VAlign_Bottom)
 			.Padding(4, 4, 0, 0)
 			[
 				SNew(SImage)
 				.Image(FAppStyle::Get().GetBrush("Icons.ArrowDown"))
-				.ColorAndOpacity_Lambda([this] { return FLinearColor::Green.CopyWithNewOpacity(FadePullArrow.GetLerp()); })
+			.ColorAndOpacity_Lambda([this] { return FLinearColor::Green.CopyWithNewOpacity(FadePullArrow.GetLerp()); })
+			]
 			]
-		]
 
 		+ SHorizontalBox::Slot()
-		.AutoWidth()
-		.VAlign(VAlign_Center)
-		.Padding(0, 0, 10, 0)
-		[
-			SNew(STextBlock)
-			.Text(MirageLabel)
-		]
-	];
-
-	RegisterActiveTimer(ArrowFadeTime, FWidgetActiveTimerDelegate::CreateSP(this, &SVirtualizationStatusIndicator::UpdateLastDataAccess));
-}
-
-EActiveTimerReturnType SVirtualizationStatusIndicator::UpdateLastDataAccess(double InCurrentTime, float InDeltaTime)
-{
-	using namespace UE::Virtualization;
-
-	FPayloadActivityInfo Info = FVirtualizationManager::Get().GetPayloadActivityInfo();
-	
-	// Don't change the fade if it is currently playing, the tick duration of this method is not that accurate
-	// and we don't want to cut out the previous fade too soon or the user might not be able to notice it.
+			.AutoWidth()
+			.VAlign(VAlign_Center)
+			.Padding(0, 0, 10, 0)
+			[
+				SNew(STextBlock)
+				.Text(MirageLabel)
+			]
+			];
 
-	if (!FadePullArrow.IsPlaying())
-	{
-		FadePullArrow.PlayRelative(this->AsShared(), PayloadsPulled != Info.PayloadsPulled);
-		PayloadsPulled = Info.PayloadsPulled;
+		RegisterActiveTimer(ArrowFadeTime, FWidgetActiveTimerDelegate::CreateSP(this, &SVirtualizationStatusIndicator::UpdateLastDataAccess));
 	}
 
-	if (!FadePushArrow.IsPlaying())
+	EActiveTimerReturnType SVirtualizationStatusIndicator::UpdateLastDataAccess(double InCurrentTime, float InDeltaTime)
 	{
-		FadePushArrow.PlayRelative(this->AsShared(), PayloadsPushed != Info.PayloadsPushed);
-		PayloadsPushed = Info.PayloadsPushed;
-	}
+		using namespace UE::Virtualization;
 
-	return EActiveTimerReturnType::Continue;
-}
+		FPayloadActivityInfo Info = FVirtualizationManager::Get().GetPayloadActivityInfo();
 
-TSharedPtr<SWidget> GetVirtualizationStatusIndicator()
-{
-	if (FVirtualizationManager::Get().IsEnabled())
-	{
-		return SNew(SVirtualizationStatusIndicator);
+		// Don't change the fade if it is currently playing, the tick duration of this method is not that accurate
+		// and we don't want to cut out the previous fade too soon or the user might not be able to notice it.
+
+		if (!FadePullArrow.IsPlaying())
+		{
+			FadePullArrow.PlayRelative(this->AsShared(), PayloadsPulled != Info.Pull.PayloadCount);
+			PayloadsPulled = Info.Pull.PayloadCount;
+		}
+
+		if (!FadePushArrow.IsPlaying())
+		{
+			FadePushArrow.PlayRelative(this->AsShared(), PayloadsPushed != Info.Push.PayloadCount);
+			PayloadsPushed = Info.Push.PayloadCount;
+		}
+
+		return EActiveTimerReturnType::Continue;
 	}
-	else
+
+	TSharedPtr<SWidget> GetVirtualizationStatusIndicator()
 	{
-		return nullptr;
+		if (FVirtualizationManager::Get().IsEnabled())
+		{
+			return SNew(SVirtualizationStatusIndicator);
+		}
+		else
+		{
+			return nullptr;
+		}
 	}
-}
 
 } // namespace UE::Virtualization
 
diff --git a/Engine/Source/Runtime/CoreUObject/Private/Virtualization/VirtualizationManager.cpp b/Engine/Source/Runtime/CoreUObject/Private/Virtualization/VirtualizationManager.cpp
--- a/Engine/Source/Runtime/CoreUObject/Private/Virtualization/VirtualizationManager.cpp
+++ b/Engine/Source/Runtime/CoreUObject/Private/Virtualization/VirtualizationManager.cpp
@@ -355,14 +355,16 @@ FPayloadActivityInfo FVirtualizationManager::GetPayloadActivityInfo() const
 #if ENABLE_COOK_STATS
 	for (const auto& Iterator : Profiling::PushStats)
 	{
-		Info.PayloadsPushed += Iterator.Value.GetAccumulatedValueAnyThread(FCookStats::CallStats::EHitOrMiss::Hit, FCookStats::CallStats::EStatType::Counter);
-		Info.TotalSizePushed += Iterator.Value.GetAccumulatedValueAnyThread(FCookStats::CallStats::EHitOrMiss::Hit, FCookStats::CallStats::EStatType::Bytes) / (1024 * 1024);
+		Info.Push.PayloadCount	+= Iterator.Value.GetAccumulatedValueAnyThread(FCookStats::CallStats::EHitOrMiss::Hit, FCookStats::CallStats::EStatType::Counter);
+		Info.Push.TotalBytes	+= Iterator.Value.GetAccumulatedValueAnyThread(FCookStats::CallStats::EHitOrMiss::Hit, FCookStats::CallStats::EStatType::Bytes);
+		Info.Push.CyclesSpent	+= Iterator.Value.GetAccumulatedValueAnyThread(FCookStats::CallStats::EHitOrMiss::Hit, FCookStats::CallStats::EStatType::Cycles);		
 	}
 
 	for (const auto& Iterator : Profiling::PullStats)
 	{
-		Info.PayloadsPulled += Iterator.Value.GetAccumulatedValueAnyThread(FCookStats::CallStats::EHitOrMiss::Hit, FCookStats::CallStats::EStatType::Counter);
-		Info.TotalSizePulled += Iterator.Value.GetAccumulatedValueAnyThread(FCookStats::CallStats::EHitOrMiss::Hit, FCookStats::CallStats::EStatType::Bytes) / (1024 * 1024);
+		Info.Pull.PayloadCount += Iterator.Value.GetAccumulatedValueAnyThread(FCookStats::CallStats::EHitOrMiss::Hit, FCookStats::CallStats::EStatType::Counter);
+		Info.Pull.TotalBytes += Iterator.Value.GetAccumulatedValueAnyThread(FCookStats::CallStats::EHitOrMiss::Hit, FCookStats::CallStats::EStatType::Bytes);
+		Info.Pull.CyclesSpent += Iterator.Value.GetAccumulatedValueAnyThread(FCookStats::CallStats::EHitOrMiss::Hit, FCookStats::CallStats::EStatType::Cycles);
 	}
 #endif // ENABLE_COOK_STATS
 
diff --git a/Engine/Source/Runtime/CoreUObject/Public/Virtualization/VirtualizationManager.h b/Engine/Source/Runtime/CoreUObject/Public/Virtualization/VirtualizationManager.h
--- a/Engine/Source/Runtime/CoreUObject/Public/Virtualization/VirtualizationManager.h
+++ b/Engine/Source/Runtime/CoreUObject/Public/Virtualization/VirtualizationManager.h
@@ -52,11 +52,18 @@ class IVirtualizationBackendFactory;
 /** Profiling data containing all activity relating to payloads. */
 struct FPayloadActivityInfo
 {
-	int64 PayloadsPulled = 0;
-	int64 PayloadsPushed = 0;
-
-	int64 TotalSizePulled = 0;
-	int64 TotalSizePushed = 0;
+	struct FActivity
+	{
+		/** The number of payloads that have been involved by the activity. */
+		int64 PayloadCount = 0;
+		/** The total size of all payloads involved in the activity, in bytes. */
+		int64 TotalBytes = 0;
+		/** The total number of cycles spent on the activity across all threads. */
+		int64 CyclesSpent = 0;
+	};
+
+	FActivity Pull;
+	FActivity Push;
 };
 
 /** This is used as a wrapper around the various potential back end implementations. 
