BulkDataRegistry bugfix: Fix FRequest null pointer access. Remove unused variable.
#rnx

diff --git a/Engine/Source/Editor/UnrealEd/Private/EditorDomain/BulkDataRegistryEditorDomain.cpp b/Engine/Source/Editor/UnrealEd/Private/EditorDomain/BulkDataRegistryEditorDomain.cpp
--- a/Engine/Source/Editor/UnrealEd/Private/EditorDomain/BulkDataRegistryEditorDomain.cpp
+++ b/Engine/Source/Editor/UnrealEd/Private/EditorDomain/BulkDataRegistryEditorDomain.cpp
@@ -506,7 +506,7 @@ void FPendingPackage::Cancel()
 	// Called from outside Owner->PendingPackagesLock, so OnBulkDataList can complete on other thread while we wait
 	// Called after removing this from Owner->PendingPackages under a previous cover of the lock
 	// If OnBulkDataList is running on other thread its attempt to remove from PendingPackages will be a noop
-	if (!BulkDataListCacheRequest->Poll())
+	if (!BulkDataListCacheRequest.Poll())
 	{
 		// Optimization: prevent WriteCache from running at all if we reach here first
 		PendingOperations.fetch_or(Flag_Canceled);
diff --git a/Engine/Source/Editor/UnrealEd/Private/EditorDomain/BulkDataRegistryEditorDomain.h b/Engine/Source/Editor/UnrealEd/Private/EditorDomain/BulkDataRegistryEditorDomain.h
--- a/Engine/Source/Editor/UnrealEd/Private/EditorDomain/BulkDataRegistryEditorDomain.h
+++ b/Engine/Source/Editor/UnrealEd/Private/EditorDomain/BulkDataRegistryEditorDomain.h
@@ -74,7 +74,6 @@ private:
 	FName PackageName;
 	TArray<UE::Virtualization::FVirtualizedUntypedBulkData> BulkDatas;
 	TArray<UE::Virtualization::FVirtualizedUntypedBulkData> CachedBulkDatas;
-	TArray<UE::DerivedData::FRequest> PayloadIdCacheRequests;
 	UE::DerivedData::FRequest BulkDataListCacheRequest;
 	FBulkDataRegistryEditorDomain* Owner;
 	/**
