Feature: Move breakpoint storage out of the uasset file
Description: Make it so breakpoints are saved in the EditorPerProjectUserSettings ini rather than blueprint uasset files so that they can be excluded from version control and be unique to the user

diff --git a/Engine/Plugins/Tests/EditorTests/Source/EditorTests/Private/UnrealEd/EditorBuildPromotionTests.cpp b/Engine/Plugins/Tests/EditorTests/Source/EditorTests/Private/UnrealEd/EditorBuildPromotionTests.cpp
--- a/Engine/Plugins/Tests/EditorTests/Source/EditorTests/Private/UnrealEd/EditorBuildPromotionTests.cpp
+++ b/Engine/Plugins/Tests/EditorTests/Source/EditorTests/Private/UnrealEd/EditorBuildPromotionTests.cpp
@@ -80,7 +80,7 @@
 #include "K2Node_CustomEvent.h"
 #include "EdGraphUtilities.h"
 #include "Kismet2/KismetDebugUtilities.h"
-#include "Engine/Breakpoint.h"
+#include "Kismet2/Breakpoint.h"
 #include "Engine/LevelScriptBlueprint.h"
 #include "Kismet2/BlueprintEditorUtils.h"
 #include "Kismet2/KismetEditorUtilities.h"
@@ -1989,11 +1989,7 @@ namespace BuildPromotionTestHelper
 			if (BlueprintObject)
 			{
 				//Add a breakpoint
-				UBreakpoint* NewBreakpoint = NewObject<UBreakpoint>(BlueprintObject);
-				FKismetDebugUtilities::SetBreakpointEnabled(NewBreakpoint, true);
-				FKismetDebugUtilities::SetBreakpointLocation(NewBreakpoint, PrintNode);
-				BlueprintObject->Breakpoints.Add(NewBreakpoint);
-				BlueprintObject->MarkPackageDirty();
+				FKismetDebugUtilities::CreateBreakpoint(BlueprintObject, PrintNode, /* bIsEnabled = */ true);
 
 				Test->AddInfo(TEXT("Set a breakpoint on the PrintString node"));
 				EditorBuildPromotionTestUtils::StartPIE(true);
diff --git a/Engine/Source/Editor/BlueprintGraph/Private/BlueprintEditorSettings.cpp b/Engine/Source/Editor/BlueprintGraph/Private/BlueprintEditorSettings.cpp
--- a/Engine/Source/Editor/BlueprintGraph/Private/BlueprintEditorSettings.cpp
+++ b/Engine/Source/Editor/BlueprintGraph/Private/BlueprintEditorSettings.cpp
@@ -9,6 +9,7 @@
 #include "Subsystems/AssetEditorSubsystem.h"
 #include "Editor.h"
 #include "BlueprintTypePromotion.h"
+#include "AssetRegistry/AssetRegistryModule.h"
 
 UBlueprintEditorSettings::UBlueprintEditorSettings(const FObjectInitializer& ObjectInitializer)
 	: Super(ObjectInitializer)
@@ -71,6 +72,28 @@ UBlueprintEditorSettings::UBlueprintEditorSettings(const FObjectInitializer& Obj
 	{
 		SaveOnCompile = SoC_SuccessOnly;
 	}
+
+	FAssetRegistryModule& AssetRegistryModule = FModuleManager::LoadModuleChecked<FAssetRegistryModule>(TEXT("AssetRegistry"));
+	AssetRegistryModule.Get().OnAssetRenamed().AddUObject(this, &UBlueprintEditorSettings::OnAssetRenamed);
+	AssetRegistryModule.Get().OnInMemoryAssetDeleted().AddUObject(this, &UBlueprintEditorSettings::OnAssetRemoved);
+}
+
+void UBlueprintEditorSettings::OnAssetRenamed(FAssetData const& AssetInfo, const FString& InOldName)
+{
+	FPerBlueprintSettings Temp;
+	if(PerBlueprintSettings.RemoveAndCopyValue(InOldName, Temp))
+	{
+		PerBlueprintSettings.Add(AssetInfo.ObjectPath.ToString(), Temp);
+		SaveConfig();
+	}
+}
+
+void UBlueprintEditorSettings::OnAssetRemoved(UObject* Object)
+{
+	if(UBlueprint* Blueprint = Cast<UBlueprint>(Object))
+	{
+		FKismetDebugUtilities::ClearBreakpoints(Blueprint);
+	}
 }
 
 void UBlueprintEditorSettings::PostEditChangeProperty(FPropertyChangedEvent& PropertyChangedEvent)
@@ -117,4 +140,4 @@ void UBlueprintEditorSettings::PostEditChangeProperty(FPropertyChangedEvent& Pro
 	Super::PostEditChangeProperty(PropertyChangedEvent);
 
 	OnSettingsChange.Broadcast(this, PropertyChangedEvent.ChangeType);
-}
\ No newline at end of file
+}
diff --git a/Engine/Source/Editor/BlueprintGraph/Private/EdGraphSchema_K2.cpp b/Engine/Source/Editor/BlueprintGraph/Private/EdGraphSchema_K2.cpp
--- a/Engine/Source/Editor/BlueprintGraph/Private/EdGraphSchema_K2.cpp
+++ b/Engine/Source/Editor/BlueprintGraph/Private/EdGraphSchema_K2.cpp
@@ -2,7 +2,7 @@
 
 #include "EdGraphSchema_K2.h"
 #include "BlueprintCompilationManager.h"
-#include "Engine/Breakpoint.h"
+#include "Kismet2/Breakpoint.h"
 #include "Modules/ModuleManager.h"
 #include "UObject/Interface.h"
 #include "UObject/UnrealType.h"
@@ -4551,10 +4551,11 @@ void UEdGraphSchema_K2::HandleGraphBeingDeleted(UEdGraph& GraphBeingRemoved) con
 		Blueprint->LastEditedDocuments.RemoveAll([&GraphBeingRemoved](const FEditedDocumentInfo& TestDoc) { return TestDoc.EditedObjectPath.ResolveObject() == &GraphBeingRemoved; });
 
 		// Remove any BPs that reference a node in this graph:
-		Blueprint->Breakpoints.RemoveAll(
-			[&GraphBeingRemoved](UBreakpoint* Breakpoint)
+		FKismetDebugUtilities::RemoveBreakpointsByPredicate(
+			Blueprint,
+			[&GraphBeingRemoved](const FBreakpoint& Breakpoint)
 			{
-				return !Breakpoint || (Breakpoint->GetLocation() && Breakpoint->GetLocation()->IsIn(&GraphBeingRemoved));
+				return (Breakpoint.GetLocation() && Breakpoint.GetLocation()->IsIn(&GraphBeingRemoved));
 			}
 		);
 	}
diff --git a/Engine/Source/Editor/BlueprintGraph/Private/K2Node.cpp b/Engine/Source/Editor/BlueprintGraph/Private/K2Node.cpp
--- a/Engine/Source/Editor/BlueprintGraph/Private/K2Node.cpp
+++ b/Engine/Source/Editor/BlueprintGraph/Private/K2Node.cpp
@@ -18,7 +18,7 @@
 #include "Editor/EditorEngine.h"
 #include "Misc/OutputDeviceNull.h"
 
-#include "Engine/Breakpoint.h"
+#include "Kismet2/Breakpoint.h"
 #include "Kismet2/KismetDebugUtilities.h"
 #include "KismetCompiler.h"
 #include "PropertyCustomizationHelpers.h"
@@ -289,7 +289,7 @@ FText UK2Node::GetToolTipHeading() const
 
 	if (Blueprint)
 	{
-		if (UBreakpoint* ExistingBreakpoint = FKismetDebugUtilities::FindBreakpointForNode(Blueprint, this))
+		if (FBreakpoint* ExistingBreakpoint = FKismetDebugUtilities::FindBreakpointForNode(this, Blueprint))
 		{
 			if (ExistingBreakpoint->IsEnabled())
 			{
diff --git a/Engine/Source/Editor/BlueprintGraph/Private/K2Node_CallFunction.cpp b/Engine/Source/Editor/BlueprintGraph/Private/K2Node_CallFunction.cpp
--- a/Engine/Source/Editor/BlueprintGraph/Private/K2Node_CallFunction.cpp
+++ b/Engine/Source/Editor/BlueprintGraph/Private/K2Node_CallFunction.cpp
@@ -1191,12 +1191,8 @@ void UK2Node_CallFunction::PostReconstructNode()
 
 	if (IsNodePure())
 	{
-		// Remove any pre-existing breakpoint on this node since pure nodes cannot have breakpoints
-		if (UBreakpoint* ExistingBreakpoint = FKismetDebugUtilities::FindBreakpointForNode(GetBlueprint(), this))
-		{
-			// Remove the breakpoint
-			FKismetDebugUtilities::StartDeletingBreakpoint(ExistingBreakpoint, GetBlueprint());
-		}
+		// Remove the breakpoint
+		FKismetDebugUtilities::RemoveBreakpointFromNode(this, GetBlueprint());
 	}
 }
 
diff --git a/Engine/Source/Editor/BlueprintGraph/Public/BlueprintEditorSettings.h b/Engine/Source/Editor/BlueprintGraph/Public/BlueprintEditorSettings.h
--- a/Engine/Source/Editor/BlueprintGraph/Public/BlueprintEditorSettings.h
+++ b/Engine/Source/Editor/BlueprintGraph/Public/BlueprintEditorSettings.h
@@ -6,7 +6,9 @@
 #include "UObject/ObjectMacros.h"
 #include "UObject/Object.h"
 #include "Engine/Blueprint.h"
+#include "Kismet2/Breakpoint.h"
 #include "Engine/DeveloperSettings.h"
+#include "Kismet2/KismetDebugUtilities.h"
 #include "BlueprintEditorSettings.generated.h"
 
 UENUM()
@@ -16,7 +18,19 @@ enum ESaveOnCompile
 	SoC_SuccessOnly UMETA(DisplayName="On Success Only"),
 	SoC_Always UMETA(DisplayName = "Always"),
 };
- 
+
+/** Blueprint Editor settings that are different for each
+*	blueprint.
+*/
+USTRUCT()
+struct BLUEPRINTGRAPH_API FPerBlueprintSettings 
+{
+	GENERATED_BODY()
+	
+	UPROPERTY()
+	TArray<FBreakpoint> Breakpoints;
+};
+
 UCLASS(config=EditorPerProjectUserSettings)
 class BLUEPRINTGRAPH_API UBlueprintEditorSettings
 	:	public UObject
@@ -186,6 +200,10 @@ public:
 	UPROPERTY(config)
 	TArray<FBPEditorBookmarkNode> BookmarkNodes;
 
+	/** Maps Blueprint path to settings such as breakpoints */
+	UPROPERTY(config)
+	TMap<FString, FPerBlueprintSettings> PerBlueprintSettings;
+
 	/** If enabled, comment nodes will be included in the tree view display in the Bookmarks tab. */
 	UPROPERTY(config)
 	bool bIncludeCommentNodesInBookmarksTab;
@@ -211,9 +229,14 @@ public:
 	{
 		OnSettingsChange.Remove(Object);
 	}
+	
 
 protected:
 	//~ Begin UObject Interface
 	virtual void PostEditChangeProperty(struct FPropertyChangedEvent& PropertyChangedEvent) override;
 	//~ End UObject Interface
+	
+	void OnAssetRenamed(FAssetData const& AssetInfo, const FString& InOldName);
+	
+	void OnAssetRemoved(UObject* Object);
 };
diff --git a/Engine/Source/Editor/GraphEditor/Private/KismetNodes/SGraphNodeK2Base.cpp b/Engine/Source/Editor/GraphEditor/Private/KismetNodes/SGraphNodeK2Base.cpp
--- a/Engine/Source/Editor/GraphEditor/Private/KismetNodes/SGraphNodeK2Base.cpp
+++ b/Engine/Source/Editor/GraphEditor/Private/KismetNodes/SGraphNodeK2Base.cpp
@@ -16,7 +16,7 @@
 #include "EdGraphSchema_K2.h"
 #include "K2Node.h"
 #include "K2Node_Timeline.h"
-#include "Engine/Breakpoint.h"
+#include "Kismet2/Breakpoint.h"
 #include "Kismet2/KismetDebugUtilities.h"
 #include "Kismet2/BlueprintEditorUtils.h"
 #include "KismetNodes/KismetNodeInfoContext.h"
@@ -408,14 +408,14 @@ void SGraphNodeK2Base::GetOverlayBrushes(bool bSelected, const FVector2D WidgetS
 	UBlueprint* OwnerBlueprint = FBlueprintEditorUtils::FindBlueprintForNode(GraphNode);
 
 	// Search for an enabled or disabled breakpoint on this node
-	UBreakpoint* Breakpoint = OwnerBlueprint ? FKismetDebugUtilities::FindBreakpointForNode(OwnerBlueprint, GraphNode) : nullptr;
+	FBreakpoint* Breakpoint = OwnerBlueprint ? FKismetDebugUtilities::FindBreakpointForNode(GraphNode, OwnerBlueprint) : nullptr;
 	if (Breakpoint != NULL)
 	{
 		FOverlayBrushInfo BreakpointOverlayInfo;
 
 		if (Breakpoint->IsEnabledByUser())
 		{
-			BreakpointOverlayInfo.Brush = FEditorStyle::GetBrush(FKismetDebugUtilities::IsBreakpointValid(Breakpoint) ? TEXT("Kismet.DebuggerOverlay.Breakpoint.EnabledAndValid") : TEXT("Kismet.DebuggerOverlay.Breakpoint.EnabledAndInvalid"));
+			BreakpointOverlayInfo.Brush = FEditorStyle::GetBrush(FKismetDebugUtilities::IsBreakpointValid(*Breakpoint) ? TEXT("Kismet.DebuggerOverlay.Breakpoint.EnabledAndValid") : TEXT("Kismet.DebuggerOverlay.Breakpoint.EnabledAndInvalid"));
 		}
 		else
 		{
diff --git a/Engine/Source/Editor/Kismet/Private/BlueprintCompilationManager.cpp b/Engine/Source/Editor/Kismet/Private/BlueprintCompilationManager.cpp
--- a/Engine/Source/Editor/Kismet/Private/BlueprintCompilationManager.cpp
+++ b/Engine/Source/Editor/Kismet/Private/BlueprintCompilationManager.cpp
@@ -1399,10 +1399,13 @@ void FBlueprintCompilationManagerImpl::FlushCompilationQueueImpl(bool bSuppressB
 					BP->BlueprintSystemVersion = UBlueprint::GetCurrentBlueprintSystemVersion();
 
 					// Reapply breakpoints to the bytecode of the new class
-					for (UBreakpoint* Breakpoint  : BP->Breakpoints)
-					{
-						FKismetDebugUtilities::ReapplyBreakpoint(Breakpoint);
-					}
+					FKismetDebugUtilities::ForeachBreakpoint(
+						BP,
+						[](FBreakpoint& Breakpoint)
+						{
+							FKismetDebugUtilities::ReapplyBreakpoint(Breakpoint);
+						}
+					);
 				}
 				else
 				{
diff --git a/Engine/Source/Editor/Kismet/Private/BlueprintEditor.cpp b/Engine/Source/Editor/Kismet/Private/BlueprintEditor.cpp
--- a/Engine/Source/Editor/Kismet/Private/BlueprintEditor.cpp
+++ b/Engine/Source/Editor/Kismet/Private/BlueprintEditor.cpp
@@ -69,7 +69,7 @@
 #include "K2Node_SetFieldsInStruct.h"
 #include "K2Node_Knot.h"
 #include "Engine/LevelScriptBlueprint.h"
-#include "Engine/Breakpoint.h"
+#include "Kismet2/Breakpoint.h"
 #include "ScopedTransaction.h"
 #include "Kismet2/KismetDebugUtilities.h"
 #include "Kismet2/BlueprintEditorUtils.h"
@@ -3956,39 +3956,38 @@ bool FBlueprintEditor::CanOpenBlueprintDebugger() const
 
 bool FBlueprintEditor::HasAnyBreakpoints() const
 {
-	return GetBlueprintObj() && GetBlueprintObj()->Breakpoints.Num() > 0;
+	const UBlueprint* Blueprint = GetBlueprintObj();
+	if(!Blueprint)
+	{
+		return false;
+	}
+	return FKismetDebugUtilities::BlueprintHasBreakpoints(Blueprint);
 }
 
 bool FBlueprintEditor::HasAnyEnabledBreakpoints() const
 {
 	if (!IsEditingSingleBlueprint()) {return false;}
 
-	for (TArray<UBreakpoint*>::TIterator BreakpointIt(GetBlueprintObj()->Breakpoints); BreakpointIt; ++BreakpointIt)
-	{
-		UBreakpoint* BP = *BreakpointIt;
-		if (BP->IsEnabledByUser())
+	return FKismetDebugUtilities::FindBreakpointByPredicate(
+		GetBlueprintObj(),
+		[](const FBreakpoint& Breakpoint)
 		{
-			return true;
+			return Breakpoint.IsEnabledByUser();
 		}
-	}
-
-	return false;
+	) != nullptr;
 }
 
 bool FBlueprintEditor::HasAnyDisabledBreakpoints() const
 {
 	if (!IsEditingSingleBlueprint()) {return false;}
 
-	for (TArray<UBreakpoint*>::TIterator BreakpointIt(GetBlueprintObj()->Breakpoints); BreakpointIt; ++BreakpointIt)
-	{
-		UBreakpoint* BP = *BreakpointIt;
-		if (!BP->IsEnabledByUser())
+	return FKismetDebugUtilities::FindBreakpointByPredicate(
+		GetBlueprintObj(),
+		[](const FBreakpoint& Breakpoint)
 		{
-			return true;
+			return !Breakpoint.IsEnabledByUser();
 		}
-	}
-
-	return false;
+	) != nullptr;
 }
 
 bool FBlueprintEditor::HasAnyWatches() const
@@ -5123,20 +5122,16 @@ void FBlueprintEditor::OnToggleBreakpoint()
 		UK2Node* SelectedNode = Cast<UK2Node>(*NodeIt);
 		if ((SelectedNode != nullptr) && SelectedNode->CanPlaceBreakpoints())
 		{
-			UBreakpoint* ExistingBreakpoint = FKismetDebugUtilities::FindBreakpointForNode(GetBlueprintObj(), SelectedNode);
+			FBreakpoint* ExistingBreakpoint = FKismetDebugUtilities::FindBreakpointForNode(SelectedNode, GetBlueprintObj());
 			if (ExistingBreakpoint == nullptr)
 			{
 				// Add a breakpoint on this node if there isn't one there already
-				UBreakpoint* NewBreakpoint = NewObject<UBreakpoint>(GetBlueprintObj());
-				FKismetDebugUtilities::SetBreakpointEnabled(NewBreakpoint, true);
-				FKismetDebugUtilities::SetBreakpointLocation(NewBreakpoint, SelectedNode);
-				GetBlueprintObj()->Breakpoints.Add(NewBreakpoint);
-				GetBlueprintObj()->MarkPackageDirty();
+				FKismetDebugUtilities::CreateBreakpoint(GetBlueprintObj(), SelectedNode, /* bIsEnabled = */ true);
 			}
 			else
 			{
 				// Remove the breakpoint if it was present
-				FKismetDebugUtilities::StartDeletingBreakpoint(ExistingBreakpoint, GetBlueprintObj());
+				FKismetDebugUtilities::RemoveBreakpointFromNode(SelectedNode, GetBlueprintObj());
 			}
 		}
 	}
@@ -5166,15 +5161,10 @@ void FBlueprintEditor::OnAddBreakpoint()
 		if ((SelectedNode != nullptr) && SelectedNode->CanPlaceBreakpoints())
 		{
 			// Add a breakpoint on this node if there isn't one there already
-			UBreakpoint* ExistingBreakpoint = FKismetDebugUtilities::FindBreakpointForNode(GetBlueprintObj(), SelectedNode);
+			const FBreakpoint* ExistingBreakpoint = FKismetDebugUtilities::FindBreakpointForNode(SelectedNode, GetBlueprintObj());
 			if (ExistingBreakpoint == nullptr)
 			{
-				// Add a new breakpoint
-				UBreakpoint* NewBreakpoint = NewObject<UBreakpoint>(GetBlueprintObj());
-				FKismetDebugUtilities::SetBreakpointEnabled(NewBreakpoint, true);
-				FKismetDebugUtilities::SetBreakpointLocation(NewBreakpoint, SelectedNode);
-				GetBlueprintObj()->Breakpoints.Add(NewBreakpoint);
-				GetBlueprintObj()->MarkPackageDirty();
+				FKismetDebugUtilities::CreateBreakpoint(GetBlueprintObj(), SelectedNode, /* bIsEnabled = */ true);
 			}
 		}
 	}
@@ -5189,7 +5179,7 @@ bool FBlueprintEditor::CanAddBreakpoint() const
 		UK2Node* SelectedNode = Cast<UK2Node>(*NodeIt);
 		if ((SelectedNode != nullptr) && SelectedNode->CanPlaceBreakpoints())
 		{
-			UBreakpoint* ExistingBreakpoint = FKismetDebugUtilities::FindBreakpointForNode(GetBlueprintObj(), SelectedNode);
+			FBreakpoint* ExistingBreakpoint = FKismetDebugUtilities::FindBreakpointForNode(SelectedNode, GetBlueprintObj());
 			if (ExistingBreakpoint == nullptr)
 			{
 				return true;
@@ -5207,12 +5197,8 @@ void FBlueprintEditor::OnRemoveBreakpoint()
 	{
 		UEdGraphNode* SelectedNode = CastChecked<UEdGraphNode>(*NodeIt);
 
-		// Remove any pre-existing breakpoint on this node
-		if (UBreakpoint* ExistingBreakpoint = FKismetDebugUtilities::FindBreakpointForNode(GetBlueprintObj(), SelectedNode))
-		{
-			// Remove the breakpoint
-			FKismetDebugUtilities::StartDeletingBreakpoint(ExistingBreakpoint, GetBlueprintObj());
-		}
+		// Remove the breakpoint
+		FKismetDebugUtilities::RemoveBreakpointFromNode(SelectedNode, GetBlueprintObj());
 	}
 }
 
@@ -5223,7 +5209,7 @@ bool FBlueprintEditor::CanRemoveBreakpoint() const
 	for (FGraphPanelSelectionSet::TConstIterator NodeIt(SelectedNodes); NodeIt; ++NodeIt)
 	{
 		UEdGraphNode* SelectedNode = CastChecked<UEdGraphNode>(*NodeIt);
-		if (UBreakpoint* ExistingBreakpoint = FKismetDebugUtilities::FindBreakpointForNode(GetBlueprintObj(), SelectedNode))
+		if (FKismetDebugUtilities::FindBreakpointForNode(SelectedNode, GetBlueprintObj()))
 		{
 			return true;
 		}
@@ -5239,10 +5225,7 @@ void FBlueprintEditor::OnDisableBreakpoint()
 	{
 		UEdGraphNode* SelectedNode = CastChecked<UEdGraphNode>(*NodeIt);
 
-		if (UBreakpoint* ExistingBreakpoint = FKismetDebugUtilities::FindBreakpointForNode(GetBlueprintObj(), SelectedNode))
-		{
-			FKismetDebugUtilities::SetBreakpointEnabled(ExistingBreakpoint, false);
-		}
+		FKismetDebugUtilities::SetBreakpointEnabled(SelectedNode, GetBlueprintObj(), false);
 	}
 }
 
@@ -5253,7 +5236,7 @@ bool FBlueprintEditor::CanDisableBreakpoint() const
 	for (FGraphPanelSelectionSet::TConstIterator NodeIt(SelectedNodes); NodeIt; ++NodeIt)
 	{
 		UEdGraphNode* SelectedNode = CastChecked<UEdGraphNode>(*NodeIt);
-		if (UBreakpoint* ExistingBreakpoint = FKismetDebugUtilities::FindBreakpointForNode(GetBlueprintObj(), SelectedNode))
+		if (FBreakpoint* ExistingBreakpoint = FKismetDebugUtilities::FindBreakpointForNode(SelectedNode, GetBlueprintObj()))
 		{
 			if (ExistingBreakpoint->IsEnabledByUser())
 			{
@@ -5272,10 +5255,7 @@ void FBlueprintEditor::OnEnableBreakpoint()
 	{
 		UEdGraphNode* SelectedNode = CastChecked<UEdGraphNode>(*NodeIt);
 
-		if (UBreakpoint* ExistingBreakpoint = FKismetDebugUtilities::FindBreakpointForNode(GetBlueprintObj(), SelectedNode))
-		{
-			FKismetDebugUtilities::SetBreakpointEnabled(ExistingBreakpoint, true);
-		}
+		FKismetDebugUtilities::SetBreakpointEnabled(SelectedNode, GetBlueprintObj(), true);
 	}
 }
 
@@ -5286,7 +5266,7 @@ bool FBlueprintEditor::CanEnableBreakpoint() const
 	for (FGraphPanelSelectionSet::TConstIterator NodeIt(SelectedNodes); NodeIt; ++NodeIt)
 	{
 		UEdGraphNode* SelectedNode = CastChecked<UEdGraphNode>(*NodeIt);
-		if (UBreakpoint* ExistingBreakpoint = FKismetDebugUtilities::FindBreakpointForNode(GetBlueprintObj(), SelectedNode))
+		if (FBreakpoint* ExistingBreakpoint = FKismetDebugUtilities::FindBreakpointForNode(SelectedNode, GetBlueprintObj()))
 		{
 			if (!ExistingBreakpoint->IsEnabledByUser())
 			{
diff --git a/Engine/Source/Editor/Kismet/Private/Debugging/KismetDebugCommands.cpp b/Engine/Source/Editor/Kismet/Private/Debugging/KismetDebugCommands.cpp
--- a/Engine/Source/Editor/Kismet/Private/Debugging/KismetDebugCommands.cpp
+++ b/Engine/Source/Editor/Kismet/Private/Debugging/KismetDebugCommands.cpp
@@ -3,6 +3,7 @@
 
 #include "Debugging/KismetDebugCommands.h"
 #include "Kismet2/BlueprintEditorUtils.h"
+#include "Kismet2/Breakpoint.h"
 #include "Kismet2/KismetDebugUtilities.h"
 
 //////////////////////////////////////////////////////////////////////////
@@ -27,21 +28,23 @@ void FDebuggingActionCallbacks::ClearBreakpoints(UBlueprint* OwnerBlueprint)
 	FKismetDebugUtilities::ClearBreakpoints(OwnerBlueprint);
 }
 
-void FDebuggingActionCallbacks::ClearBreakpoint(UBreakpoint* Breakpoint, UBlueprint* OwnerBlueprint)
+void FDebuggingActionCallbacks::ClearBreakpoint(TSoftObjectPtr<UEdGraphNode> BreakpointNode, const UBlueprint* OwnerBlueprint)
 {
-	FKismetDebugUtilities::StartDeletingBreakpoint(Breakpoint, OwnerBlueprint);
+	FKismetDebugUtilities::RemoveBreakpointFromNode(BreakpointNode.Get(), OwnerBlueprint);
 }
 
-void FDebuggingActionCallbacks::SetBreakpointEnabled(UBreakpoint* Breakpoint, bool bEnabled)
+void FDebuggingActionCallbacks::SetBreakpointEnabled(TSoftObjectPtr<UEdGraphNode> BreakpointNode, const UBlueprint* BreakpointBlueprint, bool bEnabled)
 {
-	FKismetDebugUtilities::SetBreakpointEnabled(Breakpoint, bEnabled);
+	FKismetDebugUtilities::SetBreakpointEnabled(BreakpointNode.Get(), BreakpointBlueprint, bEnabled);
 }
 
-void FDebuggingActionCallbacks::SetEnabledOnAllBreakpoints(UBlueprint* OwnerBlueprint, bool bShouldBeEnabled)
+void FDebuggingActionCallbacks::SetEnabledOnAllBreakpoints(const UBlueprint* OwnerBlueprint, bool bShouldBeEnabled)
 {
-	for (TArray<UBreakpoint*>::TIterator BreakpointIt(OwnerBlueprint->Breakpoints); BreakpointIt; ++BreakpointIt)
-	{
-		UBreakpoint* BP = *BreakpointIt;
-		FKismetDebugUtilities::SetBreakpointEnabled(BP, bShouldBeEnabled);
-	}
+	FKismetDebugUtilities::ForeachBreakpoint(
+		OwnerBlueprint,
+		[bShouldBeEnabled](FBreakpoint& Breakpoint)
+		{
+			FKismetDebugUtilities::SetBreakpointEnabled(Breakpoint, bShouldBeEnabled);
+		}
+	);
 }
diff --git a/Engine/Source/Editor/Kismet/Private/Debugging/KismetDebugCommands.h b/Engine/Source/Editor/Kismet/Private/Debugging/KismetDebugCommands.h
--- a/Engine/Source/Editor/Kismet/Private/Debugging/KismetDebugCommands.h
+++ b/Engine/Source/Editor/Kismet/Private/Debugging/KismetDebugCommands.h
@@ -7,15 +7,20 @@
 //////////////////////////////////////////////////////////////////////////
 // FDebuggingActionCallbacks
 
+struct FBreakpoint;
+class UBlueprint;
+class UEdGraphNode;
+class UEdGraphPin;
+
 class FDebuggingActionCallbacks
 {
 public:
-	static void ClearWatches(class UBlueprint* Blueprint);
-	static void ClearWatch(class UEdGraphPin* WatchedPin);
-	static void ClearBreakpoints(class UBlueprint* OwnerBlueprint);
-	static void ClearBreakpoint(class UBreakpoint* Breakpoint, class UBlueprint* OwnerBlueprint);
-	static void SetBreakpointEnabled(class UBreakpoint* Breakpoint, bool bEnabled);
-	static void SetEnabledOnAllBreakpoints(class UBlueprint* OwnerBlueprint, bool bShouldBeEnabled);
+	static void ClearWatches(UBlueprint* Blueprint);
+	static void ClearWatch(UEdGraphPin* WatchedPin);
+	static void ClearBreakpoints(UBlueprint* OwnerBlueprint);
+	static void ClearBreakpoint(TSoftObjectPtr<UEdGraphNode> BreakpointNode, const UBlueprint* OwnerBlueprint);
+	static void SetBreakpointEnabled(TSoftObjectPtr<UEdGraphNode> BreakpointNode, const UBlueprint* BreakpointBlueprint, bool bEnabled);
+	static void SetEnabledOnAllBreakpoints(const UBlueprint* OwnerBlueprint, bool bShouldBeEnabled);
 };
 
 //////////////////////////////////////////////////////////////////////////
diff --git a/Engine/Source/Editor/Kismet/Private/Debugging/SKismetDebuggingView.cpp b/Engine/Source/Editor/Kismet/Private/Debugging/SKismetDebuggingView.cpp
--- a/Engine/Source/Editor/Kismet/Private/Debugging/SKismetDebuggingView.cpp
+++ b/Engine/Source/Editor/Kismet/Private/Debugging/SKismetDebuggingView.cpp
@@ -18,7 +18,7 @@
 #include "Editor.h"
 #include "Kismet2/KismetEditorUtilities.h"
 #include "K2Node.h"
-#include "Engine/Breakpoint.h"
+#include "Kismet2/Breakpoint.h"
 #include "Kismet2/KismetDebugUtilities.h"
 #include "Kismet2/DebuggerCommands.h"
 #include "Debugging/KismetDebugCommands.h"
@@ -66,18 +66,17 @@ TSharedRef<SWidget> FDebugLineItem::GenerateValueWidget()
 
 UBlueprint* FDebugLineItem::GetBlueprintForObject(UObject* ParentObject)
 {
-	UBlueprint* ParentBlueprint = NULL;
-
-	if (ParentObject != NULL)
+	if(ParentObject == nullptr)
 	{
-		ParentBlueprint = Cast<UBlueprint>(ParentObject);
-		if (ParentBlueprint == NULL)
-		{
-			ParentBlueprint = Cast<UBlueprint>(ParentObject->GetClass()->ClassGeneratedBy);
-		}
+		return nullptr;
 	}
-
-	return ParentBlueprint;
+	if(UBlueprint* ParentBlueprint = Cast<UBlueprint>(ParentObject))
+	{
+		return ParentBlueprint;
+	}
+	
+	// recursively walk up ownership heirrarchy until we find the blueprint
+	return GetBlueprintForObject(ParentObject->GetOuter());
 }
 
 UBlueprintGeneratedClass* FDebugLineItem::GetClassForObject(UObject* ParentObject)
@@ -408,13 +407,13 @@ void FWatchLineItem::OnNavigateToWatchLocation( )
 struct FBreakpointLineItem : public FDebugLineItem
 {
 protected:
-	TWeakObjectPtr< UObject > ParentObjectRef;
-	TWeakObjectPtr< UBreakpoint > BreakpointRef;
+	TWeakObjectPtr<UObject> ParentObjectRef;
+	TSoftObjectPtr<UEdGraphNode> BreakpointNode;
 public:
-	FBreakpointLineItem(UBreakpoint* BreakpointToWatch, UObject* ParentObject)
+	FBreakpointLineItem(TSoftObjectPtr<UEdGraphNode> BreakpointToWatch, UObject* ParentObject)
 		: FDebugLineItem(DLT_Breakpoint)
 	{
-		BreakpointRef = BreakpointToWatch;
+		BreakpointNode = BreakpointToWatch;
 		ParentObjectRef = ParentObject;
 	}
 
@@ -422,18 +421,18 @@ public:
 	{
 		FBreakpointLineItem* Other = (FBreakpointLineItem*)BaseOther;
 		return (ParentObjectRef.Get() == Other->ParentObjectRef.Get()) &&
-			(BreakpointRef.Get() == Other->BreakpointRef.Get());
+			(BreakpointNode == Other->BreakpointNode);
 	}
 
 	virtual FDebugLineItem* Duplicate() const override
 	{
-		return new FBreakpointLineItem(BreakpointRef.Get(), ParentObjectRef.Get());
+		return new FBreakpointLineItem(BreakpointNode, ParentObjectRef.Get());
 	}	
 
 	virtual void MakeMenu(class FMenuBuilder& MenuBuilder) override
 	{
-		UBreakpoint* Breakpoint = BreakpointRef.Get();
-		UBlueprint* ParentBlueprint = GetBlueprintForObject(ParentObjectRef.Get());
+		FBreakpoint* Breakpoint = GetBreakpoint();
+		const UBlueprint* ParentBlueprint = GetBlueprintForObject(ParentObjectRef.Get());
 
 		// By default, we don't allow actions to execute when in debug mode. 
 		// Create an empty action to always allow execution for these commands (they are allowed in debug mode)
@@ -444,7 +443,9 @@ public:
 			const bool bNewEnabledState = !Breakpoint->IsEnabledByUser();
 
 			FUIAction ToggleThisBreakpoint(
-				FExecuteAction::CreateStatic( &FDebuggingActionCallbacks::SetBreakpointEnabled, Breakpoint, bNewEnabledState ),
+				FExecuteAction::CreateStatic(
+					&FDebuggingActionCallbacks::SetBreakpointEnabled, BreakpointNode, ParentBlueprint, bNewEnabledState
+					),
 				AlwaysAllowExecute
 				);
 
@@ -471,7 +472,7 @@ public:
 		if ((Breakpoint != NULL) && (ParentBlueprint != NULL))
 		{
 			FUIAction ClearThisBreakpoint(
-				FExecuteAction::CreateStatic( &FDebuggingActionCallbacks::ClearBreakpoint, Breakpoint, ParentBlueprint ),
+				FExecuteAction::CreateStatic( &FDebuggingActionCallbacks::ClearBreakpoint, BreakpointNode, ParentBlueprint ),
 				AlwaysAllowExecute
 				);
 
@@ -483,6 +484,18 @@ public:
 		}
 	}
 protected:
+	FBreakpoint* GetBreakpoint() const
+	{
+		if(UEdGraphNode* Node = BreakpointNode.Get())
+		{
+			if(const UBlueprint* Blueprint = GetBlueprintForObject(Node))
+			{
+				return FKismetDebugUtilities::FindBreakpointForNode(Node, Blueprint);
+			}
+		}
+		return nullptr;
+	}
+	
 	virtual TSharedRef<SWidget> GenerateNameWidget() override
 	{
 		return SNew(SHorizontalBox)
@@ -526,7 +539,7 @@ protected:
 
 FText FBreakpointLineItem::GetLocationDescription() const
 {
-	if (UBreakpoint* MyBreakpoint = BreakpointRef.Get())
+	if (FBreakpoint* MyBreakpoint = GetBreakpoint())
 	{
 		return MyBreakpoint->GetLocationDescription();
 	}
@@ -535,16 +548,16 @@ FText FBreakpointLineItem::GetLocationDescription() const
 
 FReply FBreakpointLineItem::OnUserToggledEnabled()
 {
-	if (UBreakpoint* MyBreakpoint = BreakpointRef.Get())
+	if (FBreakpoint* MyBreakpoint = GetBreakpoint())
 	{
-		FKismetDebugUtilities::SetBreakpointEnabled(MyBreakpoint, !MyBreakpoint->IsEnabledByUser());
+		FKismetDebugUtilities::SetBreakpointEnabled(*MyBreakpoint, !MyBreakpoint->IsEnabledByUser());
 	}
 	return FReply::Handled();
 }
 
 void FBreakpointLineItem::OnNavigateToBreakpointLocation()
 {
-	if (UBreakpoint* MyBreakpoint = BreakpointRef.Get())
+	if (FBreakpoint* MyBreakpoint = GetBreakpoint())
 	{
 		FKismetEditorUtilities::BringKismetToFocusAttentionOnObject(MyBreakpoint->GetLocation());
 	}
@@ -552,11 +565,11 @@ void FBreakpointLineItem::OnNavigateToBreakpointLocation()
 
 const FSlateBrush* FBreakpointLineItem::GetStatusImage() const
 {
-	if (UBreakpoint* MyBreakpoint = BreakpointRef.Get())
+	if (FBreakpoint* MyBreakpoint = GetBreakpoint())
 	{
 		if (MyBreakpoint->IsEnabledByUser())
 		{
-			return FEditorStyle::GetBrush(FKismetDebugUtilities::IsBreakpointValid(MyBreakpoint) ? TEXT("Kismet.Breakpoint.EnabledAndValid") : TEXT("Kismet.Breakpoint.EnabledAndInvalid"));
+			return FEditorStyle::GetBrush(FKismetDebugUtilities::IsBreakpointValid(*MyBreakpoint) ? TEXT("Kismet.Breakpoint.EnabledAndValid") : TEXT("Kismet.Breakpoint.EnabledAndInvalid"));
 		}
 		else
 		{
@@ -569,9 +582,9 @@ const FSlateBrush* FBreakpointLineItem::GetStatusImage() const
 
 FText FBreakpointLineItem::GetStatusTooltip() const
 {
-	if (UBreakpoint* MyBreakpoint = BreakpointRef.Get())
+	if (FBreakpoint* MyBreakpoint = GetBreakpoint())
 	{
-		if (!FKismetDebugUtilities::IsBreakpointValid(MyBreakpoint))
+		if (!FKismetDebugUtilities::IsBreakpointValid(*MyBreakpoint))
 		{
 			return LOCTEXT("Breakpoint_NoHit", "This breakpoint will not be hit because its node generated no code");
 		}
@@ -625,12 +638,13 @@ public:
 				}
 
 				// Create children for each breakpoint
-				for (int32 BreakpointIndex = 0; BreakpointIndex < ParentBP->Breakpoints.Num(); ++BreakpointIndex)
-				{
-					UBreakpoint* Breakpoint = ParentBP->Breakpoints[BreakpointIndex];
-
-					EnsureChildIsAdded(ChildrenMirrors, OutChildren, FBreakpointLineItem(Breakpoint, ParentObject));
-				}
+				FKismetDebugUtilities::ForeachBreakpoint(
+					ParentBP,
+					[this, &OutChildren, ParentObject](FBreakpoint& Breakpoint)
+					{
+						EnsureChildIsAdded(ChildrenMirrors, OutChildren, FBreakpointLineItem(Breakpoint.GetLocation(), ParentObject));
+					}
+				);
 
 				// Make sure there is something there, to let the user know if there is nothing
 				if (OutChildren.Num() == 0)
@@ -651,12 +665,13 @@ public:
 					}
 
 					// Create children for each breakpoint
-					for (int32 BreakpointIndex = 0; BreakpointIndex < ParentBP->Breakpoints.Num(); ++BreakpointIndex)
-					{
-						UBreakpoint* Breakpoint = ParentBP->Breakpoints[BreakpointIndex];
-
-						EnsureChildIsAdded(ChildrenMirrors, OutChildren, FBreakpointLineItem(Breakpoint, ParentObject));
-					}
+					FKismetDebugUtilities::ForeachBreakpoint(
+						ParentBP,
+						[this, &OutChildren, ParentObject](FBreakpoint& Breakpoint)
+						{
+							EnsureChildIsAdded(ChildrenMirrors, OutChildren, FBreakpointLineItem(Breakpoint.GetLocation(), ParentObject));
+						}
+					);
 				}
 
 				// It could also have active latent behaviors
@@ -717,7 +732,7 @@ protected:
 	{
 		if (UBlueprint* BP = Cast<UBlueprint>(ObjectRef.Get()))
 		{
-			if (BP->WatchedPins.Num() > 0)
+			if (FKismetDebugUtilities::BlueprintHasBreakpoints(BP))
 			{
 				FUIAction ClearAllWatches(
 					FExecuteAction::CreateStatic( &FDebuggingActionCallbacks::ClearWatches, BP )
@@ -728,10 +743,7 @@ protected:
 					LOCTEXT("ClearWatches_ToolTip", "Clear all watches in this blueprint"),
 					FSlateIcon(),
 					ClearAllWatches);
-			}
-
-			if (BP->Breakpoints.Num() > 0)
-			{
+			
 				FUIAction ClearAllBreakpoints(
 					FExecuteAction::CreateStatic( &FDebuggingActionCallbacks::ClearBreakpoints, BP )
 					);
@@ -1061,7 +1073,7 @@ EVisibility SKismetDebuggingView::IsDebuggerVisible() const
 bool SKismetDebuggingView::CanDisableAllBreakpoints() const
 {
 	const UBlueprint* BlueprintObj = BlueprintToWatchPtr.Get();
-	return BlueprintObj && BlueprintObj->Breakpoints.Num() > 0;
+	return BlueprintObj && FKismetDebugUtilities::BlueprintHasBreakpoints(BlueprintObj);
 }
 
 FReply SKismetDebuggingView::OnDisableAllBreakpointsClicked()
diff --git a/Engine/Source/Editor/UnrealEd/Private/Kismet2/BlueprintEditorUtils.cpp b/Engine/Source/Editor/UnrealEd/Private/Kismet2/BlueprintEditorUtils.cpp
--- a/Engine/Source/Editor/UnrealEd/Private/Kismet2/BlueprintEditorUtils.cpp
+++ b/Engine/Source/Editor/UnrealEd/Private/Kismet2/BlueprintEditorUtils.cpp
@@ -2859,10 +2859,7 @@ void FBlueprintEditorUtils::RemoveNode(UBlueprint* Blueprint, UEdGraphNode* Node
 	if (Blueprint != nullptr)
 	{
 		// Remove any breakpoints set on the node
-		if (UBreakpoint* Breakpoint = FKismetDebugUtilities::FindBreakpointForNode(Blueprint, Node))
-		{
-			FKismetDebugUtilities::StartDeletingBreakpoint(Breakpoint, Blueprint);
-		}
+		FKismetDebugUtilities::RemoveBreakpointFromNode(Node, Blueprint);
 
 		// Remove any watches set on the node's pins
 		for (int32 PinIndex = 0; PinIndex < Node->Pins.Num(); ++PinIndex)
diff --git a/Engine/Source/Editor/UnrealEd/Private/Kismet2/Breakpoint.cpp b/Engine/Source/Editor/UnrealEd/Private/Kismet2/Breakpoint.cpp
new file mode 100644
--- /dev/null
+++ b/Engine/Source/Editor/UnrealEd/Private/Kismet2/Breakpoint.cpp
@@ -0,0 +1,41 @@
+// Copyright Epic Games, Inc. All Rights Reserved.
+
+#include "Kismet2/Breakpoint.h"
+#if WITH_EDITOR
+#include "EdGraph/EdGraphNode.h"
+#endif
+
+FBreakpoint::FBreakpoint()
+{
+	bEnabled = false;
+	bStepOnce = false;
+	bStepOnce_WasPreviouslyDisabled = false;
+	bStepOnce_RemoveAfterHit = false;
+}
+
+FText FBreakpoint::GetLocationDescription() const
+{
+#if WITH_EDITORONLY_DATA
+	if (Node != NULL)
+	{
+		FString Result
+#if WITH_EDITOR
+			= Node->GetDescriptiveCompiledName()
+#endif
+			;
+		if (!Node->NodeComment.IsEmpty())
+		{
+			Result += TEXT(" // ");
+			Result += Node->NodeComment;
+		}
+
+		return FText::FromString(Result);
+	}
+	else
+	{
+		return NSLOCTEXT("UBreakpoint", "ErrorInvalidLocation", "Error: Invalid location");
+	}
+#else	//#if WITH_EDITORONLY_DATA
+	return NSLOCTEXT("UBreakpoint", "NoEditorData", "--- NO EDITOR DATA! ---");
+#endif	//#if WITH_EDITORONLY_DATA
+}
diff --git a/Engine/Source/Editor/UnrealEd/Private/Kismet2/KismetDebugUtilities.cpp b/Engine/Source/Editor/UnrealEd/Private/Kismet2/KismetDebugUtilities.cpp
--- a/Engine/Source/Editor/UnrealEd/Private/Kismet2/KismetDebugUtilities.cpp
+++ b/Engine/Source/Editor/UnrealEd/Private/Kismet2/KismetDebugUtilities.cpp
@@ -23,7 +23,7 @@
 #include "WatchPointViewer.h"
 #include "Animation/AnimBlueprintGeneratedClass.h"
 #include "UnrealEdGlobals.h"
-#include "Engine/Breakpoint.h"
+#include "Kismet2/Breakpoint.h"
 #include "ActorEditorUtils.h"
 #include "EdGraphSchema_K2.h"
 #include "K2Node.h"
@@ -40,6 +40,7 @@
 #include "AnimGraphNode_Base.h"
 #include "UObject/UnrealType.h"
 #include "AnimationGraphSchema.h"
+#include "BlueprintEditorSettings.h"
 
 #define LOCTEXT_NAMESPACE "BlueprintDebugging"
 
@@ -662,12 +663,12 @@ void FKismetDebugUtilities::AttemptToBreakExecution(UBlueprint* BlueprintObj, co
 		// Find the breakpoint object for the node, assuming we hit one
 		if (Info.GetType() == EBlueprintExceptionType::Breakpoint)
 		{
-			UBreakpoint* Breakpoint = FKismetDebugUtilities::FindBreakpointForNode(BlueprintObj, NodeStoppedAt);
+			FBreakpoint* Breakpoint = FindBreakpointForNode(NodeStoppedAt, BlueprintObj);
 
 			if (Breakpoint != NULL)
 			{
 				Data.MostRecentBreakpointInstructionPointer = NodeStoppedAt;
-				FKismetDebugUtilities::UpdateBreakpointStateWhenHit(Breakpoint, BlueprintObj);
+				UpdateBreakpointStateWhenHit(NodeStoppedAt, BlueprintObj);
 					
 				//@TODO: K2: DEBUGGING: Debug print text can go eventually
 				UE_LOG(LogBlueprintDebug, Warning, TEXT("Hit breakpoint on node '%s', from offset %d"), *(NodeStoppedAt->GetDescriptiveCompiledName()), DebugOpcodeOffset);
@@ -784,20 +785,55 @@ bool FKismetDebugUtilities::IsSingleStepping()
 		|| Data.TargetGraphStackDepth != INDEX_NONE; 
 }
 
+TArray<FBreakpoint>* FKismetDebugUtilities::GetBreakpoints(const UBlueprint* Blueprint)
+{
+	UBlueprintEditorSettings* Settings = GetMutableDefault<UBlueprintEditorSettings>();
+	check(Settings);
+	
+	FPerBlueprintSettings* Found = Settings->PerBlueprintSettings.Find(Blueprint->GetPathName());
+	if(Found)
+	{
+		check(!Found->Breakpoints.IsEmpty());
+		return &Found->Breakpoints;
+	}
+	return nullptr;
+}
+
+void FKismetDebugUtilities::SaveBreakpoints()
+{
+	UBlueprintEditorSettings* Settings = GetMutableDefault<UBlueprintEditorSettings>();
+	check(Settings);
+	Settings->SaveConfig();
+}
+
+void FKismetDebugUtilities::CleanupBreakpoints(const UBlueprint* Blueprint)
+{
+	RemoveBreakpointsByPredicate(
+		Blueprint,
+		[Blueprint](const FBreakpoint& Breakpoint)
+		{
+			if (Breakpoint.GetLocation() == nullptr)
+			{
+				UE_LOG(LogBlueprintDebug, Display, TEXT("Encountered a blueprint breakpoint in %s without an associated node. The blueprint breakpoint has been removed"), *Blueprint->GetPathName());
+				return true;
+			}
+			return false;
+		}
+	);
+}
+
 //////////////////////////////////////////////////////////////////////////
 // Breakpoint
 
 // Is the node a valid breakpoint target? (i.e., the node is impure and ended up generating code)
-bool FKismetDebugUtilities::IsBreakpointValid(UBreakpoint* Breakpoint)
+bool FKismetDebugUtilities::IsBreakpointValid(const FBreakpoint& Breakpoint)
 {
-	check(Breakpoint);
-
 	// Breakpoints on impure nodes in a macro graph are always considered valid
-	UBlueprint* Blueprint = Cast<UBlueprint>(Breakpoint->GetOuter());
-	if (Blueprint && Blueprint->BlueprintType == BPTYPE_MacroLibrary)
+	UK2Node* K2Node = Cast<UK2Node>(Breakpoint.GetLocation());
+	if (K2Node)
 	{
-		UK2Node* K2Node = Cast<UK2Node>(Breakpoint->Node);
-		if (K2Node)
+		UBlueprint* Blueprint = Cast<UBlueprint>(K2Node->GetOuter()->GetOuter());
+		if (Blueprint && Blueprint->BlueprintType == BPTYPE_MacroLibrary)
 		{
 			return K2Node->IsA<UK2Node_MacroInstance>()
 				|| (!K2Node->IsNodePure() && !K2Node->IsA<UK2Node_Tunnel>());
@@ -805,88 +841,103 @@ bool FKismetDebugUtilities::IsBreakpointValid(UBreakpoint* Breakpoint)
 	}
 
 	TArray<uint8*> InstallSites;
-	FKismetDebugUtilities::GetBreakpointInstallationSites(Breakpoint, InstallSites);
+	GetBreakpointInstallationSites(Breakpoint, InstallSites);
 	return InstallSites.Num() > 0;
 }
 
 // Set the node that the breakpoint should focus on
-void FKismetDebugUtilities::SetBreakpointLocation(UBreakpoint* Breakpoint, UEdGraphNode* NewNode)
+void FKismetDebugUtilities::SetBreakpointLocation(FBreakpoint& Breakpoint, UEdGraphNode* NewNode)
 {
-	if (NewNode != Breakpoint->Node)
+	if (NewNode != Breakpoint.Node)
 	{
 		// Uninstall it from the old site if needed
-		FKismetDebugUtilities::SetBreakpointInternal(Breakpoint, false);
+		SetBreakpointInternal(Breakpoint, false);
 
 		// Make the new site accurate
-		Breakpoint->Node = NewNode;
-		FKismetDebugUtilities::SetBreakpointInternal(Breakpoint, Breakpoint->bEnabled);
+		Breakpoint.Node = NewNode;
+		SetBreakpointInternal(Breakpoint, Breakpoint.bEnabled);
 	}
 }
 
-// Set or clear the enabled flag for the breakpoint
-void FKismetDebugUtilities::SetBreakpointEnabled(UBreakpoint* Breakpoint, bool bIsEnabled)
+void FKismetDebugUtilities::SetBreakpointEnabled(FBreakpoint& Breakpoint, bool bIsEnabled)
 {
-	if (Breakpoint->bStepOnce && !bIsEnabled)
+	if (Breakpoint.bStepOnce && !bIsEnabled)
 	{
 		// Want to be disabled, but the single-stepping is keeping it enabled
 		bIsEnabled = true;
-		Breakpoint->bStepOnce_WasPreviouslyDisabled = true;
+		Breakpoint.bStepOnce_WasPreviouslyDisabled = true;
 	}
 
-	Breakpoint->bEnabled = bIsEnabled;
-	FKismetDebugUtilities::SetBreakpointInternal(Breakpoint, Breakpoint->bEnabled);
+	Breakpoint.bEnabled = bIsEnabled;
+	SetBreakpointInternal(Breakpoint, Breakpoint.bEnabled);
+	SaveBreakpoints();
+}
+
+// Set or clear the enabled flag for the breakpoint
+void FKismetDebugUtilities::SetBreakpointEnabled(const UEdGraphNode* OwnerNode, const UBlueprint* OwnerBlueprint, bool bIsEnabled)
+{
+	if(FBreakpoint* Breakpoint = FindBreakpointForNode(OwnerNode, OwnerBlueprint))
+	{
+		SetBreakpointEnabled(*Breakpoint, bIsEnabled);
+	}
 }
 
 // Sets this breakpoint up as a single-step breakpoint (will disable or delete itself after one go if the breakpoint wasn't already enabled)
-void FKismetDebugUtilities::SetBreakpointEnabledForSingleStep(UBreakpoint* Breakpoint, bool bDeleteAfterStep)
+void FKismetDebugUtilities::SetBreakpointEnabledForSingleStep(FBreakpoint& Breakpoint, bool bDeleteAfterStep)
 {
-	Breakpoint->bStepOnce = true;
-	Breakpoint->bStepOnce_RemoveAfterHit = bDeleteAfterStep;
-	Breakpoint->bStepOnce_WasPreviouslyDisabled = !Breakpoint->bEnabled;
+	Breakpoint.bStepOnce = true;
+	Breakpoint.bStepOnce_RemoveAfterHit = bDeleteAfterStep;
+	Breakpoint.bStepOnce_WasPreviouslyDisabled = !Breakpoint.bEnabled;
 
-	FKismetDebugUtilities::SetBreakpointEnabled(Breakpoint, true);
+	SetBreakpointEnabled(Breakpoint, true);
 }
 
-void FKismetDebugUtilities::ReapplyBreakpoint(UBreakpoint* Breakpoint)
+void FKismetDebugUtilities::ReapplyBreakpoint(FBreakpoint& Breakpoint)
 {
-	FKismetDebugUtilities::SetBreakpointInternal(Breakpoint, Breakpoint->IsEnabled());
+	SetBreakpointInternal(Breakpoint, Breakpoint.IsEnabled());
 }
 
-void FKismetDebugUtilities::StartDeletingBreakpoint(UBreakpoint* Breakpoint, UBlueprint* OwnerBlueprint)
+void FKismetDebugUtilities::RemoveBreakpointFromNode(const UEdGraphNode* OwnerNode, const UBlueprint* OwnerBlueprint)
 {
 #if WITH_EDITORONLY_DATA
-	checkSlow(OwnerBlueprint->Breakpoints.Contains(Breakpoint));
-	OwnerBlueprint->Breakpoints.Remove(Breakpoint);
-	OwnerBlueprint->MarkPackageDirty();
-
-	FKismetDebugUtilities::SetBreakpointLocation(Breakpoint, NULL);
+	RemoveBreakpointsByPredicate(
+		OwnerBlueprint,
+		[OwnerNode](const FBreakpoint& Breakpoint)
+		{
+			return Breakpoint.GetLocation() == OwnerNode;
+		}
+	);
+	SaveBreakpoints();
 #endif	//#if WITH_EDITORONLY_DATA
 }
 
 // Update the internal state of the breakpoint when it got hit
-void FKismetDebugUtilities::UpdateBreakpointStateWhenHit(UBreakpoint* Breakpoint, UBlueprint* OwnerBlueprint)
+void FKismetDebugUtilities::UpdateBreakpointStateWhenHit(const UEdGraphNode* OwnerNode, const UBlueprint* OwnerBlueprint)
 {
-	// Handle single-step breakpoints
-	if (Breakpoint->bStepOnce)
+	if (FBreakpoint* Breakpoint = FindBreakpointForNode(OwnerNode, OwnerBlueprint))
 	{
-		Breakpoint->bStepOnce = false;
-
-		if (Breakpoint->bStepOnce_RemoveAfterHit)
+		// Handle single-step breakpoints
+		if (Breakpoint->bStepOnce)
 		{
-			FKismetDebugUtilities::StartDeletingBreakpoint(Breakpoint, OwnerBlueprint);
-		}
-		else if (Breakpoint->bStepOnce_WasPreviouslyDisabled)
-		{
-			FKismetDebugUtilities::SetBreakpointEnabled(Breakpoint, false);
+			Breakpoint->bStepOnce = false;
+
+			if (Breakpoint->bStepOnce_RemoveAfterHit)
+			{
+				RemoveBreakpointFromNode(Breakpoint->GetLocation(), OwnerBlueprint);
+			}
+			else if (Breakpoint->bStepOnce_WasPreviouslyDisabled)
+			{
+				SetBreakpointEnabled(*Breakpoint, false);
+			}
 		}
 	}
 }
 
 // Install/uninstall the breakpoint into/from the script code for the generated class that contains the node
-void FKismetDebugUtilities::SetBreakpointInternal(UBreakpoint* Breakpoint, bool bShouldBeEnabled)
+void FKismetDebugUtilities::SetBreakpointInternal(FBreakpoint& Breakpoint, bool bShouldBeEnabled)
 {
 	TArray<uint8*> InstallSites;
-	FKismetDebugUtilities::GetBreakpointInstallationSites(Breakpoint, InstallSites);
+	GetBreakpointInstallationSites(Breakpoint, InstallSites);
 
 	for (int i = 0; i < InstallSites.Num(); ++i)
 	{
@@ -898,21 +949,21 @@ void FKismetDebugUtilities::SetBreakpointInternal(UBreakpoint* Breakpoint, bool
 }
 
 // Returns the installation site(s); don't cache these pointers!
-void FKismetDebugUtilities::GetBreakpointInstallationSites(UBreakpoint* Breakpoint, TArray<uint8*>& InstallSites)
+void FKismetDebugUtilities::GetBreakpointInstallationSites(const FBreakpoint& Breakpoint, TArray<uint8*>& InstallSites)
 {
 	InstallSites.Empty();
 
 #if WITH_EDITORONLY_DATA
-	if (Breakpoint->Node != NULL)
+	if (Breakpoint.Node != NULL)
 	{
-		UBlueprint* Blueprint = FBlueprintEditorUtils::FindBlueprintForNode(Breakpoint->Node);
+		UBlueprint* Blueprint = FBlueprintEditorUtils::FindBlueprintForNode(Breakpoint.GetLocation());
 
 		if ((Blueprint != NULL) && (Blueprint->GeneratedClass != NULL))
 		{
 			if (UBlueprintGeneratedClass* Class = Cast<UBlueprintGeneratedClass>(*Blueprint->GeneratedClass))
 			{
 				// Find the insertion point from the debugging data
-				Class->GetDebugData().FindBreakpointInjectionSites(Breakpoint->Node, InstallSites);
+				Class->GetDebugData().FindBreakpointInjectionSites(Breakpoint.GetLocation(), InstallSites);
 			}
 		}
 	}
@@ -967,54 +1018,118 @@ void FKismetDebugUtilities::GetValidBreakpointLocations(const UK2Node_MacroInsta
 	}
 }
 
-// Finds a breakpoint for a given node if it exists, or returns NULL
-UBreakpoint* FKismetDebugUtilities::FindBreakpointForNode(UBlueprint* Blueprint, const UEdGraphNode* Node, bool bCheckSubLocations)
+void FKismetDebugUtilities::CreateBreakpoint(const UBlueprint* Blueprint, UEdGraphNode* Node, bool bIsEnabled)
+{
+	UBlueprintEditorSettings* Settings = GetMutableDefault<UBlueprintEditorSettings>();
+	check(Settings);
+	FPerBlueprintSettings &BlueprintSettings = Settings->PerBlueprintSettings.FindOrAdd(Blueprint->GetPathName());
+
+	// ensure that this node doesn't already contain a breakpoint
+	checkSlow(!BlueprintSettings.Breakpoints.ContainsByPredicate(
+		[Node](const FBreakpoint& Other)
+		{
+			return Other.Node == Node;
+		}
+	));
+	
+	BlueprintSettings.Breakpoints.Emplace();
+	SetBreakpointEnabled(BlueprintSettings.Breakpoints.Top(), bIsEnabled);
+	SetBreakpointLocation(BlueprintSettings.Breakpoints.Top(), Node);
+	SaveBreakpoints();
+}
+
+void FKismetDebugUtilities::ForeachBreakpoint(const UBlueprint* Blueprint,
+	const TFunctionRef<void(FBreakpoint&)> Task)
 {
-	// iterate backwards so we can remove invalid breakpoints as we go
-	for (int32 Index = Blueprint->Breakpoints.Num()-1; Index >= 0; --Index)
+	if(TArray<FBreakpoint>* Breakpoints = GetBreakpoints(Blueprint))
 	{
-		UBreakpoint* Breakpoint = Blueprint->Breakpoints[Index];
-		if (Breakpoint == nullptr)
+		for(FBreakpoint& Breakpoint : *Breakpoints)
 		{
-			Blueprint->Breakpoints.RemoveAtSwap(Index);
-			Blueprint->MarkPackageDirty();
-			UE_LOG(LogBlueprintDebug, Warning, TEXT("Encountered an invalid blueprint breakpoint in %s (this should not happen... if you know how your blueprint got in this state, then please notify the Engine-Blueprints team)"), *Blueprint->GetPathName());
-			continue;
+			Task(Breakpoint);
 		}
+	}
+}
 
-		const UEdGraphNode* BreakpointLocation = Breakpoint->GetLocation();
-		if (BreakpointLocation == nullptr)
+void FKismetDebugUtilities::RemoveBreakpointsByPredicate(const UBlueprint* Blueprint,
+                                                         const TFunctionRef<bool(const FBreakpoint&)> Predicate)
+{
+	if(TArray<FBreakpoint>* Breakpoints = GetBreakpoints(Blueprint))
+	{
+		// notify the debugger of the breakpoints being removed
+		for(FBreakpoint& Breakpoint : *Breakpoints)
 		{
-			Blueprint->Breakpoints.RemoveAtSwap(Index);
-			Blueprint->MarkPackageDirty();
-			UE_LOG(LogBlueprintDebug, Display, TEXT("Encountered a blueprint breakpoint in %s without an associated node. The blueprint breakpoint has been removed"), *Blueprint->GetPathName());
-			continue;
+			if(Predicate(Breakpoint))
+			{
+				SetBreakpointLocation(Breakpoint, nullptr);
+			}
 		}
 
-		// Return this breakpoint if the location matches the given node
-		if (BreakpointLocation == Node)
+		// remove the breakpoints from the data
+		if(Breakpoints->RemoveAllSwap(Predicate, false))
 		{
-			return Breakpoint;
+			if(Breakpoints->IsEmpty())
+			{
+				// keeps the ini file clean by removing empty arrays
+				ClearBreakpoints(Blueprint);
+			}
+			SaveBreakpoints();
 		}
-		else if (bCheckSubLocations)
+	}
+}
+
+FBreakpoint* FKismetDebugUtilities::FindBreakpointByPredicate(const UBlueprint* Blueprint,
+                                                              const TFunctionRef<bool(const FBreakpoint&)> Predicate)
+{
+	if(TArray<FBreakpoint>* Breakpoints = GetBreakpoints(Blueprint))
+	{
+		for(FBreakpoint& Breakpoint : *Breakpoints)
 		{
-			// If this breakpoint is set on a macro instance node, check the set of valid breakpoint locations. If we find a
-			// match in the returned set, return the breakpoint that's set on the macro instance node. This allows breakpoints
-			// to be set and hit on macro instance nodes contained in a macro graph that will be expanded during compile.
-			const UK2Node_MacroInstance* MacroInstanceNode = Cast<UK2Node_MacroInstance>(BreakpointLocation);
-			if (MacroInstanceNode)
+			if(Predicate(Breakpoint))
 			{
-				TArray<const UEdGraphNode*> ValidBreakpointLocations;
-				GetValidBreakpointLocations(MacroInstanceNode, ValidBreakpointLocations);
-				if (ValidBreakpointLocations.Contains(Node))
-				{
-					return Breakpoint;
-				}
+				return &Breakpoint;
 			}
 		}
 	}
+	return nullptr;
+}
 
-	return NULL;
+// Finds a breakpoint for a given node if it exists, or returns NULL
+FBreakpoint* FKismetDebugUtilities::FindBreakpointForNode(const UEdGraphNode* OwnerNode, const UBlueprint* OwnerBlueprint,
+                                                          bool bCheckSubLocations)
+{
+	// remove expired data from deleted nodes and such
+	CleanupBreakpoints(OwnerBlueprint);
+
+	// find breakpoint
+	return FindBreakpointByPredicate(
+		OwnerBlueprint,
+		[OwnerNode, bCheckSubLocations](const FBreakpoint &Breakpoint)
+		{
+			UEdGraphNode* BreakpointLoaction = Breakpoint.GetLocation();
+			// Return this breakpoint if the location matches the given node
+			if (BreakpointLoaction == OwnerNode)
+			{
+				return true;
+			}
+			if (bCheckSubLocations)
+			{
+				// If this breakpoint is set on a macro instance node, check the set of valid breakpoint locations. If we find a
+				// match in the returned set, return the breakpoint that's set on the macro instance node. This allows breakpoints
+				// to be set and hit on macro instance nodes contained in a macro graph that will be expanded during compile.
+				const UK2Node_MacroInstance* MacroInstanceNode = Cast<UK2Node_MacroInstance>(BreakpointLoaction);
+				if (MacroInstanceNode)
+				{
+					TArray<const UEdGraphNode*> ValidBreakpointLocations;
+					GetValidBreakpointLocations(MacroInstanceNode, ValidBreakpointLocations);
+					if (ValidBreakpointLocations.Contains(OwnerNode))
+					{
+						return true;
+					}
+				}
+			}
+			return false;
+		}
+	);
 }
 
 bool FKismetDebugUtilities::HasDebuggingData(const UBlueprint* Blueprint)
@@ -1025,6 +1140,11 @@ bool FKismetDebugUtilities::HasDebuggingData(const UBlueprint* Blueprint)
 //////////////////////////////////////////////////////////////////////////
 // Blueprint utils
 
+bool FKismetDebugUtilities::BlueprintHasBreakpoints(const UBlueprint* Blueprint)
+{
+	return GetBreakpoints(Blueprint) != nullptr;
+}
+
 // Looks thru the debugging data for any class variables associated with the node
 FProperty* FKismetDebugUtilities::FindClassPropertyForPin(UBlueprint* Blueprint, const UEdGraphPin* Pin)
 {
@@ -1057,16 +1177,27 @@ FProperty* FKismetDebugUtilities::FindClassPropertyForNode(UBlueprint* Blueprint
 }
 
 
-void FKismetDebugUtilities::ClearBreakpoints(UBlueprint* Blueprint)
+void FKismetDebugUtilities::ClearBreakpoints(const UBlueprint* Blueprint)
+{
+	ClearBreakpointsForPath(Blueprint->GetPathName());
+}
+
+
+void FKismetDebugUtilities::ClearBreakpointsForPath(const FString& BlueprintPath)
 {
-	for (int32 BreakpointIndex = 0; BreakpointIndex < Blueprint->Breakpoints.Num(); ++BreakpointIndex)
+	UBlueprintEditorSettings* Settings = GetMutableDefault<UBlueprintEditorSettings>();
+	check(Settings);
+
+	FPerBlueprintSettings Collection;
+	if (Settings->PerBlueprintSettings.RemoveAndCopyValue(BlueprintPath, Collection))
 	{
-		UBreakpoint* Breakpoint = Blueprint->Breakpoints[BreakpointIndex];
-		FKismetDebugUtilities::SetBreakpointLocation(Breakpoint, NULL);
+		for(FBreakpoint& Breakpoint : Collection.Breakpoints)
+		{
+			// notify debugger that the breakpont has been removed
+			SetBreakpointLocation(Breakpoint, nullptr);	
+		}
+		SaveBreakpoints();
 	}
-
-	Blueprint->Breakpoints.Empty();
-	Blueprint->MarkPackageDirty();
 }
 
 FKismetDebugUtilities::FOnWatchedPinsListChanged FKismetDebugUtilities::WatchedPinsListChangedEvent;
diff --git a/Engine/Source/Editor/UnrealEd/Public/Kismet2/Breakpoint.h b/Engine/Source/Editor/UnrealEd/Public/Kismet2/Breakpoint.h
new file mode 100644
--- /dev/null
+++ b/Engine/Source/Editor/UnrealEd/Public/Kismet2/Breakpoint.h
@@ -0,0 +1,65 @@
+// Copyright Epic Games, Inc. All Rights Reserved.
+
+#pragma once
+
+#include "CoreMinimal.h"
+#include "Breakpoint.generated.h"
+
+class UEdGraphNode;
+
+USTRUCT()
+struct UNREALED_API FBreakpoint
+{
+	GENERATED_BODY()
+
+private:
+	// Is the breakpoint currently enabled?
+	UPROPERTY()
+	uint8 bEnabled:1;
+
+	// Node that the breakpoint is placed on
+	UPROPERTY()
+	TSoftObjectPtr<class UEdGraphNode> Node = nullptr;
+
+	// Is this breakpoint auto-generated, and should be removed when next hit?
+	UPROPERTY()
+	uint8 bStepOnce:1;
+
+	UPROPERTY()
+	uint8 bStepOnce_WasPreviouslyDisabled:1;
+
+	UPROPERTY()
+	uint8 bStepOnce_RemoveAfterHit:1;
+
+public:
+	FBreakpoint();
+
+	/** Get the target node for the breakpoint */
+	UEdGraphNode* GetLocation() const
+	{
+		return Node.Get();
+	}
+
+	/** @returns true if the breakpoint should be enabled when debugging */
+	bool IsEnabled() const
+	{
+		return bEnabled;
+	}
+
+	/** @returns true if the user wanted the breakpoint enabled?  (single stepping, etc... could result in the breakpoint being temporarily enabled) */
+	bool IsEnabledByUser() const
+	{
+		return bEnabled && !(bStepOnce && bStepOnce_WasPreviouslyDisabled);
+	}
+
+	bool operator==(const FBreakpoint& Other) const
+	{
+		return Node == Other.Node;
+	}
+
+	/** Gets a string that describes the location */
+	FText GetLocationDescription() const;
+
+	friend class FKismetDebugUtilities;
+};
+
diff --git a/Engine/Source/Editor/UnrealEd/Public/Kismet2/KismetDebugUtilities.h b/Engine/Source/Editor/UnrealEd/Public/Kismet2/KismetDebugUtilities.h
--- a/Engine/Source/Editor/UnrealEd/Public/Kismet2/KismetDebugUtilities.h
+++ b/Engine/Source/Editor/UnrealEd/Public/Kismet2/KismetDebugUtilities.h
@@ -3,6 +3,7 @@
 #pragma once
 
 #include "CoreMinimal.h"
+
 #include "UObject/Object.h"
 #include "UObject/Class.h"
 #include "Engine/Selection.h"
@@ -10,7 +11,7 @@
 static_assert(DO_BLUEPRINT_GUARD, "KismetDebugUtilities assumes BP exception tracking is enabled");
 
 class UBlueprint;
-class UBreakpoint;
+struct FBreakpoint;
 template<typename ElementType> class TSimpleRingBuffer;
 
 DECLARE_LOG_CATEGORY_EXTERN(LogBlueprintDebug, Log, All);
@@ -165,35 +166,63 @@ public:
 	// Breakpoint utils
 
 	/** Is the node a valid breakpoint target? (i.e., the node is impure and ended up generating code) */
-	static bool IsBreakpointValid(UBreakpoint* Breakpoint);
+	static bool IsBreakpointValid(const FBreakpoint& Breakpoint);
 
 	/** Set the node that the breakpoint should focus on */
-	static void SetBreakpointLocation(UBreakpoint* Breakpoint, UEdGraphNode* NewNode);
+	static void SetBreakpointLocation(FBreakpoint& Breakpoint, UEdGraphNode* NewNode);
 
 	/** Set or clear the enabled flag for the breakpoint */
-	static void SetBreakpointEnabled(UBreakpoint* Breakpoint, bool bIsEnabled);
+	static void SetBreakpointEnabled(FBreakpoint& Breakpoint, bool bIsEnabled);
+	static void SetBreakpointEnabled(const UEdGraphNode* OwnerNode, const UBlueprint* OwnerBlueprint, bool bIsEnabled);
 
 	/** Sets this breakpoint up as a single-step breakpoint (will disable or delete itself after one go if the breakpoint wasn't already enabled) */
-	static void SetBreakpointEnabledForSingleStep(UBreakpoint* Breakpoint, bool bDeleteAfterStep);
+	static void SetBreakpointEnabledForSingleStep(FBreakpoint& Breakpoint, bool bDeleteAfterStep);
 
 	/** Reapplies the breakpoint (used after recompiling to ensure it is set if needed) */
-	static void ReapplyBreakpoint(UBreakpoint* Breakpoint);
+	static void ReapplyBreakpoint(FBreakpoint& Breakpoint);
 
 	/** Start the process of deleting this breakpoint */
-	static void StartDeletingBreakpoint(UBreakpoint* Breakpoint, UBlueprint* OwnerBlueprint);
+	static void RemoveBreakpointFromNode(const UEdGraphNode* OwnerNode, const UBlueprint* OwnerBlueprint);
 
 	/** Update the internal state of the breakpoint when it got hit */
-	static void UpdateBreakpointStateWhenHit(UBreakpoint* Breakpoint, UBlueprint* OwnerBlueprint);
+	static void UpdateBreakpointStateWhenHit(const UEdGraphNode* OwnerNode, const UBlueprint* OwnerBlueprint);
 
 	/** Returns the installation site(s); don't cache these pointers! */
-	static void GetBreakpointInstallationSites(UBreakpoint* Breakpoint, TArray<uint8*>& InstallSites);
+	static void GetBreakpointInstallationSites(const FBreakpoint& Breakpoint, TArray<uint8*>& InstallSites);
 
 	/** Install/uninstall the breakpoint into/from the script code for the generated class that contains the node */
-	static void SetBreakpointInternal(UBreakpoint* Breakpoint, bool bShouldBeEnabled);
+	static void SetBreakpointInternal(FBreakpoint& Breakpoint, bool bShouldBeEnabled);
 
 	/** Returns the set of valid macro source node breakpoint location(s) for the given macro instance node. The set may be empty. */
 	static void GetValidBreakpointLocations(const class UK2Node_MacroInstance* MacroInstanceNode, TArray<const UEdGraphNode*>& BreakpointLocations);
 
+	/** Adds a breakpoint to the provided node */
+	static void CreateBreakpoint(const UBlueprint* Blueprint, UEdGraphNode* Node, bool bIsEnabled = true);
+
+	/**
+	 * Performs a task on every breakpoint in the provided blueprint
+	 * @param Blueprint The owning blueprint of the breakpoints to iterate
+	 * @param Task function to be called on every element
+	 */
+	static void ForeachBreakpoint(const UBlueprint* Blueprint, TFunctionRef<void(FBreakpoint &)> Task);
+
+	/**
+	* Removes any breakpoint that matches the provided predicate
+	* @param Blueprint The owning blueprint of the breakpoints to iterate
+	* @param Predicate function that returns true if a breakpoint should be removed
+	*/
+	static void RemoveBreakpointsByPredicate(const UBlueprint* Blueprint, const TFunctionRef<bool(const FBreakpoint&)> Predicate);
+
+	/**
+	* Returns the first breakpoint that matches the provided predicate or nullptr if nothing matched
+	* @param Blueprint The owning blueprint of the breakpoints to iterate
+	* @param Predicate function that returns true for the found breakpoint
+	*/
+	static FBreakpoint* FindBreakpointByPredicate(const UBlueprint* Blueprint, const TFunctionRef<bool(const FBreakpoint&)> Predicate);
+
+	/** Queries whether a blueprint has breakpoints in it */
+	static bool BlueprintHasBreakpoints(const UBlueprint* Blueprint);
+	
 	// Blueprint utils 
 
 	// Looks thru the debugging data for any class variables associated with the pin
@@ -206,10 +235,11 @@ public:
 	static bool HasDebuggingData(const UBlueprint* Blueprint);
 
 	/** Returns the breakpoint associated with a node, or NULL */
-	static UBreakpoint* FindBreakpointForNode(UBlueprint* Blueprint, const UEdGraphNode* Node, bool bCheckSubLocations = false);
+	static FBreakpoint* FindBreakpointForNode(const UEdGraphNode* OwnerNode, const UBlueprint* OwnerBlueprint, bool bCheckSubLocations = false);
 
 	/** Deletes all breakpoints in this blueprint */
-	static void ClearBreakpoints(UBlueprint* Blueprint);
+	static void ClearBreakpoints(const UBlueprint* Blueprint);
+	static void ClearBreakpointsForPath(const FString &BlueprintPath);
 
 	// Notifies listeners when a watched pin is added or removed
 	static FOnWatchedPinsListChanged WatchedPinsListChangedEvent;
@@ -267,7 +297,16 @@ protected:
 	* @return	EWTR_Valid if the debug data could be found, otherwise an appropriate error code
 	*/
 	static EWatchTextResult FindDebuggingData(UBlueprint* Blueprint, UObject* ActiveObject, const UEdGraphPin* WatchPin, FProperty*& OutProperty, void*& OutData, void*& OutDelta, UObject*& OutParent, TArray<UObject*>& SeenObjects, bool* bOutIsDirectPtr = nullptr);
+	
+
+	/**	Retrieve the Array of breakpoints associated with a blueprint.
+	*	returns null if there are no breakpoints associated with this blueprint */
+	static TArray<FBreakpoint>* GetBreakpoints(const UBlueprint* Blueprint);
+	
+	/** Save any modifications made to breakpoints */
+	static void SaveBreakpoints();
 
+	static void CleanupBreakpoints(const UBlueprint* Blueprint);
 private:
 	FKismetDebugUtilities() {}
 };
diff --git a/Engine/Source/Runtime/Engine/Classes/Engine/Blueprint.h b/Engine/Source/Runtime/Engine/Classes/Engine/Blueprint.h
--- a/Engine/Source/Runtime/Engine/Classes/Engine/Blueprint.h
+++ b/Engine/Source/Runtime/Engine/Classes/Engine/Blueprint.h
@@ -630,10 +630,6 @@ PRAGMA_ENABLE_DEPRECATION_WARNINGS
 	UPROPERTY()
 	TArray<FBPEditorBookmarkNode> BookmarkNodes;
 
-	/** Persistent debugging options */
-	UPROPERTY()
-	TArray<TObjectPtr<class UBreakpoint>> Breakpoints;
-
 	UPROPERTY()
 	TArray<FEdGraphPinReference> WatchedPins;
 
diff --git a/Engine/Source/Runtime/Engine/Classes/Engine/Breakpoint.h b/Engine/Source/Runtime/Engine/Classes/Engine/Breakpoint.h
deleted file mode 100644
diff --git a/Engine/Source/Runtime/Engine/Private/Blueprint.cpp b/Engine/Source/Runtime/Engine/Private/Blueprint.cpp
--- a/Engine/Source/Runtime/Engine/Private/Blueprint.cpp
+++ b/Engine/Source/Runtime/Engine/Private/Blueprint.cpp
@@ -16,7 +16,6 @@
 #include "Misc/SecureHash.h"
 #include "Engine/BlueprintGeneratedClass.h"
 #include "EdGraph/EdGraph.h"
-#include "Engine/Breakpoint.h"
 #include "Components/TimelineComponent.h"
 #include "Modules/ModuleManager.h"
 #include "UObject/TextProperty.h"
@@ -45,9 +44,15 @@
 #include "BlueprintAssetHandler.h"
 #include "Blueprint/BlueprintExtension.h"
 #include "UObject/TextProperty.h"
+#include "Kismet2/Breakpoint.h"
 #endif
+
 #include "Engine/InheritableComponentHandler.h"
 
+#if WITH_EDITORONLY_DATA
+#include "Kismet2/KismetDebugUtilities.h"
+#endif
+
 DEFINE_LOG_CATEGORY(LogBlueprint);
 
 //////////////////////////////////////////////////////////////////////////
@@ -700,17 +705,17 @@ void UBlueprint::PostLoad()
 	// Purge any NULL graphs
 	FBlueprintEditorUtils::PurgeNullGraphs(this);
 
+#if WITH_EDITOR
 	// Remove stale breakpoints
-	for (int32 i = 0; i < Breakpoints.Num(); ++i)
-	{
-		UBreakpoint* Breakpoint = Breakpoints[i];
-		const UEdGraphNode* const Location = Breakpoint ? Breakpoint->GetLocation() : nullptr;
-		if (!Location || !Location->IsIn(this))
+	FKismetDebugUtilities::RemoveBreakpointsByPredicate(
+		this,
+		[this](const FBreakpoint& Breakpoint)
 		{
-			Breakpoints.RemoveAt(i);
-			--i;
+			const UEdGraphNode* const Location = Breakpoint.GetLocation();
+			return !Location || !Location->IsIn(this);
 		}
-	}
+	);
+# endif
 
 	// Make sure we have an SCS and ensure it's transactional
 	if( FBlueprintEditorUtils::SupportsConstructionScript(this) )
diff --git a/Engine/Source/Runtime/Engine/Private/Breakpoint.cpp b/Engine/Source/Runtime/Engine/Private/Breakpoint.cpp
deleted file mode 100644
diff --git a/Engine/Source/Runtime/Engine/Public/EngineFwd.h b/Engine/Source/Runtime/Engine/Public/EngineFwd.h
--- a/Engine/Source/Runtime/Engine/Public/EngineFwd.h
+++ b/Engine/Source/Runtime/Engine/Public/EngineFwd.h
@@ -40,7 +40,6 @@ class UBlendSpace;
 class UBlueprint;
 class UBlueprintGeneratedClass;
 class UBodySetup;
-class UBreakpoint;
 class UCameraComponent;
 class UCanvas;
 class UChannel;
