Fixes crash in editor viewport when using temporal upscaler plugin

diff --git a/Engine/Source/Editor/UnrealEd/Private/EditorViewportClient.cpp b/Engine/Source/Editor/UnrealEd/Private/EditorViewportClient.cpp
--- a/Engine/Source/Editor/UnrealEd/Private/EditorViewportClient.cpp
+++ b/Engine/Source/Editor/UnrealEd/Private/EditorViewportClient.cpp
@@ -3837,7 +3837,7 @@ void FEditorViewportClient::Draw(FViewport* InViewport, FCanvas* Canvas)
 		}
 
 		// Setup custom upscaler and screen percentage.
-		if (GCustomEditorStaticScreenPercentage && ViewFamily.ViewMode == EViewModeIndex::VMI_Lit)
+		if (GCustomEditorStaticScreenPercentage && ViewFamily.ViewMode == EViewModeIndex::VMI_Lit && ViewFamily.Scene->GetShadingPath() == EShadingPath::Deferred && ViewFamily.bRealtimeUpdate)
 		{
 			GCustomEditorStaticScreenPercentage->SetupEditorViewFamily(ViewFamily, this);
 		}
diff --git a/Engine/Source/Runtime/Engine/Private/GameViewportClient.cpp b/Engine/Source/Runtime/Engine/Private/GameViewportClient.cpp
--- a/Engine/Source/Runtime/Engine/Private/GameViewportClient.cpp
+++ b/Engine/Source/Runtime/Engine/Private/GameViewportClient.cpp
@@ -1409,7 +1409,7 @@ void UGameViewportClient::Draw(FViewport* InViewport, FCanvas* SceneCanvas)
 		}
 		#endif
 
-		if (GCustomStaticScreenPercentage && ViewFamily.ViewMode == EViewModeIndex::VMI_Lit && ViewFamily.Scene->GetShadingPath() == EShadingPath::Deferred)
+		if (GCustomStaticScreenPercentage && ViewFamily.ViewMode == EViewModeIndex::VMI_Lit && ViewFamily.Scene->GetShadingPath() == EShadingPath::Deferred && ViewFamily.bRealtimeUpdate)
 		{
 			GCustomStaticScreenPercentage->SetupMainGameViewFamily(ViewFamily);
 		}
diff --git a/Engine/Source/Runtime/Engine/Private/SceneView.cpp b/Engine/Source/Runtime/Engine/Private/SceneView.cpp
--- a/Engine/Source/Runtime/Engine/Private/SceneView.cpp
+++ b/Engine/Source/Runtime/Engine/Private/SceneView.cpp
@@ -933,7 +933,7 @@ void FSceneView::SetupAntiAliasingMethod()
 		}
 
 		// Overides the anti aliasing method to temporal AA when using a custom temporal upscaler.
-		if (Family->GetTemporalUpscalerInterface() != nullptr && Family->bRealtimeUpdate)
+		if (Family->GetTemporalUpscalerInterface() != nullptr)
 		{
 			AntiAliasingMethod = AAM_TemporalAA;
 		}
