Removed debug code submitted by error

diff --git a/Engine/Source/Editor/UnrealEd/Private/WorldPartition/WorldPartitionHLODsBuilder.cpp b/Engine/Source/Editor/UnrealEd/Private/WorldPartition/WorldPartitionHLODsBuilder.cpp
--- a/Engine/Source/Editor/UnrealEd/Private/WorldPartition/WorldPartitionHLODsBuilder.cpp
+++ b/Engine/Source/Editor/UnrealEd/Private/WorldPartition/WorldPartitionHLODsBuilder.cpp
@@ -459,12 +459,7 @@ bool UWorldPartitionHLODsBuilder::BuildHLODActors()
 
 		UE_LOG(LogWorldPartitionHLODsBuilder, Display, TEXT("    [%d/%d] Building HLOD actor %s..."), CurrentActor, HLODActorsToBuild.Num(), *HLODActor->GetActorLabel());
 
-		if (HLODActor->GetSubActorsHLODLayer()->GetLayerType() != EHLODLayerType::Instancing)
-		{
-			continue;
-		}
-
-		HLODActor->BuildHLOD(true);
+		HLODActor->BuildHLOD();
 
 		UPackage* ActorPackage = HLODActor->GetPackage();
 		if (ActorPackage->IsDirty())
