Store EditorDomain in local storage only, since we do not yet handle downloading it in batch from remote.

diff --git a/Engine/Source/Editor/UnrealEd/Private/EditorDomain/EditorDomainUtils.cpp b/Engine/Source/Editor/UnrealEd/Private/EditorDomain/EditorDomainUtils.cpp
--- a/Engine/Source/Editor/UnrealEd/Private/EditorDomain/EditorDomainUtils.cpp
+++ b/Engine/Source/Editor/UnrealEd/Private/EditorDomain/EditorDomainUtils.cpp
@@ -345,6 +345,8 @@ void FMemoryPackageStoreWriter::WriteBulkdata(const FBulkDataInfo& Info, const F
 
 bool TrySavePackage(UPackage* Package)
 {
+	using namespace UE::DerivedData;
+
 	FString ErrorMessage;
 	FPackageDigest PackageDigest;
 	EPackageDigestResult FindHashResult = GetPackageDigest(*IAssetRegistry::Get(), Package->GetFName(), PackageDigest,
@@ -413,16 +415,16 @@ bool TrySavePackage(UPackage* Package)
 		PackageBuffer = MakeSharedBufferFromArray(MoveTemp(TempBytes));
 	}
 
-	UE::DerivedData::ICache& Cache = GetDerivedDataCacheRef();
-	UE::DerivedData::FCacheRecordBuilder RecordBuilder = Cache.CreateRecord(GetEditorDomainPackageKey(PackageDigest));
+	ICache& Cache = GetDerivedDataCacheRef();
+	FCacheRecordBuilder RecordBuilder = Cache.CreateRecord(GetEditorDomainPackageKey(PackageDigest));
 	TCbWriter<256> MetaData;
 	MetaData.BeginObject();
 	MetaData << "FileSize" << PackageBuffer.GetSize();
 	MetaData.EndObject();
 	RecordBuilder.SetMeta(MetaData.Save().AsObject());
 	RecordBuilder.SetValue(PackageBuffer);
-	UE::DerivedData::FCacheRecord Record = RecordBuilder.Build();
-	Cache.Put(MakeArrayView(&Record, 1), Package->GetName());
+	FCacheRecord Record = RecordBuilder.Build();
+	Cache.Put(MakeArrayView(&Record, 1), Package->GetName(), ECachePolicy::Local);
 	return true;
 }
 
