Fix Actor Label refresh on Copy/Paste

diff --git a/Engine/Source/Editor/UnrealEd/Private/EditorActor.cpp b/Engine/Source/Editor/UnrealEd/Private/EditorActor.cpp
--- a/Engine/Source/Editor/UnrealEd/Private/EditorActor.cpp
+++ b/Engine/Source/Editor/UnrealEd/Private/EditorActor.cpp
@@ -376,6 +376,11 @@ void UUnrealEdEngine::PasteActors(TArray<AActor*>& OutPastedActors, UWorld* InWo
 			FActorLabelUtilities::SetActorLabelUnique(Actor, Actor->GetActorLabel(), &ActorLabels);
 			ActorLabels.Add(Actor->GetActorLabel());
 		}
+		else
+		{
+			// Make sure to broadcast label change after actor import as it may change between spawn and import
+			FCoreDelegates::OnActorLabelChanged.Broadcast(Actor);
+		}
 
 		LayersSubsystem->InitializeNewActorLayers(Actor);
 
