Multiple fixes to avoid dirtying Parent Actor package when Attaching/Detaching Children

diff --git a/Engine/Source/Editor/SceneOutliner/Private/ActorBrowsingMode.cpp b/Engine/Source/Editor/SceneOutliner/Private/ActorBrowsingMode.cpp
--- a/Engine/Source/Editor/SceneOutliner/Private/ActorBrowsingMode.cpp
+++ b/Engine/Source/Editor/SceneOutliner/Private/ActorBrowsingMode.cpp
@@ -1223,7 +1223,8 @@ void FActorBrowsingMode::OnDrop(ISceneOutlinerTreeItem& DropTarget, const FScene
 					if (RootComp && RootComp->GetAttachParent())
 					{
 						AActor* OldParent = RootComp->GetAttachParent()->GetOwner();
-						OldParent->Modify();
+						// Attachment is persisted on the child so modify both actors for Undo/Redo but do not mark the Parent package dirty
+						OldParent->Modify(/*bAlwaysMarkDirty=*/false);
 						RootComp->DetachFromComponent(FDetachmentTransformRules::KeepWorldTransform);
 
 						DragActor->SetFolderPath_Recursively(OldParent->GetFolderPath());
@@ -1351,7 +1352,8 @@ void FActorBrowsingMode::OnDrop(ISceneOutlinerTreeItem& DropTarget, const FScene
 				{
 					if (AActor* OldParentActor = RootComp->GetAttachParent()->GetOwner())
 					{
-						OldParentActor->Modify();
+						// Attachment is persisted on the child so modify both actors for Undo/Redo but do not mark the Parent package dirty
+						OldParentActor->Modify(/*bAlwaysMarkDirty=*/false);
 					}
 					RootComp->DetachFromComponent(FDetachmentTransformRules::KeepWorldTransform);
 				}
diff --git a/Engine/Source/Editor/UnrealEd/Private/EditorEngine.cpp b/Engine/Source/Editor/UnrealEd/Private/EditorEngine.cpp
--- a/Engine/Source/Editor/UnrealEd/Private/EditorEngine.cpp
+++ b/Engine/Source/Editor/UnrealEd/Private/EditorEngine.cpp
@@ -3964,14 +3964,15 @@ void UEditorEngine::ParentActors( AActor* ParentActor, AActor* ChildActor, const
 
 		// modify parent and child
 		const FScopedTransaction Transaction( NSLOCTEXT("Editor", "UndoAction_PerformAttachment", "Attach actors") );
+		// Attachment is persisted on the child so modify both actors for Undo/Redo but do not mark the Parent package dirty
 		ChildActor->Modify();
-		ParentActor->Modify();
+		ParentActor->Modify(/*bAlwaysMarkDirty=*/false);
 
 		// If child is already attached to something, modify the old parent and detach
 		if(ChildRoot->GetAttachParent() != nullptr)
 		{
 			AActor* OldParentActor = ChildRoot->GetAttachParent()->GetOwner();
-			OldParentActor->Modify();
+			OldParentActor->Modify(/*bAlwaysMarkDirty=*/false);
 			ChildRoot->DetachFromComponent(FDetachmentTransformRules::KeepWorldTransform);
 
 			GEngine->BroadcastLevelActorDetached(ChildActor, OldParentActor);
@@ -3980,6 +3981,7 @@ void UEditorEngine::ParentActors( AActor* ParentActor, AActor* ChildActor, const
 		// If the parent is already attached to this child, modify its parent and detach so we can allow the attachment
 		if(ParentRoot->IsAttachedTo(ChildRoot))
 		{
+			// Here its ok to mark the parent package dirty as both Parent & Child need to be saved.
 			ParentRoot->GetAttachParent()->GetOwner()->Modify();
 			ParentRoot->DetachFromComponent(FDetachmentTransformRules::KeepWorldTransform);
 		}
@@ -4006,7 +4008,7 @@ bool UEditorEngine::DetachSelectedActors()
 		if( RootComp != nullptr && RootComp->GetAttachParent() != nullptr)
 		{
 			AActor* OldParentActor = RootComp->GetAttachParent()->GetOwner();
-			OldParentActor->Modify();
+			OldParentActor->Modify(false);
 			RootComp->DetachFromComponent(FDetachmentTransformRules::KeepWorldTransform);
 			bDetachOccurred = true;
 			Actor->SetFolderPath_Recursively(OldParentActor->GetFolderPath());
diff --git a/Engine/Source/Runtime/Engine/Private/Components/SceneComponent.cpp b/Engine/Source/Runtime/Engine/Private/Components/SceneComponent.cpp
--- a/Engine/Source/Runtime/Engine/Private/Components/SceneComponent.cpp
+++ b/Engine/Source/Runtime/Engine/Private/Components/SceneComponent.cpp
@@ -2173,7 +2173,8 @@ void USceneComponent::DetachFromComponent(const FDetachmentTransformRules& Detac
 		if (DetachmentRules.bCallModify && !HasAnyFlags(RF_Transient))
 		{
 			Modify();
-			GetAttachParent()->Modify();
+			// Attachment is persisted on the child so modify both actors for Undo/Redo but do not mark the Parent package dirty
+			GetAttachParent()->Modify(/*bAlwaysMarkDirty=*/false);
 		}
 
 		PrimaryComponentTick.RemovePrerequisite(GetAttachParent(), GetAttachParent()->PrimaryComponentTick); // no longer required to tick after the attachment
