Updated DDC Widget
Added detailed resource and cache backend information.
Added new indicators for Busy, Upload, Download, LocalCache and RemoteCache
Minor modifications to StudioAnalytics to re-use the stats gathering code

diff --git a/Engine/Source/Developer/DerivedDataCache/Private/DerivedDataCacheUsageStats.cpp b/Engine/Source/Developer/DerivedDataCache/Private/DerivedDataCacheUsageStats.cpp
new file mode 100644
--- /dev/null
+++ b/Engine/Source/Developer/DerivedDataCache/Private/DerivedDataCacheUsageStats.cpp
@@ -0,0 +1,122 @@
+// Copyright Epic Games, Inc. All Rights Reserved.
+
+#include "DerivedDataCacheUsageStats.h"
+#include "DerivedDataCacheInterface.h"
+#include "DerivedDataBackendInterface.h"
+
+DERIVEDDATACACHE_API void GatherDerivedDataCacheResourceStats(TArray<FDerivedDataCacheResourceStat>& DDCResourceStats)
+{
+#if ENABLE_COOK_STATS
+	// Temporary Set to allow accumulation by asset type
+	TSet<FDerivedDataCacheResourceStat, FDerivedDataCacheResourceStatKeyFuncs> DDCResourceStatsSet;
+	
+	/** this functor will take a collected cooker stat and log it out using some custom formatting based on known stats that are collected.. */
+	auto LogStatsFunc = [&DDCResourceStatsSet]
+	(const FString& StatName, const TArray<FCookStatsManager::StringKeyValue>& StatAttributes)
+	{
+		if (StatName.EndsWith(TEXT(".Usage"), ESearchCase::IgnoreCase))
+		{
+			// Anything that ends in .Usage is assumed to be an instance of FCookStats.FDDCResourceUsageStats. We'll log that using custom formatting.
+			FString AssetType = StatName;
+			AssetType.RemoveFromEnd(TEXT(".Usage"), ESearchCase::IgnoreCase);
+			// See if the asset has a subtype (found via the "Node" parameter")
+			const FCookStatsManager::StringKeyValue* AssetSubType = StatAttributes.FindByPredicate([](const FCookStatsManager::StringKeyValue& Item) { return Item.Key == TEXT("Node"); });
+			if (AssetSubType && AssetSubType->Value.Len() > 0)
+			{
+				AssetType += FString::Printf(TEXT(" (%s)"), *AssetSubType->Value);
+			}
+
+			if (AssetType.Contains("DDC"))
+				return;
+
+			int64 AssetCount = 0;
+
+			const FCookStatsManager::StringKeyValue* CountAttr = StatAttributes.FindByPredicate([](const FCookStatsManager::StringKeyValue& Item) { return Item.Key == TEXT("Count"); });
+
+			if (CountAttr)
+			{
+				LexFromString(AssetCount, *CountAttr->Value);
+			}
+
+			if (AssetCount == 0)
+				return;
+
+			// Pull the Time and Size attributes and AddOrAccumulate them into the set of stats. Ugly string/container manipulation code courtesy of UE4/C++.
+			const FCookStatsManager::StringKeyValue* AssetTimeSecAttr = StatAttributes.FindByPredicate([](const FCookStatsManager::StringKeyValue& Item) { return Item.Key == TEXT("TimeSec"); });
+			double AssetTimeSec = 0.0;
+
+			if (AssetTimeSecAttr)
+			{
+				LexFromString(AssetTimeSec, *AssetTimeSecAttr->Value);
+			}
+
+			const FCookStatsManager::StringKeyValue* AssetSizeMBAttr = StatAttributes.FindByPredicate([](const FCookStatsManager::StringKeyValue& Item) { return Item.Key == TEXT("MB"); });
+			double AssetSizeMB = 0.0;
+
+			if (AssetSizeMBAttr)
+			{
+				LexFromString(AssetSizeMB, *AssetSizeMBAttr->Value);
+			}
+
+			const FCookStatsManager::StringKeyValue* ThreadNameAttr = StatAttributes.FindByPredicate([](const FCookStatsManager::StringKeyValue& Item) { return Item.Key == TEXT("ThreadName"); });
+			bool bIsGameThreadTime = ThreadNameAttr != nullptr && ThreadNameAttr->Value == TEXT("GameThread");
+
+			const FCookStatsManager::StringKeyValue* HitOrMissAttr = StatAttributes.FindByPredicate([](const FCookStatsManager::StringKeyValue& Item) { return Item.Key == TEXT("HitOrMiss"); });
+			bool bWasMiss = HitOrMissAttr != nullptr && HitOrMissAttr->Value == TEXT("Miss");
+
+			const FCookStatsManager::StringKeyValue* Call = StatAttributes.FindByPredicate([](const FCookStatsManager::StringKeyValue& Item) { return Item.Key == TEXT("Call"); });
+			bool bWasGet = Call != nullptr && Call->Value == TEXT("Get");
+
+			FDerivedDataCacheResourceStat Stat(AssetType);
+
+			if (bWasMiss)
+			{
+				Stat = FDerivedDataCacheResourceStat(AssetType, bIsGameThreadTime, 0, 0, 0, AssetTimeSec, AssetSizeMB, AssetCount);
+			}
+			else
+			{
+				Stat = FDerivedDataCacheResourceStat(AssetType, bIsGameThreadTime, AssetTimeSec, AssetSizeMB, AssetCount, 0, 0, 0);
+			}
+
+			FDerivedDataCacheResourceStat* ExistingStat = DDCResourceStatsSet.Find(AssetType);
+
+			if (ExistingStat)
+			{
+				ExistingStat->Accumulate(Stat);
+			}
+			else
+			{
+				DDCResourceStatsSet.Add(Stat);
+			}
+		}
+	};
+
+	// Grab the DDC stats
+	FCookStatsManager::LogCookStats(LogStatsFunc);
+
+	// Place the Set members in the final Array
+	for (const FDerivedDataCacheResourceStat& Stat : DDCResourceStatsSet)
+	{
+		DDCResourceStats.Add(Stat);
+	}
+#endif
+}
+
+#if ENABLE_COOK_STATS
+DERIVEDDATACACHE_API void GatherDerivedDataCacheSummaryStats(TArray<FDerivedDataCacheSummaryStat>& DDCSummaryStats)
+{
+	/** this functor will take a collected cooker stat and log it out using some custom formatting based on known stats that are collected.. */
+	auto LogStatsFunc = [&DDCSummaryStats]
+	(const FString& StatName, const TArray<FCookStatsManager::StringKeyValue>& StatAttributes)
+	{
+		if (StatName == TEXT("DDC.Summary"))
+		{
+			DDCSummaryStats.Append(StatAttributes);
+		}
+	};
+
+	// Grab the DDC stats
+	FCookStatsManager::LogCookStats(LogStatsFunc);
+}
+#endif
+
diff --git a/Engine/Source/Developer/DerivedDataCache/Public/DerivedDataCacheUsageStats.h b/Engine/Source/Developer/DerivedDataCache/Public/DerivedDataCacheUsageStats.h
--- a/Engine/Source/Developer/DerivedDataCache/Public/DerivedDataCacheUsageStats.h
+++ b/Engine/Source/Developer/DerivedDataCache/Public/DerivedDataCacheUsageStats.h
@@ -170,4 +170,61 @@ public:
 protected:
 	const FDerivedDataBackendInterface* BackendInterface;
 	FString CacheName;
-};
\ No newline at end of file
+};
+
+struct FDerivedDataCacheResourceStat
+{
+public:
+	FDerivedDataCacheResourceStat(FString InAssetType = TEXT("None"), bool bIsGameThreadTime = 0.0, double InLoadTimeSec = 0.0, double InLoadSizeMB = 0.0, int64 InAssetsLoaded = 0, double InBuildTimeSec = 0.0, double InBuildSizeMB = 0.0, int64 InAssetsBuilt = 0) :
+		AssetType(MoveTemp(InAssetType)),
+		LoadTimeSec(InLoadTimeSec),
+		LoadSizeMB(InLoadSizeMB),
+		LoadCount(InAssetsLoaded),
+		BuildTimeSec(InBuildTimeSec),
+		BuildSizeMB(InBuildSizeMB),
+		BuildCount(InAssetsBuilt),
+		GameThreadTimeSec(bIsGameThreadTime ? InLoadTimeSec + InBuildTimeSec : 0.0)
+	{}
+
+	void Accumulate(const FDerivedDataCacheResourceStat& OtherStat)
+	{
+		GameThreadTimeSec += OtherStat.GameThreadTimeSec;
+
+		LoadCount += OtherStat.LoadCount;
+		LoadTimeSec += OtherStat.LoadTimeSec;
+		LoadSizeMB += OtherStat.LoadSizeMB;
+
+		BuildCount += OtherStat.BuildCount;
+		BuildTimeSec += OtherStat.BuildTimeSec;
+		BuildSizeMB += OtherStat.BuildSizeMB;
+	}
+
+	FString AssetType;
+
+	double LoadTimeSec;
+	double LoadSizeMB;
+	int64 LoadCount;
+
+	double BuildTimeSec;
+	double BuildSizeMB;
+	int64 BuildCount;
+
+	double GameThreadTimeSec;
+};
+
+struct FDerivedDataCacheResourceStatKeyFuncs : BaseKeyFuncs<FDerivedDataCacheResourceStat, FString, false>
+{
+	static const FString& GetSetKey(const FDerivedDataCacheResourceStat& Element) { return Element.AssetType; }
+	static bool Matches(const FString& A, const FString& B) { return A == B; }
+	static uint32 GetKeyHash(const FString& Key) { return GetTypeHash(Key); }
+};
+
+DERIVEDDATACACHE_API void GatherDerivedDataCacheResourceStats(TArray<FDerivedDataCacheResourceStat>& DDCResourceStats);
+
+#if ENABLE_COOK_STATS
+
+using FDerivedDataCacheSummaryStat = FCookStatsManager::StringKeyValue;
+DERIVEDDATACACHE_API void GatherDerivedDataCacheSummaryStats(TArray<FDerivedDataCacheSummaryStat>& DDCSummaryStats);
+
+#endif
+
diff --git a/Engine/Source/Editor/UnrealEd/Private/DDC/SDDCInformation.cpp b/Engine/Source/Editor/UnrealEd/Private/DDC/SDDCInformation.cpp
--- a/Engine/Source/Editor/UnrealEd/Private/DDC/SDDCInformation.cpp
+++ b/Engine/Source/Editor/UnrealEd/Private/DDC/SDDCInformation.cpp
@@ -2,6 +2,7 @@
 
 #include "DDC/SDDCInformation.h"
 #include "DerivedDataCacheInterface.h"
+#include "DerivedDataCacheUsageStats.h"
 #include "Widgets/SBoxPanel.h"
 #include "Widgets/SToolTip.h"
 #include "Widgets/Text/STextBlock.h"
@@ -40,43 +41,44 @@ void SDDCInformation::Construct(const FArguments& InArgs)
 
 		+ SVerticalBox::Slot()
 		.AutoHeight()
-		.Expose(GridSlot)
+		.Expose(SummaryGridSlot)
 		[
-			GetComplexDataGrid()
+			GetSummaryGrid()
 		]
 
 		+ SVerticalBox::Slot()
 		.AutoHeight()
+		.Expose(CacheGridSlot)
 		[
-			SNew(SVerticalBox)
-			.Visibility_Lambda([] { return FSlateApplication::Get().GetModifierKeys().IsShiftDown() ? EVisibility::Collapsed : EVisibility::Visible; })
-		
-			+ SVerticalBox::Slot()
-			.AutoHeight()
-			[
-				SNew(STextBlock)
-				.Text(this, &SDDCInformation::GetSimpleCacheInformation)
-			]
+			GetCacheGrid()
+		]
 
-			+ SVerticalBox::Slot()
-			.AutoHeight()
-			[
-				SNew(STextBlock)
-				.Font(FCoreStyle::GetDefaultFontStyle("Italic", 10))
-				.ColorAndOpacity(FSlateColor::UseSubduedForeground())
-				.Text(LOCTEXT("ShiftMoreInformation", "[Press Shift for more information]"))
-			]
+		+ SVerticalBox::Slot()
+		.AutoHeight().Padding(0, 20, 0, 0)
+		.Expose(AssetGridSlot)
+		[
+			GetAssetGrid()
 		]
 	];
 
-	RegisterActiveTimer(0.5f, FWidgetActiveTimerDelegate::CreateSP(this, &SDDCInformation::UpdateComplexDataGrid));
+	RegisterActiveTimer(0.5f, FWidgetActiveTimerDelegate::CreateSP(this, &SDDCInformation::UpdateGridPanels));
 }
 
-EActiveTimerReturnType SDDCInformation::UpdateComplexDataGrid(double InCurrentTime, float InDeltaTime)
+EActiveTimerReturnType SDDCInformation::UpdateGridPanels(double InCurrentTime, float InDeltaTime)
 {
-	(*GridSlot)
+	(*SummaryGridSlot)
+	[
+		GetSummaryGrid()
+	];
+
+	(*CacheGridSlot)
+	[
+		GetCacheGrid()
+	];
+
+	(*AssetGridSlot)
 	[
-		GetComplexDataGrid()
+		GetAssetGrid()
 	];
 
 	SlatePrepass(GetPrepassLayoutScaleMultiplier());
@@ -84,81 +86,369 @@ EActiveTimerReturnType SDDCInformation::UpdateComplexDataGrid(double InCurrentTi
 	return EActiveTimerReturnType::Continue;
 }
 
-FText SDDCInformation::GetSimpleCacheInformation() const
-{
-	TStringBuilder<1024> Builder;
+TSharedRef<SWidget> SDDCInformation::GetAssetGrid()
+{	
+	TArray<FDerivedDataCacheResourceStat> DDCResourceStats;
+
+	// Grab the resource stats
+	GatherDerivedDataCacheResourceStats(DDCResourceStats);
+
+	// Sort results on asscending Load size
+	DDCResourceStats.Sort([](const FDerivedDataCacheResourceStat& LHS, const FDerivedDataCacheResourceStat& RHS) { return LHS.LoadSizeMB > RHS.LoadSizeMB; });
+
+	FDerivedDataCacheResourceStat DDCResourceStatsTotal(TEXT("Total"));
+
+	// Accumulate Totals
+	for (const FDerivedDataCacheResourceStat& Stat : DDCResourceStats)
+	{
+		DDCResourceStatsTotal.Accumulate(Stat);
+	}
 
 	PRAGMA_DISABLE_DEPRECATION_WARNINGS
 	TSharedRef<FDerivedDataCacheStatsNode> RootUsage = GetDerivedDataCache()->GatherUsageStats();
 	PRAGMA_ENABLE_DEPRECATION_WARNINGS
-	TArray<TSharedRef<const FDerivedDataCacheStatsNode>> LeafUsageStats;
+		TArray<TSharedRef<const FDerivedDataCacheStatsNode>> LeafUsageStats;
 	RootUsage->ForEachDescendant([&LeafUsageStats](TSharedRef<const FDerivedDataCacheStatsNode> Node) {
 		if (Node->Children.Num() == 0)
 		{
 			LeafUsageStats.Add(Node);
 		}
-	});
+		});
 
-	TMap<FDerivedDataBackendInterface::ESpeedClass, FDerivedDataCacheUsageStats> Stats_BySource;
+	TSharedRef<SGridPanel> Panel =
+		SNew(SGridPanel)
+		.Visibility_Lambda([] { return GetShowDetailedInformation() ? EVisibility::Visible : EVisibility::Collapsed;  });
 
-	for (int32 Index = 0; Index < LeafUsageStats.Num(); Index++)
-	{
-		const FDerivedDataBackendInterface* Backend = LeafUsageStats[Index]->GetBackendInterface();
-		FDerivedDataCacheUsageStats& Stats_FromSource = Stats_BySource.FindOrAdd(Backend->GetSpeedClass());
+	int32 Row = 0;
 
-		TSharedRef<FDerivedDataCacheStatsNode> Usage = Backend->GatherUsageStats();
-		for (const auto& KVP : Usage->Stats)
-		{
-			Stats_FromSource.Combine(KVP.Value);
-		}
+	Panel->AddSlot(2, Row)
+		[
+		SNew(STextBlock)
+		.Margin(FMargin(5, 0))
+		.ColorAndOpacity(FStyleColors::ForegroundHover)
+		.Text(LOCTEXT("Loaded", "Loaded"))
+		.Justification(ETextJustify::Center)
+		];
+
+
+	Panel->AddSlot(5, Row)
+		[
+		SNew(STextBlock)
+		.Margin(FMargin(5, 0))
+		.ColorAndOpacity(FStyleColors::ForegroundHover)
+		.Text(LOCTEXT("Built", "Built"))
+		.Justification(ETextJustify::Center)
+		];
+
+	Row++;
+
+	Panel->AddSlot(0, Row)
+		[
+		SNew(STextBlock)
+		.ColorAndOpacity(FStyleColors::ForegroundHover)
+		.Text(LOCTEXT("Asset", "Asset"))
+		];
+
+	Panel->AddSlot(1, Row)
+		[
+	
+		SNew(STextBlock)
+		.Margin(FMargin(5, 0))
+		.ColorAndOpacity(FStyleColors::ForegroundHover)
+		.Text(LOCTEXT("Count", "Count"))
+		.Justification(ETextJustify::Center)
+		];
+
+	Panel->AddSlot(2, Row)
+		[
+		SNew(STextBlock)
+		.Margin(FMargin(5, 0))
+		.ColorAndOpacity(FStyleColors::ForegroundHover)
+		.Text(LOCTEXT("Time (Sec)", "Time (Sec)"))
+		.Justification(ETextJustify::Center)
+		];
+
+	Panel->AddSlot(3, Row)
+		[
+		SNew(STextBlock)
+		.Margin(FMargin(5, 0))
+		.ColorAndOpacity(FStyleColors::ForegroundHover)
+		.Text(LOCTEXT("Size (MB)", "Size (MB)"))
+		.Justification(ETextJustify::Center)
+		];
+
+	Panel->AddSlot(4, Row)
+		[
+		SNew(STextBlock)
+		.Margin(FMargin(5, 0))
+		.ColorAndOpacity(FStyleColors::ForegroundHover)
+		.Text(LOCTEXT("Count", "Count"))
+		.Justification(ETextJustify::Center)
+		];
+
+	Panel->AddSlot(5, Row)
+		[
+		SNew(STextBlock)
+		.Margin(FMargin(5, 0))
+		.ColorAndOpacity(FStyleColors::ForegroundHover)
+		.Text(LOCTEXT("Time (Sec)", "Time (Sec)"))
+		.Justification(ETextJustify::Center)
+		];
+
+	Panel->AddSlot(6, Row)
+		[
+			SNew(STextBlock)
+			.Margin(FMargin(5, 0))
+		.ColorAndOpacity(FStyleColors::ForegroundHover)
+		.Text(LOCTEXT("Size (MB)", "Size (MB)"))
+		.Justification(ETextJustify::Center)
+		];
+
+	Row++;
+
+
+	for (const FDerivedDataCacheResourceStat& Stat : DDCResourceStats)
+	{	
+		Panel->AddSlot(0, Row)
+		[
+			SNew(STextBlock)
+			.Text(FText::FromString(Stat.AssetType))
+			.Justification(ETextJustify::Left)
+		];
+
+		Panel->AddSlot(1, Row)
+		[
+			SNew(STextBlock)
+			.Text(FText::FromString(FString::Printf(TEXT("%u"),Stat.LoadCount)))
+			.Justification(ETextJustify::Center)
+		];
+
+		Panel->AddSlot(2, Row)
+		[
+			SNew(STextBlock)
+			.Text(FText::FromString(SingleDecimalFormat(Stat.LoadTimeSec)))
+			.Justification(ETextJustify::Center)
+		];
+
+		Panel->AddSlot(3, Row)
+		[
+			SNew(STextBlock)
+			.Text(FText::FromString(SingleDecimalFormat(Stat.LoadSizeMB)))
+			.Justification(ETextJustify::Center)
+		];
+
+		Panel->AddSlot(4, Row)
+		[
+			SNew(STextBlock)
+			.Text(FText::FromString(FString::Printf(TEXT("%u"),Stat.BuildCount)))
+			.Justification(ETextJustify::Center)
+		];
+
+		Panel->AddSlot(5, Row)
+		[
+			SNew(STextBlock)
+			.Text(FText::FromString(SingleDecimalFormat(Stat.BuildTimeSec)))
+			.Justification(ETextJustify::Center)
+		];
+
+		Panel->AddSlot(6, Row)
+		[
+			SNew(STextBlock)
+			.Text(FText::FromString(SingleDecimalFormat(Stat.BuildSizeMB)))
+			.Justification(ETextJustify::Center)
+		];
+
+		Row++;
 	}
+	
+	Panel->AddSlot(0, Row)
+		[
+			SNew(STextBlock)
+			.Text(FText::FromString(DDCResourceStatsTotal.AssetType))
+			.ColorAndOpacity(FStyleColors::AccentWhite)
+			.Justification(ETextJustify::Left)
+		];
 
-	if (Stats_BySource.Num() > 0)
-	{
-		int64 TotalUploadedBytes = 0;
-		for (const auto& KVP : Stats_BySource)
-		{
-			const FDerivedDataCacheUsageStats& Stats = KVP.Value;
+	Panel->AddSlot(1, Row)
+		[
+			SNew(STextBlock)
+			.Text(FText::FromString(FString::Printf(TEXT("%u"), DDCResourceStatsTotal.LoadCount)))
+			.ColorAndOpacity(FStyleColors::ForegroundHover)
+			.Justification(ETextJustify::Center)
+		];
 
-			const int64 TotalGetHits = Stats.GetStats.GetAccumulatedValueAnyThread(FCookStats::CallStats::EHitOrMiss::Hit, FCookStats::CallStats::EStatType::Counter);
-			const int64 TotalGetMisses = Stats.GetStats.GetAccumulatedValueAnyThread(FCookStats::CallStats::EHitOrMiss::Miss, FCookStats::CallStats::EStatType::Counter);
-			const int64 TotalGets = TotalGetHits + TotalGetMisses;
+	Panel->AddSlot(2, Row)
+		[
+			SNew(STextBlock)
+			.Text(FText::FromString(SingleDecimalFormat(DDCResourceStatsTotal.LoadTimeSec)))
+			.ColorAndOpacity(FStyleColors::ForegroundHover)
+			.Justification(ETextJustify::Center)
+		];
 
-			const int64 TotalGetBytes = Stats.GetStats.GetAccumulatedValueAnyThread(FCookStats::CallStats::EHitOrMiss::Hit, FCookStats::CallStats::EStatType::Bytes);
+	Panel->AddSlot(3, Row)
+		[
+			SNew(STextBlock)
+			.Text(FText::FromString(SingleDecimalFormat(DDCResourceStatsTotal.LoadSizeMB)))
+			.ColorAndOpacity(FStyleColors::ForegroundHover)
+			.Justification(ETextJustify::Center)
+		];
 
-			const int64 TotalPutHits = Stats.PutStats.GetAccumulatedValueAnyThread(FCookStats::CallStats::EHitOrMiss::Hit, FCookStats::CallStats::EStatType::Counter);
-			const int64 TotalPutMisses = Stats.PutStats.GetAccumulatedValueAnyThread(FCookStats::CallStats::EHitOrMiss::Miss, FCookStats::CallStats::EStatType::Counter);
-			const int64 TotalPuts = TotalPutHits + TotalPutMisses;
+	Panel->AddSlot(4, Row)
+		[
+			SNew(STextBlock)
+			.Text(FText::FromString(FString::Printf(TEXT("%u"), DDCResourceStatsTotal.BuildCount)))
+			.ColorAndOpacity(FStyleColors::ForegroundHover)
+			.Justification(ETextJustify::Center)
+		];
 
-			TotalUploadedBytes += Stats.PutStats.GetAccumulatedValueAnyThread(FCookStats::CallStats::EHitOrMiss::Hit, FCookStats::CallStats::EStatType::Bytes);
+	Panel->AddSlot(5, Row)
+		[
+			SNew(STextBlock)
+			.Text(FText::FromString(SingleDecimalFormat(DDCResourceStatsTotal.BuildTimeSec)))
+			.ColorAndOpacity(FStyleColors::ForegroundHover)
+			.Justification(ETextJustify::Center)
+		];
 
-			const double LoadedMB = FUnitConversion::Convert(TotalGetBytes, EUnit::Bytes, EUnit::Megabytes);
+	Panel->AddSlot(6, Row)
+		[
+			SNew(STextBlock)
+			.Text(FText::FromString(SingleDecimalFormat(DDCResourceStatsTotal.BuildSizeMB)))
+			.ColorAndOpacity(FStyleColors::ForegroundHover)
+			.Justification(ETextJustify::Center)
+		];
 
-			if (LoadedMB > 0)
-			{
-				Builder.Append(LexToString(KVP.Key));
+	return Panel;
+}
 
-				Builder.Append(TEXT(" "));
-				Builder.Append(SingleDecimalFormat(LoadedMB));
-				Builder.Append(TEXT(" MB"));
 
-				Builder.Append(TEXT("\n"));
-			}
-		}
+TSharedRef<SWidget> SDDCInformation::GetSummaryGrid()
+{
+	TArray<FDerivedDataCacheResourceStat> DDCResourceStats;
 
-		Builder.Append(TEXT("\n"));
+	// Grab the latest resource stats
+	GatherDerivedDataCacheResourceStats(DDCResourceStats);
 
-		const double GeneratedMB = FUnitConversion::Convert(TotalUploadedBytes, EUnit::Bytes, EUnit::Megabytes);
+	FDerivedDataCacheResourceStat DDCResourceStatsTotal(TEXT("Total"));
 
-		Builder.Append(TEXT("Generated "));
-		Builder.Append(SingleDecimalFormat(GeneratedMB));
-		Builder.Append(TEXT(" MB"));
+	// Accumulate Totals
+	for (const FDerivedDataCacheResourceStat& Stat : DDCResourceStats)
+	{
+		DDCResourceStatsTotal.Accumulate(Stat);
 	}
 
-	return FText::FromString(Builder.ToString());
+	const int64 TotalCount = DDCResourceStatsTotal.LoadCount + DDCResourceStatsTotal.BuildCount;
+	const double Efficiency = TotalCount > 0 ? (double)DDCResourceStatsTotal.LoadCount / TotalCount : 0.0;
+
+	const double DownloadedBytesMB = FUnitConversion::Convert(GetDDCSizeBytes(true, false), EUnit::Bytes, EUnit::Megabytes);
+	const double UploadedBytesMB = FUnitConversion::Convert(GetDDCSizeBytes(false, false), EUnit::Bytes, EUnit::Megabytes);
+	
+	TSharedRef<SGridPanel> Panel =
+		SNew(SGridPanel)
+		.Visibility_Lambda([] { return GetShowDetailedInformation() == false ? EVisibility::Visible : EVisibility::Collapsed;  });
+
+	int32 Row = 0;
+
+	Panel->AddSlot(0, Row)
+		[
+			SNew(STextBlock)
+			.ColorAndOpacity(FStyleColors::ForegroundHover)
+			.Text(LOCTEXT("Efficiency", "Efficiency"))
+		];
+
+	Panel->AddSlot(1, Row)
+		[
+			SNew(STextBlock)
+			.Margin(FMargin(5, 0))
+			.Text(FText::FromString(SingleDecimalFormat(Efficiency * 100.0f) + TEXT(" %")))
+			.Justification(ETextJustify::Right)
+		];
+
+	Row++;
+
+	Panel->AddSlot(0, Row)
+		[
+			SNew(STextBlock)
+			.ColorAndOpacity(FStyleColors::ForegroundHover)
+			.Text(LOCTEXT("Loaded", "Loaded"))
+		];
+
+	Panel->AddSlot(1, Row)
+		[
+			SNew(STextBlock)
+			.Margin(FMargin(5, 0))
+			.Text(FText::FromString(SingleDecimalFormat(DDCResourceStatsTotal.LoadSizeMB) + TEXT(" MB")))
+			.Justification(ETextJustify::Right)
+		];
+
+	Row++;
+	
+	Panel->AddSlot(0, Row)
+		[
+			SNew(STextBlock)
+			.ColorAndOpacity(FStyleColors::ForegroundHover)
+			.Text(LOCTEXT("Built", "Built"))
+		];
+
+	Panel->AddSlot(1, Row)
+		[
+			SNew(STextBlock)
+			.Margin(FMargin(5, 0))
+			.Text(FText::FromString(SingleDecimalFormat(DDCResourceStatsTotal.BuildSizeMB) + TEXT(" MB")))
+			.Justification(ETextJustify::Right)
+		];
+
+	Row++;
+
+	Panel->AddSlot(0, Row)
+		[
+			SNew(STextBlock)
+			.ColorAndOpacity(FStyleColors::ForegroundHover)
+			.Text(LOCTEXT("Downloaded", "Downloaded"))
+		];
+
+	Panel->AddSlot(1, Row)
+		[
+			SNew(STextBlock)
+			.Margin(FMargin(5, 0))
+			.Text(FText::FromString(SingleDecimalFormat(DownloadedBytesMB) + TEXT(" MB")))
+			.Justification(ETextJustify::Right)
+		];
+
+	Row++;
+
+	Panel->AddSlot(0, Row)
+		[
+			SNew(STextBlock)
+			.ColorAndOpacity(FStyleColors::ForegroundHover)
+			.Text(LOCTEXT("Uploaded", "Uploaded"))
+		];
+
+	Panel->AddSlot(1, Row)
+		[
+			SNew(STextBlock)
+			.Margin(FMargin(5, 0))
+			.Text(FText::FromString(SingleDecimalFormat(UploadedBytesMB) + TEXT(" MB")))
+			.Justification(ETextJustify::Right)
+		];
+
+	Row++;
+
+	Panel->AddSlot(0, Row)
+		[
+			SNew(STextBlock)
+			.Font(FCoreStyle::GetDefaultFontStyle("Italic", 10))
+			.ColorAndOpacity(FSlateColor::UseSubduedForeground())
+			.Text(LOCTEXT("ShiftMoreInformation", "[Hold Shift for more information]"))
+
+		];
+
+	Row++;
+
+	return Panel;
 }
 
-TSharedRef<SWidget> SDDCInformation::GetComplexDataGrid()
+
+TSharedRef<SWidget> SDDCInformation::GetCacheGrid()
 {
 	PRAGMA_DISABLE_DEPRECATION_WARNINGS
 	TSharedRef<FDerivedDataCacheStatsNode> RootUsage = GetDerivedDataCache()->GatherUsageStats();
@@ -173,7 +463,7 @@ TSharedRef<SWidget> SDDCInformation::GetComplexDataGrid()
 
 	TSharedRef<SGridPanel> Panel = 
 		SNew(SGridPanel)
-		.Visibility_Lambda([] { return FSlateApplication::Get().GetModifierKeys().IsShiftDown() ? EVisibility::Visible : EVisibility::Collapsed; });
+		.Visibility_Lambda([] { return GetShowDetailedInformation()==true ? EVisibility::Visible : EVisibility::Collapsed; });
 
 	int32 Row = 0;
 
@@ -181,14 +471,15 @@ TSharedRef<SWidget> SDDCInformation::GetComplexDataGrid()
 	[
 		SNew(STextBlock)
 		.ColorAndOpacity(FStyleColors::ForegroundHover)
-		.Text(LOCTEXT("DDC", "DDC"))
+		.Text(LOCTEXT("Cache", "Cache"))
 	];
 
 	Panel->AddSlot(1, Row)
 	[
 		SNew(STextBlock)
-		.Margin(FMargin(5, 0))
+		.Margin(FMargin(10, 0))
 		.ColorAndOpacity(FStyleColors::ForegroundHover)
+		.Justification(ETextJustify::Right)
 		.Text(LOCTEXT("Speed", "Speed"))
 	];
 
@@ -197,7 +488,8 @@ TSharedRef<SWidget> SDDCInformation::GetComplexDataGrid()
 		SNew(STextBlock)
 		.Margin(FMargin(5, 0))
 		.ColorAndOpacity(FStyleColors::ForegroundHover)
-		.Text(LOCTEXT("Loaded", "Loaded"))
+		.Justification(ETextJustify::Right)
+		.Text(LOCTEXT("HitPercentage", "Hit%"))
 	];
 
 	Panel->AddSlot(3, Row)
@@ -205,10 +497,20 @@ TSharedRef<SWidget> SDDCInformation::GetComplexDataGrid()
 		SNew(STextBlock)
 		.Margin(FMargin(5, 0))
 		.ColorAndOpacity(FStyleColors::ForegroundHover)
-		.Text(LOCTEXT("HitPercentage", "Hit%"))
+		.Justification(ETextJustify::Right)
+		.Text(LOCTEXT("Read", "Read"))
 	];
 
 	Panel->AddSlot(4, Row)
+		[
+		SNew(STextBlock)
+		.Margin(FMargin(5, 0))
+		.ColorAndOpacity(FStyleColors::ForegroundHover)
+		.Justification(ETextJustify::Right)
+		.Text(LOCTEXT("Write", "Write"))
+		];
+
+	Panel->AddSlot(5, Row)
 	[
 		SNew(SImage)
 		.Image(FAppStyle::Get().GetBrush("Icons.Edit"))
@@ -218,6 +520,9 @@ TSharedRef<SWidget> SDDCInformation::GetComplexDataGrid()
 
 	Row++;
 
+	double SumTotalGetMB =0.0;
+	double SumTotalPutMB = 0.0;
+
 	for (TSharedRef<const FDerivedDataCacheStatsNode> Node : LeafUsageStats)
 	{
 		FDerivedDataCacheUsageStats Stats;
@@ -232,10 +537,12 @@ TSharedRef<SWidget> SDDCInformation::GetComplexDataGrid()
 
 		const int64 TotalGetBytes = Stats.GetStats.GetAccumulatedValueAnyThread(FCookStats::CallStats::EHitOrMiss::Hit, FCookStats::CallStats::EStatType::Bytes);
 		const int64 TotalPutBytes = Stats.PutStats.GetAccumulatedValueAnyThread(FCookStats::CallStats::EHitOrMiss::Hit, FCookStats::CallStats::EStatType::Bytes);
+		const int64 TotalPrefetchBytes = Stats.PrefetchStats.GetAccumulatedValueAnyThread(FCookStats::CallStats::EHitOrMiss::Hit, FCookStats::CallStats::EStatType::Bytes);
 
 		const int64 TotalGets_Hit = Stats.GetStats.GetAccumulatedValueAnyThread(FCookStats::CallStats::EHitOrMiss::Hit, FCookStats::CallStats::EStatType::Counter);
 		const int64 TotalGets_Miss = Stats.GetStats.GetAccumulatedValueAnyThread(FCookStats::CallStats::EHitOrMiss::Miss, FCookStats::CallStats::EStatType::Counter);
-		const int64 TotalGets = TotalGets_Hit + TotalGets_Miss;
+
+		const int64 TotalRequests = TotalGets_Hit + TotalGets_Miss;
 
 		//if (TotalGetBytes > 0 || TotalPutBytes > 0)
 		{
@@ -248,31 +555,49 @@ TSharedRef<SWidget> SDDCInformation::GetComplexDataGrid()
 			Panel->AddSlot(1, Row)
 			[
 				SNew(STextBlock)
-				.Margin(FMargin(5, 0))
+				.Margin(FMargin(10, 0))
+				.Justification(ETextJustify::Right)
 				.Text(FText::FromString(LexToString(Backend->GetSpeedClass())))
 			];
 
-			const double TotalGetMB = FUnitConversion::Convert(TotalGetBytes, EUnit::Bytes, EUnit::Megabytes);
+			const double HitPercentage = TotalRequests > 0 ? (100.0 * (TotalGets_Hit / (double)TotalRequests)) : 0.0;
 
 			Panel->AddSlot(2, Row)
+				.HAlign(HAlign_Right)
+				[
+					SNew(STextBlock)
+					.Margin(FMargin(5, 0))
+				.Justification(ETextJustify::Right)
+				.Text(FText::FromString(SingleDecimalFormat(HitPercentage) + TEXT(" %")))
+				];
+
+			const double TotalGetMB = FUnitConversion::Convert(TotalGetBytes, EUnit::Bytes, EUnit::Megabytes);
+
+			SumTotalGetMB += TotalGetMB;
+
+			Panel->AddSlot(3, Row)
 			.HAlign(HAlign_Right)
 			[
 				SNew(STextBlock)
 				.Margin(FMargin(5, 0))
+				.Justification(ETextJustify::Right)
 				.Text(FText::FromString(SingleDecimalFormat(TotalGetMB) + TEXT(" MB")))
 			];
 
-			const double HitPercentage = TotalGets > 0 ? (100.0 * (TotalGets_Hit / (double)TotalGets)) : 0.0;
+			const double TotalPutMB = FUnitConversion::Convert(TotalPutBytes, EUnit::Bytes, EUnit::Megabytes);
 
-			Panel->AddSlot(3, Row)
+			SumTotalPutMB += TotalPutMB;
+
+			Panel->AddSlot(4, Row)
 			.HAlign(HAlign_Right)
 			[
 				SNew(STextBlock)
 				.Margin(FMargin(5, 0))
-				.Text(FText::FromString(SingleDecimalFormat(HitPercentage) + TEXT("%")))
+				.Justification(ETextJustify::Right)
+				.Text(FText::FromString(SingleDecimalFormat(TotalPutMB) + TEXT(" MB")))
 			];
 
-			Panel->AddSlot(4, Row)
+			Panel->AddSlot(5, Row)
 			.HAlign(HAlign_Center)
 			[
 				SNew(SImage)
@@ -284,7 +609,183 @@ TSharedRef<SWidget> SDDCInformation::GetComplexDataGrid()
 		}
 	}
 
+	Panel->AddSlot(0, Row)
+		[
+			SNew(STextBlock)
+			.Text(FText::FromString(TEXT("Total")))
+			.ColorAndOpacity(FStyleColors::AccentWhite)
+			.Justification(ETextJustify::Left)	
+		];
+
+	Panel->AddSlot(3, Row)
+		.HAlign(HAlign_Right)
+		[
+		SNew(STextBlock)
+		.Margin(FMargin(5, 0))
+		.Justification(ETextJustify::Right)
+		.ColorAndOpacity(FStyleColors::ForegroundHover)
+		.Text(FText::FromString(SingleDecimalFormat(SumTotalGetMB) + TEXT(" MB")))
+		];
+
+	Panel->AddSlot(4, Row)
+		.HAlign(HAlign_Right)
+		[
+			SNew(STextBlock)
+			.Margin(FMargin(5, 0))
+		.Justification(ETextJustify::Right)
+		.ColorAndOpacity(FStyleColors::ForegroundHover)
+		.Text(FText::FromString(SingleDecimalFormat(SumTotalPutMB) + TEXT(" MB")))
+		];
+
 	return Panel;
 }
 
+double SDDCInformation::GetDDCSizeBytes(bool bGet, bool bLocal)
+{
+	PRAGMA_DISABLE_DEPRECATION_WARNINGS
+	TSharedRef<FDerivedDataCacheStatsNode> RootUsage = GetDerivedDataCache()->GatherUsageStats();
+	PRAGMA_ENABLE_DEPRECATION_WARNINGS
+	TArray<TSharedRef<const FDerivedDataCacheStatsNode>> LeafUsageStats;
+	RootUsage->ForEachDescendant([&LeafUsageStats](TSharedRef<const FDerivedDataCacheStatsNode> Node) {
+		if (Node->Children.Num() == 0)
+		{
+			LeafUsageStats.Add(Node);
+		}
+		});
+
+	int64 TotalBytes = 0;
+
+	for (int32 Index = 0; Index < LeafUsageStats.Num(); Index++)
+	{
+		const FDerivedDataBackendInterface* Backend = LeafUsageStats[Index]->GetBackendInterface();
+
+		if ((Backend->GetSpeedClass() == FDerivedDataBackendInterface::ESpeedClass::Local) != bLocal)
+			continue;
+
+		TSharedRef<FDerivedDataCacheStatsNode> Usage = Backend->GatherUsageStats();
+
+		for (const auto& KVP : Usage->Stats)
+		{
+			const FDerivedDataCacheUsageStats& Stats = KVP.Value;
+
+			if (bGet)
+			{
+				TotalBytes += Stats.GetStats.GetAccumulatedValueAnyThread(FCookStats::CallStats::EHitOrMiss::Hit, FCookStats::CallStats::EStatType::Bytes);
+			}
+			else
+			{
+				TotalBytes += Stats.PutStats.GetAccumulatedValueAnyThread(FCookStats::CallStats::EHitOrMiss::Hit, FCookStats::CallStats::EStatType::Bytes);
+			}
+		}
+	}
+
+	return TotalBytes;
+}
+
+
+double SDDCInformation::GetDDCTimeSeconds(bool bGet, bool bLocal)
+{
+	PRAGMA_DISABLE_DEPRECATION_WARNINGS
+		TSharedRef<FDerivedDataCacheStatsNode> RootUsage = GetDerivedDataCache()->GatherUsageStats();
+	PRAGMA_ENABLE_DEPRECATION_WARNINGS
+		TArray<TSharedRef<const FDerivedDataCacheStatsNode>> LeafUsageStats;
+	RootUsage->ForEachDescendant([&LeafUsageStats](TSharedRef<const FDerivedDataCacheStatsNode> Node) {
+		if (Node->Children.Num() == 0)
+		{
+			LeafUsageStats.Add(Node);
+		}
+		});
+
+	int64 TotalCycles = 0;
+
+	for (int32 Index = 0; Index < LeafUsageStats.Num(); Index++)
+	{
+		const FDerivedDataBackendInterface* Backend = LeafUsageStats[Index]->GetBackendInterface();
+
+		if ((Backend->GetSpeedClass() == FDerivedDataBackendInterface::ESpeedClass::Local) != bLocal)
+			continue;
+
+		TSharedRef<FDerivedDataCacheStatsNode> Usage = Backend->GatherUsageStats();
+
+		for (const auto& KVP : Usage->Stats)
+		{
+			const FDerivedDataCacheUsageStats& Stats = KVP.Value;
+
+			if (bGet)
+			{
+				TotalCycles +=
+					(Stats.GetStats.GetAccumulatedValueAnyThread(FCookStats::CallStats::EHitOrMiss::Hit, FCookStats::CallStats::EStatType::Cycles) +
+						Stats.GetStats.GetAccumulatedValueAnyThread(FCookStats::CallStats::EHitOrMiss::Miss, FCookStats::CallStats::EStatType::Cycles));
+
+				TotalCycles +=
+					(Stats.PrefetchStats.GetAccumulatedValueAnyThread(FCookStats::CallStats::EHitOrMiss::Hit, FCookStats::CallStats::EStatType::Cycles) +
+						Stats.PrefetchStats.GetAccumulatedValueAnyThread(FCookStats::CallStats::EHitOrMiss::Miss, FCookStats::CallStats::EStatType::Cycles));
+			}
+			else
+			{
+				TotalCycles +=
+					(Stats.PutStats.GetAccumulatedValueAnyThread(FCookStats::CallStats::EHitOrMiss::Hit, FCookStats::CallStats::EStatType::Cycles) +
+						Stats.PutStats.GetAccumulatedValueAnyThread(FCookStats::CallStats::EHitOrMiss::Miss, FCookStats::CallStats::EStatType::Cycles));
+			}
+		}
+	}
+
+	return (double)TotalCycles * FPlatformTime::GetSecondsPerCycle();
+}
+
+bool SDDCInformation::GetDDCHasLocalBackend()
+{
+	PRAGMA_DISABLE_DEPRECATION_WARNINGS
+		TSharedRef<FDerivedDataCacheStatsNode> RootUsage = GetDerivedDataCache()->GatherUsageStats();
+	PRAGMA_ENABLE_DEPRECATION_WARNINGS
+		TArray<TSharedRef<const FDerivedDataCacheStatsNode>> LeafUsageStats;
+	RootUsage->ForEachDescendant([&LeafUsageStats](TSharedRef<const FDerivedDataCacheStatsNode> Node) {
+		if (Node->Children.Num() == 0)
+		{
+			LeafUsageStats.Add(Node);
+		}
+		});
+
+	for (int32 Index = 0; Index < LeafUsageStats.Num(); Index++)
+	{
+		const FDerivedDataBackendInterface* Backend = LeafUsageStats[Index]->GetBackendInterface();
+
+		if (Backend->GetSpeedClass() == FDerivedDataBackendInterface::ESpeedClass::Local)
+			return true;
+	}
+
+	return false;
+}
+
+
+bool SDDCInformation::GetDDCHasRemoteBackend()
+{
+	PRAGMA_DISABLE_DEPRECATION_WARNINGS
+		TSharedRef<FDerivedDataCacheStatsNode> RootUsage = GetDerivedDataCache()->GatherUsageStats();
+	PRAGMA_ENABLE_DEPRECATION_WARNINGS
+		TArray<TSharedRef<const FDerivedDataCacheStatsNode>> LeafUsageStats;
+	RootUsage->ForEachDescendant([&LeafUsageStats](TSharedRef<const FDerivedDataCacheStatsNode> Node) {
+		if (Node->Children.Num() == 0)
+		{
+			LeafUsageStats.Add(Node);
+		}
+		});
+
+	for (int32 Index = 0; Index < LeafUsageStats.Num(); Index++)
+	{
+		const FDerivedDataBackendInterface* Backend = LeafUsageStats[Index]->GetBackendInterface();
+
+		if (Backend->GetSpeedClass() != FDerivedDataBackendInterface::ESpeedClass::Local)
+			return true;
+	}
+
+	return false;
+}
+
+bool SDDCInformation::GetShowDetailedInformation()
+{
+	return FSlateApplication::Get().GetModifierKeys().IsShiftDown();
+}
+
+
 #undef LOCTEXT_NAMESPACE
\ No newline at end of file
diff --git a/Engine/Source/Editor/UnrealEd/Private/DDC/SDDCStatusIndicator.cpp b/Engine/Source/Editor/UnrealEd/Private/DDC/SDDCStatusIndicator.cpp
--- a/Engine/Source/Editor/UnrealEd/Private/DDC/SDDCStatusIndicator.cpp
+++ b/Engine/Source/Editor/UnrealEd/Private/DDC/SDDCStatusIndicator.cpp
@@ -2,7 +2,9 @@
 
 #include "DDC/SDDCStatusIndicator.h"
 #include "DDC/SDDCInformation.h"
+#include "DerivedDataCacheUsageStats.h"
 #include "DerivedDataCacheInterface.h"
+#include "DerivedDataBackendInterface.h"
 #include "Widgets/SBoxPanel.h"
 #include "Widgets/SToolTip.h"
 #include "Widgets/Text/STextBlock.h"
@@ -25,10 +27,7 @@
 #define LOCTEXT_NAMESPACE "SDDCStatusIndicator"
 
 void SDDCStatusIndicator::Construct(const FArguments& InArgs)
-{
-	FDerivedDataCacheInterface* DDC = GetDerivedDataCache();
-	const FString DDCGraphName = FName::NameToDisplayString(FString(DDC->GetGraphName()), false);
-	const FText DDCLabel = FText::FromString(DDC->IsDefaultGraph() ? TEXT("DDC") : DDCGraphName);
+{	
 	/*FDerivedDataCacheInterface::FOnDDCNotification& DDCNotificationEvent = GetDerivedDataCacheRef().GetDDCNotificationEvent();
 
 	if (bSubscribe)
@@ -48,6 +47,17 @@ void SDDCStatusIndicator::Construct(const FArguments& InArgs)
 	[
 		SNew(SHorizontalBox)
 		.ToolTip(SNew(SToolTip)[SNew(SDDCInformation)])
+
+		+ SHorizontalBox::Slot()
+		.AutoWidth()
+		.VAlign(VAlign_Center)
+		.Padding(0,0,3,0)
+		[
+			SNew(SImage)
+			.Image(FAppStyle::Get().GetBrush("Icons.Settings"))
+			.ColorAndOpacity_Lambda([this] { return bBusy? FLinearColor::Green.CopyWithNewOpacity(0.5f + (0.5f * FMath::MakePulsatingValue(BusyPulseSequence.GetLerp(), 1))) : FSlateColor::UseSubduedForeground(); })
+		]
+
 		+ SHorizontalBox::Slot()
 		.AutoWidth()
 		.VAlign(VAlign_Center)
@@ -62,7 +72,7 @@ void SDDCStatusIndicator::Construct(const FArguments& InArgs)
 			[
 				SNew(SImage)
 				.Image(FAppStyle::Get().GetBrush("Icons.ArrowUp"))
-				.ColorAndOpacity_Lambda([this] { return FLinearColor::Red.CopyWithNewOpacity(FadePutSequence.GetLerp()); })
+				.ColorAndOpacity_Lambda([this] { return bPutActive ? FLinearColor::Green.CopyWithNewOpacity(0.5f + (0.5f * FMath::MakePulsatingValue(FadePutSequence.GetLerp(), 1))) : FSlateColor::UseSubduedForeground(); })
 			]
 
 			+ SOverlay::Slot()
@@ -72,7 +82,7 @@ void SDDCStatusIndicator::Construct(const FArguments& InArgs)
 			[
 				SNew(SImage)
 				.Image(FAppStyle::Get().GetBrush("Icons.ArrowDown"))
-				.ColorAndOpacity_Lambda([this] { return FLinearColor::Green.CopyWithNewOpacity(FadeGetSequence.GetLerp()); })
+				.ColorAndOpacity_Lambda([this] { return bGetActive ? FLinearColor::Green.CopyWithNewOpacity(0.5f + (0.5f * FMath::MakePulsatingValue(FadeGetSequence.GetLerp(), 1))) : FSlateColor::UseSubduedForeground(); })
 			]
 		]
 
@@ -82,10 +92,20 @@ void SDDCStatusIndicator::Construct(const FArguments& InArgs)
 		.Padding(0,0,3,0)
 		[
 			SNew(SImage)
-			.Image(FAppStyle::Get().GetBrush("Icons.Server"))
+			.Image(FAppStyle::Get().GetBrush("Icons.Local"))
 			.ColorAndOpacity_Lambda([this] {
-				return FStyleColors::Foreground.GetSpecifiedColor().CopyWithNewOpacity(0.5f + (0.5f * FMath::MakePulsatingValue(BusyPulseSequence.GetLerp(), 1)));
-			})
+			return SDDCInformation::GetDDCHasLocalBackend() ? FStyleColors::AccentBlue : FSlateColor::UseSubduedForeground(); })
+		]
+
+		+ SHorizontalBox::Slot()
+			.AutoWidth()
+			.VAlign(VAlign_Center)
+			.Padding(0, 0, 3, 0)
+		[
+			SNew(SImage)
+			.Image(FAppStyle::Get().GetBrush("Icons.Cloud"))
+			.ColorAndOpacity_Lambda([this] { 
+			return SDDCInformation::GetDDCHasRemoteBackend()? FStyleColors::AccentBlue : FSlateColor::UseSubduedForeground(); })
 		]
 
 		+ SHorizontalBox::Slot()
@@ -94,12 +114,12 @@ void SDDCStatusIndicator::Construct(const FArguments& InArgs)
 		.Padding(0, 0, 10, 0)
 		[
 			SNew(STextBlock)
-			.Text(DDCLabel)
+			.Text_Lambda([this] { return GetInformationText(); })
 		]
 	];
 
-	LastDDCGetTime = GetDDCTime(true);
-	LastDDCPutTime = GetDDCTime(false);
+	LastDDCGetTime = SDDCInformation::GetDDCTimeSeconds(true, false);
+	LastDDCPutTime = SDDCInformation::GetDDCTimeSeconds(false, false);
 	RegisterActiveTimer(0.5f, FWidgetActiveTimerDelegate::CreateSP(this, &SDDCStatusIndicator::UpdateBusyIndicator));
 	RegisterActiveTimer(5.0f, FWidgetActiveTimerDelegate::CreateSP(this, &SDDCStatusIndicator::UpdateWarnings));
 }
@@ -109,13 +129,17 @@ EActiveTimerReturnType SDDCStatusIndicator::UpdateBusyIndicator(double InCurrent
 	const double OldLastDDCGetTime = LastDDCGetTime;
 	const double OldLastDDCPutTime = LastDDCPutTime;
 
-	LastDDCGetTime = GetDDCTime(true);
-	LastDDCPutTime = GetDDCTime(false);
+	LastDDCGetTime = SDDCInformation::GetDDCTimeSeconds(true, false);
+	LastDDCPutTime = SDDCInformation::GetDDCTimeSeconds(false, false);
+
+	bGetActive = OldLastDDCGetTime != LastDDCGetTime;
+	bPutActive = OldLastDDCPutTime != LastDDCPutTime;
+	bBusy = GetDerivedDataCache()->AnyAsyncRequestsRemaining();
 
-	FadeGetSequence.PlayRelative(this->AsShared(), OldLastDDCGetTime != LastDDCGetTime);
-	FadePutSequence.PlayRelative(this->AsShared(), OldLastDDCPutTime != LastDDCPutTime);
+	FadeGetSequence.PlayRelative(this->AsShared(), bGetActive);
+	FadePutSequence.PlayRelative(this->AsShared(), bPutActive);
 
-	if (GetDerivedDataCache()->AnyAsyncRequestsRemaining())
+	if (bBusy)
 	{
 		if (!BusyPulseSequence.IsPlaying())
 		{
@@ -131,38 +155,6 @@ EActiveTimerReturnType SDDCStatusIndicator::UpdateBusyIndicator(double InCurrent
 	return EActiveTimerReturnType::Continue;
 }
 
-double SDDCStatusIndicator::GetDDCTime(bool bGet) const
-{
-	int64 TotalCycles = 0;
-
-	PRAGMA_DISABLE_DEPRECATION_WARNINGS
-	TSharedRef<FDerivedDataCacheStatsNode> RootUsage = GetDerivedDataCache()->GatherUsageStats();
-	PRAGMA_ENABLE_DEPRECATION_WARNINGS
-	for (const auto& KVP : RootUsage->Stats)
-	{
-		const FDerivedDataCacheUsageStats& Stats = KVP.Value;
-		
-		if (bGet)
-		{
-			TotalCycles +=
-				(Stats.GetStats.GetAccumulatedValueAnyThread(FCookStats::CallStats::EHitOrMiss::Hit, FCookStats::CallStats::EStatType::Cycles) +
-				Stats.GetStats.GetAccumulatedValueAnyThread(FCookStats::CallStats::EHitOrMiss::Miss, FCookStats::CallStats::EStatType::Cycles));
-		}
-		else
-		{
-			TotalCycles +=
-				(Stats.PrefetchStats.GetAccumulatedValueAnyThread(FCookStats::CallStats::EHitOrMiss::Hit, FCookStats::CallStats::EStatType::Cycles) +
-				Stats.PrefetchStats.GetAccumulatedValueAnyThread(FCookStats::CallStats::EHitOrMiss::Miss, FCookStats::CallStats::EStatType::Cycles));
-
-			TotalCycles +=
-				(Stats.PutStats.GetAccumulatedValueAnyThread(FCookStats::CallStats::EHitOrMiss::Hit, FCookStats::CallStats::EStatType::Cycles) +
-				Stats.PutStats.GetAccumulatedValueAnyThread(FCookStats::CallStats::EHitOrMiss::Miss, FCookStats::CallStats::EStatType::Cycles));
-		}
-	}
-
-	return TotalCycles * FPlatformTime::GetSecondsPerCycle();
-}
-
 EActiveTimerReturnType SDDCStatusIndicator::UpdateWarnings(double InCurrentTime, float InDeltaTime)
 {
 	const UEditorSettings* Settings = GetDefault<UEditorSettings>();
@@ -233,4 +225,13 @@ EActiveTimerReturnType SDDCStatusIndicator::UpdateWarnings(double InCurrentTime,
 	return EActiveTimerReturnType::Stop;
 }
 
+
+FText SDDCStatusIndicator::GetInformationText() const
+{
+	FDerivedDataCacheInterface* DDC = GetDerivedDataCache();
+	const FString DDCGraphName = FName::NameToDisplayString(FString(DDC->GetGraphName()), false);
+	const FText DDCLabel = FText::FromString(DDC->IsDefaultGraph() ? TEXT("DDC") : DDCGraphName);
+	return DDCLabel;
+}
+
 #undef LOCTEXT_NAMESPACE
\ No newline at end of file
diff --git a/Engine/Source/Editor/UnrealEd/Public/DDC/SDDCInformation.h b/Engine/Source/Editor/UnrealEd/Public/DDC/SDDCInformation.h
--- a/Engine/Source/Editor/UnrealEd/Public/DDC/SDDCInformation.h
+++ b/Engine/Source/Editor/UnrealEd/Public/DDC/SDDCInformation.h
@@ -18,10 +18,23 @@ class UNREALED_API SDDCInformation : public SCompoundWidget
 
 	void Construct(const FArguments& InArgs);
 
+	static double GetDDCTimeSeconds(bool bGet, bool bLocal);
+	static double GetDDCSizeBytes(bool bGet, bool bLocal);
+	static bool	GetDDCHasLocalBackend();
+	static bool	GetDDCHasRemoteBackend();
+	static bool	GetShowDetailedInformation();
+
 private:
-	FText GetSimpleCacheInformation() const;
-	TSharedRef<SWidget> GetComplexDataGrid();
-	EActiveTimerReturnType UpdateComplexDataGrid(double InCurrentTime, float InDeltaTime);
 
-	SVerticalBox::FSlot* GridSlot = nullptr;
+	TSharedRef<SWidget> GetAssetGrid();
+	TSharedRef<SWidget> GetSummaryGrid();
+	TSharedRef<SWidget> GetCacheGrid();
+
+	EActiveTimerReturnType UpdateGridPanels(double InCurrentTime, float InDeltaTime);
+
+	SVerticalBox::FSlot* SummaryGridSlot = nullptr;
+	SVerticalBox::FSlot* AssetGridSlot = nullptr;
+	SVerticalBox::FSlot* CacheGridSlot = nullptr;
+
+
 };
\ No newline at end of file
diff --git a/Engine/Source/Editor/UnrealEd/Public/DDC/SDDCStatusIndicator.h b/Engine/Source/Editor/UnrealEd/Public/DDC/SDDCStatusIndicator.h
--- a/Engine/Source/Editor/UnrealEd/Public/DDC/SDDCStatusIndicator.h
+++ b/Engine/Source/Editor/UnrealEd/Public/DDC/SDDCStatusIndicator.h
@@ -20,13 +20,18 @@ class UNREALED_API SDDCStatusIndicator : public SCompoundWidget
 
 private:
 	EActiveTimerReturnType UpdateBusyIndicator(double InCurrentTime, float InDeltaTime);
-	double GetDDCTime(bool bGet) const;
 	EActiveTimerReturnType UpdateWarnings(double InCurrentTime, float InDeltaTime);
+	FText GetInformationText() const;
 
 	double LastDDCGetTime = 0;
 	double LastDDCPutTime = 0;
 
+	bool bGetActive = false;
+	bool bPutActive = false;
+	bool bBusy = false;
+
 	FCurveSequence BusyPulseSequence;
 	FCurveSequence FadeGetSequence;
 	FCurveSequence FadePutSequence;
+	
 };
\ No newline at end of file
diff --git a/Engine/Source/Editor/UnrealEd/UnrealEd.Build.cs b/Engine/Source/Editor/UnrealEd/UnrealEd.Build.cs
--- a/Engine/Source/Editor/UnrealEd/UnrealEd.Build.cs
+++ b/Engine/Source/Editor/UnrealEd/UnrealEd.Build.cs
@@ -118,6 +118,7 @@ public class UnrealEd : ModuleRules
 				"SubobjectDataInterface",
 				"SubobjectEditor",
 				"PhysicsUtilities",
+				"DerivedDataCache",
 			}
 		);
 
@@ -206,6 +207,7 @@ public class UnrealEd : ModuleRules
 				"MaterialBaking",
 				"CookOnTheFly",
 				"BazelExecutor",
+				"DerivedDataCache",
 			}
 		);
 
diff --git a/Engine/Source/Runtime/Engine/Engine.Build.cs b/Engine/Source/Runtime/Engine/Engine.Build.cs
--- a/Engine/Source/Runtime/Engine/Engine.Build.cs
+++ b/Engine/Source/Runtime/Engine/Engine.Build.cs
@@ -126,6 +126,7 @@ public class Engine : ModuleRules
 				"CrunchCompression",
 				"IntelISPC",
 				"TraceLog",
+				"DerivedDataCache",
 			}
 		);
 
diff --git a/Engine/Source/Runtime/Engine/Private/Analytics/StudioAnalytics.cpp b/Engine/Source/Runtime/Engine/Private/Analytics/StudioAnalytics.cpp
--- a/Engine/Source/Runtime/Engine/Private/Analytics/StudioAnalytics.cpp
+++ b/Engine/Source/Runtime/Engine/Private/Analytics/StudioAnalytics.cpp
@@ -18,7 +18,7 @@
 #include "Interfaces/IAnalyticsProvider.h"
 #include "Templates/SharedPointer.h"
 #include "HAL/PlatformProcess.h"
-#include "ProfilingDebugging/CookStats.h"
+#include "DerivedDataCacheUsageStats.h"
 
 bool FStudioAnalytics::bInitialized = false;
 volatile double FStudioAnalytics::TimeEstimation = 0;
@@ -166,151 +166,66 @@ void FStudioAnalytics::FireEvent_Loading(const FString& LoadingName, double Seco
 		FStudioAnalytics::GetProvider().RecordEvent(TEXT("Performance.Loading"), Attributes);
 
 #if ENABLE_COOK_STATS
-		//// Sends each DDC stat to the studio analytics system.
-		//auto SendDDCStatsToAnalytics = [&Attributes](const FString& StatName, const TArray<FCookStatsManager::StringKeyValue>& StatAttributes)
-		//{
-		//	// We'reonly interested in DDC Summary stats
-		//	if (StatName.Contains("DDC.Summary"))
-		//	{
-		//		for (const auto& Attr : StatAttributes)
-		//		{
-		//			FString FormattedAttrName = StatName + "." + Attr.Key;
-		//			Attributes.Emplace(FormattedAttrName, Attr.Value);
-		//		}		
-		//	}
-		//};
-
-		//// Grab the DDC stats
-		//FCookStatsManager::LogCookStats(SendDDCStatsToAnalytics);
-
-		/** Used for custom logging of DDC Resource usage stats. */
-		struct FDDCResourceUsageStat
-		{
-		public:
-			FDDCResourceUsageStat(FString InAssetType, double InTotalTimeSec, bool bIsGameThreadTime, double InSizeMB, int64 InAssetsBuilt) : AssetType(MoveTemp(InAssetType)), TotalTimeSec(InTotalTimeSec), GameThreadTimeSec(bIsGameThreadTime ? InTotalTimeSec : 0.0), SizeMB(InSizeMB), AssetsBuilt(InAssetsBuilt) {}
-			void Accumulate(const FDDCResourceUsageStat& OtherStat)
-			{
-				TotalTimeSec += OtherStat.TotalTimeSec;
-				GameThreadTimeSec += OtherStat.GameThreadTimeSec;
-				SizeMB += OtherStat.SizeMB;
-				AssetsBuilt += OtherStat.AssetsBuilt;
-			}
-			FString AssetType;
-			double TotalTimeSec;
-			double GameThreadTimeSec;
-			double SizeMB;
-			int64 AssetsBuilt;
-		};
-
-		/** Used for custom TSet comparison of DDC Resource usage stats. */
-		struct FDDCResourceUsageStatKeyFuncs : BaseKeyFuncs<FDDCResourceUsageStat, FString, false>
+
+		TArray<FDerivedDataCacheResourceStat> DDCResourceStats;
+
+		// Grab the latest resource stats
+		GatherDerivedDataCacheResourceStats(DDCResourceStats);
+
+		FDerivedDataCacheResourceStat DDCResourceStatsTotal(TEXT("Total"));
+
+		// Accumulate Totals
+		for (const FDerivedDataCacheResourceStat& Stat : DDCResourceStats)
 		{
-			static const FString& GetSetKey(const FDDCResourceUsageStat& Element) { return Element.AssetType; }
-			static bool Matches(const FString& A, const FString& B) { return A == B; }
-			static uint32 GetKeyHash(const FString& Key) { return GetTypeHash(Key); }
-		};
-
-		// instead of printing the usage stats generically, we capture them so we can log a subset of them in an easy-to-read way.
-		TSet<FDDCResourceUsageStat, FDDCResourceUsageStatKeyFuncs> DDCResourceUsageStats;
-		TArray<FCookStatsManager::StringKeyValue> DDCSummaryStats;
-
-		int64 TotalAssetsBuilt = 0;
-		double TotalAssetTimeSec = 0.0;
-		double TotalAssetSizeMB = 0.0;
-		
-		/** this functor will take a collected cooker stat and log it out using some custom formatting based on known stats that are collected.. */
-		auto LogStatsFunc = [&DDCResourceUsageStats, &DDCSummaryStats, &TotalAssetsBuilt, &TotalAssetTimeSec, &TotalAssetSizeMB]
-		(const FString& StatName, const TArray<FCookStatsManager::StringKeyValue>& StatAttributes)
+			DDCResourceStatsTotal.Accumulate(Stat);
+		}
+
+		DDCResourceStats.Emplace(DDCResourceStatsTotal);
+
+		for (const FDerivedDataCacheResourceStat& Stat : DDCResourceStats)
 		{
-			if (StatName.EndsWith(TEXT(".Usage"), ESearchCase::IgnoreCase))
+			const FString BaseName = TEXT("DDC.Resource.") + Stat.AssetType;
+
 			{
-				// Anything that ends in .Usage is assumed to be an instance of FCookStats.FDDCResourceUsageStats. We'll log that using custom formatting.
-				FString AssetType = StatName;
-				AssetType.RemoveFromEnd(TEXT(".Usage"), ESearchCase::IgnoreCase);
-				// See if the asset has a subtype (found via the "Node" parameter")
-				const FCookStatsManager::StringKeyValue* AssetSubType = StatAttributes.FindByPredicate([](const FCookStatsManager::StringKeyValue& Item) { return Item.Key == TEXT("Node"); });
-				if (AssetSubType && AssetSubType->Value.Len() > 0)
-				{
-					AssetType += FString::Printf(TEXT(" (%s)"), *AssetSubType->Value);
-				}
-				// Pull the Time and Size attributes and AddOrAccumulate them into the set of stats. Ugly string/container manipulation code courtesy of UE4/C++.
-				const FCookStatsManager::StringKeyValue* AssetTimeSecAttr = StatAttributes.FindByPredicate([](const FCookStatsManager::StringKeyValue& Item) { return Item.Key == TEXT("TimeSec"); });
-				double AssetTimeSec = 0.0;
-				if (AssetTimeSecAttr)
-				{
-					LexFromString(AssetTimeSec, *AssetTimeSecAttr->Value);
-				}
-				const FCookStatsManager::StringKeyValue* AssetSizeMBAttr = StatAttributes.FindByPredicate([](const FCookStatsManager::StringKeyValue& Item) { return Item.Key == TEXT("MB"); });
-				double AssetSizeMB = 0.0;
-				if (AssetSizeMBAttr)
-				{
-					LexFromString(AssetSizeMB, *AssetSizeMBAttr->Value);
-				}
-				const FCookStatsManager::StringKeyValue* ThreadNameAttr = StatAttributes.FindByPredicate([](const FCookStatsManager::StringKeyValue& Item) { return Item.Key == TEXT("ThreadName"); });
-				bool bIsGameThreadTime = ThreadNameAttr != nullptr && ThreadNameAttr->Value == TEXT("GameThread");
-
-				const FCookStatsManager::StringKeyValue* HitOrMissAttr = StatAttributes.FindByPredicate([](const FCookStatsManager::StringKeyValue& Item) { return Item.Key == TEXT("HitOrMiss"); });
-				bool bWasMiss = HitOrMissAttr != nullptr && HitOrMissAttr->Value == TEXT("Miss");
-				int64 AssetsBuilt = 0;
-				if (bWasMiss)
-				{
-					const FCookStatsManager::StringKeyValue* CountAttr = StatAttributes.FindByPredicate([](const FCookStatsManager::StringKeyValue& Item) { return Item.Key == TEXT("Count"); });
-					if (CountAttr)
-					{
-						LexFromString(AssetsBuilt, *CountAttr->Value);
-					}
-				}
-
-				TotalAssetsBuilt++;
-				TotalAssetTimeSec += AssetTimeSec;
-				TotalAssetSizeMB += AssetSizeMB;
-
-				FDDCResourceUsageStat Stat(AssetType, AssetTimeSec, bIsGameThreadTime, AssetSizeMB, AssetsBuilt);
-				FDDCResourceUsageStat* ExistingStat = DDCResourceUsageStats.Find(Stat.AssetType);
-				if (ExistingStat)
-				{
-					ExistingStat->Accumulate(Stat);
-				}
-				else
-				{
-					DDCResourceUsageStats.Add(Stat);
-				}
+				FString AttrName = BaseName + TEXT(".BuildCount");
+				Attributes.Emplace(MoveTemp(AttrName), Stat.BuildCount);
 			}
-			else if (StatName == TEXT("DDC.Summary"))
+			
 			{
-				DDCSummaryStats.Append(StatAttributes);
+				FString AttrName = BaseName + TEXT(".BuildTimeSec");
+				Attributes.Emplace(MoveTemp(AttrName), Stat.BuildTimeSec);
 			}
-		};
 
-		// Grab the DDC stats
-		FCookStatsManager::LogCookStats(LogStatsFunc);
+			{
+				FString AttrName = BaseName + TEXT(".BuildSizeMB");
+				Attributes.Emplace(MoveTemp(AttrName), Stat.BuildSizeMB);
+			}
 
-		for (const FDDCResourceUsageStat& Stat : DDCResourceUsageStats)
-		{
 			{
-				FString AttrName = "DDC.Resource." + Stat.AssetType + ".Built";
-				Attributes.Emplace(AttrName, Stat.AssetsBuilt);
+				FString AttrName = BaseName + TEXT(".LoadCount");
+				Attributes.Emplace(MoveTemp(AttrName), Stat.LoadCount);
 			}
-			
+
 			{
-				FString AttrName = "DDC.Resource." + Stat.AssetType + ".TimeSec";
-				Attributes.Emplace(AttrName, Stat.TotalTimeSec);
+				FString AttrName = BaseName + TEXT(".LoadTimeSec");
+				Attributes.Emplace(MoveTemp(AttrName), Stat.LoadTimeSec);
 			}
 
 			{
-				FString AttrName = "DDC.Resource." + Stat.AssetType + ".SizeMB";
-				Attributes.Emplace(AttrName, Stat.SizeMB);
+				FString AttrName = BaseName + TEXT(".LoadSizeMB");
+				Attributes.Emplace(MoveTemp(AttrName), Stat.LoadSizeMB);
 			}
 		}
 
-		Attributes.Emplace(TEXT("DDC.Resource.TotalAssetsBuilt"), TotalAssetsBuilt);
-		Attributes.Emplace(TEXT("DDC.Resource.TotalAssetTimeSec"), TotalAssetTimeSec);
-		Attributes.Emplace(TEXT("DDC.Resource.TotalAssetSizeMB"), TotalAssetSizeMB);
+		TArray<FDerivedDataCacheSummaryStat> DDCSummaryStats;
+
+		// Grab the summary stats
+		GatherDerivedDataCacheSummaryStats(DDCSummaryStats);
 		
-		for (const FCookStatsManager::StringKeyValue& Attr : DDCSummaryStats)
+		for (const FDerivedDataCacheSummaryStat& Stat : DDCSummaryStats)
 		{
-			FString FormattedAttrName = "DDC.Summary." + Attr.Key;
-			Attributes.Emplace(FormattedAttrName, Attr.Value);
+			FString FormattedAttrName = "DDC.Summary." + Stat.Key;
+			Attributes.Emplace(FormattedAttrName, Stat.Value);
 		}
 		
 		FStudioAnalytics::GetProvider().RecordEvent(TEXT("Core.Loading"), Attributes);
diff --git a/Engine/Source/Runtime/SlateCore/Private/Styling/StarshipCoreStyle.cpp b/Engine/Source/Runtime/SlateCore/Private/Styling/StarshipCoreStyle.cpp
--- a/Engine/Source/Runtime/SlateCore/Private/Styling/StarshipCoreStyle.cpp
+++ b/Engine/Source/Runtime/SlateCore/Private/Styling/StarshipCoreStyle.cpp
@@ -295,6 +295,8 @@ TSharedRef<ISlateStyle> FStarshipCoreStyle::Create()
 	
 		Style->Set("Icons.Download", new IMAGE_BRUSH("Icons/icon_Downloads_16x", Icon16x16));
 		Style->Set("Icons.Server", new IMAGE_BRUSH_SVG("Starship/Common/server", Icon16x16));
+		Style->Set("Icons.Cloud", new IMAGE_BRUSH("Icons/icon_Downloads_16x", Icon16x16));
+		Style->Set("Icons.Local", new IMAGE_BRUSH_SVG("Starship/Common/server", Icon16x16));
 
 		Style->Set("Icons.Error", new IMAGE_BRUSH_SVG("Starship/Common/alert-triangle", Icon16x16));
 		Style->Set("Icons.ErrorWithColor", new IMAGE_BRUSH_SVG("Starship/Common/alert-triangle", Icon16x16, FStyleColors::Error));
