PR for adding shaded preview for volumes

diff --git a/Engine/Source/Editor/UnrealEd/Private/LevelEditorViewport.cpp b/Engine/Source/Editor/UnrealEd/Private/LevelEditorViewport.cpp
--- a/Engine/Source/Editor/UnrealEd/Private/LevelEditorViewport.cpp
+++ b/Engine/Source/Editor/UnrealEd/Private/LevelEditorViewport.cpp
@@ -4324,46 +4324,60 @@ void FLevelEditorViewportClient::Draw(const FSceneView* View,FPrimitiveDrawInter
 
 void FLevelEditorViewportClient::DrawBrushDetails(const FSceneView* View, FPrimitiveDrawInterface* PDI)
 {
-	if (GEditor->bShowBrushMarkerPolys)
+	// Draw translucent polygons on brushes and volumes
+
+	for (TActorIterator<ABrush> It(GetWorld()); It; ++It)
 	{
-		// Draw translucent polygons on brushes and volumes
+		ABrush* Brush = *It;
 
-		for (TActorIterator<ABrush> It(GetWorld()); It; ++It)
+		// Brush->Brush is checked to safe from brushes that were created without having their brush members attached.
+		// Check whether it satisfies the ShowBrushMarkersPoly condition of being selected or whether bDisplayShadedVolume is true
+		if (Brush->Brush && (FActorEditorUtils::IsABuilderBrush(Brush) || Brush->IsVolumeBrush()) && 
+			(Brush->bDisplayShadedVolume || (GEditor->bShowBrushMarkerPolys && ModeTools->GetSelectedActors()->IsSelected(Brush))))
 		{
-			ABrush* Brush = *It;
+			// Build a mesh by basically drawing the triangles of each 
+			FDynamicMeshBuilder MeshBuilder(View->GetFeatureLevel());
+			int32 VertexOffset = 0;
 
-			// Brush->Brush is checked to safe from brushes that were created without having their brush members attached.
-			if (Brush->Brush && (FActorEditorUtils::IsABuilderBrush(Brush) || Brush->IsVolumeBrush()) && ModeTools->GetSelectedActors()->IsSelected(Brush))
+			for (int32 PolyIdx = 0; PolyIdx < Brush->Brush->Polys->Element.Num(); ++PolyIdx)
 			{
-				// Build a mesh by basically drawing the triangles of each 
-				FDynamicMeshBuilder MeshBuilder(View->GetFeatureLevel());
-				int32 VertexOffset = 0;
+				const FPoly* Poly = &Brush->Brush->Polys->Element[PolyIdx];
 
-				for (int32 PolyIdx = 0; PolyIdx < Brush->Brush->Polys->Element.Num(); ++PolyIdx)
+				if (Poly->Vertices.Num() > 2)
 				{
-					const FPoly* Poly = &Brush->Brush->Polys->Element[PolyIdx];
-
-					if (Poly->Vertices.Num() > 2)
-					{
-						const FVector Vertex0 = Poly->Vertices[0];
-						FVector Vertex1 = Poly->Vertices[1];
-
-						MeshBuilder.AddVertex(Vertex0, FVector2D::ZeroVector, FVector(1, 0, 0), FVector(0, 1, 0), FVector(0, 0, 1), FColor::White);
-						MeshBuilder.AddVertex(Vertex1, FVector2D::ZeroVector, FVector(1, 0, 0), FVector(0, 1, 0), FVector(0, 0, 1), FColor::White);
+					const FVector Vertex0 = Poly->Vertices[0];
+					FVector Vertex1 = Poly->Vertices[1];
 
-						for (int32 VertexIdx = 2; VertexIdx < Poly->Vertices.Num(); ++VertexIdx)
-						{
-							const FVector Vertex2 = Poly->Vertices[VertexIdx];
-							MeshBuilder.AddVertex(Vertex2, FVector2D::ZeroVector, FVector(1, 0, 0), FVector(0, 1, 0), FVector(0, 0, 1), FColor::White);
-							MeshBuilder.AddTriangle(VertexOffset, VertexOffset + VertexIdx, VertexOffset + VertexIdx - 1);
-							Vertex1 = Vertex2;
-						}
+					MeshBuilder.AddVertex(Vertex0, FVector2D::ZeroVector, FVector(1, 0, 0), FVector(0, 1, 0), FVector(0, 0, 1), FColor::White);
+					MeshBuilder.AddVertex(Vertex1, FVector2D::ZeroVector, FVector(1, 0, 0), FVector(0, 1, 0), FVector(0, 0, 1), FColor::White);
 
-						// Increment the vertex offset so the next polygon uses the correct vertex indices.
-						VertexOffset += Poly->Vertices.Num();
+					for (int32 VertexIdx = 2; VertexIdx < Poly->Vertices.Num(); ++VertexIdx)
+					{
+						const FVector Vertex2 = Poly->Vertices[VertexIdx];
+						MeshBuilder.AddVertex(Vertex2, FVector2D::ZeroVector, FVector(1, 0, 0), FVector(0, 1, 0), FVector(0, 0, 1), FColor::White);
+						MeshBuilder.AddTriangle(VertexOffset, VertexOffset + VertexIdx, VertexOffset + VertexIdx - 1);
+						Vertex1 = Vertex2;
 					}
+
+					// Increment the vertex offset so the next polygon uses the correct vertex indices.
+					VertexOffset += Poly->Vertices.Num();
 				}
+			}
+			// Use the material below when bDisplayShadedVolume is true
+			if (Brush->bDisplayShadedVolume)
+			{
+				// Allocate the material proxy and register it so it can be deleted properly once the rendering is done with it.
+				FColor VolumeColor = Brush->GetWireColor();
+				VolumeColor.A = Brush->ShadedVolumeOpacityValue * 255.0f;
+				FDynamicColoredMaterialRenderProxy* MaterialProxy = new FDynamicColoredMaterialRenderProxy(GEngine->GeomMaterial->GetRenderProxy(), VolumeColor);
+				PDI->RegisterDynamicResource(MaterialProxy);
 
+				// Flush the mesh triangles.
+				// This will get the default behavior of receiving decals but not having a hitproxy ID
+				MeshBuilder.Draw(PDI, Brush->ActorToWorld().ToMatrixWithScale(), MaterialProxy, SDPG_World, false);
+			}
+			else
+			{
 				// Allocate the material proxy and register it so it can be deleted properly once the rendering is done with it.
 				FDynamicColoredMaterialRenderProxy* MaterialProxy = new FDynamicColoredMaterialRenderProxy(GEngine->EditorBrushMaterial->GetRenderProxy(), Brush->GetWireColor());
 				PDI->RegisterDynamicResource(MaterialProxy);
diff --git a/Engine/Source/Runtime/Engine/Classes/Engine/Brush.h b/Engine/Source/Runtime/Engine/Classes/Engine/Brush.h
--- a/Engine/Source/Runtime/Engine/Classes/Engine/Brush.h
+++ b/Engine/Source/Runtime/Engine/Classes/Engine/Brush.h
@@ -113,6 +113,14 @@ public:
 #if WITH_EDITORONLY_DATA
 	UPROPERTY(VisibleAnywhere, Instanced, Category=BrushBuilder)
 	TObjectPtr<class UBrushBuilder> BrushBuilder;
+	
+	/** If true, display the brush with a shaded volume */
+	UPROPERTY(Transient, EditAnywhere, Category=BrushSettings)
+	uint32 bDisplayShadedVolume:1;
+
+	/** Value used to set the opacity for the shaded volume, between 0-1 */
+	UPROPERTY(Transient, EditAnywhere, Category=BrushSettings, meta=(ClampMin=0.0, ClampMax=1.0))
+	float ShadedVolumeOpacityValue = 0.25f;
 #endif
 
 	/** Flag set when we are in a manipulation (scaling, translation, brush builder param change etc.) */
