Unshelved from pending changelist '16937032':

Fix incorrect dependency on DDC for non-editor shipping builds.

diff --git a/Engine/Source/Editor/UnrealEd/UnrealEd.Build.cs b/Engine/Source/Editor/UnrealEd/UnrealEd.Build.cs
--- a/Engine/Source/Editor/UnrealEd/UnrealEd.Build.cs
+++ b/Engine/Source/Editor/UnrealEd/UnrealEd.Build.cs
@@ -118,7 +118,6 @@ public class UnrealEd : ModuleRules
 				"SubobjectDataInterface",
 				"SubobjectEditor",
 				"PhysicsUtilities",
-				"DerivedDataCache",
 			}
 		);
 
