DDC2 remote execution cleanup:
-use newer FCompressedBuffer::Compress method when providing buffers for worker executable files
-use the preferred *.action and *.output file extensions over *.uddba and *.uddbo
-use the preferred "-Build=" commandline argument over "-b="

diff --git a/Engine/Source/Developer/DesktopPlatform/Private/TargetReceiptBuildWorker.cpp b/Engine/Source/Developer/DesktopPlatform/Private/TargetReceiptBuildWorker.cpp
--- a/Engine/Source/Developer/DesktopPlatform/Private/TargetReceiptBuildWorker.cpp
+++ b/Engine/Source/Developer/DesktopPlatform/Private/TargetReceiptBuildWorker.cpp
@@ -2,6 +2,7 @@
 
 #include "TargetReceiptBuildWorker.h"
 
+#include "Compression/OodleDataCompression.h"
 #include "DerivedDataBuildWorker.h"
 #include "Features/IModularFeatures.h"
 #include "HAL/FileManager.h"
@@ -18,7 +19,7 @@
 DEFINE_LOG_CATEGORY_STATIC(LogTargetReceiptBuildWorker, Log, All);
 
 // This needs to match the version exported in DerivedDataBuildWorker.Build.cs 
-const FGuid FTargetReceiptBuildWorker::WorkerReceiptVersion(TEXT("a547b2fa-e2f9-47c6-80c8-df688fb99f34"));
+const FGuid FTargetReceiptBuildWorker::WorkerReceiptVersion(TEXT("dab5352e-a5a7-4793-a7a3-1d4acad6aff2"));
 
 FTargetReceiptBuildWorker::FWorkerPath::FWorkerPath(FStringView InLocalPathRoot, FStringView InLocalPathSuffix, FStringView InRemotePathRoot, FStringView InRemotePathSuffix)
 {
@@ -110,7 +111,11 @@ public:
 							Ar->Serialize(MutableBuffer.GetData(), TotalSize);
 							if (Ar->Close())
 							{
-								FileDataBuffers.Emplace(FCompressedBuffer::Compress(NAME_None, MutableBuffer.MoveToShared()));
+								FileDataBuffers.Emplace(FCompressedBuffer::Compress(
+										MutableBuffer.MoveToShared(),
+										FOodleDataCompression::ECompressor::NotSet,
+										FOodleDataCompression::ECompressionLevel::None
+									));
 							}
 						}
 					}
diff --git a/Engine/Source/Editor/UnrealEd/Private/DerivedDataBuildRemoteExecutor.cpp b/Engine/Source/Editor/UnrealEd/Private/DerivedDataBuildRemoteExecutor.cpp
--- a/Engine/Source/Editor/UnrealEd/Private/DerivedDataBuildRemoteExecutor.cpp
+++ b/Engine/Source/Editor/UnrealEd/Private/DerivedDataBuildRemoteExecutor.cpp
@@ -42,7 +42,6 @@ public:
 	FRemoteBuildWorkerExecutor()
 	: GlobalExecutionTimeoutSeconds(-1)
 	, Salt(TEXT("807c6a49-0657-49f3-b498-fd457213c0a7"))
-	, BaseDirectoryPath(TEXT("Engine/Binaries/Win64"))
 	, RemoteExecutor(nullptr)
 	, ContentAddressableStorage(nullptr)
 	, bEnabled(false)
@@ -158,6 +157,7 @@ public:
 		TMultiMap<FDigest, FVariantIndex> DigestFilesystemIndex;
 		TMap<FStringView, int32> PathToDirectoryIndex;
 		TMap<int32, FString> FileIndexToInputKey;
+		FString BaseDirectoryPath;
 
 		// Unique items in the tree
 		FCommand Command;
@@ -307,7 +307,10 @@ public:
 
 		// This base directory must be created as worker executables (even those that don't exist in this directory) will attempt to change directories into it during startup.
 		int32 BaseDirectoryIndex = INDEX_NONE;
-		GetOrAddMerkleTreeDirectory(State, BaseDirectoryPath, BaseDirectoryIndex);
+		TStringBuilder<128> BaseDirectoryPathBuilder;
+		FPathViews::Append(BaseDirectoryPathBuilder, TEXT("Engine/Binaries/"), State.BuildWorker.GetHostPlatform());
+		State.BaseDirectoryPath = BaseDirectoryPathBuilder.ToString();
+		GetOrAddMerkleTreeDirectory(State, State.BaseDirectoryPath, BaseDirectoryIndex);
 
 		FCbWriter BuildActionWriter;
 		State.BuildAction.Save(BuildActionWriter);
@@ -316,7 +319,7 @@ public:
 		State.BuildActionContentBytes = FCompositeBuffer(UncompressedBuildActionContentBytes.MoveToShared());
 		State.BuildActionDigest.Hash = FIoHash::HashBuffer(State.BuildActionContentBytes);
 		State.BuildActionDigest.SizeBytes = State.BuildActionContentBytes.GetSize();
-		AddMerkleTreeFile(State, TEXT("build.uddba"), State.BuildActionDigest.Hash, State.BuildActionDigest.SizeBytes, false, EFileType::BuildAction, State.BuildActionContentBytes);
+		AddMerkleTreeFile(State, TEXT("Build.action"), State.BuildActionDigest.Hash, State.BuildActionDigest.SizeBytes, false, EFileType::BuildAction, State.BuildActionContentBytes);
 
 
 		if (!State.PathToDirectoryIndex.IsEmpty())
@@ -327,9 +330,9 @@ public:
 		}
 
 		State.Command.OutputPaths.Add("Outputs");
-		State.Command.OutputPaths.Add("build.uddbo");
+		State.Command.OutputPaths.Add("Build.output");
 		State.Command.Arguments.Add(FString(State.BuildWorker.GetPath()));
-		State.Command.Arguments.Add("-b=build.uddba");
+		State.Command.Arguments.Add("-Build=Build.action");
 		State.BuildWorker.IterateEnvironment([&State] (FStringView Name, FStringView Value)
 			{
 				FCommand::FEnvironmentVariable& EnvVar = State.Command.EnvironmentVariables.AddDefaulted_GetRef();
@@ -546,7 +549,7 @@ public:
 		TOptional<FDigest> BuildOutputDigest;
 		for (const FOutputFile& ExecuteOutputFile : State.ExecuteResponse.Result.OutputFiles)
 		{
-			if (ExecuteOutputFile.Path == TEXT("build.uddbo"))
+			if (ExecuteOutputFile.Path == TEXT("Build.output"))
 			{
 				BuildOutputDigest = ExecuteOutputFile.Digest;
 				break;
@@ -725,7 +728,7 @@ public:
 
 	TConstArrayView<FStringView> GetHostPlatforms() const final
 	{
-		static constexpr FStringView HostPlatforms[]{TEXT("Win64"_SV)};
+		static constexpr FStringView HostPlatforms[]{TEXT("Win64"_SV), TEXT("Linux"_SV), TEXT("Mac"_SV)};
 		return HostPlatforms;
 	}
 
@@ -880,7 +883,6 @@ private:
 	FString InstanceName;
 	int GlobalExecutionTimeoutSeconds;
 	const FStringView Salt;
-	const FString BaseDirectoryPath;
 	IRemoteExecutor* RemoteExecutor;
 	IContentAddressableStorage* ContentAddressableStorage;
 	bool bEnabled;
diff --git a/Engine/Source/Programs/DerivedDataBuildWorker/DerivedDataBuildWorker.Build.cs b/Engine/Source/Programs/DerivedDataBuildWorker/DerivedDataBuildWorker.Build.cs
--- a/Engine/Source/Programs/DerivedDataBuildWorker/DerivedDataBuildWorker.Build.cs
+++ b/Engine/Source/Programs/DerivedDataBuildWorker/DerivedDataBuildWorker.Build.cs
@@ -76,6 +76,6 @@ public class DerivedDataBuildWorker : ModuleRules
 		}
 
 		// This needs to match the version defined in TargetReceiptBuildWorkerFactory.cpp
-		AdditionalPropertiesForReceipt.Add("DerivedDataBuildWorkerReceiptVersion", "a547b2fa-e2f9-47c6-80c8-df688fb99f34");
+		AdditionalPropertiesForReceipt.Add("DerivedDataBuildWorkerReceiptVersion", "dab5352e-a5a7-4793-a7a3-1d4acad6aff2");
 	}
 }
diff --git a/Engine/Source/Programs/DerivedDataBuildWorker/Private/DerivedDataBuildWorker.cpp b/Engine/Source/Programs/DerivedDataBuildWorker/Private/DerivedDataBuildWorker.cpp
--- a/Engine/Source/Programs/DerivedDataBuildWorker/Private/DerivedDataBuildWorker.cpp
+++ b/Engine/Source/Programs/DerivedDataBuildWorker/Private/DerivedDataBuildWorker.cpp
@@ -290,7 +290,7 @@ void FBuildWorkerProgram::BuildComplete(FBuildActionCompleteParams&& Params) con
 		}
 	}
 
-	const FString OutputPath = FPaths::ChangeExtension(FString(Output.GetName()), TEXT("uddbo"));
+	const FString OutputPath = FPaths::ChangeExtension(FString(Output.GetName()), TEXT("output"));
 	if (TUniquePtr<FArchive> Ar{IFileManager::Get().CreateFileWriter(*OutputPath, FILEWRITE_Silent)})
 	{
 		FCbWriter OutputWriter;
