Change the DDC2 remote executor to use async operations.  Only exception is worker file loading which is still synchronous until IRequest changes happen to make chaining easier.

diff --git a/Engine/Source/Editor/Experimental/BazelExecutor/Private/Execution.cpp b/Engine/Source/Editor/Experimental/BazelExecutor/Private/Execution.cpp
--- a/Engine/Source/Editor/Experimental/BazelExecutor/Private/Execution.cpp
+++ b/Engine/Source/Editor/Experimental/BazelExecutor/Private/Execution.cpp
@@ -81,20 +81,20 @@ bool FExecution::Execute(const FExecuteRequest& Request, FExecuteResponse& Respo
 				break;
 			}
 			ProtoConverter::FromProto(ProtoResponse, Response);
-			if ((grpc::StatusCode)ProtoResponse.status().code() != grpc::StatusCode::OK) {
-
+			if (!Response.Status.Ok())
+			{
 				if ((Response.Status.Code == EStatusCode::RESOURCE_EXHAUSTED) ||
 					(Response.Status.Code == EStatusCode::UNAVAILABLE))
 				{
 					UE_LOG(LogBazelExecutor, Display, TEXT("Execute: %s Info: %s"),
 						UTF8_TO_TCHAR(Operation.name().c_str()),
-						UTF8_TO_TCHAR(ProtoResponse.status().message().c_str()));
+						*Response.Status.Message);
 				}
 				else
 				{
 					UE_LOG(LogBazelExecutor, Error, TEXT("Execute: %s Error: %s"),
 						UTF8_TO_TCHAR(Operation.name().c_str()),
-						UTF8_TO_TCHAR(ProtoResponse.status().message().c_str()));
+						*Response.Status.Message);
 				}
 				break;
 			}
@@ -181,14 +181,25 @@ TFuture<FExecuteResponse> FExecution::ExecuteAsync(const FExecuteRequest& Reques
 			return;
 		}
 
-		if ((grpc::StatusCode)ProtoResponse.status().code() != grpc::StatusCode::OK) {
-			UE_LOG(LogBazelExecutor, Error, TEXT("ExecuteAsync: %s Error: %s"),
-				UTF8_TO_TCHAR(Operation.name().c_str()),
-				UTF8_TO_TCHAR(ProtoResponse.status().message().c_str()));
-		}
-
 		FExecuteResponse Response;
 		ProtoConverter::FromProto(ProtoResponse, Response);
+		if (!Response.Status.Ok())
+		{
+			if ((Response.Status.Code == EStatusCode::RESOURCE_EXHAUSTED) ||
+				(Response.Status.Code == EStatusCode::UNAVAILABLE))
+			{
+				UE_LOG(LogBazelExecutor, Display, TEXT("ExecuteAsync: %s Info: %s"),
+					UTF8_TO_TCHAR(Operation.name().c_str()),
+					*Response.Status.Message);
+			}
+			else
+			{
+				UE_LOG(LogBazelExecutor, Error, TEXT("ExecuteAsync: %s Error: %s"),
+					UTF8_TO_TCHAR(Operation.name().c_str()),
+					*Response.Status.Message);
+			}
+		}
+
 		ReturnPromise->EmplaceValue(MoveTemp(Response));
 	};
 
diff --git a/Engine/Source/Editor/UnrealEd/Private/DerivedDataBuildRemoteExecutor.cpp b/Engine/Source/Editor/UnrealEd/Private/DerivedDataBuildRemoteExecutor.cpp
--- a/Engine/Source/Editor/UnrealEd/Private/DerivedDataBuildRemoteExecutor.cpp
+++ b/Engine/Source/Editor/UnrealEd/Private/DerivedDataBuildRemoteExecutor.cpp
@@ -9,6 +9,8 @@
 #include "DerivedDataBuildOutput.h"
 #include "DerivedDataBuildWorker.h"
 #include "Features/IModularFeatures.h"
+#include "HAL/Event.h"
+#include "HAL/PlatformProcess.h"
 #include "IContentAddressableStorage.h"
 #include "IRemoteExecutor.h"
 #include "Messages.h"
@@ -36,53 +38,50 @@ namespace UE::DerivedData
 
 DEFINE_LOG_CATEGORY_STATIC(LogDerivedDataBuildRemoteExecutor, Log, All);
 
-class FRemoteBuildWorkerExecutor final: public IBuildWorkerExecutor
+class FRemoteBuildWorkerExecutor;
+
+class FRemoteBuildExecutionRequest final : public IRequest, public FThreadSafeRefCountedObject
 {
 public:
-	FRemoteBuildWorkerExecutor()
-	: GlobalExecutionTimeoutSeconds(-1)
-	, Salt(TEXT("807c6a49-0657-49f3-b498-fd457213c0a7"))
-	, RemoteExecutor(nullptr)
-	, ContentAddressableStorage(nullptr)
-	, bEnabled(false)
-	{
-		check(IsInGameThread()); // initialization from the main thread is expected to allow config reading for the limiting heuristics
-		check(GConfig && GConfig->IsReadyForUse());
-
-		GConfig->GetString(TEXT("DerivedDataBuildRemoteExecutor"), TEXT("InstanceName"), InstanceName, GEngineIni);
-		GConfig->GetInt(TEXT("DerivedDataBuildRemoteExecutor"), TEXT("GlobalExecutionTimeoutSeconds"), GlobalExecutionTimeoutSeconds, GEngineIni);
-
-		const FName RemoteExecutionFeatureName(TEXT("RemoteExecution"));
-		IModularFeatures& ModularFeatures = IModularFeatures::Get();
+	FRemoteBuildExecutionRequest(
+		FRemoteBuildWorkerExecutor& InExecutor,
+		const FBuildAction& Action,
+		const FOptionalBuildInputs& Inputs,
+		const FBuildWorker& Worker,
+		EBuildPolicy Policy,
+		EPriority Priority,
+		FOnBuildWorkerActionComplete&& OnComplete);
 
-		if (FParse::Param(FCommandLine::Get(), TEXT("DDC2RemoteExecution")))
-		{
-			FModuleManager::Get().LoadModule("BazelExecutor");
-			if (ModularFeatures.IsModularFeatureAvailable(RemoteExecutionFeatureName))
-			{
-				RemoteExecutor = &ModularFeatures.GetModularFeature<IRemoteExecutor>(RemoteExecutionFeatureName);
-				if (RemoteExecutor)
-				{
-					ContentAddressableStorage = RemoteExecutor->GetContentAddressableStorage();
-				}
-			}
-		}
+	virtual ~FRemoteBuildExecutionRequest();
 
-		bEnabled = (RemoteExecutor != nullptr) && (ContentAddressableStorage != nullptr) && !InstanceName.IsEmpty();
-		if (bEnabled)
-		{
-			ModularFeatures.RegisterModularFeature(IBuildWorkerExecutor::GetFeatureName(), this);
-		}
+	// IRequest interface
+	virtual void SetPriority(EPriority Priority) override
+	{
 	}
 
-	virtual ~FRemoteBuildWorkerExecutor()
+	virtual void Cancel() override
 	{
-		if (bEnabled)
-		{
-			IModularFeatures::Get().UnregisterModularFeature(IBuildWorkerExecutor::GetFeatureName(), this);
-		}
+		bCancelPending.store(true, std::memory_order_relaxed);
+		Wait();
+	}
+	virtual void Wait() override
+	{
+		CompletionEvent->Wait(MAX_uint32);
+	}
+	virtual bool Poll() override
+	{
+		return CompletionEvent->Wait(0);
+	}
+	virtual void AddRef() const override
+	{
+		FThreadSafeRefCountedObject::AddRef();
+	}
+	virtual void Release() const override
+	{
+		FThreadSafeRefCountedObject::Release();
 	}
 
+private:
 	enum class ENodeType
 	{
 		File,
@@ -90,7 +89,6 @@ public:
 		//Symlink,
 	};
 
-
 	enum class EFileType
 	{
 		Worker,
@@ -175,7 +173,6 @@ public:
 
 		// Step 3: Batch update blobs (upload)
 		FBatchUpdateBlobsRequest BatchUpdateBlobsRequest;
-		FBatchUpdateBlobsResponse BatchUpdateBlobsResponse;
 
 		// Step 4: Execute
 		FExecuteRequest ExecuteRequest;
@@ -183,450 +180,86 @@ public:
 
 		// Step 5: Batch read blobs (download)
 		FBatchReadBlobsRequest BatchReadBlobsRequest;
-		FBatchReadBlobsResponse BatchReadBlobsResponse;
 	};
 
-	static FMerkleTreeDirectoryBuilder& GetOrAddMerkleTreeDirectory(FRemoteExecutionState& State, FStringView Path, int32& OutDirectoryBuilderIndex)
-	{
-		int32& DirectoryBuilderIndex = State.PathToDirectoryIndex.FindOrAdd(Path, INDEX_NONE);
-		if (DirectoryBuilderIndex == INDEX_NONE)
-		{
-			DirectoryBuilderIndex = State.Directories.Num();
-			FMerkleTreeDirectoryBuilder& NewNode = State.Directories.AddDefaulted_GetRef();
-			NewNode.Name = Path.IsEmpty() ? Path : FPathViews::GetCleanFilename(Path);
-		}
-
-		OutDirectoryBuilderIndex = DirectoryBuilderIndex;
-
-		int32 ContainingDirectoryIndex = DirectoryBuilderIndex;
-		while (!Path.IsEmpty())
-		{
-			int32 SubDirIndex = ContainingDirectoryIndex;
-			Path = FPathViews::GetPath(Path);
-			GetOrAddMerkleTreeDirectory(State, Path, ContainingDirectoryIndex).SubDirIndices.AddUnique(SubDirIndex);
-		}
-
-		return State.Directories[OutDirectoryBuilderIndex];
-	}
-
-	static int32 AddMerkleTreeFile(FRemoteExecutionState& State, FStringView Path, const FIoHash& RawHash, uint64 RawSize, bool bIsExecutable, EFileType FileType, FCompositeBuffer ContentBytes = FCompositeBuffer::Null)
-	{
-		FStringView ContainingDirectoryPath = FPathViews::GetPath(Path);
-		int32 ContainingDirectoryIndex = INDEX_NONE;
-		FMerkleTreeDirectoryBuilder* ContainingDirectoryBuilder = &GetOrAddMerkleTreeDirectory(State, ContainingDirectoryPath, ContainingDirectoryIndex);
-		FFileNode& NewNode = ContainingDirectoryBuilder->Directory.Files.AddDefaulted_GetRef();
-		NewNode.Name = FPathViews::GetCleanFilename(Path);
-		NewNode.Digest.Hash = RawHash;
-		NewNode.Digest.SizeBytes = RawSize;
-		NewNode.IsExecutable = bIsExecutable;
-
-		int32 NewFileIndex = State.Files.Num();
-		State.DigestFilesystemIndex.Add(NewNode.Digest, FVariantIndex(ENodeType::File, NewFileIndex));
-		FMerkleTreeFileBuilder& FileBuilder = State.Files.AddDefaulted_GetRef();
-		FileBuilder.File = NewNode; // Duplicates the node in the state's file array
-		FileBuilder.Type = FileType;
-		if (ContentBytes)
-		{
-			FileBuilder.ContentBytes = ContentBytes;
-		}
-
-		return NewFileIndex;
-	}
-
-	const FDigest& BuildMerkleTreeDirectoryDigest(FRemoteExecutionState& State, int32 Index)
-	{
-		FMerkleTreeDirectoryBuilder& DirBuilder = State.Directories[Index];
-
-		if (DirBuilder.Digest.IsSet())
-		{
-			return DirBuilder.Digest.GetValue();
-		}
-
-		for (int32 SubDirIndex : DirBuilder.SubDirIndices)
-		{
-			FDirectoryNode& SubDirNode = DirBuilder.Directory.Directories.AddDefaulted_GetRef();
-			SubDirNode.Name = State.Directories[SubDirIndex].Name;
-			SubDirNode.Digest = BuildMerkleTreeDirectoryDigest(State, SubDirIndex);
-			State.DigestFilesystemIndex.Add(SubDirNode.Digest, FVariantIndex(ENodeType::Directory, SubDirIndex));
-		}
-
-		DirBuilder.Directory.Directories.Sort([] (const FDirectoryNode& A, const FDirectoryNode& B)
-		{
-			return A.Name < B.Name;
-		});
-
-		DirBuilder.Directory.Files.Sort([] (const FFileNode& A, const FFileNode& B)
-		{
-			return A.Name < B.Name;
-		});
-
-		ContentAddressableStorage->ToBlob(DirBuilder.Directory, DirBuilder.ContentBytes, DirBuilder.Digest.Emplace());
-
-		return DirBuilder.Digest.GetValue();
-	}
-
-	void BuildMerkleTreeNodes(FRemoteExecutionState& State)
-	{
-		TArray<FIoHash> WorkerFileHashes;
-		TArray<TTuple<FStringView, bool>> WorkerFileMeta;
-		State.BuildWorker.IterateExecutables([&WorkerFileHashes, &WorkerFileMeta] (FStringView Path, const FIoHash& RawHash, uint64 RawSize)
-			{
-				WorkerFileHashes.Emplace(RawHash);
-				WorkerFileMeta.Emplace(Path, true);
-			});
-
-		State.BuildWorker.IterateFiles([&WorkerFileHashes, &WorkerFileMeta] (FStringView Path, const FIoHash& RawHash, uint64 RawSize)
-			{
-				WorkerFileHashes.Emplace(RawHash);
-				WorkerFileMeta.Emplace(Path, false);
-			});
-
-		State.BuildWorker.FindFileData(WorkerFileHashes, EPriority::Normal,
-			[&State, &WorkerFileMeta] (FBuildWorkerFileDataCompleteParams&& Params)
-			{
-				uint32 MetaIndex = 0;
-				for (const FCompressedBuffer& Buffer : Params.Files)
-				{
-					const TTuple<FStringView, bool>& Meta = WorkerFileMeta[MetaIndex];
-					FCompositeBuffer DecompressedComposite = Buffer.DecompressToComposite();
-					AddMerkleTreeFile(State, Meta.Key, FIoHash::HashBuffer(DecompressedComposite), DecompressedComposite.GetSize(), Meta.Value, EFileType::Worker, Buffer.DecompressToComposite());
-					++MetaIndex;
-				}
-			}).Wait();
-
-		State.BuildAction.IterateInputs([&State] (FStringView Key, const FIoHash& RawHash, uint64 RawSize)
-			{
-				TStringBuilder<128> InputPath;
-				InputPath << TEXT("Inputs/") << RawHash;
-				const FString& NewInputPath = State.InputPaths.Emplace_GetRef(InputPath);
-				const FCompressedBuffer& Buffer = State.BuildInputs.Get().FindInput(Key);
-				check(!Buffer.IsNull());
-				int32 FileIndex = AddMerkleTreeFile(State, NewInputPath, FIoHash::HashBuffer(Buffer.GetCompressed()), Buffer.GetCompressedSize(), false, EFileType::Input, Buffer.GetCompressed());
-				State.FileIndexToInputKey.Emplace(FileIndex, Key);
-			});
-
-		// This base directory must be created as worker executables (even those that don't exist in this directory) will attempt to change directories into it during startup.
-		int32 BaseDirectoryIndex = INDEX_NONE;
-		TStringBuilder<128> BaseDirectoryPathBuilder;
-		FPathViews::Append(BaseDirectoryPathBuilder, TEXT("Engine/Binaries/"), State.BuildWorker.GetHostPlatform());
-		State.BaseDirectoryPath = BaseDirectoryPathBuilder.ToString();
-		GetOrAddMerkleTreeDirectory(State, State.BaseDirectoryPath, BaseDirectoryIndex);
-
-		FCbWriter BuildActionWriter;
-		State.BuildAction.Save(BuildActionWriter);
-		FUniqueBuffer UncompressedBuildActionContentBytes = FUniqueBuffer::Alloc(BuildActionWriter.GetSaveSize());
-		BuildActionWriter.Save(UncompressedBuildActionContentBytes);
-		State.BuildActionContentBytes = FCompositeBuffer(UncompressedBuildActionContentBytes.MoveToShared());
-		State.BuildActionDigest.Hash = FIoHash::HashBuffer(State.BuildActionContentBytes);
-		State.BuildActionDigest.SizeBytes = State.BuildActionContentBytes.GetSize();
-		AddMerkleTreeFile(State, TEXT("Build.action"), State.BuildActionDigest.Hash, State.BuildActionDigest.SizeBytes, false, EFileType::BuildAction, State.BuildActionContentBytes);
-
-
-		if (!State.PathToDirectoryIndex.IsEmpty())
-		{
-			int32 RootDirectoryIndex = State.PathToDirectoryIndex.FindChecked(TEXT(""));
-			State.Action.InputRootDigest = BuildMerkleTreeDirectoryDigest(State, RootDirectoryIndex);
-			State.DigestFilesystemIndex.Add(State.Action.InputRootDigest, FVariantIndex(ENodeType::Directory, RootDirectoryIndex));
-		}
-
-		State.Command.OutputPaths.Add("Outputs");
-		State.Command.OutputPaths.Add("Build.output");
-		State.Command.Arguments.Add(FString(State.BuildWorker.GetPath()));
-		State.Command.Arguments.Add("-Build=Build.action");
-		State.BuildWorker.IterateEnvironment([&State] (FStringView Name, FStringView Value)
-			{
-				FCommand::FEnvironmentVariable& EnvVar = State.Command.EnvironmentVariables.AddDefaulted_GetRef();
-				EnvVar.Name = Name;
-				EnvVar.Value = Value;
-			});
-		State.Command.EnvironmentVariables.Sort([] (const FCommand::FEnvironmentVariable& A, const FCommand::FEnvironmentVariable& B)
-			{
-				if (A.Name < B.Name)
-				{
-					return true;
-				}
-				else if (B.Name < A.Name)
-				{
-					return false;
-				}
-				return A.Value < B.Value;
-			});
-		ContentAddressableStorage->ToBlob(State.Command, State.CommandContentBytes, State.Action.CommandDigest);
-
-		State.Action.Salt = Salt;
-		ContentAddressableStorage->ToBlob(State.Action, State.ActionContentBytes, State.ActionDigest);
-	}
-
-	// Each step represents a synchronous unit of work that ends with an operation we want to have async in the future
-	FStatus DetermineMissingBlobs(FRemoteExecutionState& State)
-	{
-		BuildMerkleTreeNodes(State);
-
-		State.FindMissingBlobsRequest.InstanceName = InstanceName;
-		State.FindMissingBlobsRequest.BlobDigests.Add(State.ActionDigest);
-		State.FindMissingBlobsRequest.BlobDigests.Add(State.Action.CommandDigest);
-		for (const TPair<FDigest, FVariantIndex>& FilesystemItem : State.DigestFilesystemIndex)
-		{
-			State.FindMissingBlobsRequest.BlobDigests.AddUnique(FilesystemItem.Key);
-		}
-		return ContentAddressableStorage->FindMissingBlobs(State.FindMissingBlobsRequest, State.FindMissingBlobsResponse);
-	}
-
-	void GatherMissingInputFileBlobs(FRemoteExecutionState& State, TArray<FStringView>& OutMissingInputs)
-	{
-		for (const FDigest& MissingItem : State.FindMissingBlobsResponse.MissingBlobDigests)
-		{
-			TArray<FVariantIndex> DigestFilesystemEntries;
-			State.DigestFilesystemIndex.MultiFind(MissingItem, DigestFilesystemEntries);
-
-			for (FVariantIndex& VariantIndex : DigestFilesystemEntries)
-			{
-				if (VariantIndex.NodeType == ENodeType::File)
-				{
-					const FMerkleTreeFileBuilder& File = State.Files[VariantIndex.Index];
-					if ((File.Type == EFileType::Input) && File.ContentBytes.IsNull())
-					{
-						OutMissingInputs.Add(State.FileIndexToInputKey[VariantIndex.Index]);
-					}
-				}
-			}
-		}
-	}
+	FRemoteExecutionState State;
+	FOnBuildWorkerActionComplete CompletionCallback;
+	FRemoteBuildWorkerExecutor& Executor;
+	FEventRef CompletionEvent{EEventMode::ManualReset};
+	std::atomic<bool> bCancelPending;
+	bool bHeuristicBuildStarted;
+
+	// Merkle tree operations
+	FMerkleTreeDirectoryBuilder& GetOrAddMerkleTreeDirectory(FStringView Path, int32& OutDirectoryBuilderIndex);
+	int32 AddMerkleTreeFile(FStringView Path, const FIoHash& RawHash, uint64 RawSize, bool bIsExecutable, EFileType FileType, FCompositeBuffer ContentBytes = FCompositeBuffer::Null);
+	const FDigest& BuildMerkleTreeDirectoryDigest(int32 Index);
+	void BuildMerkleTreeNodes();
+
+	// General utility methods
+	void GatherMissingInputFileBlobs(TArray<FStringView>& OutMissingInputs);
+	bool ValidateUploadSuccess(const FBatchUpdateBlobsResponse& BatchUpdateBlobsResponse);
+	FOptionalBuildOutput ComposeBuildOutput(FBatchReadBlobsResponse& BatchReadBlobsResponse, EStatus& OutStatus);
+	bool ProcessCancellation();
+
+	// Async steps
+	TFuture<TPair<FStatus, FFindMissingBlobsResponse>> DetermineMissingBlobsAsync();
+	void LoadMissingWorkerFileBlobsAsync();
+	TFuture<TPair<FStatus, FBatchUpdateBlobsResponse>> UploadMissingBlobsAsync();
+	TFuture<FExecuteResponse> ExecuteBuildAsync();
+	TFuture<TPair<FStatus, FBatchReadBlobsResponse>> DownloadResultsAsync();
+
+	// Post-step flow
+	void OnMissingBlobsDetermined(FFindMissingBlobsResponse&& Result);
+	void OnMissingBlobsUploaded(const FBatchUpdateBlobsResponse& Result);
+	void OnExecutionCompleted(FExecuteResponse&& Result);
+	void OnOutputBlobsDownloaded(FBatchReadBlobsResponse&& Result);
+};
 
-	void LoadMissingWorkerFileBlobs(FRemoteExecutionState& State)
+class FRemoteBuildWorkerExecutor final: public IBuildWorkerExecutor
+{
+public:
+	FRemoteBuildWorkerExecutor()
+	: GlobalExecutionTimeoutSeconds(-1)
+	, Salt(TEXT("807c6a49-0657-49f3-b498-fd457213c0a7"))
+	, RemoteExecutor(nullptr)
+	, ContentAddressableStorage(nullptr)
+	, bEnabled(false)
 	{
-		TArray<FIoHash> WorkerFileHashes;
-		TMultiMap<FIoHash, uint32> WorkerFileMapping;
-		for (const FDigest& MissingItem : State.FindMissingBlobsResponse.MissingBlobDigests)
-		{
-			TArray<FVariantIndex> DigestFilesystemEntries;
-			State.DigestFilesystemIndex.MultiFind(MissingItem, DigestFilesystemEntries);
-
-			for (FVariantIndex& VariantIndex : DigestFilesystemEntries)
-			{
-				if (VariantIndex.NodeType == ENodeType::File)
-				{
-					const FMerkleTreeFileBuilder& File = State.Files[VariantIndex.Index];
-					if ((File.Type == EFileType::Worker) && File.ContentBytes.IsNull())
-					{
-						WorkerFileHashes.Emplace(MissingItem.Hash);
-						WorkerFileMapping.Add(MissingItem.Hash, VariantIndex.Index);
-					}
-				}
-			}
-		}
-
-		State.BuildWorker.FindFileData(WorkerFileHashes, EPriority::Normal,
-			[&State, &WorkerFileMapping] (FBuildWorkerFileDataCompleteParams&& Params)
-			{
-				uint32 MetaIndex = 0;
-				for (const FCompressedBuffer& Buffer : Params.Files)
-				{
-					FCompositeBuffer UncompressedWorkerFile = Buffer.DecompressToComposite();
-					TArray<uint32> WorkerFileIndices;
-					WorkerFileMapping.MultiFind(Buffer.GetRawHash(), WorkerFileIndices);
-
-					for (int32 FileIndex : WorkerFileIndices)
-					{
-						State.Files[FileIndex].ContentBytes = UncompressedWorkerFile;
-					}
-				}
-			}).Wait();
-	}
+		check(IsInGameThread()); // initialization from the main thread is expected to allow config reading for the limiting heuristics
+		check(GConfig && GConfig->IsReadyForUse());
 
-	FStatus UploadMissingBlobs(FRemoteExecutionState& State)
-	{
-		State.BatchUpdateBlobsRequest.InstanceName = InstanceName;
-		for (const FDigest& MissingItem : State.FindMissingBlobsResponse.MissingBlobDigests)
-		{
-			if (MissingItem == State.ActionDigest)
-			{
-				Stats.TotalActionBlobsUploaded.AddBlob(State.ActionContentBytes.Num());
-				FBatchUpdateBlobsRequest::FRequest& NewRequest = State.BatchUpdateBlobsRequest.Requests.AddDefaulted_GetRef();
-				NewRequest.Digest = State.ActionDigest;
-				NewRequest.Data = MoveTemp(State.ActionContentBytes);
-				UE_LOG(LogDerivedDataBuildRemoteExecutor, Verbose, TEXT("Uploading action (hash: %s) of upload size %d."), *::LexToString(NewRequest.Digest.Hash), NewRequest.Data.Num());
-			}
-			else if (MissingItem == State.Action.CommandDigest)
-			{
-				Stats.TotalCommandBlobsUploaded.AddBlob(State.CommandContentBytes.Num());
-				FBatchUpdateBlobsRequest::FRequest& NewRequest = State.BatchUpdateBlobsRequest.Requests.AddDefaulted_GetRef();
-				NewRequest.Digest = State.Action.CommandDigest;
-				NewRequest.Data = MoveTemp(State.CommandContentBytes);
-				UE_LOG(LogDerivedDataBuildRemoteExecutor, Verbose, TEXT("Uploading command (hash: %s) of upload size %d."), *::LexToString(NewRequest.Digest.Hash), NewRequest.Data.Num());
-			}
-			else
-			{
-				const FVariantIndex VariantIndex = State.DigestFilesystemIndex.FindChecked(MissingItem);
-				FBatchUpdateBlobsRequest::FRequest& NewRequest = State.BatchUpdateBlobsRequest.Requests.AddDefaulted_GetRef();
-				NewRequest.Digest = MissingItem;
-				switch (VariantIndex.NodeType)
-				{
-				case ENodeType::Directory:
-					Stats.TotalDirectoryBlobsUploaded.AddBlob(State.Directories[VariantIndex.Index].ContentBytes.Num());
-					NewRequest.Digest = State.Directories[VariantIndex.Index].Digest.GetValue();
-					NewRequest.Data = MoveTemp(State.Directories[VariantIndex.Index].ContentBytes);
-					UE_LOG(LogDerivedDataBuildRemoteExecutor, Verbose, TEXT("Uploading directory '%s' (hash: %s) of upload size %d."), *FString(State.Directories[VariantIndex.Index].Name), *::LexToString(NewRequest.Digest.Hash), NewRequest.Data.Num());
-					break;
-				case ENodeType::File:
-					{
-						Stats.TotalFileBlobsUploaded.AddBlob(State.Files[VariantIndex.Index].ContentBytes.GetSize());
-						NewRequest.Digest.Hash = FIoHash::HashBuffer(State.Files[VariantIndex.Index].ContentBytes);
-						NewRequest.Digest.SizeBytes = State.Files[VariantIndex.Index].ContentBytes.GetSize();
-						FCompositeBuffer& FileBuffer = State.Files[VariantIndex.Index].ContentBytes;
-						check(!FileBuffer.IsNull());
-						NewRequest.Data.Reserve(FileBuffer.GetSize());
-						for (const FSharedBuffer& Segment : State.Files[VariantIndex.Index].ContentBytes.GetSegments())
-						{
-							NewRequest.Data.Append((const char *)Segment.GetData(), Segment.GetSize());
-						}
-						FileBuffer.Reset();
-						UE_LOG(LogDerivedDataBuildRemoteExecutor, Verbose, TEXT("Uploading file '%s' (hash: %s, type: %s) of upload size %d."), *State.Files[VariantIndex.Index].File.Name, *::LexToString(NewRequest.Digest.Hash), *LexToString(State.Files[VariantIndex.Index].Type), NewRequest.Data.Num());
-					}
-					break;
-				default:
-					checkNoEntry();
-					break;
-				}
-			}
-		}
+		GConfig->GetString(TEXT("DerivedDataBuildRemoteExecutor"), TEXT("InstanceName"), InstanceName, GEngineIni);
+		GConfig->GetInt(TEXT("DerivedDataBuildRemoteExecutor"), TEXT("GlobalExecutionTimeoutSeconds"), GlobalExecutionTimeoutSeconds, GEngineIni);
 
-		return ContentAddressableStorage->BatchUpdateBlobs(State.BatchUpdateBlobsRequest, State.BatchUpdateBlobsResponse);
-	}
+		const FName RemoteExecutionFeatureName(TEXT("RemoteExecution"));
+		IModularFeatures& ModularFeatures = IModularFeatures::Get();
 
-	bool ValidateUploadSuccess(FRemoteExecutionState& State)
-	{
-		bool bSuccess = true;
-		for (const FBatchUpdateBlobsRequest::FRequest& RequestedUpload : State.BatchUpdateBlobsRequest.Requests)
+		if (FParse::Param(FCommandLine::Get(), TEXT("DDC2RemoteExecution")))
 		{
-			if (const FBatchUpdateBlobsResponse::FResponse* FoundResponse = Algo::FindByPredicate(State.BatchUpdateBlobsResponse.Responses, [&RequestedUpload] (const FBatchUpdateBlobsResponse::FResponse& Response)
-				{
-					return Response.Digest == RequestedUpload.Digest;
-				}))
+			FModuleManager::Get().LoadModule("BazelExecutor");
+			if (ModularFeatures.IsModularFeatureAvailable(RemoteExecutionFeatureName))
 			{
-				if (!FoundResponse->Status.Ok())
+				RemoteExecutor = &ModularFeatures.GetModularFeature<IRemoteExecutor>(RemoteExecutionFeatureName);
+				if (RemoteExecutor)
 				{
-					FStringView ActionName = State.BuildAction.GetName();
-					UE_LOG(LogDerivedDataBuildRemoteExecutor, Log, TEXT("Remote execution system error: data for action '%.*s' could not be uploaded (hash: %s, size: %u)"), ActionName.Len(), ActionName.GetData(), *::LexToString(RequestedUpload.Digest.Hash), RequestedUpload.Digest.SizeBytes);
-					bSuccess = false;
+					ContentAddressableStorage = RemoteExecutor->GetContentAddressableStorage();
 				}
 			}
 		}
-		return bSuccess;
-	}
-
-	bool ExecuteBuild(FRemoteExecutionState& State)
-	{
-		State.ExecuteRequest.InstanceName = InstanceName;
-		State.ExecuteRequest.ActionDigest = State.ActionDigest;
-		State.ExecuteRequest.SkipCacheLookup = true;
-		int64 TimeoutMs = 0;
-		if (GlobalExecutionTimeoutSeconds > 0)
-		{
-			TimeoutMs = GlobalExecutionTimeoutSeconds * 1000LL;
-		}
-		return RemoteExecutor->GetExecution()->Execute(State.ExecuteRequest, State.ExecuteResponse, TimeoutMs);
-	}
-
-	FStatus DownloadResults(FRemoteExecutionState& State)
-	{
-		State.BatchReadBlobsRequest.InstanceName = InstanceName;
-		//State.BatchReadBlobsRequest.Digests.Add(State.ExecuteResponse.Result.StdoutDigest);
-		//State.BatchReadBlobsRequest.Digests.Add(State.ExecuteResponse.Result.StderrDigest);
 
-		for (const FOutputFile& OutputFile : State.ExecuteResponse.Result.OutputFiles)
+		bEnabled = (RemoteExecutor != nullptr) && (ContentAddressableStorage != nullptr) && !InstanceName.IsEmpty();
+		if (bEnabled)
 		{
-			State.BatchReadBlobsRequest.Digests.Add(OutputFile.Digest);
+			ModularFeatures.RegisterModularFeature(IBuildWorkerExecutor::GetFeatureName(), this);
 		}
-
-		return ContentAddressableStorage->BatchReadBlobs(State.BatchReadBlobsRequest, State.BatchReadBlobsResponse);
 	}
 
-	FOptionalBuildOutput ComposeBuildOutput(FRemoteExecutionState& State, EStatus& OutStatus)
+	virtual ~FRemoteBuildWorkerExecutor()
 	{
-		OutStatus = EStatus::Error;
-
-		TOptional<FDigest> BuildOutputDigest;
-		for (const FOutputFile& ExecuteOutputFile : State.ExecuteResponse.Result.OutputFiles)
-		{
-			if (ExecuteOutputFile.Path == TEXT("Build.output"))
-			{
-				BuildOutputDigest = ExecuteOutputFile.Digest;
-				break;
-			}
-		}
-
-		if (!BuildOutputDigest.IsSet())
-		{
-			UE_LOG(LogDerivedDataBuildRemoteExecutor, Warning, TEXT("Worker error: build output structure not produced!"));
-			return FOptionalBuildOutput();
-		}
-
-		FOptionalBuildOutput RemoteBuildOutput;
-
-		TMap<FIoHash, FCompressedBuffer> PayloadResponses;
-		for (FBatchReadBlobsResponse::FResponse& OutputReadResponse : State.BatchReadBlobsResponse.Responses)
-		{
-			if (!OutputReadResponse.Status.Ok())
-			{
-				UE_LOG(LogDerivedDataBuildRemoteExecutor, Warning, TEXT("Remote execution system error: output not downloaded!"));
-				return FOptionalBuildOutput();
-			}
-			Stats.TotalFileBlobsDownloaded.AddBlob(OutputReadResponse.Data.Num());
-
-			if (OutputReadResponse.Digest == BuildOutputDigest)
-			{
-				FSharedBuffer BuildOutputBuffer = MakeSharedBufferFromArray(MoveTemp(OutputReadResponse.Data));
-
-				if (ValidateCompactBinaryRange(BuildOutputBuffer, ECbValidateMode::Default) != ECbValidateError::None)
-				{
-					UE_LOG(LogDerivedDataBuildRemoteExecutor, Warning, TEXT("Worker error: build output structure not valid!"));
-					return FOptionalBuildOutput();
-				}
-
-				RemoteBuildOutput = GetDerivedDataBuildRef().LoadOutput(State.BuildAction.GetName(), State.BuildAction.GetFunction(), FCbObject(BuildOutputBuffer));
-			}
-			else
-			{
-				FCompressedBuffer NewBuffer = FCompressedBuffer::FromCompressed(MakeSharedBufferFromArray(MoveTemp(OutputReadResponse.Data)));
-				PayloadResponses.Add(NewBuffer.GetRawHash(), MoveTemp(NewBuffer));
-			}
-		}
-
-		if (RemoteBuildOutput.IsNull())
-		{
-			UE_LOG(LogDerivedDataBuildRemoteExecutor, Warning, TEXT("Remote execution system error: build output blob missing!"));
-			return FOptionalBuildOutput();
-		}
-
-		FBuildOutputBuilder OutputBuilder = GetDerivedDataBuildRef().CreateOutput(State.BuildAction.GetName(), State.BuildAction.GetFunction());
-		
-		RemoteBuildOutput.Get().IterateDiagnostics( [&OutputBuilder](const FBuildDiagnostic& Diagnostic)
-		{
-			if (Diagnostic.Level == EBuildDiagnosticLevel::Warning)
-			{
-				OutputBuilder.AddWarning(Diagnostic.Category, Diagnostic.Message);
-			}
-			else if (Diagnostic.Level == EBuildDiagnosticLevel::Error)
-			{
-				OutputBuilder.AddError(Diagnostic.Category, Diagnostic.Message);
-			}
-		});
-
-		for (const FPayload& Payload : RemoteBuildOutput.Get().GetPayloads())
+		if (bEnabled)
 		{
-			FCompressedBuffer* BufferForPayload = PayloadResponses.Find(Payload.GetRawHash());
-			if (!BufferForPayload)
-			{
-				UE_LOG(LogDerivedDataBuildRemoteExecutor, Warning, TEXT("Remote execution system error: payload blob missing!"));
-				return FOptionalBuildOutput();
-			}
-
-			OutputBuilder.AddPayload(FPayload(Payload.GetId(), *BufferForPayload));
+			IModularFeatures::Get().UnregisterModularFeature(IBuildWorkerExecutor::GetFeatureName(), this);
 		}
-
-		OutStatus = State.ExecuteResponse.Result.ExitCode == 0 ? EStatus::Ok : EStatus::Error;
-		return OutputBuilder.Build();
 	}
 
 	FRequest BuildAction(
@@ -637,10 +270,6 @@ public:
 		EPriority Priority,
 		FOnBuildWorkerActionComplete&& OnComplete) final
 	{
-		TArray<FString> MissingInputs;
-		TArray<FStringView> MissingInputViews;
-		uint64 TotalInputSize = 0;
-		uint64 TotalMissingInputSize = 0;
 		{
 			// TODO: This block forces resolution of inputs before we attempt to determine which
 			//		 inputs need to be uploaded.  This is required because we can't refer to inputs
@@ -648,6 +277,11 @@ public:
 			//		 CompressedSize.  Once the remote execution API allows us to represent inputs with RawHash/
 			//		 RawSize, this block can be removed and we can find missing CAS inputs without having resolved
 			//		 the inputs first.
+			TArray<FString> MissingInputs;
+			TArray<FStringView> MissingInputViews;
+			uint64 TotalInputSize = 0;
+			uint64 TotalMissingInputSize = 0;
+
 			Action.IterateInputs([&MissingInputs, &MissingInputViews, &Inputs, &TotalInputSize, &TotalMissingInputSize] (FStringView Key, const FIoHash& RawHash, uint64 RawSize)
 				{
 					if (Inputs.IsNull() || Inputs.Get().FindInput(Key).IsNull())
@@ -672,71 +306,20 @@ public:
 			}
 		}
 
-		FRemoteExecutionState State { Action, Inputs, Worker, Policy, Priority };
+		return FRequest(new FRemoteBuildExecutionRequest(*this, Action, Inputs, Worker, Policy, Priority, MoveTemp(OnComplete)));
+	}
 
-		DetermineMissingBlobs(State);
+	TConstArrayView<FStringView> GetHostPlatforms() const final
+	{
+		static constexpr FStringView HostPlatforms[]{TEXT("Win64"_SV), TEXT("Linux"_SV), TEXT("Mac"_SV)};
+		return HostPlatforms;
+	}
 
-		GatherMissingInputFileBlobs(State, MissingInputViews);
-		if (!MissingInputViews.IsEmpty())
+	void DumpStats()
+	{
+		if (Stats.TotalRemoteBuilds.load() == 0)
 		{
-			OnComplete({State.BuildAction.GetKey(), {}, MissingInputViews, EStatus::Ok});
-			return FRequest();
-		}
-
-		if (!LimitingHeuristics.TryStartNewBuild(Stats))
-		{
-			OnComplete({Action.GetKey(), {}, {}, EStatus::Error});
-			return FRequest();
-		}
-		ON_SCOPE_EXIT
-		{
-			LimitingHeuristics.FinishBuild(Stats);
-		};
-
-		LoadMissingWorkerFileBlobs(State);
-		if (!State.FindMissingBlobsResponse.MissingBlobDigests.IsEmpty())
-		{
-			UploadMissingBlobs(State);
-			UE_LOG(LogDerivedDataBuildRemoteExecutor, Verbose, TEXT("Uploaded %d data blobs for remote execution."), State.BatchUpdateBlobsRequest.Requests.Num());
-			if (!ValidateUploadSuccess(State))
-			{
-				OnComplete({Action.GetKey(), {}, {}, EStatus::Error});
-				return FRequest();
-			}
-		}
-
-		if (ExecuteBuild(State))
-		{
-			DownloadResults(State);
-
-			EStatus BuildStatus = EStatus::Error;
-			FOptionalBuildOutput BuildOutput = ComposeBuildOutput(State, BuildStatus);
-			if (BuildStatus == EStatus::Ok)
-			{
-				Stats.TotalSuccessfulRemoteBuilds.fetch_add(1, std::memory_order_relaxed);
-			}
-			OnComplete({State.BuildAction.GetKey(), MoveTemp(BuildOutput), {}, BuildStatus});
-		}
-		else
-		{
-			UE_LOG(LogDerivedDataBuildRemoteExecutor, Warning, TEXT("Remote execution system error: Failed to execute build operation!"));
-			OnComplete({State.BuildAction.GetKey(), {}, {}, EStatus::Error});
-		}
-
-		return FRequest();
-	}
-
-	TConstArrayView<FStringView> GetHostPlatforms() const final
-	{
-		static constexpr FStringView HostPlatforms[]{TEXT("Win64"_SV), TEXT("Linux"_SV), TEXT("Mac"_SV)};
-		return HostPlatforms;
-	}
-
-	void DumpStats()
-	{
-		if (Stats.TotalRemoteBuilds.load() == 0)
-		{
-			return;
+			return;
 		}
 
 		Stats.Dump();
@@ -878,6 +461,8 @@ private:
 		bool bEnableLimits{false};
 	};
 
+	friend class FRemoteBuildExecutionRequest;
+
 	FStats Stats;
 	FLimitingHeuristics LimitingHeuristics;
 	FString InstanceName;
@@ -888,6 +473,595 @@ private:
 	bool bEnabled;
 };
 
+FRemoteBuildExecutionRequest::FRemoteBuildExecutionRequest(
+		FRemoteBuildWorkerExecutor& InExecutor,
+		const FBuildAction& Action,
+		const FOptionalBuildInputs& Inputs,
+		const FBuildWorker& Worker,
+		EBuildPolicy Policy,
+		EPriority Priority,
+		FOnBuildWorkerActionComplete&& OnComplete)
+: State {Action, Inputs, Worker, Policy, Priority}
+, CompletionCallback(MoveTemp(OnComplete))
+, Executor(InExecutor)
+, bCancelPending(false)
+, bHeuristicBuildStarted(false)
+{
+	DetermineMissingBlobsAsync()
+		.Next([this] (TPair<FStatus, FFindMissingBlobsResponse>&& Result) { OnMissingBlobsDetermined(MoveTemp(Result.Value)); });
+}
+
+FRemoteBuildExecutionRequest::~FRemoteBuildExecutionRequest()
+{
+	if (bHeuristicBuildStarted)
+	{
+		Executor.LimitingHeuristics.FinishBuild(Executor.Stats);
+	};
+}
+
+FRemoteBuildExecutionRequest::FMerkleTreeDirectoryBuilder& FRemoteBuildExecutionRequest::GetOrAddMerkleTreeDirectory(FStringView Path, int32& OutDirectoryBuilderIndex)
+{
+	int32& DirectoryBuilderIndex = State.PathToDirectoryIndex.FindOrAdd(Path, INDEX_NONE);
+	if (DirectoryBuilderIndex == INDEX_NONE)
+	{
+		DirectoryBuilderIndex = State.Directories.Num();
+		FMerkleTreeDirectoryBuilder& NewNode = State.Directories.AddDefaulted_GetRef();
+		NewNode.Name = Path.IsEmpty() ? Path : FPathViews::GetCleanFilename(Path);
+	}
+
+	OutDirectoryBuilderIndex = DirectoryBuilderIndex;
+
+	int32 ContainingDirectoryIndex = DirectoryBuilderIndex;
+	while (!Path.IsEmpty())
+	{
+		int32 SubDirIndex = ContainingDirectoryIndex;
+		Path = FPathViews::GetPath(Path);
+		GetOrAddMerkleTreeDirectory(Path, ContainingDirectoryIndex).SubDirIndices.AddUnique(SubDirIndex);
+	}
+
+	return State.Directories[OutDirectoryBuilderIndex];
+}
+
+int32 FRemoteBuildExecutionRequest::AddMerkleTreeFile(FStringView Path, const FIoHash& RawHash, uint64 RawSize, bool bIsExecutable, EFileType FileType, FCompositeBuffer ContentBytes)
+{
+	FStringView ContainingDirectoryPath = FPathViews::GetPath(Path);
+	int32 ContainingDirectoryIndex = INDEX_NONE;
+	FMerkleTreeDirectoryBuilder* ContainingDirectoryBuilder = &GetOrAddMerkleTreeDirectory(ContainingDirectoryPath, ContainingDirectoryIndex);
+	FFileNode& NewNode = ContainingDirectoryBuilder->Directory.Files.AddDefaulted_GetRef();
+	NewNode.Name = FPathViews::GetCleanFilename(Path);
+	NewNode.Digest.Hash = RawHash;
+	NewNode.Digest.SizeBytes = RawSize;
+	NewNode.IsExecutable = bIsExecutable;
+
+	int32 NewFileIndex = State.Files.Num();
+	State.DigestFilesystemIndex.Add(NewNode.Digest, FVariantIndex(ENodeType::File, NewFileIndex));
+	FMerkleTreeFileBuilder& FileBuilder = State.Files.AddDefaulted_GetRef();
+	FileBuilder.File = NewNode; // Duplicates the node in the state's file array
+	FileBuilder.Type = FileType;
+	if (ContentBytes)
+	{
+		FileBuilder.ContentBytes = ContentBytes;
+	}
+
+	return NewFileIndex;
+}
+
+const FDigest& FRemoteBuildExecutionRequest::BuildMerkleTreeDirectoryDigest(int32 Index)
+{
+	FMerkleTreeDirectoryBuilder& DirBuilder = State.Directories[Index];
+
+	if (DirBuilder.Digest.IsSet())
+	{
+		return DirBuilder.Digest.GetValue();
+	}
+
+	for (int32 SubDirIndex : DirBuilder.SubDirIndices)
+	{
+		FDirectoryNode& SubDirNode = DirBuilder.Directory.Directories.AddDefaulted_GetRef();
+		SubDirNode.Name = State.Directories[SubDirIndex].Name;
+		SubDirNode.Digest = BuildMerkleTreeDirectoryDigest(SubDirIndex);
+		State.DigestFilesystemIndex.Add(SubDirNode.Digest, FVariantIndex(ENodeType::Directory, SubDirIndex));
+	}
+
+	DirBuilder.Directory.Directories.Sort([] (const FDirectoryNode& A, const FDirectoryNode& B)
+	{
+		return A.Name < B.Name;
+	});
+
+	DirBuilder.Directory.Files.Sort([] (const FFileNode& A, const FFileNode& B)
+	{
+		return A.Name < B.Name;
+	});
+
+	Executor.ContentAddressableStorage->ToBlob(DirBuilder.Directory, DirBuilder.ContentBytes, DirBuilder.Digest.Emplace());
+
+	return DirBuilder.Digest.GetValue();
+}
+
+void FRemoteBuildExecutionRequest::BuildMerkleTreeNodes()
+{
+	TArray<FIoHash> WorkerFileHashes;
+	TArray<TTuple<FStringView, bool>> WorkerFileMeta;
+	State.BuildWorker.IterateExecutables([&WorkerFileHashes, &WorkerFileMeta] (FStringView Path, const FIoHash& RawHash, uint64 RawSize)
+		{
+			WorkerFileHashes.Emplace(RawHash);
+			WorkerFileMeta.Emplace(Path, true);
+		});
+
+	State.BuildWorker.IterateFiles([&WorkerFileHashes, &WorkerFileMeta] (FStringView Path, const FIoHash& RawHash, uint64 RawSize)
+		{
+			WorkerFileHashes.Emplace(RawHash);
+			WorkerFileMeta.Emplace(Path, false);
+		});
+
+	State.BuildWorker.FindFileData(WorkerFileHashes, EPriority::Normal,
+		[this, &WorkerFileMeta] (FBuildWorkerFileDataCompleteParams&& Params)
+		{
+			uint32 MetaIndex = 0;
+			for (const FCompressedBuffer& Buffer : Params.Files)
+			{
+				const TTuple<FStringView, bool>& Meta = WorkerFileMeta[MetaIndex];
+				FCompositeBuffer DecompressedComposite = Buffer.DecompressToComposite();
+				AddMerkleTreeFile(Meta.Key, FIoHash::HashBuffer(DecompressedComposite), DecompressedComposite.GetSize(), Meta.Value, EFileType::Worker, Buffer.DecompressToComposite());
+				++MetaIndex;
+			}
+		}).Wait();
+
+	State.BuildAction.IterateInputs([this] (FStringView Key, const FIoHash& RawHash, uint64 RawSize)
+		{
+			TStringBuilder<128> InputPath;
+			InputPath << TEXT("Inputs/") << RawHash;
+			const FString& NewInputPath = State.InputPaths.Emplace_GetRef(InputPath);
+			const FCompressedBuffer& Buffer = State.BuildInputs.Get().FindInput(Key);
+			check(!Buffer.IsNull());
+			int32 FileIndex = AddMerkleTreeFile(NewInputPath, FIoHash::HashBuffer(Buffer.GetCompressed()), Buffer.GetCompressedSize(), false, EFileType::Input, Buffer.GetCompressed());
+			State.FileIndexToInputKey.Emplace(FileIndex, Key);
+		});
+
+	// This base directory must be created as worker executables (even those that don't exist in this directory) will attempt to change directories into it during startup.
+	int32 BaseDirectoryIndex = INDEX_NONE;
+	TStringBuilder<128> BaseDirectoryPathBuilder;
+	FPathViews::Append(BaseDirectoryPathBuilder, TEXT("Engine/Binaries/"), State.BuildWorker.GetHostPlatform());
+	State.BaseDirectoryPath = BaseDirectoryPathBuilder.ToString();
+	GetOrAddMerkleTreeDirectory(State.BaseDirectoryPath, BaseDirectoryIndex);
+
+	FCbWriter BuildActionWriter;
+	State.BuildAction.Save(BuildActionWriter);
+	FUniqueBuffer UncompressedBuildActionContentBytes = FUniqueBuffer::Alloc(BuildActionWriter.GetSaveSize());
+	BuildActionWriter.Save(UncompressedBuildActionContentBytes);
+	State.BuildActionContentBytes = FCompositeBuffer(UncompressedBuildActionContentBytes.MoveToShared());
+	State.BuildActionDigest.Hash = FIoHash::HashBuffer(State.BuildActionContentBytes);
+	State.BuildActionDigest.SizeBytes = State.BuildActionContentBytes.GetSize();
+	AddMerkleTreeFile(TEXT("Build.action"), State.BuildActionDigest.Hash, State.BuildActionDigest.SizeBytes, false, EFileType::BuildAction, State.BuildActionContentBytes);
+
+
+	if (!State.PathToDirectoryIndex.IsEmpty())
+	{
+		int32 RootDirectoryIndex = State.PathToDirectoryIndex.FindChecked(TEXT(""));
+		State.Action.InputRootDigest = BuildMerkleTreeDirectoryDigest(RootDirectoryIndex);
+		State.DigestFilesystemIndex.Add(State.Action.InputRootDigest, FVariantIndex(ENodeType::Directory, RootDirectoryIndex));
+	}
+
+	State.Command.OutputPaths.Add("Outputs");
+	State.Command.OutputPaths.Add("Build.output");
+	State.Command.Arguments.Add(FString(State.BuildWorker.GetPath()));
+	State.Command.Arguments.Add("-Build=Build.action");
+	State.BuildWorker.IterateEnvironment([this] (FStringView Name, FStringView Value)
+		{
+			FCommand::FEnvironmentVariable& EnvVar = State.Command.EnvironmentVariables.AddDefaulted_GetRef();
+			EnvVar.Name = Name;
+			EnvVar.Value = Value;
+		});
+	State.Command.EnvironmentVariables.Sort([] (const FCommand::FEnvironmentVariable& A, const FCommand::FEnvironmentVariable& B)
+		{
+			if (A.Name < B.Name)
+			{
+				return true;
+			}
+			else if (B.Name < A.Name)
+			{
+				return false;
+			}
+			return A.Value < B.Value;
+		});
+	Executor.ContentAddressableStorage->ToBlob(State.Command, State.CommandContentBytes, State.Action.CommandDigest);
+
+	State.Action.Salt = Executor.Salt;
+	Executor.ContentAddressableStorage->ToBlob(State.Action, State.ActionContentBytes, State.ActionDigest);
+}
+
+void FRemoteBuildExecutionRequest::GatherMissingInputFileBlobs(TArray<FStringView>& OutMissingInputs)
+{
+	for (const FDigest& MissingItem : State.FindMissingBlobsResponse.MissingBlobDigests)
+	{
+		TArray<FVariantIndex> DigestFilesystemEntries;
+		State.DigestFilesystemIndex.MultiFind(MissingItem, DigestFilesystemEntries);
+
+		for (FVariantIndex& VariantIndex : DigestFilesystemEntries)
+		{
+			if (VariantIndex.NodeType == ENodeType::File)
+			{
+				const FMerkleTreeFileBuilder& File = State.Files[VariantIndex.Index];
+				if ((File.Type == EFileType::Input) && File.ContentBytes.IsNull())
+				{
+					OutMissingInputs.Add(State.FileIndexToInputKey[VariantIndex.Index]);
+				}
+			}
+		}
+	}
+}
+
+bool FRemoteBuildExecutionRequest::ValidateUploadSuccess(const FBatchUpdateBlobsResponse& BatchUpdateBlobsResponse)
+{
+	bool bSuccess = true;
+	for (const FBatchUpdateBlobsRequest::FRequest& RequestedUpload : State.BatchUpdateBlobsRequest.Requests)
+	{
+		if (const FBatchUpdateBlobsResponse::FResponse* FoundResponse = Algo::FindByPredicate(BatchUpdateBlobsResponse.Responses, [&RequestedUpload] (const FBatchUpdateBlobsResponse::FResponse& Response)
+			{
+				return Response.Digest == RequestedUpload.Digest;
+			}))
+		{
+			if (!FoundResponse->Status.Ok())
+			{
+				FStringView ActionName = State.BuildAction.GetName();
+				UE_LOG(LogDerivedDataBuildRemoteExecutor, Log, TEXT("Remote execution system error: data for action '%.*s' could not be uploaded (hash: %s, size: %u)"), ActionName.Len(), ActionName.GetData(), *::LexToString(RequestedUpload.Digest.Hash), RequestedUpload.Digest.SizeBytes);
+				bSuccess = false;
+			}
+		}
+	}
+	return bSuccess;
+}
+
+FOptionalBuildOutput FRemoteBuildExecutionRequest::ComposeBuildOutput(FBatchReadBlobsResponse& BatchReadBlobsResponse, EStatus& OutStatus)
+{
+	OutStatus = EStatus::Error;
+
+	TOptional<FDigest> BuildOutputDigest;
+	for (const FOutputFile& ExecuteOutputFile : State.ExecuteResponse.Result.OutputFiles)
+	{
+		if (ExecuteOutputFile.Path == TEXT("Build.output"))
+		{
+			BuildOutputDigest = ExecuteOutputFile.Digest;
+			break;
+		}
+	}
+
+	if (!BuildOutputDigest.IsSet())
+	{
+		UE_LOG(LogDerivedDataBuildRemoteExecutor, Warning, TEXT("Worker error: build output structure not produced!"));
+		return FOptionalBuildOutput();
+	}
+
+	FOptionalBuildOutput RemoteBuildOutput;
+
+	TMap<FIoHash, FCompressedBuffer> PayloadResponses;
+	for (FBatchReadBlobsResponse::FResponse& OutputReadResponse : BatchReadBlobsResponse.Responses)
+	{
+		if (!OutputReadResponse.Status.Ok())
+		{
+			UE_LOG(LogDerivedDataBuildRemoteExecutor, Warning, TEXT("Remote execution system error: output not downloaded!"));
+			return FOptionalBuildOutput();
+		}
+		Executor.Stats.TotalFileBlobsDownloaded.AddBlob(OutputReadResponse.Data.Num());
+
+		if (OutputReadResponse.Digest == BuildOutputDigest)
+		{
+			FSharedBuffer BuildOutputBuffer = MakeSharedBufferFromArray(MoveTemp(OutputReadResponse.Data));
+
+			if (ValidateCompactBinaryRange(BuildOutputBuffer, ECbValidateMode::Default) != ECbValidateError::None)
+			{
+				UE_LOG(LogDerivedDataBuildRemoteExecutor, Warning, TEXT("Worker error: build output structure not valid!"));
+				return FOptionalBuildOutput();
+			}
+
+			RemoteBuildOutput = GetDerivedDataBuildRef().LoadOutput(State.BuildAction.GetName(), State.BuildAction.GetFunction(), FCbObject(BuildOutputBuffer));
+		}
+		else
+		{
+			FCompressedBuffer NewBuffer = FCompressedBuffer::FromCompressed(MakeSharedBufferFromArray(MoveTemp(OutputReadResponse.Data)));
+			PayloadResponses.Add(NewBuffer.GetRawHash(), MoveTemp(NewBuffer));
+		}
+	}
+
+	if (RemoteBuildOutput.IsNull())
+	{
+		UE_LOG(LogDerivedDataBuildRemoteExecutor, Warning, TEXT("Remote execution system error: build output blob missing!"));
+		return FOptionalBuildOutput();
+	}
+
+	FBuildOutputBuilder OutputBuilder = GetDerivedDataBuildRef().CreateOutput(State.BuildAction.GetName(), State.BuildAction.GetFunction());
+	
+	RemoteBuildOutput.Get().IterateDiagnostics( [&OutputBuilder](const FBuildDiagnostic& Diagnostic)
+	{
+		if (Diagnostic.Level == EBuildDiagnosticLevel::Warning)
+		{
+			OutputBuilder.AddWarning(Diagnostic.Category, Diagnostic.Message);
+		}
+		else if (Diagnostic.Level == EBuildDiagnosticLevel::Error)
+		{
+			OutputBuilder.AddError(Diagnostic.Category, Diagnostic.Message);
+		}
+	});
+
+	for (const FPayload& Payload : RemoteBuildOutput.Get().GetPayloads())
+	{
+		FCompressedBuffer* BufferForPayload = PayloadResponses.Find(Payload.GetRawHash());
+		if (!BufferForPayload)
+		{
+			UE_LOG(LogDerivedDataBuildRemoteExecutor, Warning, TEXT("Remote execution system error: payload blob missing!"));
+			return FOptionalBuildOutput();
+		}
+
+		OutputBuilder.AddPayload(FPayload(Payload.GetId(), *BufferForPayload));
+	}
+
+	OutStatus = State.ExecuteResponse.Result.ExitCode == 0 ? EStatus::Ok : EStatus::Error;
+	return OutputBuilder.Build();
+}
+
+bool FRemoteBuildExecutionRequest::ProcessCancellation()
+{
+	if (bCancelPending.load(std::memory_order_relaxed))
+	{
+		TRefCountPtr<FRemoteBuildExecutionRequest> SelfRef(this);
+		CompletionCallback({State.BuildAction.GetKey(), {}, {}, EStatus::Canceled});
+		CompletionEvent->Trigger();
+		return true;
+	}
+	return false;
+}
+
+TFuture<TPair<FStatus, FFindMissingBlobsResponse>> FRemoteBuildExecutionRequest::DetermineMissingBlobsAsync()
+{
+	BuildMerkleTreeNodes();
+
+	State.FindMissingBlobsRequest.InstanceName = Executor.InstanceName;
+	State.FindMissingBlobsRequest.BlobDigests.Add(State.ActionDigest);
+	State.FindMissingBlobsRequest.BlobDigests.Add(State.Action.CommandDigest);
+	for (const TPair<FDigest, FVariantIndex>& FilesystemItem : State.DigestFilesystemIndex)
+	{
+		State.FindMissingBlobsRequest.BlobDigests.AddUnique(FilesystemItem.Key);
+	}
+	return Executor.ContentAddressableStorage->FindMissingBlobsAsync(State.FindMissingBlobsRequest);
+}
+
+void FRemoteBuildExecutionRequest::LoadMissingWorkerFileBlobsAsync()
+{
+	TArray<FIoHash> WorkerFileHashes;
+	TMultiMap<FIoHash, uint32> WorkerFileMapping;
+	for (const FDigest& MissingItem : State.FindMissingBlobsResponse.MissingBlobDigests)
+	{
+		TArray<FVariantIndex> DigestFilesystemEntries;
+		State.DigestFilesystemIndex.MultiFind(MissingItem, DigestFilesystemEntries);
+
+		for (FVariantIndex& VariantIndex : DigestFilesystemEntries)
+		{
+			if (VariantIndex.NodeType == ENodeType::File)
+			{
+				const FMerkleTreeFileBuilder& File = State.Files[VariantIndex.Index];
+				if ((File.Type == EFileType::Worker) && File.ContentBytes.IsNull())
+				{
+					WorkerFileHashes.Emplace(MissingItem.Hash);
+					WorkerFileMapping.Add(MissingItem.Hash, VariantIndex.Index);
+				}
+			}
+		}
+	}
+
+	State.BuildWorker.FindFileData(WorkerFileHashes, EPriority::Normal,
+		[this, &WorkerFileMapping] (FBuildWorkerFileDataCompleteParams&& Params)
+		{
+			uint32 MetaIndex = 0;
+			for (const FCompressedBuffer& Buffer : Params.Files)
+			{
+				FCompositeBuffer UncompressedWorkerFile = Buffer.DecompressToComposite();
+				TArray<uint32> WorkerFileIndices;
+				WorkerFileMapping.MultiFind(Buffer.GetRawHash(), WorkerFileIndices);
+
+				for (int32 FileIndex : WorkerFileIndices)
+				{
+					State.Files[FileIndex].ContentBytes = UncompressedWorkerFile;
+				}
+			}
+		}).Wait();
+}
+
+TFuture<TPair<FStatus, FBatchUpdateBlobsResponse>> FRemoteBuildExecutionRequest::UploadMissingBlobsAsync()
+{
+	State.BatchUpdateBlobsRequest.InstanceName = Executor.InstanceName;
+	for (const FDigest& MissingItem : State.FindMissingBlobsResponse.MissingBlobDigests)
+	{
+		if (MissingItem == State.ActionDigest)
+		{
+			Executor.Stats.TotalActionBlobsUploaded.AddBlob(State.ActionContentBytes.Num());
+			FBatchUpdateBlobsRequest::FRequest& NewRequest = State.BatchUpdateBlobsRequest.Requests.AddDefaulted_GetRef();
+			NewRequest.Digest = State.ActionDigest;
+			NewRequest.Data = MoveTemp(State.ActionContentBytes);
+			UE_LOG(LogDerivedDataBuildRemoteExecutor, Verbose, TEXT("Uploading action (hash: %s) of upload size %d."), *::LexToString(NewRequest.Digest.Hash), NewRequest.Data.Num());
+		}
+		else if (MissingItem == State.Action.CommandDigest)
+		{
+			Executor.Stats.TotalCommandBlobsUploaded.AddBlob(State.CommandContentBytes.Num());
+			FBatchUpdateBlobsRequest::FRequest& NewRequest = State.BatchUpdateBlobsRequest.Requests.AddDefaulted_GetRef();
+			NewRequest.Digest = State.Action.CommandDigest;
+			NewRequest.Data = MoveTemp(State.CommandContentBytes);
+			UE_LOG(LogDerivedDataBuildRemoteExecutor, Verbose, TEXT("Uploading command (hash: %s) of upload size %d."), *::LexToString(NewRequest.Digest.Hash), NewRequest.Data.Num());
+		}
+		else
+		{
+			const FVariantIndex VariantIndex = State.DigestFilesystemIndex.FindChecked(MissingItem);
+			FBatchUpdateBlobsRequest::FRequest& NewRequest = State.BatchUpdateBlobsRequest.Requests.AddDefaulted_GetRef();
+			NewRequest.Digest = MissingItem;
+			switch (VariantIndex.NodeType)
+			{
+			case ENodeType::Directory:
+				Executor.Stats.TotalDirectoryBlobsUploaded.AddBlob(State.Directories[VariantIndex.Index].ContentBytes.Num());
+				NewRequest.Digest = State.Directories[VariantIndex.Index].Digest.GetValue();
+				NewRequest.Data = MoveTemp(State.Directories[VariantIndex.Index].ContentBytes);
+				UE_LOG(LogDerivedDataBuildRemoteExecutor, Verbose, TEXT("Uploading directory '%s' (hash: %s) of upload size %d."), *FString(State.Directories[VariantIndex.Index].Name), *::LexToString(NewRequest.Digest.Hash), NewRequest.Data.Num());
+				break;
+			case ENodeType::File:
+				{
+					Executor.Stats.TotalFileBlobsUploaded.AddBlob(State.Files[VariantIndex.Index].ContentBytes.GetSize());
+					NewRequest.Digest.Hash = FIoHash::HashBuffer(State.Files[VariantIndex.Index].ContentBytes);
+					NewRequest.Digest.SizeBytes = State.Files[VariantIndex.Index].ContentBytes.GetSize();
+					FCompositeBuffer& FileBuffer = State.Files[VariantIndex.Index].ContentBytes;
+					check(!FileBuffer.IsNull());
+					NewRequest.Data.Reserve(FileBuffer.GetSize());
+					for (const FSharedBuffer& Segment : State.Files[VariantIndex.Index].ContentBytes.GetSegments())
+					{
+						NewRequest.Data.Append((const char *)Segment.GetData(), Segment.GetSize());
+					}
+					FileBuffer.Reset();
+					UE_LOG(LogDerivedDataBuildRemoteExecutor, Verbose, TEXT("Uploading file '%s' (hash: %s, type: %s) of upload size %d."), *State.Files[VariantIndex.Index].File.Name, *::LexToString(NewRequest.Digest.Hash), *LexToString(State.Files[VariantIndex.Index].Type), NewRequest.Data.Num());
+				}
+				break;
+			default:
+				checkNoEntry();
+				break;
+			}
+		}
+	}
+
+	return Executor.ContentAddressableStorage->BatchUpdateBlobsAsync(State.BatchUpdateBlobsRequest);
+}
+
+TFuture<FExecuteResponse> FRemoteBuildExecutionRequest::ExecuteBuildAsync()
+{
+	State.ExecuteRequest.InstanceName = Executor.InstanceName;
+	State.ExecuteRequest.ActionDigest = State.ActionDigest;
+	State.ExecuteRequest.SkipCacheLookup = true;
+	int64 TimeoutMs = 0;
+	if (Executor.GlobalExecutionTimeoutSeconds > 0)
+	{
+		TimeoutMs = Executor.GlobalExecutionTimeoutSeconds * 1000LL;
+	}
+	return Executor.RemoteExecutor->GetExecution()->ExecuteAsync(State.ExecuteRequest, nullptr, TimeoutMs);
+}
+
+TFuture<TPair<FStatus, FBatchReadBlobsResponse>> FRemoteBuildExecutionRequest::DownloadResultsAsync()
+{
+	State.BatchReadBlobsRequest.InstanceName = Executor.InstanceName;
+	//State.BatchReadBlobsRequest.Digests.Add(State.ExecuteResponse.Result.StdoutDigest);
+	//State.BatchReadBlobsRequest.Digests.Add(State.ExecuteResponse.Result.StderrDigest);
+
+	for (const FOutputFile& OutputFile : State.ExecuteResponse.Result.OutputFiles)
+	{
+		State.BatchReadBlobsRequest.Digests.Add(OutputFile.Digest);
+	}
+
+	return Executor.ContentAddressableStorage->BatchReadBlobsAsync(State.BatchReadBlobsRequest);
+}
+
+void FRemoteBuildExecutionRequest::OnMissingBlobsDetermined(FFindMissingBlobsResponse&& Result)
+{
+	if (ProcessCancellation())
+	{
+		return;
+	}
+
+	State.FindMissingBlobsResponse = MoveTemp(Result);
+	TArray<FStringView> MissingInputViews;
+	GatherMissingInputFileBlobs(MissingInputViews);
+	if (!MissingInputViews.IsEmpty())
+	{
+		TRefCountPtr<FRemoteBuildExecutionRequest> SelfRef(this);
+		CompletionCallback({State.BuildAction.GetKey(), {}, MissingInputViews, EStatus::Ok});
+		CompletionEvent->Trigger();
+		return;
+	}
+
+	if (!Executor.LimitingHeuristics.TryStartNewBuild(Executor.Stats))
+	{
+		TRefCountPtr<FRemoteBuildExecutionRequest> SelfRef(this);
+		CompletionCallback({State.BuildAction.GetKey(), {}, {}, EStatus::Error});
+		CompletionEvent->Trigger();
+		return;
+	}
+
+	bHeuristicBuildStarted = true;
+
+	// TODO: This should be async but isn't.  Requires IRequest chaining.
+	LoadMissingWorkerFileBlobsAsync();
+
+	if (!State.FindMissingBlobsResponse.MissingBlobDigests.IsEmpty())
+	{
+		UploadMissingBlobsAsync()
+			.Next([this] (const TPair<FStatus, FBatchUpdateBlobsResponse>& InnerResult) { OnMissingBlobsUploaded(InnerResult.Value); });
+	}
+	else
+	{
+		ExecuteBuildAsync()
+			.Next([this] (FExecuteResponse&& Result) { OnExecutionCompleted(MoveTemp(Result)); });
+	}
+}
+
+void FRemoteBuildExecutionRequest::OnMissingBlobsUploaded(const FBatchUpdateBlobsResponse& Result)
+{
+	if (ProcessCancellation())
+	{
+		return;
+	}
+
+	UE_LOG(LogDerivedDataBuildRemoteExecutor, Verbose, TEXT("Uploaded %d data blobs for remote execution."), State.BatchUpdateBlobsRequest.Requests.Num());
+	if (!ValidateUploadSuccess(Result))
+	{
+		TRefCountPtr<FRemoteBuildExecutionRequest> SelfRef(this);
+		CompletionCallback({State.BuildAction.GetKey(), {}, {}, EStatus::Error});
+		CompletionEvent->Trigger();
+		return;
+	}
+
+	ExecuteBuildAsync()
+		.Next([this] (FExecuteResponse&& Result) { OnExecutionCompleted(MoveTemp(Result)); });
+
+}
+
+void FRemoteBuildExecutionRequest::OnExecutionCompleted(FExecuteResponse&& Result)
+{
+	if (ProcessCancellation())
+	{
+		return;
+	}
+
+	State.ExecuteResponse = MoveTemp(Result);
+	if (State.ExecuteResponse.Status.Ok())
+	{
+		DownloadResultsAsync()
+			.Next([this] (TPair<FStatus, FBatchReadBlobsResponse>&& InnerResult)
+				{
+					OnOutputBlobsDownloaded(MoveTemp(InnerResult.Value));
+				});
+	}
+	else
+	{
+		TRefCountPtr<FRemoteBuildExecutionRequest> SelfRef(this);
+		UE_LOG(LogDerivedDataBuildRemoteExecutor, Warning, TEXT("Remote execution system error: Failed to execute build operation!"));
+		CompletionCallback({State.BuildAction.GetKey(), {}, {}, EStatus::Error});
+		CompletionEvent->Trigger();
+	}
+
+}
+
+void FRemoteBuildExecutionRequest::OnOutputBlobsDownloaded(FBatchReadBlobsResponse&& Result)
+{
+	if (ProcessCancellation())
+	{
+		return;
+	}
+
+	EStatus BuildStatus = EStatus::Error;
+	FOptionalBuildOutput BuildOutput = ComposeBuildOutput(Result, BuildStatus);
+	if (BuildStatus == EStatus::Ok)
+	{
+		Executor.Stats.TotalSuccessfulRemoteBuilds.fetch_add(1, std::memory_order_relaxed);
+	}
+
+	TRefCountPtr<FRemoteBuildExecutionRequest> SelfRef(this);
+	CompletionCallback({State.BuildAction.GetKey(), MoveTemp(BuildOutput), {}, BuildStatus});
+	CompletionEvent->Trigger();
+}
+
 } // namespace UE::DerivedData
 
 TOptional<UE::DerivedData::FRemoteBuildWorkerExecutor> GRemoteBuildWorkerExecutor;
