Adding a commandlet called GatherShaderStatsFromMaterials.  This will make an exhaustive text-based dump of all the shaders in a given project for a given platform.

The commandlet can be run via:

    ProjectName -run=GatherShaderStatsFromMaterials -platform=windows

Key things to note:

- We want to avoid compiling shaders to get this information (too time consuming).  To achieve this goal I've added a new function GetShaderTypes on UMaterialInterface which will discover what shaders are cached for the given material w/o having to trigger shader compilation.
- The commandlet uses the asset registry to discover all the Materials and MaterialInstances in a project.
- Information is dumped in a hierarchy grouped by material, vertex factory, and then all the shader types for that VF type.
- This also provides a nice summary of static parameter usage for material instances.
- The file is dumped out to the project's Saved/MaterialStats folder.

Here is a sample of the output of the text-based dump:

Material: WorldGridMaterial - 1096 shaders

	FGPUSpriteVertexFactory - 34 shaders
		TShadowDepthPSPixelShadowDepth_PerspectiveCorrect
		FLumenCardPS<false>
		TShadowDepthPSPixelShadowDepth_NonPerspectiveCorrect
		TVirtualTextureVSBaseColorNormalSpecular
		TVirtualTexturePSBaseColor
		FMaterialTexCoordScalePS
		TVirtualTextureVSWorldHeight
		TVirtualTextureVSBaseColor
		TShadowDepthVSForGSVertexShadowDepth_OnePassPointLight
		FHitProxyVS
		FHitProxyPS
		FPrimitiveDistanceAccuracyPS
		FMeshTexCoordSizeAccuracyPS
		TBasePassPSFNoLightMapPolicy
		TShadowDepthVSVertexShadowDepth_OutputDepth
		FLumenCardVS
		TShadowDepthPSPixelShadowDepth_VirtualShadowMap
		TShadowDepthVSVertexShadowDepth_VSLayer
		TVirtualTexturePSWorldHeight
		FRequiredTextureResolutionPS
		TShadowDepthVSVertexShadowDepth_PerspectiveCorrect
		FComplexityAccumulatePS
		FComplexityAccumulatePS
		TBasePassPSFNoLightMapPolicySkylight
		TShadowDepthPSPixelShadowDepth_OnePassPointLight
		TVirtualTexturePSBaseColorNormalSpecular
		FLODColorationPS
		TShadowDepthVSVertexShadowDepth_VirtualShadowMap
		TShadowDepthVSVertexShadowDepth_OnePassPointLight
		TDepthOnlyVS<false>
		FDebugViewModeVS
		TBasePassVSFNoLightMapPolicy
		FOnePassPointShadowDepthGS
		DepthNoPixelPipeline
			TDepthOnlyVS<false>

diff --git a/Engine/Source/Editor/UnrealEd/Classes/Commandlets/GatherShaderStatsFromMaterialsCommandlet.h b/Engine/Source/Editor/UnrealEd/Classes/Commandlets/GatherShaderStatsFromMaterialsCommandlet.h
new file mode 100644
--- /dev/null
+++ b/Engine/Source/Editor/UnrealEd/Classes/Commandlets/GatherShaderStatsFromMaterialsCommandlet.h
@@ -0,0 +1,18 @@
+// Copyright Epic Games, Inc. All Rights Reserved.
+
+#pragma once
+
+#include "CoreMinimal.h"
+#include "UObject/ObjectMacros.h"
+#include "Commandlets/Commandlet.h"
+#include "GatherShaderStatsFromMaterialsCommandlet.generated.h"
+
+UCLASS(config=Editor)
+class UGatherShaderStatsFromMaterialsCommandlet : public UCommandlet
+{
+	GENERATED_UCLASS_BODY()
+
+	//~ Begin UCommandlet Interface
+	virtual int32 Main(const FString& Params) override;
+	//~ End UCommandlet Interface
+};
diff --git a/Engine/Source/Editor/UnrealEd/Private/Commandlets/GatherShaderStatsFromMaterialsCommandlet.cpp b/Engine/Source/Editor/UnrealEd/Private/Commandlets/GatherShaderStatsFromMaterialsCommandlet.cpp
new file mode 100644
--- /dev/null
+++ b/Engine/Source/Editor/UnrealEd/Private/Commandlets/GatherShaderStatsFromMaterialsCommandlet.cpp
@@ -0,0 +1,293 @@
+// Copyright Epic Games, Inc. All Rights Reserved.
+
+#include "Commandlets/GatherShaderStatsFromMaterialsCommandlet.h"
+#include "AssetData.h"
+#include "AssetRegistryModule.h"
+#include "Materials/Material.h"
+#include "Materials/MaterialInstance.h"
+#include "Materials/MaterialInstanceConstant.h"
+#include "CollectionManagerTypes.h"
+#include "ICollectionManager.h"
+#include "CollectionManagerModule.h"
+
+DEFINE_LOG_CATEGORY_STATIC(LogGatherShaderStatsFromMaterialsCommandlet, Log, All);
+
+class FLogger
+{
+public:
+	FLogger() = delete;
+	FLogger(const FString& FileName)
+	{
+		DebugWriter = IFileManager::Get().CreateFileWriter(*FileName);
+	}
+	~FLogger()
+	{
+		DebugWriter->Close();
+		delete DebugWriter;
+	}
+
+	void Log(const FString& OutString)
+	{
+		//UE_LOG(LogGatherShaderStatsFromMaterialsCommandlet, Log, TEXT("%s"), *OutString);
+		FString OuputLine = OutString + TEXT("\n");
+		DebugWriter->Serialize(const_cast<ANSICHAR*>(StringCast<ANSICHAR>(*OuputLine).Get()), OuputLine.Len());
+	}
+
+private:
+	FArchive* DebugWriter = nullptr;
+};
+
+UGatherShaderStatsFromMaterialsCommandlet::UGatherShaderStatsFromMaterialsCommandlet(const FObjectInitializer& ObjectInitializer)
+	: Super(ObjectInitializer)
+{
+}
+
+int GetTotalShaders(const TArray<FDebugShaderTypeInfo>& OutShaderInfo)
+{
+	int TotalShadersForMaterial = 0;
+	for (const FDebugShaderTypeInfo& ShaderInfo : OutShaderInfo)
+	{
+		TotalShadersForMaterial += ShaderInfo.ShaderTypes.Num();
+
+		for (const FDebugShaderPipelineInfo& PipelineInfo : ShaderInfo.Pipelines)
+		{
+			TotalShadersForMaterial += PipelineInfo.ShaderTypes.Num();
+		}
+	}
+	return TotalShadersForMaterial;
+}
+
+void PrintDebugShaderInfo(FLogger& Output, const TArray<FDebugShaderTypeInfo>& OutShaderInfo)
+{
+	for (const FDebugShaderTypeInfo& ShaderInfo : OutShaderInfo)
+	{
+		int TotalShadersForVF = 0;
+		TotalShadersForVF += ShaderInfo.ShaderTypes.Num();
+
+		for (const FDebugShaderPipelineInfo& PipelineInfo : ShaderInfo.Pipelines)
+		{
+			TotalShadersForVF += PipelineInfo.ShaderTypes.Num();
+		}
+
+		Output.Log(TEXT(""));
+		Output.Log(FString::Printf(TEXT("\t%s - %d shaders"), ShaderInfo.VFType->GetName(), TotalShadersForVF));
+
+		for (FShaderType* ShaderType : ShaderInfo.ShaderTypes)
+		{
+			Output.Log(FString::Printf(TEXT("\t\t%s"), ShaderType->GetName()));
+		}
+
+		for (const FDebugShaderPipelineInfo& PipelineInfo : ShaderInfo.Pipelines)
+		{
+			Output.Log(FString::Printf(TEXT("\t\t%s"), PipelineInfo.Pipeline->GetName()));
+
+			for (FShaderType* ShaderType : PipelineInfo.ShaderTypes)
+			{
+				Output.Log(FString::Printf(TEXT("\t\t\t%s"), ShaderType->GetName()));
+			}
+		}
+
+		Output.Log(TEXT(""));
+	}
+}
+
+int ProcessMaterials(const EShaderPlatform ShaderPlatform, FLogger& Output, TArray<FAssetData>& MaterialList)
+{
+	int TotalShaders = 0;
+
+	for (const FAssetData& AssetData : MaterialList)
+	{
+		if (UMaterial* Material = Cast<UMaterial>(AssetData.GetAsset()))
+		{
+			TArray<FDebugShaderTypeInfo> OutShaderInfo;
+			Material->GetShaderTypes(ShaderPlatform, OutShaderInfo);
+
+			const int TotalShadersForMaterial = GetTotalShaders(OutShaderInfo);
+			TotalShaders += TotalShadersForMaterial;
+
+			Output.Log(TEXT(""));
+			Output.Log(FString::Printf(TEXT("Material: %s - %d shaders"), *AssetData.AssetName.ToString(), TotalShadersForMaterial));
+
+			PrintDebugShaderInfo(Output, OutShaderInfo);
+		}
+	}
+
+	Output.Log(TEXT(""));
+	Output.Log(TEXT("Summary"));
+	Output.Log(FString::Printf(TEXT("Total Materials: %d"), MaterialList.Num()));
+	Output.Log(FString::Printf(TEXT("Total Shaders: %d"), TotalShaders));
+
+	return TotalShaders;
+}
+
+int ProcessMaterialInstances(const EShaderPlatform ShaderPlatform, FLogger& Output, TArray<FAssetData>& MaterialInstanceList)
+{
+	int TotalShaders = 0;
+
+	int StaticPermutations = 0;
+	for (const FAssetData& AssetData : MaterialInstanceList)
+	{
+		if (UMaterialInstance* MaterialInstance = Cast<UMaterialInstance>(AssetData.GetAsset()))
+		{
+			TArray<FDebugShaderTypeInfo> OutShaderInfo;
+			MaterialInstance->GetShaderTypes(ShaderPlatform, OutShaderInfo);
+
+			const int TotalShadersForMaterial = GetTotalShaders(OutShaderInfo);
+			TotalShaders += TotalShadersForMaterial;
+
+			FString StaticParameterString(TEXT(""));
+
+			if (MaterialInstance->bHasStaticPermutationResource)
+			{
+				const FStaticParameterSet& ParameterSet = MaterialInstance->GetStaticParameters();
+				for (int32 StaticSwitchIndex = 0; StaticSwitchIndex < ParameterSet.StaticSwitchParameters.Num(); ++StaticSwitchIndex)
+				{
+					const FStaticSwitchParameter& StaticSwitchParameter = ParameterSet.StaticSwitchParameters[StaticSwitchIndex];
+					StaticParameterString += FString::Printf(
+						TEXT(", StaticSwitch'%s'=%s"),
+						*StaticSwitchParameter.ParameterInfo.ToString(),
+						StaticSwitchParameter.Value ? TEXT("True") : TEXT("False")
+					);
+				}
+			}
+
+			Output.Log(TEXT(""));
+			Output.Log(FString::Printf(TEXT("Material Instance: %s - %d shaders"), *AssetData.AssetName.ToString(), TotalShadersForMaterial));
+			Output.Log(FString::Printf(TEXT("Static Parameter %s"), *StaticParameterString));
+			Output.Log(FString::Printf(TEXT("Parent: %s"), MaterialInstance->Parent ? *MaterialInstance->Parent->GetName() : TEXT("NO PARENT")));
+
+			PrintDebugShaderInfo(Output, OutShaderInfo);
+
+			if (MaterialInstance->bHasStaticPermutationResource)
+			{
+				StaticPermutations++;
+			}
+		}
+	}
+
+	Output.Log(TEXT(""));
+	Output.Log(TEXT("Summary"));
+	Output.Log(FString::Printf(TEXT("Total Material Instances: %d"), MaterialInstanceList.Num()));
+	Output.Log(FString::Printf(TEXT("Material Instances w/ Static Permutations: %d"), StaticPermutations));
+	Output.Log(FString::Printf(TEXT("Total Shaders: %d"), TotalShaders));
+
+	return TotalShaders;
+}
+
+int32 UGatherShaderStatsFromMaterialsCommandlet::Main(const FString& Params)
+{
+	TArray<FString> Tokens;
+	TArray<FString> Switches;
+	TMap<FString, FString> ParamVals;
+	UCommandlet::ParseCommandLine(*Params, Tokens, Switches, ParamVals);
+
+	// Display help
+	if (Switches.Contains("help"))
+	{
+		UE_LOG(LogGatherShaderStatsFromMaterialsCommandlet, Log, TEXT("GatherShaderStatsFromMaterials"));
+		UE_LOG(LogGatherShaderStatsFromMaterialsCommandlet, Log, TEXT("This commandlet will dump to a human readable plain text file of all the shaders that would be compiled for all materials in a project."));
+		UE_LOG(LogGatherShaderStatsFromMaterialsCommandlet, Log, TEXT("Options:"));
+		UE_LOG(LogGatherShaderStatsFromMaterialsCommandlet, Log, TEXT(" Required: -platform=<platform>     (Which shader platform do you want results for?)"));
+		UE_LOG(LogGatherShaderStatsFromMaterialsCommandlet, Log, TEXT(" Optional: -collection=<name>       (You can alternatively specify a collection of assets to run this on.)"));
+		return 0;
+	}
+
+	// Parse platform
+	FString PlatformName;
+	ITargetPlatform* TargetPlatform = nullptr;
+	{
+		if (!FParse::Value(*Params, TEXT("platform="), PlatformName, true))
+		{
+			UE_LOG(LogGatherShaderStatsFromMaterialsCommandlet, Warning, TEXT("You must include a target platform with -platform=<platform>"));
+			return 1;
+		}
+		ITargetPlatformManagerModule& TPM = GetTargetPlatformManagerRef();
+		TargetPlatform = TPM.FindTargetPlatform(PlatformName);
+		if (TargetPlatform == nullptr)
+		{
+			UE_LOG(LogGatherShaderStatsFromMaterialsCommandlet, Display, TEXT("Target platform '%s' was not found.  Valid platforms are:"), *PlatformName);
+			for (ITargetPlatform* platform : TPM.GetTargetPlatforms())
+			{
+				UE_LOG(LogGatherShaderStatsFromMaterialsCommandlet, Display, TEXT("\t'%s'"), *platform->PlatformName());
+			}
+			return 1;
+		}
+
+		TargetPlatform->RefreshSettings();
+	}
+
+	TArray<FName> DesiredShaderFormats;
+	TargetPlatform->GetAllTargetedShaderFormats(DesiredShaderFormats);
+
+	IAssetRegistry& AssetRegistry = FModuleManager::Get().LoadModuleChecked<FAssetRegistryModule>("AssetRegistry").Get();
+	AssetRegistry.SearchAllAssets(true);
+
+	TArray<FAssetData> MaterialList;
+	TArray<FAssetData> MaterialInstanceList;
+
+	// Parse collection
+	FString CollectionName;
+	if (FParse::Value(*Params, TEXT("collection="), CollectionName, true))
+	{
+		if (!CollectionName.IsEmpty())
+		{
+			// Get the list of materials from a collection
+			FARFilter Filter;
+			Filter.PackagePaths.Add(FName(TEXT("/Game")));
+			Filter.bRecursivePaths = true;
+			Filter.ClassNames.Add(UMaterial::StaticClass()->GetFName());
+
+			FCollectionManagerModule& CollectionManagerModule = FCollectionManagerModule::GetModule();
+			CollectionManagerModule.Get().GetObjectsInCollection(FName(*CollectionName), ECollectionShareType::CST_All, Filter.ObjectPaths, ECollectionRecursionFlags::SelfAndChildren);
+
+			AssetRegistry.GetAssets(Filter, MaterialList);
+
+			Filter.ClassNames.Empty();
+			Filter.ClassNames.Add(UMaterialInstance::StaticClass()->GetFName());
+			Filter.ClassNames.Add(UMaterialInstanceConstant::StaticClass()->GetFName());
+
+			AssetRegistry.GetAssets(Filter, MaterialInstanceList);
+		}
+	}
+	else
+	{
+		if (!AssetRegistry.IsLoadingAssets())
+		{
+			AssetRegistry.GetAssetsByClass(UMaterial::StaticClass()->GetFName(), MaterialList, true);
+			AssetRegistry.GetAssetsByClass(UMaterialInstance::StaticClass()->GetFName(), MaterialInstanceList, true);
+		}
+	}
+
+	const double StartTime = FPlatformTime::Seconds();
+
+	const FString TimeNow = FDateTime::Now().ToString();
+	FString FileName = FPaths::Combine(*FPaths::ProjectSavedDir(), FString::Printf(TEXT("MaterialStats/ShaderStatsFromMaterials-%s.txt"), *TimeNow));
+
+	FLogger Output(FileName);
+
+	int TotalShaders = 0;
+	int TotalAssets = 0;
+
+	// Cache for all the shader formats that the cooking target requires
+	for (int32 FormatIndex = 0; FormatIndex < DesiredShaderFormats.Num(); FormatIndex++)
+	{
+		const EShaderPlatform LegacyShaderPlatform = ShaderFormatToLegacyShaderPlatform(DesiredShaderFormats[FormatIndex]);
+
+		TotalShaders += ProcessMaterials(LegacyShaderPlatform, Output, MaterialList);
+		TotalAssets += MaterialList.Num();
+
+		TotalShaders += ProcessMaterialInstances(LegacyShaderPlatform, Output, MaterialInstanceList);
+		TotalAssets += MaterialInstanceList.Num();
+	}
+
+	Output.Log(TEXT(""));
+	Output.Log(TEXT("Summary"));
+	Output.Log(FString::Printf(TEXT("Total Assets: %d"), TotalAssets));
+	Output.Log(FString::Printf(TEXT("Total Shaders: %d"), TotalShaders));
+
+	const double EndTime = FPlatformTime::Seconds() - StartTime;
+	Output.Log(TEXT(""));
+	Output.Log(FString::Printf(TEXT("Commandlet Took: %lf"), EndTime));
+
+	return true;
+}
diff --git a/Engine/Source/Runtime/Engine/Classes/Materials/Material.h b/Engine/Source/Runtime/Engine/Classes/Materials/Material.h
--- a/Engine/Source/Runtime/Engine/Classes/Materials/Material.h
+++ b/Engine/Source/Runtime/Engine/Classes/Materials/Material.h
@@ -1553,6 +1553,8 @@ private:
 	 */
 	void CacheResourceShadersForCooking(EShaderPlatform Platform, TArray<FMaterialResource*>& OutCachedMaterialResources, const ITargetPlatform* TargetPlatform = nullptr);
 
+	void GetNewResources(EShaderPlatform ShaderPlatform, TArray<FMaterialResource*>& NewResourcesToCache);
+
 	/** Caches shader maps for an array of material resources. */
 	void CacheShadersForResources(EShaderPlatform ShaderPlatform, const TArray<FMaterialResource*>& ResourcesToCache, EMaterialShaderPrecompileMode PrecompileMode = EMaterialShaderPrecompileMode::Default, const ITargetPlatform* TargetPlatform = nullptr);
 
@@ -1798,6 +1800,18 @@ public:
 	void SaveShaderStableKeys(const class ITargetPlatform* TP);
 	ENGINE_API virtual void SaveShaderStableKeysInner(const class ITargetPlatform* TP, const struct FStableShaderKeyAndValue& SaveKeyVal) override;
 
+#if WITH_EDITOR
+	/**
+	*	Gathers a list of shader types sorted by vertex factory types that should be cached for this material.  Avoids doing expensive material
+	*	and shader compilation to acquire this information.
+	*
+	*	@param	Platform		The shader platform to get info for.
+	*	@param	OutShaderInfo	Array of results sorted by vertex factory type, and shader type.
+	*
+	*/
+	ENGINE_API virtual void GetShaderTypes(EShaderPlatform Platform, TArray<FDebugShaderTypeInfo>& OutShaderInfo) override;
+#endif // WITH_EDITOR
+
 	bool HasBaseColorConnected() const { return BaseColor.IsConnected() || GetCachedExpressionData().IsMaterialAttributePropertyConnected(MP_BaseColor); }
 	bool HasRoughnessConnected() const { return Roughness.IsConnected() || GetCachedExpressionData().IsMaterialAttributePropertyConnected(MP_Roughness); }
 	bool HasAmbientOcclusionConnected() const { return AmbientOcclusion.IsConnected() || GetCachedExpressionData().IsMaterialAttributePropertyConnected(MP_AmbientOcclusion); }
diff --git a/Engine/Source/Runtime/Engine/Classes/Materials/MaterialInstance.h b/Engine/Source/Runtime/Engine/Classes/Materials/MaterialInstance.h
--- a/Engine/Source/Runtime/Engine/Classes/Materials/MaterialInstance.h
+++ b/Engine/Source/Runtime/Engine/Classes/Materials/MaterialInstance.h
@@ -726,6 +726,18 @@ public:
 	void SaveShaderStableKeys(const class ITargetPlatform* TP);
 	ENGINE_API virtual void SaveShaderStableKeysInner(const class ITargetPlatform* TP, const struct FStableShaderKeyAndValue& SaveKeyVal) override;
 
+#if WITH_EDITOR
+	/**
+	*	Gathers a list of shader types sorted by vertex factory types that should be cached for this material.  Avoids doing expensive material
+	*	and shader compilation to acquire this information.
+	*
+	*	@param	Platform		The shader platform to get info for.
+	*	@param	OutShaderInfo	Array of results sorted by vertex factory type, and shader type.
+	*
+	*/
+	ENGINE_API virtual void GetShaderTypes(EShaderPlatform Platform, TArray<FDebugShaderTypeInfo>& OutShaderInfo) override;
+#endif // WITH_EDITOR
+
 #if WITH_EDITOR
 	void BeginAllowCachingStaticParameterValues();
 	void EndAllowCachingStaticParameterValues();
diff --git a/Engine/Source/Runtime/Engine/Classes/Materials/MaterialInterface.h b/Engine/Source/Runtime/Engine/Classes/Materials/MaterialInterface.h
--- a/Engine/Source/Runtime/Engine/Classes/Materials/MaterialInterface.h
+++ b/Engine/Source/Runtime/Engine/Classes/Materials/MaterialInterface.h
@@ -911,6 +911,18 @@ public:
 	*/
 	ENGINE_API void SortTextureStreamingData(bool bForceSort, bool bFinalSort);
 
+#if WITH_EDITOR
+	/**
+	*	Gathers a list of shader types sorted by vertex factory types that should be cached for this material.  Avoids doing expensive material
+	*	and shader compilation to acquire this information.
+	*
+	*	@param	Platform		The shader platform to get info for.
+	*	@param	OutShaderInfo	Array of results sorted by vertex factory type, and shader type.
+	*
+	*/
+	ENGINE_API virtual void GetShaderTypes(EShaderPlatform Platform, TArray<FDebugShaderTypeInfo>& OutShaderInfo) {};
+#endif // WITH_EDITOR
+
 protected:
 
 	/** Returns a bitfield indicating which feature levels should be compiled for rendering. GMaxRHIFeatureLevel is always present */
diff --git a/Engine/Source/Runtime/Engine/Private/Materials/Material.cpp b/Engine/Source/Runtime/Engine/Private/Materials/Material.cpp
--- a/Engine/Source/Runtime/Engine/Private/Materials/Material.cpp
+++ b/Engine/Source/Runtime/Engine/Private/Materials/Material.cpp
@@ -3073,6 +3073,16 @@ void UMaterial::CacheResourceShadersForRendering(bool bRegenerateId, EMaterialSh
 }
 
 void UMaterial::CacheResourceShadersForCooking(EShaderPlatform ShaderPlatform, TArray<FMaterialResource*>& OutCachedMaterialResources, const ITargetPlatform* TargetPlatform)
+{
+	TArray<FMaterialResource*> NewResourcesToCache;	// only new resources need to have CacheShaders() called on them, whereas OutCachedMaterialResources may already contain resources for another shader platform
+	GetNewResources(ShaderPlatform, NewResourcesToCache);
+
+	CacheShadersForResources(ShaderPlatform, NewResourcesToCache, EMaterialShaderPrecompileMode::Background, TargetPlatform);
+
+	OutCachedMaterialResources.Append(NewResourcesToCache);
+}
+
+void UMaterial::GetNewResources(EShaderPlatform ShaderPlatform, TArray<FMaterialResource*>& NewResourcesToCache)
 {
 	ERHIFeatureLevel::Type TargetFeatureLevel = GetMaxSupportedFeatureLevel(ShaderPlatform);
 
@@ -3082,7 +3092,6 @@ void UMaterial::CacheResourceShadersForCooking(EShaderPlatform ShaderPlatform, T
 	const UShaderPlatformQualitySettings* MaterialQualitySettings = UMaterialShaderQualitySettings::Get()->GetShaderPlatformQualitySettings(ShaderPlatform);
 	bool bNeedDefaultQuality = false;
 
-	TArray<FMaterialResource*> NewResourcesToCache;	// only new resources need to have CacheShaders() called on them, whereas OutCachedMaterialResources may already contain resources for another shader platform
 	for (int32 QualityLevelIndex = 0; QualityLevelIndex < EMaterialQualityLevel::Num; QualityLevelIndex++)
 	{
 		// Add all quality levels actually used
@@ -3109,10 +3118,6 @@ void UMaterial::CacheResourceShadersForCooking(EShaderPlatform ShaderPlatform, T
 		NewResource->SetMaterial(this, nullptr, (ERHIFeatureLevel::Type)TargetFeatureLevel);
 		NewResourcesToCache.Add(NewResource);
 	}
-
-	CacheShadersForResources(ShaderPlatform, NewResourcesToCache, EMaterialShaderPrecompileMode::Background, TargetPlatform);
-
-	OutCachedMaterialResources.Append(NewResourcesToCache);
 }
 
 void UMaterial::CacheShadersForResources(EShaderPlatform ShaderPlatform, const TArray<FMaterialResource*>& ResourcesToCache, EMaterialShaderPrecompileMode PrecompileMode, const ITargetPlatform* TargetPlatform)
@@ -4362,6 +4367,21 @@ void UMaterial::SaveShaderStableKeysInner(const class ITargetPlatform* TP, const
 #endif
 }
 
+#if WITH_EDITOR
+void UMaterial::GetShaderTypes(EShaderPlatform ShaderPlatform, TArray<FDebugShaderTypeInfo>& OutShaderInfo)
+{
+	TArray<FMaterialResource*> NewResourcesToCache;
+	GetNewResources(ShaderPlatform, NewResourcesToCache);
+
+	for (FMaterialResource* Resource : NewResourcesToCache)
+	{
+		Resource->GetShaderTypes(ShaderPlatform, OutShaderInfo);
+		delete Resource;
+	}
+
+	NewResourcesToCache.Empty();
+}
+#endif // WITH_EDITOR
 
 void UMaterial::PropagateDataToMaterialProxy()
 {
diff --git a/Engine/Source/Runtime/Engine/Private/Materials/MaterialInstance.cpp b/Engine/Source/Runtime/Engine/Private/Materials/MaterialInstance.cpp
--- a/Engine/Source/Runtime/Engine/Private/Materials/MaterialInstance.cpp
+++ b/Engine/Source/Runtime/Engine/Private/Materials/MaterialInstance.cpp
@@ -4836,6 +4836,38 @@ void UMaterialInstance::SaveShaderStableKeysInner(const class ITargetPlatform* T
 #endif
 }
 
+#if WITH_EDITOR
+void UMaterialInstance::GetShaderTypes(EShaderPlatform Platform, TArray<FDebugShaderTypeInfo>& OutShaderInfo)
+{
+	if (bHasStaticPermutationResource)
+	{
+		check(IsA(UMaterialInstanceConstant::StaticClass()));
+		UMaterial* BaseMaterial = GetMaterial();
+
+		uint32 FeatureLevelsToCompile = GetFeatureLevelsToCompileForRendering();
+		const EMaterialQualityLevel::Type ActiveQualityLevel = GetCachedScalabilityCVars().MaterialQualityLevel;
+
+		TArray<FMaterialResource*> ResourcesToCache;
+		while (FeatureLevelsToCompile != 0)
+		{
+			const ERHIFeatureLevel::Type FeatureLevel = (ERHIFeatureLevel::Type)FBitSet::GetAndClearNextBit(FeatureLevelsToCompile);
+			const EShaderPlatform ShaderPlatform = GShaderPlatformForFeatureLevel[FeatureLevel];
+
+			// Only cache shaders for the quality level that will actually be used to render
+			// In cooked build, there is no shader compilation but this is still needed to
+			// register the loaded shadermap
+			FMaterialResource* CurrentResource = FindOrCreateMaterialResource(StaticPermutationMaterialResources, BaseMaterial, this, FeatureLevel, ActiveQualityLevel);
+			check(CurrentResource);
+		}
+
+		for (FMaterialResource* CurrentResource : StaticPermutationMaterialResources)
+		{
+			CurrentResource->GetShaderTypes(Platform, OutShaderInfo);
+		}
+	}
+}
+#endif // WITH_EDITOR
+
 #if WITH_EDITOR
 void UMaterialInstance::BeginAllowCachingStaticParameterValues()
 {
diff --git a/Engine/Source/Runtime/Engine/Private/Materials/MaterialShared.cpp b/Engine/Source/Runtime/Engine/Private/Materials/MaterialShared.cpp
--- a/Engine/Source/Runtime/Engine/Private/Materials/MaterialShared.cpp
+++ b/Engine/Source/Runtime/Engine/Private/Materials/MaterialShared.cpp
@@ -3822,6 +3822,45 @@ void FMaterial::SaveShaderStableKeys(EShaderPlatform TargetShaderPlatform, FStab
 #endif
 }
 
+#if WITH_EDITOR
+void FMaterial::GetShaderTypes(EShaderPlatform Platform, TArray<FDebugShaderTypeInfo>& OutShaderInfo)
+{
+	const FMaterialShaderParameters MaterialParameters(this);
+	const FMaterialShaderMapLayout& Layout = AcquireMaterialShaderMapLayout(Platform, GetCurrentShaderPermutationFlags(), MaterialParameters);
+
+	for (const FMeshMaterialShaderMapLayout& MeshLayout : Layout.MeshShaderMaps)
+	{
+		FDebugShaderTypeInfo debugInfo;
+		debugInfo.VFType = MeshLayout.VertexFactoryType;
+
+		for (const FShaderLayoutEntry& Shader : MeshLayout.Shaders)
+		{
+			if (ShouldCache(Platform, Shader.ShaderType, MeshLayout.VertexFactoryType))
+			{
+				debugInfo.ShaderTypes.Add(Shader.ShaderType);
+			}
+		}
+
+		for (const FShaderPipelineType* Pipeline : MeshLayout.ShaderPipelines)
+		{
+			if (ShouldCachePipeline(Platform, Pipeline, MeshLayout.VertexFactoryType))
+			{
+				FDebugShaderPipelineInfo pipelineInfo;
+				pipelineInfo.Pipeline = Pipeline;
+
+				for (const FShaderType* Type : Pipeline->GetStages())
+				{
+					pipelineInfo.ShaderTypes.Add((FShaderType*)Type);
+				}
+
+				debugInfo.Pipelines.Add(pipelineInfo);
+			}
+		}
+
+		OutShaderInfo.Add(debugInfo);
+	}
+}
+#endif
 
 FMaterialUpdateContext::FMaterialUpdateContext(uint32 Options, EShaderPlatform InShaderPlatform)
 {
diff --git a/Engine/Source/Runtime/Engine/Public/MaterialShared.h b/Engine/Source/Runtime/Engine/Public/MaterialShared.h
--- a/Engine/Source/Runtime/Engine/Public/MaterialShared.h
+++ b/Engine/Source/Runtime/Engine/Public/MaterialShared.h
@@ -672,6 +672,19 @@ public:
 	LAYOUT_BITFIELD(uint8, bUsesAnisotropy, 1);
 };
 
+struct FDebugShaderPipelineInfo
+{
+	const FShaderPipelineType* Pipeline = nullptr;
+	TArray<FShaderType*> ShaderTypes;
+};
+
+struct FDebugShaderTypeInfo
+{
+	FVertexFactoryType* VFType = nullptr;
+	TArray<FShaderType*> ShaderTypes;
+	TArray<FDebugShaderPipelineInfo> Pipelines;
+};
+
 /** 
  * Usage options for a shader map.
  * The purpose of EMaterialShaderMapUsage is to allow creating a unique yet deterministic (no appCreateGuid) Id,
@@ -1946,6 +1959,18 @@ public:
 	void DumpDebugInfo();
 	void SaveShaderStableKeys(EShaderPlatform TargetShaderPlatform, struct FStableShaderKeyAndValue& SaveKeyVal); // arg is non-const, we modify it as we go
 
+#if WITH_EDITOR
+/**
+*	Gathers a list of shader types sorted by vertex factory types that should be cached for this material.  Avoids doing expensive material
+*	and shader compilation to acquire this information.
+*
+*	@param	Platform		The shader platform to get info for.
+*	@param	OutShaderInfo	Array of results sorted by vertex factory type, and shader type.
+*
+*/
+	void GetShaderTypes(EShaderPlatform Platform, TArray<FDebugShaderTypeInfo>& OutShaderInfo);
+#endif // WITH_EDITOR
+
 #if WITH_EDITOR
 	/** 
 	 * Adds an FMaterial to the global list.
