Updating a conditional to avoid asset editor crashes after layout reloading

diff --git a/Engine/Source/Editor/UnrealEd/Private/Toolkits/SStandaloneAssetEditorToolkitHost.cpp b/Engine/Source/Editor/UnrealEd/Private/Toolkits/SStandaloneAssetEditorToolkitHost.cpp
--- a/Engine/Source/Editor/UnrealEd/Private/Toolkits/SStandaloneAssetEditorToolkitHost.cpp
+++ b/Engine/Source/Editor/UnrealEd/Private/Toolkits/SStandaloneAssetEditorToolkitHost.cpp
@@ -138,7 +138,7 @@ void SStandaloneAssetEditorToolkitHost::CreateDefaultStandaloneMenuBar(UToolMenu
 	MenuBar->FindOrAddSection(NAME_None).AddDynamicEntry("DynamicAssetEntry", FNewToolMenuSectionDelegate::CreateLambda([](FToolMenuSection& InSection)
 	{
 		UAssetEditorToolkitMenuContext* Context = InSection.FindContext<UAssetEditorToolkitMenuContext>();
-		if (Context->Toolkit.Pin()->IsActuallyAnAsset())
+		if (Context && Context->Toolkit.IsValid() && Context->Toolkit.Pin()->IsActuallyAnAsset())
 		{
 			InSection.AddSubMenu(
 				"Asset",
