Don't remove unprocessed packages prematurely on DDC fill commandlet

This fix is a slightly modified version of PR #8248  (Contributed by [REDACTED_PERSON])
Can be reproed with QAGame -run=DerivedDataCache -fill

#rnx

diff --git a/Engine/Source/Editor/UnrealEd/Private/Commandlets/DerivedDataCacheCommandlet.cpp b/Engine/Source/Editor/UnrealEd/Private/Commandlets/DerivedDataCacheCommandlet.cpp
--- a/Engine/Source/Editor/UnrealEd/Private/Commandlets/DerivedDataCacheCommandlet.cpp
+++ b/Engine/Source/Editor/UnrealEd/Private/Commandlets/DerivedDataCacheCommandlet.cpp
@@ -78,9 +78,9 @@ public:
 		GUObjectArray.RemoveUObjectCreateListener(this);
 	}
 
-	TSet<UPackage*> GetNewPackages()
+	TSet<UPackage*>& GetNewPackages()
 	{
-		return MoveTemp(NewPackages);
+		return NewPackages;
 	}
 
 private:
@@ -116,13 +116,11 @@ UDerivedDataCacheCommandlet::UDerivedDataCacheCommandlet(FVTableHelper& Helper)
 
 UDerivedDataCacheCommandlet::UDerivedDataCacheCommandlet(const FObjectInitializer& ObjectInitializer)
 	: Super(ObjectInitializer)
-	, PackageListener(MakeUnique<FPackageListener>())
-	, ObjectReferencer(MakeUnique<FObjectReferencer>(CachingObjects))
 {
 	LogToConsole = false;
 }
 
-void UDerivedDataCacheCommandlet::MaybeMarkPackageAsAlreadyLoaded(UPackage *Package)
+void UDerivedDataCacheCommandlet::MaybeMarkPackageAsAlreadyLoaded(UPackage* Package)
 {
 	if (ProcessedPackages.Contains(Package->GetFName()))
 	{
@@ -195,22 +193,30 @@ void UDerivedDataCacheCommandlet::CacheLoadedPackages(UPackage* CurrentPackage,
 
 	const double BeginCacheTimeStart = FPlatformTime::Seconds();
 
+	// We will only remove what we process from the list to avoid unprocessed package being forever forgotten.
+	TSet<UPackage*>& NewPackages = PackageListener->GetNewPackages();
+
 	TArray<UObject*> ObjectsWithOuter;
-	for (UPackage* NewPackage : PackageListener->GetNewPackages())
+	for (auto NewPackageIt = NewPackages.CreateIterator(); NewPackageIt; ++NewPackageIt)
 	{
+		UPackage* NewPackage = *NewPackageIt;
 		const FName NewPackageName = NewPackage->GetFName();
 		if (!ProcessedPackages.Contains(NewPackageName))
 		{
 			if ((PackageFilter & NORMALIZE_ExcludeEnginePackages) != 0 && NewPackage->GetName().StartsWith(TEXT("/Engine")))
 			{
+				UE_LOG(LogDerivedDataCacheCommandlet, Display, TEXT("Skipping %s as Engine package"), *NewPackageName.ToString());
+
 				// Add it so we don't convert the FName to a string everytime we encounter this package
 				ProcessedPackages.Add(NewPackageName);
+				NewPackageIt.RemoveCurrent();
 			}
 			else if (NewPackage == CurrentPackage || !PackagesToProcess.Contains(NewPackageName))
 			{
 				UE_LOG(LogDerivedDataCacheCommandlet, Display, TEXT("Processing %s"), *NewPackageName.ToString());
 
 				ProcessedPackages.Add(NewPackageName);
+				NewPackageIt.RemoveCurrent();
 
 				ObjectsWithOuter.Reset();
 				GetObjectsWithOuter(NewPackage, ObjectsWithOuter, true /* bIncludeNestedObjects */, RF_ClassDefaultObject /* ExclusionFlags */);
@@ -225,6 +231,10 @@ void UDerivedDataCacheCommandlet::CacheLoadedPackages(UPackage* CurrentPackage,
 				}
 			}
 		}
+		else
+		{
+			NewPackageIt.RemoveCurrent();
+		}
 	}
 
 	BeginCacheTime += FPlatformTime::Seconds() - BeginCacheTimeStart;
@@ -367,6 +377,11 @@ void UDerivedDataCacheCommandlet::CacheWorldPackages(UWorld* World, uint8 Packag
 
 int32 UDerivedDataCacheCommandlet::Main( const FString& Params )
 {
+	// Avoid putting those directly in the constructor because we don't
+	// want the CDO to have a second copy of these being active.
+	PackageListener  = MakeUnique<FPackageListener>();
+	ObjectReferencer = MakeUnique<FObjectReferencer>(CachingObjects);
+
 	TArray<FString> Tokens, Switches;
 	ParseCommandLine(*Params, Tokens, Switches);
 
@@ -638,4 +653,4 @@ int32 UDerivedDataCacheCommandlet::Main( const FString& Params )
 	UE_LOG(LogDerivedDataCacheCommandlet, Display, TEXT("BeginCacheTime=%.2lfs, FinishCacheTime=%.2lfs, GCTime=%.2lfs."), BeginCacheTime, FinishCacheTime, GCTime);
 
 	return 0;
-}
+}
\ No newline at end of file
