Fix blowback from 16981291.

SummarizeTraceCommandlet could no longer receive any events from a utrace.
Access to this batch/end scope data is through an array view now.

#rnx

diff --git a/Engine/Source/Editor/UnrealEd/Private/Commandlets/SummarizeTraceCommandlet.cpp b/Engine/Source/Editor/UnrealEd/Private/Commandlets/SummarizeTraceCommandlet.cpp
--- a/Engine/Source/Editor/UnrealEd/Private/Commandlets/SummarizeTraceCommandlet.cpp
+++ b/Engine/Source/Editor/UnrealEd/Private/Commandlets/SummarizeTraceCommandlet.cpp
@@ -16,6 +16,7 @@
 #include "Trace/Analysis.h"
 #include "Trace/Analyzer.h"
 #include "Trace/DataStream.h"
+#include "TraceServices/Utils.h"
 #include "ProfilingDebugging/CountersTrace.h"
 
 /*
@@ -122,8 +123,9 @@ void FCpuAnalyzer::OnBatch(const FOnEventContext& Context)
 
 	uint32 ThreadId = Context.ThreadInfo.GetId();
 
-	const uint8* Cursor = EventData.GetAttachment();
-	const uint8* End = Cursor + EventData.GetAttachmentSize();
+	TArrayView<const uint8> DataView = TraceServices::FTraceAnalyzerUtils::LegacyAttachmentArray("Data", Context);
+	const uint8* Cursor = DataView.GetData();
+	const uint8* End = Cursor + DataView.Num();
 	uint64 LastCycle = 0;
 	while (Cursor < End)
 	{
diff --git a/Engine/Source/Editor/UnrealEd/UnrealEd.Build.cs b/Engine/Source/Editor/UnrealEd/UnrealEd.Build.cs
--- a/Engine/Source/Editor/UnrealEd/UnrealEd.Build.cs
+++ b/Engine/Source/Editor/UnrealEd/UnrealEd.Build.cs
@@ -200,6 +200,7 @@ public class UnrealEd : ModuleRules
 				"EditorInteractiveToolsFramework",
 				"TraceLog",
 				"TraceAnalysis",
+				"TraceServices",
 				"DeveloperSettings",
 				"AnimationModifiers",
 				"AnimationBlueprintLibrary",
