Remove UE3 reference

diff --git a/Engine/Source/Editor/UnrealEd/Private/Fbx/FbxAnimationExport.cpp b/Engine/Source/Editor/UnrealEd/Private/Fbx/FbxAnimationExport.cpp
--- a/Engine/Source/Editor/UnrealEd/Private/Fbx/FbxAnimationExport.cpp
+++ b/Engine/Source/Editor/UnrealEd/Private/Fbx/FbxAnimationExport.cpp
@@ -625,7 +625,7 @@ void FFbxExporter::ExportMatineeGroup(class AMatineeActor* MatineeActor, USkelet
 	if(Owner && Owner->GetRootComponent())
 	{
 		// Set the default position of the actor on the transforms
-		// The UE3 transformation is different from FBX's Z-up: invert the Y-axis for translations and the Y/Z angle values in rotations.
+		// The UE transformation is different from FBX's Z-up: invert the Y-axis for translations and the Y/Z angle values in rotations.
 		BaseNode->LclTranslation.Set(Converter.ConvertToFbxPos(Owner->GetActorLocation()));
 		BaseNode->LclRotation.Set(Converter.ConvertToFbxRot(Owner->GetActorRotation().Euler()));
 		BaseNode->LclScaling.Set(Converter.ConvertToFbxScale(Owner->GetRootComponent()->GetRelativeScale3D()));
