Remove UE4 references UnrealEd

diff --git a/Engine/Source/Editor/UnrealEd/Classes/Editor/EditorEngine.h b/Engine/Source/Editor/UnrealEd/Classes/Editor/EditorEngine.h
--- a/Engine/Source/Editor/UnrealEd/Classes/Editor/EditorEngine.h
+++ b/Engine/Source/Editor/UnrealEd/Classes/Editor/EditorEngine.h
@@ -407,7 +407,7 @@ public:
 	UPROPERTY()
 	uint32 bIsToggleBetweenPIEandSIEQueued:1;
 
-	/** Allows multiple PIE worlds under a single instance. If false, you can only do multiple UE4 processes for pie networking */
+	/** Allows multiple PIE worlds under a single instance. If false, you can only do multiple UE processes for pie networking */
 	UPROPERTY(globalconfig)
 	uint32 bAllowMultiplePIEWorlds:1;
 
diff --git a/Engine/Source/Editor/UnrealEd/Classes/Factories/FbxAssetImportData.h b/Engine/Source/Editor/UnrealEd/Classes/Factories/FbxAssetImportData.h
--- a/Engine/Source/Editor/UnrealEd/Classes/Factories/FbxAssetImportData.h
+++ b/Engine/Source/Editor/UnrealEd/Classes/Factories/FbxAssetImportData.h
@@ -27,15 +27,15 @@ class UNREALED_API UFbxAssetImportData : public UAssetImportData
 	float ImportUniformScale;
 
 	/** Whether to convert scene from FBX scene. */
-	UPROPERTY(EditAnywhere, BlueprintReadWrite, config, Category = Miscellaneous, meta = (ImportType = "StaticMesh|SkeletalMesh|Animation", ImportCategory = "Miscellaneous", ToolTip = "Convert the scene from FBX coordinate system to UE4 coordinate system"))
+	UPROPERTY(EditAnywhere, BlueprintReadWrite, config, Category = Miscellaneous, meta = (ImportType = "StaticMesh|SkeletalMesh|Animation", ImportCategory = "Miscellaneous", ToolTip = "Convert the scene from FBX coordinate system to UE coordinate system"))
 	bool bConvertScene;
 
 	/** Whether to force the front axis to be align with X instead of -Y. */
-	UPROPERTY(EditAnywhere, BlueprintReadWrite, config, Category = Miscellaneous, meta = (editcondition = "bConvertScene", ImportType = "StaticMesh|SkeletalMesh|Animation", ImportCategory = "Miscellaneous", ToolTip = "Convert the scene from FBX coordinate system to UE4 coordinate system with front X axis instead of -Y"))
+	UPROPERTY(EditAnywhere, BlueprintReadWrite, config, Category = Miscellaneous, meta = (editcondition = "bConvertScene", ImportType = "StaticMesh|SkeletalMesh|Animation", ImportCategory = "Miscellaneous", ToolTip = "Convert the scene from FBX coordinate system to UE coordinate system with front X axis instead of -Y"))
 	bool bForceFrontXAxis;
 
-	/** Whether to convert the scene from FBX unit to UE4 unit (centimeter). */
-	UPROPERTY(EditAnywhere, BlueprintReadWrite, config, Category = Miscellaneous, meta = (ImportType = "StaticMesh|SkeletalMesh|Animation", ImportCategory = "Miscellaneous", ToolTip = "Convert the scene from FBX unit to UE4 unit (centimeter)."))
+	/** Whether to convert the scene from FBX unit to UE unit (centimeter). */
+	UPROPERTY(EditAnywhere, BlueprintReadWrite, config, Category = Miscellaneous, meta = (ImportType = "StaticMesh|SkeletalMesh|Animation", ImportCategory = "Miscellaneous", ToolTip = "Convert the scene from FBX unit to UE unit (centimeter)."))
 	bool bConvertSceneUnit;
 
 	/* Use by the reimport factory to answer CanReimport, if true only factory for scene reimport will return true */
diff --git a/Engine/Source/Editor/UnrealEd/Classes/Settings/EditorLoadingSavingSettings.h b/Engine/Source/Editor/UnrealEd/Classes/Settings/EditorLoadingSavingSettings.h
--- a/Engine/Source/Editor/UnrealEd/Classes/Settings/EditorLoadingSavingSettings.h
+++ b/Engine/Source/Editor/UnrealEd/Classes/Settings/EditorLoadingSavingSettings.h
@@ -113,7 +113,7 @@ public:
 	bool bAutoCreateAssets;
 	UPROPERTY(EditAnywhere, config, AdvancedDisplay, Category=AutoReimport, meta=(DisplayName="Auto Delete Assets", ToolTip="When enabled, deleting a source content file will automatically prompt the deletion of any related assets."))
 	bool bAutoDeleteAssets;
-	UPROPERTY(EditAnywhere, config, AdvancedDisplay, Category=AutoReimport, meta=(DisplayName="Detect Changes On Startup", ToolTip="When enabled, changes to monitored directories since UE4 was closed will be detected on restart.\n(Not recommended when working in collaboration with others using source control)."))
+	UPROPERTY(EditAnywhere, config, AdvancedDisplay, Category=AutoReimport, meta=(DisplayName="Detect Changes On Startup", ToolTip="When enabled, changes to monitored directories since UE was closed will be detected on restart.\n(Not recommended when working in collaboration with others using source control)."))
 	bool bDetectChangesOnStartup;
 	UPROPERTY(EditAnywhere, config, AdvancedDisplay, Category=AutoReimport, meta=(DisplayName="Prompt Before Action", ToolTip="Whether to prompt the user to import detected changes."))
 	bool bPromptBeforeAutoImporting;
diff --git a/Engine/Source/Editor/UnrealEd/Classes/Settings/LevelEditorPlaySettings.h b/Engine/Source/Editor/UnrealEd/Classes/Settings/LevelEditorPlaySettings.h
--- a/Engine/Source/Editor/UnrealEd/Classes/Settings/LevelEditorPlaySettings.h
+++ b/Engine/Source/Editor/UnrealEd/Classes/Settings/LevelEditorPlaySettings.h
@@ -348,7 +348,7 @@ private:
 	UPROPERTY(config, EditAnywhere, Category="Multiplayer Options")
 	TEnumAsByte<EPlayNetMode> PlayNetMode;
 
-	/** Spawn multiple player windows in a single instance of UE4. This will load much faster, but has potential to have more issues.  */
+	/** Spawn multiple player windows in a single instance of UE. This will load much faster, but has potential to have more issues.  */
 	UPROPERTY(config, EditAnywhere, Category="Multiplayer Options")
 	bool RunUnderOneProcess;
 
diff --git a/Engine/Source/Editor/UnrealEd/Private/Analytics/EditorAnalytics.cpp b/Engine/Source/Editor/UnrealEd/Private/Analytics/EditorAnalytics.cpp
--- a/Engine/Source/Editor/UnrealEd/Private/Analytics/EditorAnalytics.cpp
+++ b/Engine/Source/Editor/UnrealEd/Private/Analytics/EditorAnalytics.cpp
@@ -137,7 +137,7 @@ FString FEditorAnalytics::TranslateErrorCode(int32 ErrorCode)
 	case EAnalyticsErrorCodes::FailedToDeleteStagingDirectory :
 		return TEXT("Failed to delete staging directory.  This could be because something is currently using the staging directory (ps4/xbox/etc)");
 	case EAnalyticsErrorCodes::MissingExecutable:
-		return LOCTEXT("UATErrorMissingExecutable", "Missing UnrealGame binary.\nYou may have to build the UE4 project with your IDE. Alternatively, build using UnrealBuildTool with the commandline:\nUnrealGame <Platform> <Configuration>").ToString();
+		return LOCTEXT("UATErrorMissingExecutable", "Missing UnrealGame binary.\nYou may have to build the UE project with your IDE. Alternatively, build using UnrealBuildTool with the commandline:\nUnrealGame <Platform> <Configuration>").ToString();
 	case EAnalyticsErrorCodes::FilesInstallFailed:
 		return TEXT("Failed to deploy files to device.  Check to make sure your device is connected.");
 	case EAnalyticsErrorCodes::DeviceNotSetupForDevelopment:
diff --git a/Engine/Source/Editor/UnrealEd/Private/Commandlets/CookCommandlet.cpp b/Engine/Source/Editor/UnrealEd/Private/Commandlets/CookCommandlet.cpp
--- a/Engine/Source/Editor/UnrealEd/Private/Commandlets/CookCommandlet.cpp
+++ b/Engine/Source/Editor/UnrealEd/Private/Commandlets/CookCommandlet.cpp
@@ -248,7 +248,7 @@ namespace DetailedCookStats
 				{
 					AssetType += FString::Printf(TEXT(" (%s)"), *AssetSubType->Value);
 				}
-				// Pull the Time and Size attributes and AddOrAccumulate them into the set of stats. Ugly string/container manipulation code courtesy of UE4/C++.
+				// Pull the Time and Size attributes and AddOrAccumulate them into the set of stats. Ugly string/container manipulation code courtesy of UE/C++.
 				const FCookStatsManager::StringKeyValue* AssetTimeSecAttr = StatAttributes.FindByPredicate([](const FCookStatsManager::StringKeyValue& Item) { return Item.Key == TEXT("TimeSec"); });
 				double AssetTimeSec = 0.0;
 				if (AssetTimeSecAttr)
diff --git a/Engine/Source/Editor/UnrealEd/Private/Commandlets/GatherTextFromAssetsCommandlet.cpp b/Engine/Source/Editor/UnrealEd/Private/Commandlets/GatherTextFromAssetsCommandlet.cpp
--- a/Engine/Source/Editor/UnrealEd/Private/Commandlets/GatherTextFromAssetsCommandlet.cpp
+++ b/Engine/Source/Editor/UnrealEd/Private/Commandlets/GatherTextFromAssetsCommandlet.cpp
@@ -990,7 +990,7 @@ bool UGatherTextFromAssetsCommandlet::ConfigureFromScript(const FString& GatherT
 
 	GetPathArrayFromConfig(*SectionName, TEXT("ManifestDependencies"), ManifestDependenciesList, GatherTextConfigPath);
 
-	// Get whether we should gather editor-only data. Typically only useful for the localization of UE4 itself.
+	// Get whether we should gather editor-only data. Typically only useful for the localization of UE itself.
 	if (!GetBoolFromConfig(*SectionName, TEXT("ShouldGatherFromEditorOnlyData"), ShouldGatherFromEditorOnlyData, GatherTextConfigPath))
 	{
 		ShouldGatherFromEditorOnlyData = false;
diff --git a/Engine/Source/Editor/UnrealEd/Private/Commandlets/GatherTextFromMetadataCommandlet.cpp b/Engine/Source/Editor/UnrealEd/Private/Commandlets/GatherTextFromMetadataCommandlet.cpp
--- a/Engine/Source/Editor/UnrealEd/Private/Commandlets/GatherTextFromMetadataCommandlet.cpp
+++ b/Engine/Source/Editor/UnrealEd/Private/Commandlets/GatherTextFromMetadataCommandlet.cpp
@@ -105,7 +105,7 @@ int32 UGatherTextFromMetaDataCommandlet::Main( const FString& Params )
 
 	FGatherTextDelegates::GetAdditionalGatherPaths.Broadcast(GatherManifestHelper->GetTargetName(), IncludePathFilters, ExcludePathFilters);
 
-	// Get whether we should gather editor-only data. Typically only useful for the localization of UE4 itself.
+	// Get whether we should gather editor-only data. Typically only useful for the localization of UE itself.
 	if (!GetBoolFromConfig(*SectionName, TEXT("ShouldGatherFromEditorOnlyData"), ShouldGatherFromEditorOnlyData, GatherTextConfigPath))
 	{
 		ShouldGatherFromEditorOnlyData = false;
diff --git a/Engine/Source/Editor/UnrealEd/Private/Commandlets/GatherTextFromSourceCommandlet.cpp b/Engine/Source/Editor/UnrealEd/Private/Commandlets/GatherTextFromSourceCommandlet.cpp
--- a/Engine/Source/Editor/UnrealEd/Private/Commandlets/GatherTextFromSourceCommandlet.cpp
+++ b/Engine/Source/Editor/UnrealEd/Private/Commandlets/GatherTextFromSourceCommandlet.cpp
@@ -269,7 +269,7 @@ int32 UGatherTextFromSourceCommandlet::Main( const FString& Params )
 	// Init a parse context to track the state of the file parsing 
 	FSourceFileParseContext ParseCtxt(this);
 
-	// Get whether we should gather editor-only data. Typically only useful for the localization of UE4 itself.
+	// Get whether we should gather editor-only data. Typically only useful for the localization of UE itself.
 	if (!GetBoolFromConfig(*SectionName, TEXT("ShouldGatherFromEditorOnlyData"), ParseCtxt.ShouldGatherFromEditorOnlyData, GatherTextConfigPath))
 	{
 		ParseCtxt.ShouldGatherFromEditorOnlyData = false;
diff --git a/Engine/Source/Editor/UnrealEd/Private/Editor.cpp b/Engine/Source/Editor/UnrealEd/Private/Editor.cpp
--- a/Engine/Source/Editor/UnrealEd/Private/Editor.cpp
+++ b/Engine/Source/Editor/UnrealEd/Private/Editor.cpp
@@ -1725,7 +1725,7 @@ void ExecuteInvalidateCachedShaders(const TArray< FString >& Args)
 		}
 		else if(SourceControlState->IsCheckedOutOther())
 		{
-			UE_LOG(LogConsoleResponse, Display, TEXT("r.InvalidateCachedShaders failed\n\"ShaderVersion.ush\" is already checked out by someone else\n(UE4 SourceControl needs to be fixed to allow multiple checkout.)"));
+			UE_LOG(LogConsoleResponse, Display, TEXT("r.InvalidateCachedShaders failed\n\"ShaderVersion.ush\" is already checked out by someone else\n(UE SourceControl needs to be fixed to allow multiple checkout.)"));
 			return;
 		}
 		else if(SourceControlState->IsDeleted())
diff --git a/Engine/Source/Editor/UnrealEd/Private/EditorCommandLineUtils.cpp b/Engine/Source/Editor/UnrealEd/Private/EditorCommandLineUtils.cpp
--- a/Engine/Source/Editor/UnrealEd/Private/EditorCommandLineUtils.cpp
+++ b/Engine/Source/Editor/UnrealEd/Private/EditorCommandLineUtils.cpp
@@ -281,7 +281,7 @@ bool FMergeAsset::Load(FCommandLineErrorReporter& ErrorReporter)
 		return false;
 	}
 
-	// UE4 cannot open files with certain special characters in them (like 
+	// UE cannot open files with certain special characters in them (like 
 	// the # symbol), so we make a copy of the file with a more UE digestible 
 	// path (since this may be a perforce temp file)
 	if (IFileManager::Get().Copy(*DestFilePath, *SrcFilePath) != COPY_OK)
diff --git a/Engine/Source/Editor/UnrealEd/Private/EditorComponents.cpp b/Engine/Source/Editor/UnrealEd/Private/EditorComponents.cpp
--- a/Engine/Source/Editor/UnrealEd/Private/EditorComponents.cpp
+++ b/Engine/Source/Editor/UnrealEd/Private/EditorComponents.cpp
@@ -408,7 +408,7 @@ void FEditorCommonDrawHelper::Draw(const FSceneView* View,class FPrimitiveDrawIn
 /** Draw green lines to indicate what the selected actor(s) are based on. */
 void FEditorCommonDrawHelper::DrawBaseInfo(const FSceneView* View,FPrimitiveDrawInterface* PDI)
 {
-// @todo UE4 - reimplement with new component attachment system
+// @todo UE - reimplement with new component attachment system
 }
 
 void FEditorCommonDrawHelper::DrawOldGrid(const FSceneView* View,FPrimitiveDrawInterface* PDI)
diff --git a/Engine/Source/Editor/UnrealEd/Private/Factories/EditorFactories.cpp b/Engine/Source/Editor/UnrealEd/Private/Factories/EditorFactories.cpp
--- a/Engine/Source/Editor/UnrealEd/Private/Factories/EditorFactories.cpp
+++ b/Engine/Source/Editor/UnrealEd/Private/Factories/EditorFactories.cpp
@@ -3989,7 +3989,7 @@ UObject* UTextureFactory::FactoryCreateBinary
 			// Don't try to rename the texture if we are doing a reimport.
 			if (!bIsDoingAReimport)
 			{ 
-				// Exclude UDIM number from the name of the UE4 texture asset we create
+				// Exclude UDIM number from the name of the UE texture asset we create
 				// We do this even in the case where we're only importing a single 1001 image, which won't technically be a UDIM
 				// We still want to strip the UDIM suffix in this case, as industry standard still considers this a UDIM imageset
 				const FString ShortPackageName = ObjectTools::SanitizeInvalidChars(BaseUDIMName, INVALID_LONGPACKAGE_CHARACTERS);
@@ -4544,7 +4544,7 @@ bool UTextureFactory::IsImportResolutionValid(int32 Width, int32 Height, bool bA
 {
 	static const auto CVarVirtualTexturesEnabled = IConsoleManager::Get().FindTConsoleVariableDataInt(TEXT("r.VirtualTextures")); check(CVarVirtualTexturesEnabled);
 
-	// In theory this value could be much higher, but various UE4 image code currently uses 32bit size/offset values
+	// In theory this value could be much higher, but various UE image code currently uses 32bit size/offset values
 	const int32 MaximumSupportedVirtualTextureResolution = 16 * 1024;
 
 	// Calculate the maximum supported resolution utilizing the global max texture mip count
diff --git a/Engine/Source/Editor/UnrealEd/Private/Fbx/FbxMainExport.cpp b/Engine/Source/Editor/UnrealEd/Private/Fbx/FbxMainExport.cpp
--- a/Engine/Source/Editor/UnrealEd/Private/Fbx/FbxMainExport.cpp
+++ b/Engine/Source/Editor/UnrealEd/Private/Fbx/FbxMainExport.cpp
@@ -2925,7 +2925,7 @@ void FFbxExporter::ExportChannelToFbxCurve(FbxAnimCurve& InFbxCurve, const FMovi
 				float LeaveTangent = KeyValue.Tangent.LeaveTangent * TickResolution.AsDecimal();
 				float NextArriveTangent = Values[Index + 1].Tangent.ArriveTangent * TickResolution.AsDecimal();
 
-				//Need to convert Ue4 tangent weight which is the length of the hypotenuse to FBX normalized X(time) weight
+				//Need to convert UE tangent weight which is the length of the hypotenuse to FBX normalized X(time) weight
 				if (WeightedMode == FbxAnimCurveDef::eWeightedAll || WeightedMode == FbxAnimCurveDef::eWeightedRight)
 				{
 					const float XVal = FMath::Sqrt((KeyValue.Tangent.LeaveTangentWeight * KeyValue.Tangent.LeaveTangentWeight) / (1.0f + LeaveTangent * LeaveTangent));
@@ -4524,7 +4524,7 @@ FbxNode* FFbxExporter::ExportStaticMeshToFbx(const UStaticMesh* StaticMesh, int3
 			check(UVsLayer);
 
 			FString UVChannelNameBuilder = TEXT("UVmap_") + FString::FromInt(TexCoordSourceIndex);
-			const char* UVChannelName = TCHAR_TO_UTF8(*UVChannelNameBuilder); // actually UTF8 as required by Fbx, but can't use UE4's UTF8CHAR type because that's a uint8 aka *unsigned* char
+			const char* UVChannelName = TCHAR_TO_UTF8(*UVChannelNameBuilder); // actually UTF8 as required by Fbx, but can't use UE's UTF8CHAR type because that's a uint8 aka *unsigned* char
 			if ((LightmapUVChannel >= 0) || ((LightmapUVChannel == -1) && (TexCoordSourceIndex == StaticMesh->GetLightMapCoordinateIndex())))
 			{
 				UVChannelName = "LightMapUV";
@@ -4846,7 +4846,7 @@ void FFbxExporter::ExportSplineMeshToFbx(const USplineMeshComponent* SplineMeshC
 			UVsLayer = Mesh->GetLayer(TexCoordSourceIndex);
 		}
 		FString UVChannelNameBuilder = TEXT("UVmap_") + FString::FromInt(TexCoordSourceIndex);
-		const char* UVChannelName = TCHAR_TO_UTF8(*UVChannelNameBuilder); // actually UTF8 as required by Fbx, but can't use UE4's UTF8CHAR type because that's a uint8 aka *unsigned* char
+		const char* UVChannelName = TCHAR_TO_UTF8(*UVChannelNameBuilder); // actually UTF8 as required by Fbx, but can't use UE's UTF8CHAR type because that's a uint8 aka *unsigned* char
 		if (TexCoordSourceIndex == StaticMesh->GetLightMapCoordinateIndex())
 		{
 			UVChannelName = "LightMapUV";
diff --git a/Engine/Source/Editor/UnrealEd/Private/Fbx/FbxSceneImportFactory.cpp b/Engine/Source/Editor/UnrealEd/Private/Fbx/FbxSceneImportFactory.cpp
--- a/Engine/Source/Editor/UnrealEd/Private/Fbx/FbxSceneImportFactory.cpp
+++ b/Engine/Source/Editor/UnrealEd/Private/Fbx/FbxSceneImportFactory.cpp
@@ -942,7 +942,7 @@ void UFbxSceneImportFactory::ChangeFrontAxis(void* VoidFbxImporter, void* VoidSc
 		UnFbx::FbxNodeInfo &RootNodeInfo = SceneInfo->HierarchyInfo[0];
 		check(RootNodeInfo.UniqueId == RootNode->GetUniqueID());
 		RootNodeInfo.Transform = RootNode->EvaluateGlobalTransform();
-		//Set the UE4 data
+		//Set the UE data
 		TSharedPtr<FFbxNodeInfo> RootNodeInfoPtr = GetNodeInfoPtrById(SceneInfoPtr->HierarchyInfo, RootNodeInfo.UniqueId);
 		if (RootNodeInfoPtr.IsValid())
 		{
diff --git a/Engine/Source/Editor/UnrealEd/Private/FbxExporter.h b/Engine/Source/Editor/UnrealEd/Private/FbxExporter.h
--- a/Engine/Source/Editor/UnrealEd/Private/FbxExporter.h
+++ b/Engine/Source/Editor/UnrealEd/Private/FbxExporter.h
@@ -497,7 +497,7 @@ private:
 
 	/**
 	 * Finds the given actor in the already-exported list of structures
-	 * @return FbxNode* the FBX node created from the UE4 actor
+	 * @return FbxNode* the FBX node created from the UE actor
 	 */
 	FbxNode* FindActor(AActor* Actor, INodeNameAdapter* NodeNameAdapter = nullptr);
 
diff --git a/Engine/Source/Editor/UnrealEd/Private/Lightmass/Lightmass.cpp b/Engine/Source/Editor/UnrealEd/Private/Lightmass/Lightmass.cpp
--- a/Engine/Source/Editor/UnrealEd/Private/Lightmass/Lightmass.cpp
+++ b/Engine/Source/Editor/UnrealEd/Private/Lightmass/Lightmass.cpp
@@ -70,7 +70,7 @@ UNREALED_API bool GLightmassStatsMode = false;
 #if !UE_BUILD_DOCS
 void FSwarmDebugOptions::Touch()
 {
-	//@todo UE4. For some reason, the global instance is not initializing to the default settings...
+	//@todo For some reason, the global instance is not initializing to the default settings...
 	if (bInitialized == false)
 	{
 		bDistributionEnabled = true;
diff --git a/Engine/Source/Editor/UnrealEd/Private/Lightmass/LightmassRender.cpp b/Engine/Source/Editor/UnrealEd/Private/Lightmass/LightmassRender.cpp
--- a/Engine/Source/Editor/UnrealEd/Private/Lightmass/LightmassRender.cpp
+++ b/Engine/Source/Editor/UnrealEd/Private/Lightmass/LightmassRender.cpp
@@ -1134,7 +1134,7 @@ bool FLightmassMaterialRenderer::GenerateMaterialPropertyData(
 		// If the ShaderMap is NULL that's because it failed to compile, which is ok as the default material will be used for exporting
 		check(!MaterialProxy->GetGameThreadShaderMap() || MaterialProxy->GetGameThreadShaderMap()->IsValidForRendering());
 
-		//@todo UE4. The format may be determined by the material property...
+		//@todo The format may be determined by the material property...
 		// For example, if Diffuse doesn't need to be F16 it can create a standard RGBA8 target.
 		EPixelFormat Format = PF_FloatRGBA;
 		if (MaterialProxy->GetRenderTargetFormatAndSize(InMaterialProperty, Format, InMaterial.GetExportResolutionScale(), InOutSizeX, InOutSizeY))
diff --git a/Engine/Source/Editor/UnrealEd/Private/PackageTools.cpp b/Engine/Source/Editor/UnrealEd/Private/PackageTools.cpp
--- a/Engine/Source/Editor/UnrealEd/Private/PackageTools.cpp
+++ b/Engine/Source/Editor/UnrealEd/Private/PackageTools.cpp
@@ -1268,7 +1268,7 @@ UPackageTools::UPackageTools(const FObjectInitializer& ObjectInitializer)
 			}
 			else if (NumberOfAssets > 1)
 			{
-				// this shouldn't happen in UE4
+				// this shouldn't happen 
 				bShouldGenerateUniquePackageName = true;
 			}
 
diff --git a/Engine/Source/Editor/UnrealEd/Private/Settings/LevelEditorPlaySettingsCustomization.cpp b/Engine/Source/Editor/UnrealEd/Private/Settings/LevelEditorPlaySettingsCustomization.cpp
--- a/Engine/Source/Editor/UnrealEd/Private/Settings/LevelEditorPlaySettingsCustomization.cpp
+++ b/Engine/Source/Editor/UnrealEd/Private/Settings/LevelEditorPlaySettingsCustomization.cpp
@@ -525,7 +525,7 @@ FText FLevelEditorPlaySettingsCustomization::HandleMultiplayerOptionsDescription
 	FString Desc;
 	if ( CanRunUnderOneProcess )
 	{
-		Desc += LOCTEXT( "MultiplayerDescription_OneProcess", "The following will all run under one UE4 instance:\n" ).ToString();
+		Desc += LOCTEXT( "MultiplayerDescription_OneProcess", "The following will all run under one UE instance:\n" ).ToString();
 		if ( PlayNetMode == EPlayNetMode::PIE_Client )
 		{
 			Desc += LOCTEXT( "MultiplayerDescription_DedicatedServerHidden", "A hidden dedicated server instance will run in editor. " ).ToString();
@@ -568,7 +568,7 @@ FText FLevelEditorPlaySettingsCustomization::HandleMultiplayerOptionsDescription
 	}
 	else
 	{
-		Desc += LOCTEXT( "MultiplayerDescription_MultiProcess", "The following will run with multiple UE4 instances:\n" ).ToString();
+		Desc += LOCTEXT( "MultiplayerDescription_MultiProcess", "The following will run with multiple UE instances:\n" ).ToString();
 		if ( PlayNetMode == PIE_Standalone )
 		{
 			if ( PlayNumberOfClients == 1 )
diff --git a/Engine/Source/Editor/UnrealEd/Private/StaticLightingSystem/StaticLightingSystem.cpp b/Engine/Source/Editor/UnrealEd/Private/StaticLightingSystem/StaticLightingSystem.cpp
--- a/Engine/Source/Editor/UnrealEd/Private/StaticLightingSystem/StaticLightingSystem.cpp
+++ b/Engine/Source/Editor/UnrealEd/Private/StaticLightingSystem/StaticLightingSystem.cpp
@@ -1634,7 +1634,7 @@ void FStaticLightingSystem::AddBSPStaticLightingInfo(ULevel* Level, bool bBuildL
 		if (NodeGroup->Nodes.Num())
 		{
 			// get one of the surfaces/components from the NodeGroup
-			// @todo UE4: Remove need for GetSurfaceLightMapResolution to take a surfaceindex, or a ModelComponent :)
+			// @todo: Remove need for GetSurfaceLightMapResolution to take a surfaceindex, or a ModelComponent :)
 			UModelComponent* SomeModelComponent = Level->ModelComponents[Model->Nodes[NodeGroup->Nodes[0]].ComponentIndex];
 			int32 SurfaceIndex = Model->Nodes[NodeGroup->Nodes[0]].iSurf;
 
@@ -1795,7 +1795,7 @@ void FStaticLightingSystem::AddBSPStaticLightingInfo(ULevel* Level, TArray<FNode
 		if (NodeGroup && NodeGroup->Nodes.Num())
 		{
 			// get one of the surfaces/components from the NodeGroup
-			// @todo UE4: Remove need for GetSurfaceLightMapResolution to take a surfaceindex, or a ModelComponent :)
+			// @todo: Remove need for GetSurfaceLightMapResolution to take a surfaceindex, or a ModelComponent :)
 			UModelComponent* SomeModelComponent = Level->ModelComponents[Model->Nodes[NodeGroup->Nodes[0]].ComponentIndex];
 			int32 SurfaceIndex = Model->Nodes[NodeGroup->Nodes[0]].iSurf;
 
diff --git a/Engine/Source/Editor/UnrealEd/Private/ThumbnailRendering/TextureThumbnailRenderer.cpp b/Engine/Source/Editor/UnrealEd/Private/ThumbnailRendering/TextureThumbnailRenderer.cpp
--- a/Engine/Source/Editor/UnrealEd/Private/ThumbnailRendering/TextureThumbnailRenderer.cpp
+++ b/Engine/Source/Editor/UnrealEd/Private/ThumbnailRendering/TextureThumbnailRenderer.cpp
@@ -111,7 +111,7 @@ void UTextureThumbnailRenderer::Draw(UObject* Object, int32 X, int32 Y, uint32 W
 		CanvasTile.BatchedElementParameters = BatchedElementParameters;
 		if (bIsVirtualTexture && Texture->Source.GetNumBlocks() > 1)
 		{
-			// Adjust UVs to display entire UDIM range, acounting for UE4 inverted V-axis
+			// Adjust UVs to display entire UDIM range, acounting for UE inverted V-axis
 			// We're not actually rendering a VT here, but the editor-only texture we're using is still using the UDIM tile layout
 			// So we use inverted Y-axis, but then normalize back to [0,1)
 			const FIntPoint BlockSize = Texture->Source.GetSizeInBlocks();
diff --git a/Engine/Source/Editor/UnrealEd/Private/UnrealEdSrv.cpp b/Engine/Source/Editor/UnrealEd/Private/UnrealEdSrv.cpp
--- a/Engine/Source/Editor/UnrealEd/Private/UnrealEdSrv.cpp
+++ b/Engine/Source/Editor/UnrealEd/Private/UnrealEdSrv.cpp
@@ -2384,7 +2384,7 @@ bool UUnrealEdEngine::Exec_Actor( UWorld* InWorld, const TCHAR* Str, FOutputDevi
 		}
 		else if( FParse::Command(&Str,TEXT("BASED")) ) // ACTOR SELECT BASED
 		{
-			// @TODO UE4 - no longer meaningful
+			// @TODO no longer meaningful
 			return true;
 		}
 		else if( FParse::Command(&Str,TEXT("BYPROPERTY")) ) // ACTOR SELECT BYPROPERTY
