EditorDomain BulkData: Read BulkDatas as attachments rather than including them all in one cache value.
#rnx

diff --git a/Engine/Source/Editor/UnrealEd/Private/EditorDomain/EditorDomain.cpp b/Engine/Source/Editor/UnrealEd/Private/EditorDomain/EditorDomain.cpp
--- a/Engine/Source/Editor/UnrealEd/Private/EditorDomain/EditorDomain.cpp
+++ b/Engine/Source/Editor/UnrealEd/Private/EditorDomain/EditorDomain.cpp
@@ -289,6 +289,7 @@ FOpenPackageResult FEditorDomain::OpenReadPackage(const FPackagePath& PackagePat
 	FPackagePath* OutUpdatedPath)
 {
 	using namespace UE::EditorDomain;
+	using namespace UE::DerivedData;
 
 	FScopeLock ScopeLock(&Locks->Lock);
 	if (PackageSegment != EPackageSegment::Header)
@@ -308,22 +309,26 @@ FOpenPackageResult FEditorDomain::OpenReadPackage(const FPackagePath& PackagePat
 	// Unlock before requesting the package because the completion callback takes the lock.
 	ScopeLock.Unlock();
 
-	UE::DerivedData::FRequest Request = RequestEditorDomainPackage(PackagePath, PackageSourceDigest,
-		UE::DerivedData::ECachePolicy::None, UE::DerivedData::EPriority::Normal,
-		[Result](UE::DerivedData::FCacheGetCompleteParams&& Params)
+	// Fetch only meta-data in the initial request
+	ECachePolicy SkipFlags = ECachePolicy::SkipData & ~ECachePolicy::SkipMeta;
+	FRequest Request = RequestEditorDomainPackage(PackagePath, PackageSourceDigest,
+		SkipFlags, EPriority::Normal,
+		[Result](FCacheGetCompleteParams&& Params)
 		{
 			// Note that ~FEditorDomainReadArchive waits for this callback to be called, so Result cannot dangle
-			Result->OnCacheRequestComplete(MoveTemp(Params));
+			Result->OnRecordRequestComplete(MoveTemp(Params));
 		});
 	Result->SetRequest(Request);
 
+	// Precache the exports segment
+	// EDITOR_DOMAIN_TODO: Skip doing this for OpenReadPackage calls from bulk data
+	Result->Precache(0, 0);
+
 	if (OutUpdatedPath)
 	{
 		*OutUpdatedPath = PackagePath;
 	}
 
-	// EDITOR_DOMAIN_TODO: Reading GetPackageFormat forces us to wait for the cache response
-	// We should read just the metadata for the package so we don't have to block here on the transfer of the bytes.
 	const EPackageFormat Format = bHasEditorSource ? EPackageFormat::Binary : Result->GetPackageFormat();
 	return FOpenPackageResult{ TUniquePtr<FArchive>(Result), Format };
 }
@@ -331,6 +336,7 @@ FOpenPackageResult FEditorDomain::OpenReadPackage(const FPackagePath& PackagePat
 IAsyncReadFileHandle* FEditorDomain::OpenAsyncReadPackage(const FPackagePath& PackagePath, EPackageSegment PackageSegment)
 {
 	using namespace UE::EditorDomain;
+	using namespace UE::DerivedData;
 
 	FScopeLock ScopeLock(&Locks->Lock);
 	if (PackageSegment != EPackageSegment::Header)
@@ -345,13 +351,15 @@ IAsyncReadFileHandle* FEditorDomain::OpenAsyncReadPackage(const FPackagePath& Pa
 		return Workspace->OpenAsyncReadPackage(PackagePath, PackageSegment);
 	}
 
+	// Fetch meta-data only in the initial request
+	ECachePolicy SkipFlags = ECachePolicy::SkipData & ~ECachePolicy::SkipMeta;
 	FEditorDomainAsyncReadFileHandle* Result = new FEditorDomainAsyncReadFileHandle(Locks, PackagePath, PackageSource);
-	UE::DerivedData::FRequest Request = RequestEditorDomainPackage(PackagePath, PackageSource->Digest,
-		UE::DerivedData::ECachePolicy::None, UE::DerivedData::EPriority::Normal,
-		[Result](UE::DerivedData::FCacheGetCompleteParams&& Params)
+	FRequest Request = RequestEditorDomainPackage(PackagePath, PackageSource->Digest,
+		SkipFlags, EPriority::Normal,
+		[Result](FCacheGetCompleteParams&& Params)
 		{
 			// Note that ~FEditorDomainAsyncReadFileHandle waits for this callback to be called, so Result cannot dangle
-			Result->OnCacheRequestComplete(MoveTemp(Params));
+			Result->OnRecordRequestComplete(MoveTemp(Params));
 		});
 	Result->SetRequest(Request);
 
diff --git a/Engine/Source/Editor/UnrealEd/Private/EditorDomain/EditorDomain.h b/Engine/Source/Editor/UnrealEd/Private/EditorDomain/EditorDomain.h
--- a/Engine/Source/Editor/UnrealEd/Private/EditorDomain/EditorDomain.h
+++ b/Engine/Source/Editor/UnrealEd/Private/EditorDomain/EditorDomain.h
@@ -154,6 +154,7 @@ private:
 	static FEditorDomain* RegisteredEditorDomain;
 
 	friend class FEditorDomainRegisterAsPackageResourceManager;
+	friend class FEditorDomainPackageSegments;
 	friend class FEditorDomainReadArchive;
 	friend class FEditorDomainAsyncReadFileHandle;
 };
\ No newline at end of file
diff --git a/Engine/Source/Editor/UnrealEd/Private/EditorDomain/EditorDomainArchive.cpp b/Engine/Source/Editor/UnrealEd/Private/EditorDomain/EditorDomainArchive.cpp
--- a/Engine/Source/Editor/UnrealEd/Private/EditorDomain/EditorDomainArchive.cpp
+++ b/Engine/Source/Editor/UnrealEd/Private/EditorDomain/EditorDomainArchive.cpp
@@ -5,8 +5,11 @@
 #include "AssetRegistry/AssetData.h"
 #include "Async/AsyncFileHandleNull.h"
 #include "DerivedDataCache.h"
+#include "DerivedDataCacheInterface.h"
+#include "DerivedDataCacheKey.h"
 #include "DerivedDataCacheRecord.h"
 #include "EditorDomain/EditorDomainSave.h"
+#include "EditorDomain/EditorDomainUtils.h"
 #include "HAL/UnrealMemory.h"
 #include "Memory/SharedBuffer.h"
 #include "Misc/AssertionMacros.h"
@@ -14,268 +17,624 @@
 #include "Misc/ScopeLock.h"
 #include "Serialization/CompactBinary.h"
 
-FEditorDomainReadArchive::FEditorDomainReadArchive(const TRefCountPtr<FEditorDomain::FLocks>& InLocks,
+FEditorDomainPackageSegments::FEditorDomainPackageSegments(const TRefCountPtr<FEditorDomain::FLocks>& InLocks,
 	const FPackagePath& InPackagePath, const TRefCountPtr<FEditorDomain::FPackageSource>& InPackageSource)
 	: EditorDomainLocks(InLocks)
 	, PackagePath(InPackagePath)
 	, PackageSource(InPackageSource)
+	, PackageDigest(InPackageSource->Digest)
 {
-	this->SetIsLoading(true);
-	this->SetIsPersistent(true);
 }
 
-FEditorDomainReadArchive::~FEditorDomainReadArchive()
+FEditorDomainPackageSegments::~FEditorDomainPackageSegments()
 {
 	Close();
 }
 
-void FEditorDomainReadArchive::SetRequest(const UE::DerivedData::FRequest& InRequest)
+void FEditorDomainPackageSegments::SetRequest(const UE::DerivedData::FRequest& InRequest)
 {
 	Request = InRequest;
 }
 
-void FEditorDomainReadArchive::Seek(int64 InPos)
+FEditorDomainPackageSegments::ESource FEditorDomainPackageSegments::GetSource() const
 {
-	switch (Source)
-	{
-	case ESource::Archive:
-		InnerArchive->Seek(InPos);
-		break;
-	default:
-		Pos = InPos;
-		break;
-	}
+	return Source;
 }
 
-int64 FEditorDomainReadArchive::Tell()
+std::atomic<FEditorDomainPackageSegments::ESource>& FEditorDomainPackageSegments::GetAsyncSource()
 {
-	switch (Source)
+	return AsyncSource;
+}
+
+void FEditorDomainPackageSegments::WaitForReady() const
+{
+	if (Source != ESource::Uninitialized)
 	{
-	case ESource::Archive:
-		return InnerArchive->Tell();
-	default:
-		return Pos;
+		return;
 	}
+	Request.Wait();
+	Source = AsyncSource;
 }
 
-int64 FEditorDomainReadArchive::TotalSize()
+void FEditorDomainPackageSegments::Seek(int64 InPos)
 {
-	WaitForReady();
+	Pos = InPos;
+}
+
+int64 FEditorDomainPackageSegments::Tell() const
+{
+	return Pos;
+}
+
+int64 FEditorDomainPackageSegments::TotalSize() const
+{
+	check(Source == ESource::Segments);
 	return Size;
 }
 
-bool FEditorDomainReadArchive::Close()
+const FPackagePath& FEditorDomainPackageSegments::GetPackagePath() const
 {
-	{
-		FScopeLock ScopeLock(&Lock);
-		if (AsyncSource == ESource::Uninitialized)
-		{
-			AsyncSource = ESource::Closed;
-		}
-	}
-	Request.Cancel();
-	InnerArchive.Reset();
-	Bytes.Reset();
-	Source = ESource::Closed;
-	return true;
+	return PackagePath;
 }
 
-void FEditorDomainReadArchive::Serialize(void* V, int64 Length)
+void FEditorDomainPackageSegments::Close()
 {
-	switch (Source)
+	// Called from Interface-only
+	AsyncSource = ESource::Closed;
+	Request.Cancel();
+	for (FSegment& Segment : Segments)
 	{
-	case ESource::Uninitialized:
-		WaitForReady();
-		check(Source != ESource::Uninitialized);
-		Serialize(V, Length);
-		break;
-	case ESource::Bytes:
-		if (Pos + Length > Size)
-		{
-			SetError();
-			UE_LOG(LogEditorDomain, Error, TEXT("Requested read of %d bytes when %d bytes remain (file=%s, size=%d)"),
-				Length, Size - Pos, *PackagePath.GetDebugName(), Size);
-			return;
-		}
-		FMemory::Memcpy(V, static_cast<const uint8*>(Bytes.GetData()) + Pos, Length);
-		Pos += Length;
-		break;
-	case ESource::Archive:
-		InnerArchive->Serialize(V, Length);
-		break;
-	case ESource::Closed:
-		UE_LOG(LogEditorDomain, Error, TEXT("Requested read after close (file=%s)"), *PackagePath.GetDebugName());
-		break;
-	default:
-		checkNoEntry();
-		break;
+		Segment.Request.Cancel();
 	}
-}
 
-FString FEditorDomainReadArchive::GetArchiveName() const
-{
-	return PackagePath.GetDebugName();
+	Segments.Reset();
+	Source = ESource::Closed;
 }
 
-void FEditorDomainReadArchive::Flush()
+bool FEditorDomainPackageSegments::Serialize(void* V, int64 Length)
 {
-	switch (Source)
+	// Called from Interface-only
+	check(Source == ESource::Segments); // Caller should not call in any other case
+	int64 End = Pos + Length;
+	if (Pos < 0)
 	{
-	case ESource::Uninitialized:
-		WaitForReady();
-		check(Source != ESource::Uninitialized);
-		Flush();
-		break;
-	case ESource::Archive:
-		InnerArchive->Flush();
-		break;
-	default:
-		break;
+		UE_LOG(LogEditorDomain, Error, TEXT("Invalid negative Pos %" INT64_FMT " (file = %s)"), Pos,
+			*PackagePath.GetDebugName());
+		return false;
+	}
+	if (End > Size)
+	{
+		UE_LOG(LogEditorDomain, Error, TEXT("Requested read of %" INT64_FMT " bytes when %" INT64_FMT " bytes remain (file = %s, size=%" INT64_FMT ")"),
+			Length, Size - Pos, *PackagePath.GetDebugName(), Size);
+		return false;
 	}
+	bool bIsReady;
+	FSegment* Segment = EnsureSegmentRange(Pos, End, true, bIsReady);
+	while (Pos < End)
+	{
+		int64 SegmentEnd = Segment < &Segments.Last() ? (Segment + 1)->Start : Size;
+		int64 LengthInSegment = FMath::Min(SegmentEnd, End) - Pos;
+		int64 SegmentPos = Pos - Segment->Start;
+		FMemory::Memcpy(V, static_cast<const uint8*>(Segment->Data.GetData()) + SegmentPos, LengthInSegment);
+		Pos += LengthInSegment;
+		++Segment;
+	}
+	return true;
 }
 
-void FEditorDomainReadArchive::FlushCache()
+void FEditorDomainPackageSegments::Precache(int64 InStart, int64 InSize, bool* bOutIsReady)
 {
 	switch (Source)
 	{
 	case ESource::Uninitialized:
-		WaitForReady();
-		check(Source != ESource::Uninitialized);
-		FlushCache();
-		break;
-	case ESource::Archive:
-		InnerArchive->FlushCache();
+		if (AsyncSource != ESource::Uninitialized)
+		{
+			WaitForReady();
+			return Precache(InStart, InSize, bOutIsReady);
+		}
+		else
+		{
+			// Set the Initial request offset. We only support a single initial offset; overwrite any
+			// previously existing offset. 
+			InitialRequestOffset = InStart;
+			// THe offset will not be used if the callback thread reached the callback and checked the
+			// value of InitialRequestOffset in between our read of AsyncSource and our write of
+			// InitialRequestOffset. Check AsyncSource again to handle that case.
+			if (AsyncSource != ESource::Uninitialized)
+			{
+				WaitForReady();
+				return Precache(InStart, InSize, bOutIsReady);
+			}
+			if (bOutIsReady)
+			{
+				*bOutIsReady = false;
+			}
+		}
 		break;
+	case ESource::Segments:
+	{
+		int64 End = InStart + InSize;
+		End = FMath::Min(End, this->Size);
+		int64 Start = FMath::Clamp(InStart, 0, End);
+
+		bool bIsReady;
+		EnsureSegmentRange(Start, End, false /* bWait */, bIsReady);
+		if (bOutIsReady)
+		{
+			*bOutIsReady = bIsReady;
+		}
+	}
 	default:
 		break;
 	}
 }
 
-bool FEditorDomainReadArchive::Precache(int64 PrecacheOffset, int64 PrecacheSize)
+int32 FEditorDomainPackageSegments::GetSegmentContainingOffset(uint64 Start) const
 {
-	switch (Source)
+	int32 StartIndex = Algo::LowerBoundBy(Segments, Start, [](const FSegment& Segment) { return Segment.Start; });
+	check(0 <= StartIndex && StartIndex <= Segments.Num());
+	if (StartIndex == Segments.Num() || Segments[StartIndex].Start > Start)
 	{
-	case ESource::Uninitialized:
-		WaitForReady();
-		check(Source != ESource::Uninitialized);
-		return Precache(PrecacheOffset, PrecacheSize);
-	case ESource::Archive:
-		return InnerArchive->Precache(PrecacheOffset, PrecacheSize);
-	default:
-		return true;
+		check(StartIndex > 0);
+		--StartIndex;
 	}
+	check(Segments[StartIndex].Start <= Start && (StartIndex == Segments.Num() - 1 || Start < Segments[StartIndex + 1].Start));
+	return StartIndex;
 }
 
-EPackageFormat FEditorDomainReadArchive::GetPackageFormat() const
+FEditorDomainPackageSegments::FSegment* FEditorDomainPackageSegments::EnsureSegmentRange(uint64 Start, uint64 End, bool bWait, bool& bOutIsReady)
 {
-	WaitForReady();
-	return PackageFormat;
-}
+	// Called from Interface-only
+	check(Source == ESource::Segments); // Caller should not call in any other case
+
+	int32 StartIndex = GetSegmentContainingOffset(Start);
+	FSegment* StartSegment = &Segments[StartIndex];
 
+	bOutIsReady = true;
+	bool bNeedsWait = false;
+	FSegment* Segment = StartSegment;
+	FSegment* FileEndSegment = &Segments[0] + Segments.Num();
+	do
+	{
+		if (!Segment->bComplete)
+		{
+			if (!Segment->Request.IsValid())
+			{
+				SendSegmentRequest(*Segment);
+			}
+			if (Segment->bAsyncComplete)
+			{
+				Segment->bComplete = true;
+			}
+			else
+			{
+				bOutIsReady = false;
+			}
+		}
+		++Segment;
+	} while (Segment < FileEndSegment && Segment->Start < End);
 
-void FEditorDomainReadArchive::OnCacheRequestComplete(UE::DerivedData::FCacheGetCompleteParams&& Params)
+	if (!bOutIsReady && bWait)
+	{
+		Segment = StartSegment;
+		do
+		{
+			if (!Segment->bComplete)
+			{
+				Segment->Request.Wait();
+				check(Segment->bAsyncComplete);
+				Segment->bComplete = true;
+			}
+			++Segment;
+		} while (Segment < FileEndSegment && Segment->Start < End);
+		bOutIsReady = true;
+	}
+	return StartSegment;
+}
+
+void FEditorDomainPackageSegments::OnRecordRequestComplete(UE::DerivedData::FCacheGetCompleteParams&& Params,
+	TUniqueFunction<void(bool bValid)>&& CreateSegmentData,
+	TUniqueFunction<bool(FEditorDomain& EditorDomain)>&& TryCreateFallbackData)
 {
-	FScopeLock ScopeLock(&Lock);
-	if (AsyncSource == ESource::Closed)
 	{
-		return;
+		ESource LocalAsyncSource = AsyncSource;
+		if (AsyncSource == ESource::Closed)
+		{
+			return;
+		}
+		check(LocalAsyncSource == ESource::Uninitialized);
 	}
 
-	check(AsyncSource == ESource::Uninitialized);
-	FScopeLock DomainScopeLock(&EditorDomainLocks->Lock);
-	if ((PackageSource->Source == FEditorDomain::EPackageSource::Undecided || PackageSource->Source == FEditorDomain::EPackageSource::Editor) &&
-		Params.Status == UE::DerivedData::EStatus::Ok)
+	Size = 0;
+
+	ESource NewAsyncSource = ESource::Uninitialized;
+	int32 InitialRequestIndex = -1;
+	if (Params.Status == UE::DerivedData::EStatus::Ok)
 	{
-		const FCbObject& MetaData = Params.Record.GetMeta();
-		int64 FileSize = MetaData["FileSize"].AsInt64(-1);
-		FSharedBuffer LocalBytes = MoveTemp(Params.Record).GetValue();
-		if (static_cast<int64>(LocalBytes.GetSize()) != FileSize)
+		UE::DerivedData::FCacheRecord& Record = Params.Record;
+		uint64 FileSize = Record.GetMeta()["FileSize"].AsUInt64();
+		const UE::DerivedData::FPayload& ValuePayload = Record.GetValuePayload();
+
+		Segments.Reserve(1 + Record.GetAttachmentPayloads().Num());
+		uint64 CompositeSize = 0;
+		uint64 ValueSize = ValuePayload.GetRawSize();
+		if (ValueSize > 0)
+		{
+			Segments.Emplace(ValuePayload.GetId(), CompositeSize);
+			CompositeSize += ValueSize;
+		}
+		for (const UE::DerivedData::FPayload& Payload : Record.GetAttachmentPayloads())
+		{
+			uint64 PayloadSize = Payload.GetRawSize();
+			if (PayloadSize > 0)
+			{
+				Segments.Emplace(Payload.GetId(), CompositeSize);
+				CompositeSize += PayloadSize;
+			}
+		}
+
+		if (CompositeSize != FileSize)
 		{
 			UE_LOG(LogEditorDomain, Warning, TEXT("Package %s received invalid record from EditorDomainPackage table ")
-				TEXT("with blob size %" INT64_FMT " not equal to FileSize in metadata %" INT64_FMT)
+				TEXT("with size of all segments %" UINT64_FMT " not equal to FileSize in metadata %" UINT64_FMT)
 				TEXT(". Reading from workspace domain instead."),
-				*PackagePath.GetDebugName(), static_cast<int64>(LocalBytes.GetSize()), FileSize);
+				*PackagePath.GetDebugName(), CompositeSize, FileSize);
+			Segments.Empty();
 		}
 		else
 		{
-			AsyncSource = ESource::Bytes;
-			Size = FileSize;
-			Bytes = MoveTemp(LocalBytes).MakeOwned();
-			PackageFormat = EPackageFormat::Binary;
-			PackageSource->Source = FEditorDomain::EPackageSource::Editor;
+			bool bRegisterSuccessful = false;
+			bool bEditorDomainAvailable = false;
+			{
+				FScopeLock DomainScopeLock(&EditorDomainLocks->Lock);
+				if (EditorDomainLocks->Owner)
+				{
+					bEditorDomainAvailable = true;
+					if (PackageSource->Source == FEditorDomain::EPackageSource::Undecided || PackageSource->Source == FEditorDomain::EPackageSource::Editor)
+					{
+						PackageSource->Source = FEditorDomain::EPackageSource::Editor;
+						bRegisterSuccessful = true;
+					}
+				}
+			}
+			if (bRegisterSuccessful)
+			{
+				Size = FileSize;
+				NewAsyncSource = ESource::Segments;
+				CreateSegmentData(true /* bValid */);
+				if (0 <= InitialRequestOffset && static_cast<uint32>(InitialRequestOffset) < FileSize)
+				{
+					InitialRequestIndex = GetSegmentContainingOffset(InitialRequestOffset);
+				}
+			}
+			else
+			{
+				Segments.Empty();
+				if (!bEditorDomainAvailable)
+				{
+					UE_LOG(LogEditorDomain, Warning, TEXT("%s read after EditorDomain shutdown. Archive Set to Error."),
+						*PackagePath.GetDebugName());
+					NewAsyncSource = ESource::Segments;
+					CreateSegmentData(false /* bValid */);
+				}
+			}
 		}
 	}
-	if (AsyncSource == ESource::Uninitialized)
+
+	if (NewAsyncSource == ESource::Uninitialized)
 	{
-		checkf(PackageSource->Source == FEditorDomain::EPackageSource::Undecided || PackageSource->Source == FEditorDomain::EPackageSource::Workspace,
-			TEXT("%s was previously loaded from the EditorDomain but now is unavailable."), *PackagePath.GetDebugName());
+		FScopeLock DomainScopeLock(&EditorDomainLocks->Lock);
 		if (EditorDomainLocks->Owner)
 		{
+			checkf(PackageSource->Source == FEditorDomain::EPackageSource::Undecided || PackageSource->Source == FEditorDomain::EPackageSource::Workspace,
+				TEXT("%s was previously loaded from the EditorDomain but now is unavailable."), *PackagePath.GetDebugName());
 			FEditorDomain& EditorDomain(*EditorDomainLocks->Owner);
-			IPackageResourceManager& Workspace = *EditorDomain.Workspace;
-			FOpenPackageResult Result = Workspace.OpenReadPackage(PackagePath, EPackageSegment::Header);
-			if (Result.Archive)
+			bool bSucceeded = TryCreateFallbackData(EditorDomain);
+			if (bSucceeded)
 			{
 				EditorDomain.MarkNeedsLoadFromWorkspace(PackagePath, PackageSource);
-				InnerArchive = MoveTemp(Result.Archive);
-				AsyncSource = ESource::Archive;
-				Size = InnerArchive->TotalSize();
-				PackageFormat = Result.Format;
+				NewAsyncSource = ESource::Fallback;
 			}
 			else
 			{
-				UE_LOG(LogEditorDomain, Warning, TEXT("%s could not be read from WorkspaceDomain. Archive Set to Error."),
-					*PackagePath.GetDebugName());
-				AsyncSource = ESource::Bytes;
-				Size = 0;
-				PackageFormat = EPackageFormat::Binary;
-				SetError();
+				NewAsyncSource = ESource::Segments;
 			}
 		}
 		else
 		{
 			UE_LOG(LogEditorDomain, Warning, TEXT("%s read after EditorDomain shutdown. Archive Set to Error."),
 				*PackagePath.GetDebugName());
-			AsyncSource = ESource::Bytes;
-			Size = 0;
-			PackageFormat = EPackageFormat::Binary;
-			SetError();
+			NewAsyncSource = ESource::Segments;
+			CreateSegmentData(false /* bValid */);
 		}
 	}
+
+	ESource ExpectedValue = ESource::Uninitialized;
+	bool bSetSuccessful = AsyncSource.compare_exchange_strong(ExpectedValue, NewAsyncSource);
+	if (bSetSuccessful && InitialRequestIndex >= 0)
+	{
+		SendSegmentRequest(Segments[InitialRequestIndex]);
+	}
+}
+
+void FEditorDomainPackageSegments::SendSegmentRequest(FSegment& Segment)
+{
+	// Called from Callback-only until AsyncSource is set, then from Interface-only
+	UE::DerivedData::ICache& Cache = GetDerivedDataCacheRef();
+
+	// Note that Segment.Request is Interface-only and so we can write it outside the lock
+	Segment.Request = Cache.GetPayload({ UE::DerivedData::FCachePayloadKey{ UE::EditorDomain::GetEditorDomainPackageKey(PackageDigest), Segment.PayloadId} },
+		PackagePath.GetDebugName(), UE::DerivedData::ECachePolicy::Local, UE::DerivedData::EPriority::Normal,
+		[this, &Segment](UE::DerivedData::FCacheGetPayloadCompleteParams&& Params)
+		{
+			if (AsyncSource == ESource::Closed)
+			{
+				return;
+			}
+			if (Segment.bAsyncComplete)
+			{
+				// Race condition; another request got there before us
+				return;
+			}
+
+			if (Params.Status != UE::DerivedData::EStatus::Ok)
+			{
+				UE_LOG(LogEditorDomain, Error, TEXT("Package %s is missing cache data in segment %d."), *PackagePath.GetDebugName(), (int32)(&Segment - &Segments[0]));
+			}
+			else
+			{
+				Segment.Data = Params.Payload.GetData().Decompress();
+				uint64 SegmentEnd = &Segment < &Segments.Last() ? (&Segment + 1)->Start : Size;
+				uint64 SegmentSize = SegmentEnd - Segment.Start;
+				if (Segment.Data.GetSize() != SegmentSize)
+				{
+					UE_LOG(LogEditorDomain, Error, TEXT("Package %s has corrupted cache data in segment %d."), *PackagePath.GetDebugName(), (int32)(&Segment - &Segments[0]));
+					Segment.Data.Reset();
+				}
+			}
+			Segment.bAsyncComplete = true;
+		});
+}
+
+FEditorDomainReadArchive::FEditorDomainReadArchive(const TRefCountPtr<FEditorDomain::FLocks>& InLocks,
+	const FPackagePath& InPackagePath, const TRefCountPtr<FEditorDomain::FPackageSource>& InPackageSource)
+	: Segments(InLocks, InPackagePath, InPackageSource)
+{
+	this->SetIsLoading(true);
+	this->SetIsPersistent(true);
+}
+
+void FEditorDomainReadArchive::SetRequest(const UE::DerivedData::FRequest& InRequest)
+{
+	Segments.SetRequest(InRequest);
 }
 
 void FEditorDomainReadArchive::WaitForReady() const
 {
-	if (Source != ESource::Uninitialized)
+	if (Segments.GetSource() != FEditorDomainPackageSegments::ESource::Uninitialized)
 	{
 		return;
 	}
-	Request.Wait();
+
+	Segments.WaitForReady();
+
+	// If we had any seeks, implement the seek behavior for the last seek, now that we know what that behavior is.
+	int64 Pos = Segments.Tell();
+	if (Pos != 0)
+	{
+		switch (Segments.GetSource())
+		{
+		case FEditorDomainPackageSegments::ESource::Fallback:
+		{
+			InnerArchive->Seek(Pos);
+			break;
+		}
+		case FEditorDomainPackageSegments::ESource::Segments:
+		{
+			const_cast<FEditorDomainReadArchive&>(*this).Segments.Precache(Pos, 0);
+			break;
+		}
+		default:
+			checkNoEntry();
+			break;
+		}
+	}
+}
+
+void FEditorDomainReadArchive::OnRecordRequestComplete(UE::DerivedData::FCacheGetCompleteParams&& Params)
+{
+	Segments.OnRecordRequestComplete(MoveTemp(Params),
+		[this](bool bValid) { CreateSegmentData(bValid); },
+		[this](FEditorDomain& EditorDomain) { return TryCreateFallbackData(EditorDomain); }
+	);
+}
+
+void FEditorDomainReadArchive::CreateSegmentData(bool bValid)
+{
+	PackageFormat = EPackageFormat::Binary;
+	if (!bValid)
 	{
-		// Even though we know that the asynchronous task has left the critical section,
-		// we still need to synchronize the memory order.
-		// Entering the critical section will activate the equivalent of std::memory_order::acquire that we need.
-		FScopeLock ReceiverScopeLock(&Lock);
-		Source = AsyncSource;
+		SetError();
 	}
+}
 
-	switch (Source)
+bool FEditorDomainReadArchive::TryCreateFallbackData(FEditorDomain& EditorDomain)
+{
+	IPackageResourceManager& Workspace = *EditorDomain.Workspace;
+	// EDITOR_DOMAIN_TODO: Editor platforms that return false in FPlatformMisc::SupportsMultithreadedFileHandles
+	// will fail here because we are creating the file handle on the callback thread and using it on the 
+	// interface thread which is on the game thread. Need to make all editor platforms support that, or
+	// have this fallback occur when WaitForReady is called.
+	FOpenPackageResult Result = Workspace.OpenReadPackage(Segments.GetPackagePath(), EPackageSegment::Header);
+	if (Result.Archive)
+	{
+		InnerArchive = MoveTemp(Result.Archive);
+		PackageFormat = Result.Format;
+		return true;
+	}
+	else
+	{
+		UE_LOG(LogEditorDomain, Warning, TEXT("%s could not be read from WorkspaceDomain. Archive Set to Error."),
+			*Segments.GetPackagePath().GetDebugName());
+		PackageFormat = EPackageFormat::Binary;
+		SetError();
+		return false;
+	}
+}
+
+void FEditorDomainReadArchive::Seek(int64 InPos)
+{
+	switch (Segments.GetSource())
 	{
-	case ESource::Archive:
-		// Copy local variables over to InnerArchive
-		if (Pos != 0)
+	case FEditorDomainPackageSegments::ESource::Uninitialized:
+		if (Segments.GetAsyncSource() != FEditorDomainPackageSegments::ESource::Uninitialized)
 		{
-			InnerArchive->Seek(Pos);
+			WaitForReady();
+			check(Segments.GetSource() != FEditorDomainPackageSegments::ESource::Uninitialized);
+			return Seek(InPos);
 		}
+		Segments.Seek(InPos);
+		break;
+	case FEditorDomainPackageSegments::ESource::Segments:
+		Segments.Seek(InPos);
+		Segments.Precache(InPos, 0);
+		break;
+	case FEditorDomainPackageSegments::ESource::Fallback:
+		InnerArchive->Seek(InPos);
+		break;
+	case FEditorDomainPackageSegments::ESource::Closed:
+		break;
+	default:
+		checkNoEntry();
+		break;
+	}
+}
+
+int64 FEditorDomainReadArchive::Tell()
+{
+	switch (Segments.GetSource())
+	{
+	case FEditorDomainPackageSegments::ESource::Fallback:
+		return InnerArchive->Tell();
+	default:
+		return Segments.Tell();
+	}
+}
+
+int64 FEditorDomainReadArchive::TotalSize()
+{
+	WaitForReady();
+	switch (Segments.GetSource())
+	{
+	case FEditorDomainPackageSegments::ESource::Fallback:
+		return InnerArchive->TotalSize();
+	default:
+		return Segments.TotalSize();
+	}
+}
+
+bool FEditorDomainReadArchive::Close()
+{
+	Segments.Close();
+	InnerArchive.Reset();
+	return true;
+}
+
+void FEditorDomainReadArchive::Serialize(void* V, int64 Length)
+{
+	switch (Segments.GetSource())
+	{
+	case FEditorDomainPackageSegments::ESource::Uninitialized:
+		WaitForReady();
+		check(Segments.GetSource() != FEditorDomainPackageSegments::ESource::Uninitialized);
+		Serialize(V, Length);
+		break;
+	case FEditorDomainPackageSegments::ESource::Segments:
+		if (!Segments.Serialize(V, Length))
+		{
+			SetError();
+		}
+		break;
+	case FEditorDomainPackageSegments::ESource::Fallback:
+		InnerArchive->Serialize(V, Length);
+		break;
+	case FEditorDomainPackageSegments::ESource::Closed:
+		UE_LOG(LogEditorDomain, Error, TEXT("Requested read after close (file=%s)"), *Segments.GetPackagePath().GetDebugName());
+		break;
+	default:
+		checkNoEntry();
+		break;
+	}
+}
+
+FString FEditorDomainReadArchive::GetArchiveName() const
+{
+	return Segments.GetPackagePath().GetDebugName();
+}
+
+void FEditorDomainReadArchive::Flush()
+{
+	switch (Segments.GetSource())
+	{
+	case FEditorDomainPackageSegments::ESource::Uninitialized:
+		WaitForReady();
+		check(Segments.GetSource() != FEditorDomainPackageSegments::ESource::Uninitialized);
+		Flush();
+		break;
+	case FEditorDomainPackageSegments::ESource::Fallback:
+		InnerArchive->Flush();
+		break;
+	default:
+		break;
+	}
+}
+
+void FEditorDomainReadArchive::FlushCache()
+{
+	switch (Segments.GetSource())
+	{
+	case FEditorDomainPackageSegments::ESource::Uninitialized:
+		WaitForReady();
+		check(Segments.GetSource() != FEditorDomainPackageSegments::ESource::Uninitialized);
+		FlushCache();
 		break;
-	case ESource::Bytes:
+	case FEditorDomainPackageSegments::ESource::Fallback:
+		InnerArchive->FlushCache();
 		break;
 	default:
-		check(false);
 		break;
 	}
 }
 
+bool FEditorDomainReadArchive::Precache(int64 PrecacheOffset, int64 PrecacheSize)
+{
+	switch (Segments.GetSource())
+	{
+	case FEditorDomainPackageSegments::ESource::Uninitialized:
+	{
+		bool bIsReady;
+		Segments.Precache(PrecacheOffset, PrecacheSize, &bIsReady);
+		return bIsReady;
+	}
+	case FEditorDomainPackageSegments::ESource::Segments:
+	{
+		bool bIsReady;
+		Segments.Precache(PrecacheOffset, PrecacheSize, &bIsReady);
+		return bIsReady;
+	}
+	case FEditorDomainPackageSegments::ESource::Fallback:
+		return InnerArchive->Precache(PrecacheOffset, PrecacheSize);
+	default:
+		return true;
+	}
+}
+
+EPackageFormat FEditorDomainReadArchive::GetPackageFormat() const
+{
+	WaitForReady();
+	return PackageFormat;
+}
+
 /** An IAsyncReadRequest SizeRequest that returns a value known at construction time. */
 class FAsyncSizeRequestConstant : public IAsyncReadRequest
 {
@@ -297,26 +656,18 @@ protected:
 	}
 };
 
-/** An IAsyncReadRequest that reads from a SharedBuffer that was already populated at construction time. */
+/** An IAsyncReadRequest that reads from memory that was already populated. */
 class FAsyncReadRequestConstant : public IAsyncReadRequest
 {
 public:
-	FAsyncReadRequestConstant(const FSharedBuffer& Bytes, FAsyncFileCallBack* InCallback, int64 Offset,
-		int64 BytesToRead, uint8* UserSuppliedMemory, const FPackagePath& PackagePath)
-		: IAsyncReadRequest(InCallback, false /* bInSizeRequest */, UserSuppliedMemory)
+	FAsyncReadRequestConstant(uint8* Data, bool bIsUserSupplied, FAsyncFileCallBack* InCallback)
+		: IAsyncReadRequest(InCallback, false /* bInSizeRequest */, bIsUserSupplied ? Data : nullptr)
 	{
-		if (Offset < 0 || BytesToRead < 0 || Offset + BytesToRead > static_cast<int64>(Bytes.GetSize()))
+		if (!bIsUserSupplied)
 		{
-			UE_LOG(LogEditorDomain, Fatal, TEXT("FAsyncReadRequestConstant bogus request Offset = %" INT64_FMT)
-				TEXT(" BytesToRead = %" INT64_FMT " Bytes.GetSize() == %" UINT64_FMT " File = %s"),
-				Offset, BytesToRead, Bytes.GetSize(), *PackagePath.GetDebugName());
-		}
-		if (!UserSuppliedMemory)
-		{
-			Memory = (uint8*)FMemory::Malloc(BytesToRead);
+			Memory = Data;
 		}
 		check(Memory);
-		FMemory::Memcpy(Memory, reinterpret_cast<const uint8*>(Bytes.GetData()) + Offset, BytesToRead);
 		SetComplete();
 	}
 
@@ -346,41 +697,47 @@ protected:
 
 FEditorDomainAsyncReadFileHandle::FEditorDomainAsyncReadFileHandle(const TRefCountPtr<FEditorDomain::FLocks>& InLocks,
 	const FPackagePath& InPackagePath, const TRefCountPtr<FEditorDomain::FPackageSource>& InPackageSource)
-	: EditorDomainLocks(InLocks)
-	, PackagePath(InPackagePath)
-	, PackageSource(InPackageSource)
+	: Segments(InLocks, InPackagePath, InPackageSource)
 {
 }
 
-FEditorDomainAsyncReadFileHandle::~FEditorDomainAsyncReadFileHandle()
+void FEditorDomainAsyncReadFileHandle::SetRequest(const UE::DerivedData::FRequest& InRequest)
 {
-	{
-		FScopeLock ScopeLock(&Lock);
-		if (AsyncSource == ESource::Uninitialized)
-		{
-			AsyncSource = ESource::Closed;
-		}
-	}
-	Request.Cancel();
-	Source = ESource::Closed;
+	Segments.SetRequest(InRequest);
 }
 
-void FEditorDomainAsyncReadFileHandle::SetRequest(const UE::DerivedData::FRequest& InRequest)
+void FEditorDomainAsyncReadFileHandle::OnRecordRequestComplete(UE::DerivedData::FCacheGetCompleteParams&& Params)
 {
-	Request = InRequest;
+	Segments.OnRecordRequestComplete(MoveTemp(Params),
+		[this](bool bValid) { CreateSegmentData(bValid); },
+		[this](FEditorDomain& EditorDomain) { return TryCreateFallbackData(EditorDomain); }
+	);
+}
+
+void FEditorDomainAsyncReadFileHandle::CreateSegmentData(bool bValid)
+{
+}
+
+bool FEditorDomainAsyncReadFileHandle::TryCreateFallbackData(FEditorDomain& EditorDomain)
+{
+	IPackageResourceManager& Workspace = *EditorDomain.Workspace;
+	IAsyncReadFileHandle* Result = Workspace.OpenAsyncReadPackage(Segments.GetPackagePath(), EPackageSegment::Header);
+	check(Result);
+	InnerArchive.Reset(Result);
+	return true;
 }
 
 IAsyncReadRequest* FEditorDomainAsyncReadFileHandle::SizeRequest(FAsyncFileCallBack* CompleteCallback)
 {
-	switch (Source)
+	switch (Segments.GetSource())
 	{
-	case ESource::Uninitialized:
-		WaitForReady();
-		check(Source != ESource::Uninitialized);
+	case FEditorDomainPackageSegments::ESource::Uninitialized:
+		Segments.WaitForReady();
+		check(Segments.GetSource() != FEditorDomainPackageSegments::ESource::Uninitialized);
 		return SizeRequest(CompleteCallback);
-	case ESource::Bytes:
-		return new FAsyncSizeRequestConstant(Bytes.GetSize(), CompleteCallback);
-	case ESource::Archive:
+	case FEditorDomainPackageSegments::ESource::Segments:
+		return new FAsyncSizeRequestConstant(Segments.TotalSize(), CompleteCallback);
+	case FEditorDomainPackageSegments::ESource::Fallback:
 		return InnerArchive->SizeRequest(CompleteCallback);
 	default:
 		checkNoEntry();
@@ -392,92 +749,41 @@ IAsyncReadRequest* FEditorDomainAsyncReadFileHandle::ReadRequest(int64 Offset, i
 	EAsyncIOPriorityAndFlags PriorityAndFlags, FAsyncFileCallBack* CompleteCallback,
 	uint8* UserSuppliedMemory)
 {
-	switch (Source)
+	switch (Segments.GetSource())
 	{
-	case ESource::Uninitialized:
-		WaitForReady();
-		check(Source != ESource::Uninitialized);
+	case FEditorDomainPackageSegments::ESource::Uninitialized:
+		Segments.WaitForReady();
+		check(Segments.GetSource() != FEditorDomainPackageSegments::ESource::Uninitialized);
 		return ReadRequest(Offset, BytesToRead, PriorityAndFlags, CompleteCallback, UserSuppliedMemory);
-	case ESource::Bytes:
-		return new FAsyncReadRequestConstant(Bytes, CompleteCallback, Offset, BytesToRead, UserSuppliedMemory, PackagePath);
-	case ESource::Archive:
-		return InnerArchive->ReadRequest(Offset, BytesToRead, PriorityAndFlags, CompleteCallback, UserSuppliedMemory);
-	default:
-		checkNoEntry();
-		return nullptr;
-	}
-}
-
-bool FEditorDomainAsyncReadFileHandle::UsesCache()
-{
-	return false;
-}
-
-void FEditorDomainAsyncReadFileHandle::OnCacheRequestComplete(UE::DerivedData::FCacheGetCompleteParams&& Params)
-{
-	FScopeLock ScopeLock(&Lock);
-	if (AsyncSource == ESource::Closed)
-	{
-		return;
-	}
-
-	check(AsyncSource == ESource::Uninitialized);
-	FScopeLock DomainScopeLock(&EditorDomainLocks->Lock);
-	if ((PackageSource->Source == FEditorDomain::EPackageSource::Undecided || PackageSource->Source == FEditorDomain::EPackageSource::Editor) &&
-		Params.Status == UE::DerivedData::EStatus::Ok)
+	case FEditorDomainPackageSegments::ESource::Segments:
 	{
-		const FCbObject& MetaData = Params.Record.GetMeta();
-		int64 FileSize = MetaData["FileSize"].AsInt64(-1);
-		FSharedBuffer LocalBytes = MoveTemp(Params.Record).GetValue();
-		if (static_cast<int64>(LocalBytes.GetSize()) != FileSize)
+		// EDITOR_DOMAIN_TODO: Make a FAsyncReadRequest that accepts a TFuture instead of synchronously reading upfront
+		Segments.Precache(Offset, 0);
+		Segments.Seek(Offset);
+		uint8* Memory = UserSuppliedMemory;
+		if (!Memory)
 		{
-			UE_LOG(LogEditorDomain, Warning, TEXT("Package %s received invalid record from EditorDomainPackage ")
-				TEXT("table with blob size %" INT64_FMT "not equal to FileSize in metadata %" INT64_FMT)
-				TEXT(". Reading from workspace domain instead."),
-				*PackagePath.GetDebugName(), static_cast<int64>(LocalBytes.GetSize()), FileSize);
+			Memory = (uint8*)FMemory::Malloc(BytesToRead);
 		}
-		else
+		if (!Segments.Serialize(Memory, BytesToRead))
 		{
-			AsyncSource = ESource::Bytes;
-			Bytes = MoveTemp(LocalBytes).MakeOwned();
-			PackageSource->Source = FEditorDomain::EPackageSource::Editor;
+			if (!UserSuppliedMemory)
+			{
+				FMemory::Free(Memory);
+			}
+			return nullptr;
 		}
+		return new FAsyncReadRequestConstant(Memory, UserSuppliedMemory != nullptr, CompleteCallback);
 	}
-	if (AsyncSource == ESource::Uninitialized)
-	{
-		checkf(PackageSource->Source == FEditorDomain::EPackageSource::Undecided || PackageSource->Source == FEditorDomain::EPackageSource::Workspace,
-			TEXT("%s was previously loaded from the EditorDomain but now is unavailable."), *PackagePath.GetDebugName());
-		if (EditorDomainLocks->Owner)
-		{
-			FEditorDomain& EditorDomain(*EditorDomainLocks->Owner);
-			IPackageResourceManager& Workspace = *EditorDomain.Workspace;
-			IAsyncReadFileHandle* Result = Workspace.OpenAsyncReadPackage(PackagePath, EPackageSegment::Header);
-			check(Result);
-			AsyncSource = ESource::Archive;
-			InnerArchive.Reset(Result);
-			EditorDomain.MarkNeedsLoadFromWorkspace(PackagePath, PackageSource);
-		}
-		else
-		{
-			UE_LOG(LogEditorDomain, Warning, TEXT("%s read after EditorDomain shutdown. Returning null archive"),
-				*PackagePath.GetDebugName());
-			AsyncSource = ESource::Archive;
-			InnerArchive.Reset(new FAsyncReadFileHandleNull());
-		}
+	case FEditorDomainPackageSegments::ESource::Fallback:
+		return InnerArchive->ReadRequest(Offset, BytesToRead, PriorityAndFlags, CompleteCallback, UserSuppliedMemory);
+	default:
+		checkNoEntry();
+		return nullptr;
 	}
 }
 
-void FEditorDomainAsyncReadFileHandle::WaitForReady() const
+bool FEditorDomainAsyncReadFileHandle::UsesCache()
 {
-	if (Source != ESource::Uninitialized)
-	{
-		return;
-	}
-	Request.Wait();
-	{
-		// Even though we know that the asynchronous task has left the critical section, we still need to synchronize the memory order
-		// Entering the critical section will activate the equivalent of std::memory_order::acquire that we need
-		FScopeLock ReceiverScopeLock(&Lock);
-		Source = AsyncSource;
-	}
+	return true;
 }
diff --git a/Engine/Source/Editor/UnrealEd/Private/EditorDomain/EditorDomainArchive.h b/Engine/Source/Editor/UnrealEd/Private/EditorDomain/EditorDomainArchive.h
--- a/Engine/Source/Editor/UnrealEd/Private/EditorDomain/EditorDomainArchive.h
+++ b/Engine/Source/Editor/UnrealEd/Private/EditorDomain/EditorDomainArchive.h
@@ -3,6 +3,9 @@
 #pragma once
 
 #include "Async/AsyncFileHandle.h"
+#include "Containers/Array.h"
+#include "DerivedDataCache.h"
+#include "DerivedDataPayload.h"
 #include "DerivedDataRequest.h"
 #include "EditorDomain/EditorDomain.h"
 #include "HAL/CriticalSection.h"
@@ -10,6 +13,7 @@
 #include "Memory/SharedBuffer.h"
 #include "Misc/PackagePath.h"
 #include "Serialization/Archive.h"
+#include "Templates/Function.h"
 #include "Templates/RefCounting.h"
 #include "Templates/UniquePtr.h"
 #include "UObject/PackageResourceManager.h"
@@ -17,6 +21,168 @@
 class FAssetPackageData;
 namespace UE { namespace DerivedData { struct FCacheGetCompleteParams; } }
 
+/**
+ * Helper class for archive classes that server either an EditorDomain version of a package
+ * or the WorkspaceDomain version. Starts with a request to the Cache to load the metadata for
+ * an EditorDomainPackage. If the request fails, goes dormant and the owner class falls
+ * back to workspace domain. If the request succeeds, it makes all Attachments from the
+ * CacheRecord available as individually fetched segments.
+ */
+class FEditorDomainPackageSegments
+{
+public:
+	enum class ESource : uint8
+	{
+		Uninitialized,
+		Segments,
+		Fallback,
+		Closed
+	};
+
+	FEditorDomainPackageSegments(const TRefCountPtr<FEditorDomain::FLocks>& InLocks,
+		const FPackagePath& InPackagePath, const TRefCountPtr<FEditorDomain::FPackageSource>& InPackageSource);
+	~FEditorDomainPackageSegments();
+
+	/** Set the CacheRequest handle that will feed this archive. */
+	void SetRequest(const UE::DerivedData::FRequest& InRequest);
+	/** Callback from the request of the Record; set whether we're reading from EditorDomain bytes or WorkspaceDomain archive. */
+	void OnRecordRequestComplete(UE::DerivedData::FCacheGetCompleteParams&& Params,
+		TUniqueFunction<void(bool bValid)>&& CreateSegmentData,
+		TUniqueFunction<bool(FEditorDomain& EditorDomain)>&& TryCreateFallbackData);
+
+	/**
+	 * Get the cache of GetAsyncSource as seen on the Interface thread.
+	 * If Uninitialized, caller should early exit, WaitForReady, or call GetAsyncSource.
+	 */
+	ESource GetSource() const;
+
+	/** Get the ESource that is an atomic used from all threads. */
+	std::atomic<ESource>& GetAsyncSource();
+
+	/** Wait for the handle to call OnCacheRequestComplete and make the size and list of Segments available. */
+	void WaitForReady() const;
+
+	/** Set Pos to the given value. Caller is responsible for ensuring in range. Can be called before ready. */
+	void Seek(int64 InPos);
+
+	/** Report the current value of Pos. Can be called before ready. */
+	int64 Tell() const;
+
+	/**
+	 * Report the combined size of all segments.
+	 * Must not be called unless GetSource is ESource::Segments.
+	 */
+	int64 TotalSize() const;
+
+	/** Return the PackagePath. Can be called before ready. */
+	const FPackagePath& GetPackagePath() const;
+
+	/** Cancel any requests, wait for them to complete, release memory. */
+	void Close();
+
+	/**
+	 * Copy Length bytes from current Pos into V.
+	 * Must not be called unless GetSource is ESource::Segments.
+	 */
+	bool Serialize(void* V, int64 Length);
+
+	/**
+	 * Ensure requests have been sent for segments touching the given range. Out of range requests are ignored.
+	 * Must not be called unless GetSource is ESource::Segments.
+
+	 * @param bOutIsReady Set to true iff all segments touching the range have completd their request.
+	 */
+	void Precache(int64 InStart, int64 InSize, bool* bOutIsReady = nullptr);
+
+private:
+	/**
+	 * A segment of bytes that is sourced from an attachment in the CacheRecord and
+	 * is requested on demand.
+	 */
+	struct FSegment
+	{
+		FSegment(const UE::DerivedData::FPayloadId& InPayloadId, uint64 InStart)
+			: PayloadId(InPayloadId), Start(InStart)
+		{}
+
+		/** Attachment ID to request from cache for this segment. Read-only. */
+		UE::DerivedData::FPayloadId PayloadId;
+		/** The request that will be (has been) sent for the bytes from cache. Interface-only. */
+		UE::DerivedData::FRequest Request;
+		/** Offset from the start of the entire archive. Read-only. */
+		uint64 Start;
+		/** Payload bytes. Callback-only before bAsyncComplete, Interface-only after. */
+		FSharedBuffer Data;
+		/** Interface thread cache of bAsyncComplete. Interface-only. */
+		bool bComplete = false;
+		/** Set to true whent the callback for the cache request has finished writing Data. */
+		std::atomic<bool> bAsyncComplete{ false };
+	};
+
+	/** Map offset to segment index. Start must be in range 0 <= Start < Size. Must be called after Segments assigned. */
+	int32 GetSegmentContainingOffset(uint64 Start) const;
+
+	/**
+	 * Make sure requests are sent for segments touching [Start, End); optionally wait for requests to finish.
+	 * Must not be called unless GetSource is ESource::Segments. Caller must ensure valid range Start <= End <= Size.
+	 *
+	 * @param bOutIsReady Set to true if all segments touching the range have finished their request.
+	 * @return Address of the first segment touching the range.
+	 */
+	FSegment* EnsureSegmentRange(uint64 Start, uint64 End, bool bWait, bool& bOutIsReady);
+
+	/**
+	 * Send the request for the given segment; does not check if it has already been sent.
+	 * Losing Race values are ignored.
+	 */
+	void SendSegmentRequest(FSegment& Segment);
+
+	/**
+	 * Lock from the EditorDomain to allow shutting down this archive if callback outlives editor domain.
+	 * Pointer is read - only, * pointer has internal lock.
+	 */
+	TRefCountPtr<FEditorDomain::FLocks> EditorDomainLocks;
+	/** Initial Request for the list of segments. Read-only. */
+	UE::DerivedData::FRequest Request;
+	/** PackagePath for fallback to WorkspaceDomain. Read-only. */
+	FPackagePath PackagePath;
+	/** Offset to next use in Serialize. Interface-only. */
+	int64 Pos = 0;
+	/**
+	 * Interface-thread cache of AsyncSource. Used for performance and to trigger copying of async data on first
+	 * read after Request complete. Interface-only.
+	 */
+	mutable ESource Source = ESource::Uninitialized;
+	/**
+	 * PackageSource for read/write of whether this Package comes from editor domain or workspace domain.
+	 * Read-only, *pointer requires EditorDomainLocks.Lock.
+	 */
+	TRefCountPtr<FEditorDomain::FPackageSource> PackageSource;
+	/** PackageDigest for requesting payloads, copied from PackageSource. Read-only. */
+	UE::EditorDomain::FPackageDigest PackageDigest;
+
+	/**
+	 * Data read from the EditorDomain; each Segment is requested individually.
+	 * Array allocation is Callback-only before Request returns, read-only after.
+	 * Thread properties of elements are mixed; see FSegment.
+	 */
+	TArray<FSegment> Segments;
+	/** Size if Source is Segments, 0 otherwise. Callback-only before Request returns, read-only after. */
+	int64 Size = 0;
+	/** 
+	 * Specifies which data to use based on whether EditorDomain data exists.
+	 * Set by Close to shutdown both Callback and Interface activity.
+	 * Interface thread polls this in some cases to check whether it should wait for WaitForReady.
+	 */
+	std::atomic<ESource> AsyncSource{ ESource::Uninitialized };
+	/**
+	 * Instructs the Callback thread to immediately send a SegmentRequest for the given offset.
+	 * Written by Interface thread. Callback thread will use it if its set in time(shortly after AsyncSource is set),
+	 * otherwise will ignore or use stale value.
+	 */
+	std::atomic<int64> InitialRequestOffset{ -1 }; 
+};
+
 /**
  * An Archive that asynchronously waits for the cache request to complete, and reads either from the returned cache bytes
  * or from the fallback WorkspaceDomain Archive for the given PackagePath.
@@ -29,12 +195,11 @@ class FEditorDomainReadArchive : public FArchive
 public:
 	FEditorDomainReadArchive(const TRefCountPtr<FEditorDomain::FLocks>& InLocks,
 		const FPackagePath& InPackagePath, const TRefCountPtr<FEditorDomain::FPackageSource>& InPackageSource);
-	~FEditorDomainReadArchive();
 
 	/** Set the CacheRequest handle that will feed this archive. */
 	void SetRequest(const UE::DerivedData::FRequest& InRequest);
-	/** Callback from the CacheRequest; set whether we're reading from EditorDomain bytes or WorkspaceDomain archive. */
-	void OnCacheRequestComplete(UE::DerivedData::FCacheGetCompleteParams&& Params);
+	/** Callback from the request of the Record; set whether we're reading from EditorDomain bytes or WorkspaceDomain archive. */
+	void OnRecordRequestComplete(UE::DerivedData::FCacheGetCompleteParams&& Params);
 	/** Get the PackageFormat, which depends on the domain the data is read from. */
 	EPackageFormat GetPackageFormat() const;
 
@@ -51,34 +216,12 @@ public:
 	virtual bool NeedsEngineVersionChecks() const override { return false; }
 
 private:
-	enum class ESource : uint8
-	{
-		Uninitialized,
-		Bytes,
-		Archive,
-		Closed
-	};
-
-	/** Wait for the handle to call OnCacheRequestComplete and make the size and bytes available. */
 	void WaitForReady() const;
+	void CreateSegmentData(bool bValid);
+	bool TryCreateFallbackData(FEditorDomain& EditorDomain);
 
-	/** Lock to synchronize the CacheCompletion and the public interface thread. */
-	mutable FCriticalSection Lock;
-
-	// Data in this section is either read-only, or is read and written only on the public interface thread.
-	TRefCountPtr<FEditorDomain::FLocks> EditorDomainLocks; // Pointer is read-only, *pointer has internal lock
-	UE::DerivedData::FRequest Request; // Read-only
-	FPackagePath PackagePath; // Read-only
-	int64 Pos = 0; // Interface-thread-only
-	mutable ESource Source = ESource::Uninitialized; // Interface-thread-only
-	TRefCountPtr<FEditorDomain::FPackageSource> PackageSource; // Read-only, *pointer requires EditorDomainLocks.Lock
-
-	// Data in this section is written on the callback thread.
-	// It can not be read until after WaitForReady or Cancel.
+	FEditorDomainPackageSegments Segments;
 	TUniquePtr<FArchive> InnerArchive;
-	FSharedBuffer Bytes;
-	int64 Size = 0;
-	ESource AsyncSource = ESource::Uninitialized;
 	EPackageFormat PackageFormat = EPackageFormat::Binary;
 };
 
@@ -91,12 +234,11 @@ class FEditorDomainAsyncReadFileHandle : public IAsyncReadFileHandle
 public:
 	FEditorDomainAsyncReadFileHandle(const TRefCountPtr<FEditorDomain::FLocks>& InLocks,
 		const FPackagePath& InPackagePath, const TRefCountPtr<FEditorDomain::FPackageSource>& InPackageSource);
-	~FEditorDomainAsyncReadFileHandle();
 
 	/** Set the CacheRequest handle that will feed this archive. */
 	void SetRequest(const UE::DerivedData::FRequest& InRequest);
-	/** Callback from the CacheRequest; set whether we're reading from EditorDomain bytes or WorkspaceDomain archive. */
-	void OnCacheRequestComplete(UE::DerivedData::FCacheGetCompleteParams&& Params);
+	/** Callback from the request of the Record; set whether we're reading from EditorDomain bytes or WorkspaceDomain archive. */
+	void OnRecordRequestComplete(UE::DerivedData::FCacheGetCompleteParams&& Params);
 
 	// IAsyncReadFileHandle interface
 	virtual IAsyncReadRequest* SizeRequest(FAsyncFileCallBack* CompleteCallback = nullptr) override;
@@ -106,30 +248,9 @@ public:
 	virtual bool UsesCache();
 
 private:
-	enum class ESource : uint8
-	{
-		Uninitialized,
-		Bytes,
-		Archive,
-		Closed
-	};
-
-	/** Wait for the handle to call OnCacheRequestComplete and make the size and bytes available. */
-	void WaitForReady() const;
-
-	/** Lock to synchronize the CacheCompletion and the public interface thread. */
-	mutable FCriticalSection Lock;
-
-	// Data in this section is either read-only, or is read and written only on the public interface thread.
-	TRefCountPtr<FEditorDomain::FLocks> EditorDomainLocks; // Pointer is read-only, *pointer has internal lock
-	UE::DerivedData::FRequest Request; // Read-only
-	FPackagePath PackagePath; // Read-only
-	mutable ESource Source = ESource::Uninitialized; // Interface-thread-only
-	TRefCountPtr<FEditorDomain::FPackageSource> PackageSource; // Read-only, *pointer requires EditorDomainLocks.Lock
+	void CreateSegmentData(bool bValid);
+	bool TryCreateFallbackData(FEditorDomain& EditorDomain);
 
-	// Data in this section is written on the callback thread.
-	// It can not be read until after WaitForReady or Cancel.
+	FEditorDomainPackageSegments Segments;
 	TUniquePtr<IAsyncReadFileHandle> InnerArchive;
-	FSharedBuffer Bytes;
-	ESource AsyncSource = ESource::Uninitialized;
 };
\ No newline at end of file
diff --git a/Engine/Source/Editor/UnrealEd/Private/EditorDomain/EditorDomainUtils.cpp b/Engine/Source/Editor/UnrealEd/Private/EditorDomain/EditorDomainUtils.cpp
--- a/Engine/Source/Editor/UnrealEd/Private/EditorDomain/EditorDomainUtils.cpp
+++ b/Engine/Source/Editor/UnrealEd/Private/EditorDomain/EditorDomainUtils.cpp
@@ -37,7 +37,7 @@ FClassDigestMap& GetClassDigests()
 }
 
 // Change to a new guid when EditorDomain needs to be invalidated
-const TCHAR* EditorDomainVersion = TEXT("C217EB656E9B4C04816D3DC0E21901F6");
+const TCHAR* EditorDomainVersion = TEXT("D2B58CFE87DA471C95A91F36D2EA9214");
 // Identifier of the CacheBuckets for EditorDomain tables
 const TCHAR* EditorDomainPackageBucketName = TEXT("EditorDomainPackage");
 const TCHAR* EditorDomainBulkDataListBucketName = TEXT("EditorDomainBulkDataList");
@@ -281,19 +281,18 @@ public:
 
 	// FMemoryPackageStoreWriter interface
 	uint64 GetHeaderSize() const { return HeaderSize; }
-	FIoBuffer GetHeaderAndExports() const { return HeaderAndExports; }
+	FSharedBuffer& GetHeaderAndExports() { return HeaderAndExports; }
 	/**
 	 * BulkDatas in this array are views into the FMemoryPackageStoreWriter.
 	 * Callers should call MakeOwned if they make copies that need to outlive the FMemoryPackageStoreWriter.
 	 */
-	const TArray<FIoBuffer>& GetBulkDatas() const { return BulkDataRegions; }
+	TArray<FSharedBuffer>& GetBulkDatas() { return BulkDataRegions; }
 
 private:
 	void SetPackageName(FName InPackageName);
 
-	FIoBuffer HeaderAndExports;
-	TArray<FIoBuffer> BulkDataArchives;
-	TArray<FIoBuffer> BulkDataRegions;
+	FSharedBuffer HeaderAndExports;
+	TArray<FSharedBuffer> BulkDataRegions;
 	FCommitEvent CommitEvent;
 	FName PackageName;
 	uint64 HeaderSize = 0;
@@ -316,6 +315,14 @@ void FMemoryPackageStoreWriter::BeginPackage(const FBeginPackageInfo& Info)
 	SetPackageName(Info.PackageName);
 }
 
+auto IoBufferToSharedBuffer = [](const FIoBuffer& InBuffer) -> FSharedBuffer {
+	InBuffer.EnsureOwned();
+	const uint64 DataSize = InBuffer.DataSize();
+	FIoBuffer MutableBuffer(InBuffer);
+	uint8* DataPtr = MutableBuffer.Release().ValueOrDie();
+	return FSharedBuffer{ FSharedBuffer::TakeOwnership(DataPtr, DataSize, FMemory::Free) };
+};
+
 void FMemoryPackageStoreWriter::WritePackageData(const FPackageInfo& Info, const FIoBuffer& PackageData, const TArray<FFileRegion>& FileRegions)
 {
 	for (const FFileRegion& FileRegion : FileRegions)
@@ -326,7 +333,7 @@ void FMemoryPackageStoreWriter::WritePackageData(const FPackageInfo& Info, const
 	// Info.LooseFilePath is ignored
 	HeaderSize = Info.HeaderSize;
 	// Info.ChunkId is ignored
-	HeaderAndExports = PackageData;
+	HeaderAndExports = IoBufferToSharedBuffer(PackageData);
 }
 
 void FMemoryPackageStoreWriter::WriteBulkdata(const FBulkDataInfo& Info, const FIoBuffer& BulkData, const TArray<FFileRegion>& FileRegions)
@@ -336,14 +343,15 @@ void FMemoryPackageStoreWriter::WriteBulkdata(const FBulkDataInfo& Info, const F
 	// Info.ChunkId is ignored
 	checkf(Info.BulkdataType == IPackageStoreWriter::FBulkDataInfo::Standard, TEXT("MemoryPackageStoreWriter does not currently support BulkData types other than Standard."));
 
-	FIoBuffer& BulkDataOwner = BulkDataArchives.Add_GetRef(BulkData);
-	const uint8* BulkDataStart = BulkData.Data();
-	uint64 BulkDataLen = BulkData.DataSize();
+	FSharedBuffer BulkDataOwner = IoBufferToSharedBuffer(BulkData);
+	const uint8* BulkDataStart = reinterpret_cast<const uint8*>(BulkDataOwner.GetData());
+	uint64 BulkDataLen = BulkDataOwner.GetSize();
 	for (const FFileRegion& FileRegion : FileRegions)
 	{
 		checkf(FileRegion.Type == EFileRegionType::None, TEXT("FMemoryPackageStoreWriter does not currently support FileRegion types other than None."));
 		checkf(FileRegion.Offset + FileRegion.Length <= BulkDataLen, TEXT("FileRegions in WriteBulkdata were outside of the range of the BulkData's size."));
-		BulkDataRegions.Add(FIoBuffer(BulkDataStart + FileRegion.Offset, FileRegion.Length, BulkDataOwner));
+		check(FileRegion.Length > 0); // SavePackage is not allowed to call WriteBulkData with empty bulkdatas
+		BulkDataRegions.Add(FSharedBuffer::MakeView(BulkDataStart + FileRegion.Offset, FileRegion.Length, BulkDataOwner));
 	}
 }
 
@@ -351,14 +359,15 @@ void FMemoryPackageStoreWriter::WriteLinkerAdditionalData(const FLinkerAdditiona
 {
 	SetPackageName(Info.PackageName);
 
-	FIoBuffer& DataOwner = BulkDataArchives.Add_GetRef(Data);
-	const uint8* DataStart = Data.Data();
-	uint64 DataLen = Data.DataSize();
+	FSharedBuffer DataOwner = IoBufferToSharedBuffer(Data);
+	const uint8* DataStart = reinterpret_cast<const uint8*>(DataOwner.GetData());
+	uint64 DataLen = DataOwner.GetSize();
 	for (const FFileRegion& FileRegion : FileRegions)
 	{
 		checkf(FileRegion.Type == EFileRegionType::None, TEXT("FMemoryPackageStoreWriter does not currently support FileRegion types other than None."));
 		checkf(FileRegion.Offset + FileRegion.Length <= DataLen, TEXT("FileRegions in WriteLinkerAdditionalData were outside of the range of the Data's size."));
-		BulkDataRegions.Add(FIoBuffer(DataStart + FileRegion.Offset, FileRegion.Length, DataOwner));
+		check(FileRegion.Length > 0); // SavePackage is not allowed to call WriteLinkerAdditionalData with empty regions
+		BulkDataRegions.Add(FSharedBuffer::MakeView(DataStart + FileRegion.Offset, FileRegion.Length, DataOwner));
 	}
 
 }
@@ -422,27 +431,42 @@ bool TrySavePackage(UPackage* Package)
 		return false;
 	}
 
-	// EDITOR_DOMAIN_TODO: Convert FIoBuffer to use FSharedBuffer and refer to the SharedBuffers rather than copying
-	FSharedBuffer PackageBuffer;
-	{
-		TArray64<uint8> TempBytes;
-		const FIoBuffer& MainBuffer = PackageStoreWriter->GetHeaderAndExports();
-		TempBytes.Append(MainBuffer.Data(), MainBuffer.DataSize());
-		for (const FIoBuffer& BulkBuffer : PackageStoreWriter->GetBulkDatas())
-		{
-			TempBytes.Append(BulkBuffer.Data(), BulkBuffer.DataSize());
-		}
-		PackageBuffer = MakeSharedBufferFromArray(MoveTemp(TempBytes));
-	}
-
 	ICache& Cache = GetDerivedDataCacheRef();
 	FCacheRecordBuilder RecordBuilder = Cache.CreateRecord(GetEditorDomainPackageKey(PackageDigest));
-	TCbWriter<256> MetaData;
+
+	// We use a counter for PayloadIds rather than hashes of the Attachments. We do this because
+	// some attachments may be identical, and Attachments are not allowed to have identical PayloadIds.
+	// We need to keep the duplicate copies of identical payloads because BulkDatas were written into
+	// the exports with offsets that expect all attachment segments to exist in the segmented archive.
+	auto IntToPayloadId = [](int32 Value)
+	{
+		FPayloadId::ByteArray Bytes;
+		static_assert(sizeof(Bytes) >= sizeof(int32), "We are storing an int32 counter in the Bytes array");
+		FMemory::Memset(&Bytes, 0, sizeof(Bytes));
+		int32* IntView = reinterpret_cast<int32*>(&Bytes);
+		IntView[0] = Value;
+		return FPayloadId(Bytes);
+	};
+
+	int32 AttachmentIndex = 0;
+	const FSharedBuffer& ExportsBuffer = PackageStoreWriter->GetHeaderAndExports();
+	check(ExportsBuffer.GetSize() > 0); // Header+Exports segment is non-zero in length
+	RecordBuilder.AddAttachment(ExportsBuffer, IntToPayloadId(AttachmentIndex++));
+	uint64 FileSize = ExportsBuffer.GetSize();
+	for (const FSharedBuffer& BulkBuffer : PackageStoreWriter->GetBulkDatas())
+	{
+		int64 BulkSize = BulkBuffer.GetSize();
+		check(BulkSize > 0); // We checked this before adding the Region to the PackageStoreWriter
+		RecordBuilder.AddAttachment(BulkBuffer, IntToPayloadId(AttachmentIndex++));
+		FileSize += BulkSize;
+	}
+
+	TCbWriter<16> MetaData;
 	MetaData.BeginObject();
-	MetaData << "FileSize" << PackageBuffer.GetSize();
+	MetaData << "FileSize" << FileSize;
 	MetaData.EndObject();
+
 	RecordBuilder.SetMeta(MetaData.Save().AsObject());
-	RecordBuilder.SetValue(PackageBuffer);
 	FCacheRecord Record = RecordBuilder.Build();
 	Cache.Put(MakeArrayView(&Record, 1), Package->GetName(), ECachePolicy::Local);
 	return true;
diff --git a/Engine/Source/Editor/UnrealEd/Private/EditorDomain/EditorDomainUtils.h b/Engine/Source/Editor/UnrealEd/Private/EditorDomain/EditorDomainUtils.h
--- a/Engine/Source/Editor/UnrealEd/Private/EditorDomain/EditorDomainUtils.h
+++ b/Engine/Source/Editor/UnrealEd/Private/EditorDomain/EditorDomainUtils.h
@@ -19,6 +19,8 @@ class FSharedBuffer;
 class IAssetRegistry;
 class UPackage;
 
+namespace UE::DerivedData { struct FCacheKey; }
+
 namespace UE::EditorDomain
 {
 
@@ -61,6 +63,7 @@ void PrecacheClassDigests(TConstArrayView<FName> ClassNames);
 UE::DerivedData::FRequest RequestEditorDomainPackage(const FPackagePath& PackagePath,
 	const FPackageDigest& PackageDigest, UE::DerivedData::ECachePolicy SkipFlags,
 	UE::DerivedData::EPriority CachePriority, UE::DerivedData::FOnCacheGetComplete&& Callback);
+UE::DerivedData::FCacheKey GetEditorDomainPackageKey(const FPackageDigest& PackageDigest);
 
 /** Save the given package into the EditorDomain. */
 bool TrySavePackage(UPackage* Package);
