CookOnTheFlyServer: Fix long package name warning message to include the failure reason.
#rnx

diff --git a/Engine/Source/Editor/UnrealEd/Private/CookOnTheFlyServer.cpp b/Engine/Source/Editor/UnrealEd/Private/CookOnTheFlyServer.cpp
--- a/Engine/Source/Editor/UnrealEd/Private/CookOnTheFlyServer.cpp
+++ b/Engine/Source/Editor/UnrealEd/Private/CookOnTheFlyServer.cpp
@@ -6576,7 +6576,7 @@ void UCookOnTheFlyServer::GenerateLongPackageNames(TArray<FName>& FilesInPath)
 			}
 			else
 			{
-				LogCookerMessage(FString::Printf(TEXT("Unable to generate long package name, %s"), *FileInPath,
+				LogCookerMessage(FString::Printf(TEXT("Unable to generate long package name, %s. %s"), *FileInPath,
 					*FPackageName::FormatErrorAsString(FileInPath, FailureReason)), EMessageSeverity::Warning);
 			}
 		}
