EditorDomain: 1) Fix incorrect format of EditorDomain records; the first record was getting a payloadid of 0 which was interpreted as empty payloadid, and so was replacing it with the hash.
2) Fix deadlock in OnPostEngineInit. Package->IsFullyLoaded can call FEditorDomain::FileSize, so it must be called outside the lock.
3) Sort the list of attachments in EditorDomain segment reads by PayloadId, in case something changes in the future and a backend gives us arbitrarily-ordered attachments rather than the current implementation that has them sorted by payloadid.
4) Include bEditorDomainSaveUnversioned in the PackageDigest key, since different values result for the package payload depending on that config variable.
#rnx

diff --git a/Engine/Source/Editor/UnrealEd/Private/EditorDomain/EditorDomain.cpp b/Engine/Source/Editor/UnrealEd/Private/EditorDomain/EditorDomain.cpp
--- a/Engine/Source/Editor/UnrealEd/Private/EditorDomain/EditorDomain.cpp
+++ b/Engine/Source/Editor/UnrealEd/Private/EditorDomain/EditorDomain.cpp
@@ -459,7 +459,6 @@ void FEditorDomain::OnEndLoadPackage(TConstArrayView<UPackage*> LoadedPackages)
 
 void FEditorDomain::OnPostEngineInit()
 {
-	TArray<UPackage*> PackagesToSave;
 	{
 		FScopeLock ScopeLock(&Locks->Lock);
 		bHasPassedPostEngineInit = true;
@@ -467,17 +466,22 @@ void FEditorDomain::OnPostEngineInit()
 		{
 			return;
 		}
+	}
 
-		FString PackageName;
-		for (TObjectIterator<UPackage> It; It; ++It)
+	TArray<UPackage*> PackagesToSave;
+	FString PackageName;
+	for (TObjectIterator<UPackage> It; It; ++It)
+	{
+		UPackage* Package = *It;
+		Package->GetName(PackageName);
+		if (Package->IsFullyLoaded() && !FPackageName::IsScriptPackage(PackageName))
 		{
-			UPackage* Package = *It;
-			Package->GetName(PackageName);
-			if (Package->IsFullyLoaded() && !FPackageName::IsScriptPackage(PackageName))
-			{
-				PackagesToSave.Add(Package);
-			}
+			PackagesToSave.Add(Package);
 		}
+	}
+
+	{
+		FScopeLock ScopeLock(&Locks->Lock);
 		FilterKeepPackagesToSave(PackagesToSave);
 	}
 
diff --git a/Engine/Source/Editor/UnrealEd/Private/EditorDomain/EditorDomainArchive.cpp b/Engine/Source/Editor/UnrealEd/Private/EditorDomain/EditorDomainArchive.cpp
--- a/Engine/Source/Editor/UnrealEd/Private/EditorDomain/EditorDomainArchive.cpp
+++ b/Engine/Source/Editor/UnrealEd/Private/EditorDomain/EditorDomainArchive.cpp
@@ -236,6 +236,8 @@ void FEditorDomainPackageSegments::OnRecordRequestComplete(UE::DerivedData::FCac
 	TUniqueFunction<void(bool bValid)>&& CreateSegmentData,
 	TUniqueFunction<bool(FEditorDomain& EditorDomain)>&& TryCreateFallbackData)
 {
+	using namespace UE::DerivedData;
+
 	{
 		ESource LocalAsyncSource = AsyncSource;
 		if (AsyncSource == ESource::Closed)
@@ -249,29 +251,37 @@ void FEditorDomainPackageSegments::OnRecordRequestComplete(UE::DerivedData::FCac
 
 	ESource NewAsyncSource = ESource::Uninitialized;
 	int32 InitialRequestIndex = -1;
-	if (Params.Status == UE::DerivedData::EStatus::Ok)
+	if (Params.Status == EStatus::Ok)
 	{
-		UE::DerivedData::FCacheRecord& Record = Params.Record;
+		FCacheRecord& Record = Params.Record;
 		uint64 FileSize = Record.GetMeta()["FileSize"].AsUInt64();
-		const UE::DerivedData::FPayload& ValuePayload = Record.GetValuePayload();
-
-		Segments.Reserve(1 + Record.GetAttachmentPayloads().Num());
-		uint64 CompositeSize = 0;
-		uint64 ValueSize = ValuePayload.GetRawSize();
-		if (ValueSize > 0)
+		const FPayload& ValuePayload = Record.GetValuePayload();
+
+		// Sort the attachment payloads by PayloadId, in case they are not already sorted. EditorDomainPackage attachments
+		// are written with PayloadIds in the same order as the segment order
+		TConstArrayView<FPayload> AttachmentPayloads = Record.GetAttachmentPayloads();
+		TArray<const FPayload*, TInlineAllocator<4>> SortedPayloads;
+		SortedPayloads.Reserve(AttachmentPayloads.Num()+1);
+		if (ValuePayload.GetRawSize() > 0)
 		{
-			Segments.Emplace(ValuePayload.GetId(), CompositeSize);
-			CompositeSize += ValueSize;
+			SortedPayloads.Add(&ValuePayload);
 		}
-		for (const UE::DerivedData::FPayload& Payload : Record.GetAttachmentPayloads())
+		for (const FPayload& Payload : AttachmentPayloads)
 		{
-			uint64 PayloadSize = Payload.GetRawSize();
-			if (PayloadSize > 0)
+			if (Payload.GetRawSize() > 0)
 			{
-				Segments.Emplace(Payload.GetId(), CompositeSize);
-				CompositeSize += PayloadSize;
+				SortedPayloads.Add(&Payload);
 			}
 		}
+		SortedPayloads.Sort([](const FPayload& A, const FPayload& B) { return A.GetId() < B.GetId(); });
+
+		uint64 CompositeSize = 0;
+		Segments.Reserve(SortedPayloads.Num());
+		for (const FPayload* Payload : SortedPayloads)
+		{
+			Segments.Emplace(Payload->GetId(), CompositeSize);
+			CompositeSize += Payload->GetRawSize();
+		}
 
 		if (CompositeSize != FileSize)
 		{
@@ -359,13 +369,15 @@ void FEditorDomainPackageSegments::OnRecordRequestComplete(UE::DerivedData::FCac
 
 void FEditorDomainPackageSegments::SendSegmentRequest(FSegment& Segment)
 {
+	using namespace UE::DerivedData;
+
 	// Called from Callback-only until AsyncSource is set, then from Interface-only
-	UE::DerivedData::ICache& Cache = GetDerivedDataCacheRef();
+	ICache& Cache = GetDerivedDataCacheRef();
 
 	// Note that Segment.Request is Interface-only and so we can write it outside the lock
-	Segment.Request = Cache.GetPayload({ UE::DerivedData::FCachePayloadKey{ UE::EditorDomain::GetEditorDomainPackageKey(PackageDigest), Segment.PayloadId} },
-		PackagePath.GetDebugName(), UE::DerivedData::ECachePolicy::Local, UE::DerivedData::EPriority::Normal,
-		[this, &Segment](UE::DerivedData::FCacheGetPayloadCompleteParams&& Params)
+	Segment.Request = Cache.GetPayload({ FCachePayloadKey{ UE::EditorDomain::GetEditorDomainPackageKey(PackageDigest), Segment.PayloadId} },
+		PackagePath.GetDebugName(), ECachePolicy::Local, EPriority::Normal,
+		[this, &Segment](FCacheGetPayloadCompleteParams&& Params)
 		{
 			if (AsyncSource == ESource::Closed)
 			{
@@ -377,7 +389,7 @@ void FEditorDomainPackageSegments::SendSegmentRequest(FSegment& Segment)
 				return;
 			}
 
-			if (Params.Status != UE::DerivedData::EStatus::Ok)
+			if (Params.Status != EStatus::Ok)
 			{
 				UE_LOG(LogEditorDomain, Error, TEXT("Package %s is missing cache data in segment %d."), *PackagePath.GetDebugName(), (int32)(&Segment - &Segments[0]));
 			}
diff --git a/Engine/Source/Editor/UnrealEd/Private/EditorDomain/EditorDomainUtils.cpp b/Engine/Source/Editor/UnrealEd/Private/EditorDomain/EditorDomainUtils.cpp
--- a/Engine/Source/Editor/UnrealEd/Private/EditorDomain/EditorDomainUtils.cpp
+++ b/Engine/Source/Editor/UnrealEd/Private/EditorDomain/EditorDomainUtils.cpp
@@ -37,18 +37,26 @@ FClassDigestMap& GetClassDigests()
 }
 
 // Change to a new guid when EditorDomain needs to be invalidated
-const TCHAR* EditorDomainVersion = TEXT("D2B58CFE87DA471C95A91F36D2EA9214");
+const TCHAR* EditorDomainVersion = TEXT("D1718C34CA7C47AEB87A1607568E25B0");
 // Identifier of the CacheBuckets for EditorDomain tables
 const TCHAR* EditorDomainPackageBucketName = TEXT("EditorDomainPackage");
 const TCHAR* EditorDomainBulkDataListBucketName = TEXT("EditorDomainBulkDataList");
 const TCHAR* EditorDomainBulkDataPayloadIdBucketName = TEXT("EditorDomainBulkDataPayloadId");
 
+static bool GetEditorDomainSaveUnversioned()
+{
+	bool bParsedValue;
+	static bool bEditorDomainSaveUnversioned = GConfig->GetBool(TEXT("CookSettings"), TEXT("EditorDomainSaveUnversioned"), bParsedValue, GEditorIni) ? bParsedValue : true;
+	return bEditorDomainSaveUnversioned;
+}
+
 EPackageDigestResult AppendPackageDigest(FCbWriter& Writer, FString& OutErrorMessage,
 	const FAssetPackageData& PackageData, FName PackageName)
 {
 	int32 CurrentFileVersionUE = GPackageFileUEVersion;
 	int32 CurrentFileVersionLicenseeUE = GPackageFileLicenseeUEVersion;
 	Writer << EditorDomainVersion;
+	Writer << GetEditorDomainSaveUnversioned();
 	PRAGMA_DISABLE_DEPRECATION_WARNINGS
 	Writer << const_cast<FGuid&>(PackageData.PackageGuid);
 	PRAGMA_ENABLE_DEPRECATION_WARNINGS
@@ -396,9 +404,7 @@ bool TrySavePackage(UPackage* Package)
 		// | SAVE_KeepGUID;			// Prevent indeterminism by keeping the Guid
 		;
 
-	bool bEditorDomainSaveUnversioned = false;
-	GConfig->GetBool(TEXT("CookSettings"), TEXT("EditorDomainSaveUnversioned"), bEditorDomainSaveUnversioned, GEditorIni);
-	if (bEditorDomainSaveUnversioned)
+	if (GetEditorDomainSaveUnversioned())
 	{
 		// With some exceptions, EditorDomain packages are saved unversioned; 
 		// editors request the appropriate version of the EditorDomain package matching their serialization version
@@ -448,7 +454,7 @@ bool TrySavePackage(UPackage* Package)
 		return FPayloadId(Bytes);
 	};
 
-	int32 AttachmentIndex = 0;
+	int32 AttachmentIndex = 1; // 0 is not a valid value for IntToPayloadId
 	const FSharedBuffer& ExportsBuffer = PackageStoreWriter->GetHeaderAndExports();
 	check(ExportsBuffer.GetSize() > 0); // Header+Exports segment is non-zero in length
 	RecordBuilder.AddAttachment(ExportsBuffer, IntToPayloadId(AttachmentIndex++));
