Fix NonUnity build errors

#rnx

diff --git a/Engine/Plugins/Animation/ControlRig/Source/ControlRigEditor/Private/Editor/DetailsViewWrapperObject.cpp b/Engine/Plugins/Animation/ControlRig/Source/ControlRigEditor/Private/Editor/DetailsViewWrapperObject.cpp
--- a/Engine/Plugins/Animation/ControlRig/Source/ControlRigEditor/Private/Editor/DetailsViewWrapperObject.cpp
+++ b/Engine/Plugins/Animation/ControlRig/Source/ControlRigEditor/Private/Editor/DetailsViewWrapperObject.cpp
@@ -3,6 +3,7 @@
 #include "DetailsViewWrapperObject.h"
 #include "Units/RigUnit.h"
 #include "ControlRigElementDetails.h"
+#include "Modules/ModuleManager.h"
 
 #if WITH_EDITOR
 #include "PropertyEditorModule.h"
diff --git a/Engine/Source/Editor/Kismet/Private/Debugging/KismetDebugCommands.h b/Engine/Source/Editor/Kismet/Private/Debugging/KismetDebugCommands.h
--- a/Engine/Source/Editor/Kismet/Private/Debugging/KismetDebugCommands.h
+++ b/Engine/Source/Editor/Kismet/Private/Debugging/KismetDebugCommands.h
@@ -3,6 +3,7 @@
 #pragma once
 
 #include "CoreMinimal.h"
+#include "UObject/SoftObjectPtr.h"
 
 //////////////////////////////////////////////////////////////////////////
 // FDebuggingActionCallbacks
diff --git a/Engine/Source/Editor/UMGEditor/Private/WidgetBlueprintEditorUtils.cpp b/Engine/Source/Editor/UMGEditor/Private/WidgetBlueprintEditorUtils.cpp
--- a/Engine/Source/Editor/UMGEditor/Private/WidgetBlueprintEditorUtils.cpp
+++ b/Engine/Source/Editor/UMGEditor/Private/WidgetBlueprintEditorUtils.cpp
@@ -9,6 +9,7 @@
 #include "Internationalization/TextPackageNamespaceUtil.h"
 #include "UObject/PropertyPortFlags.h"
 #include "Blueprint/WidgetTree.h"
+#include "Misc/ConfigCacheIni.h"
 #include "Modules/ModuleManager.h"
 #include "MovieScene.h"
 #include "WidgetBlueprint.h"
diff --git a/Engine/Source/Editor/UnrealEd/Private/Commandlets/GatherShaderStatsFromMaterialsCommandlet.cpp b/Engine/Source/Editor/UnrealEd/Private/Commandlets/GatherShaderStatsFromMaterialsCommandlet.cpp
--- a/Engine/Source/Editor/UnrealEd/Private/Commandlets/GatherShaderStatsFromMaterialsCommandlet.cpp
+++ b/Engine/Source/Editor/UnrealEd/Private/Commandlets/GatherShaderStatsFromMaterialsCommandlet.cpp
@@ -3,6 +3,9 @@
 #include "Commandlets/GatherShaderStatsFromMaterialsCommandlet.h"
 #include "AssetData.h"
 #include "AssetRegistryModule.h"
+#include "HAL/FileManager.h"
+#include "Interfaces/ITargetPlatform.h"
+#include "Interfaces/ItargetPlatformManagerModule.h"
 #include "Materials/Material.h"
 #include "Materials/MaterialInstance.h"
 #include "Materials/MaterialInstanceConstant.h"
diff --git a/Engine/Source/Runtime/Engine/Private/LightWeightInstanceSubsystem.cpp b/Engine/Source/Runtime/Engine/Private/LightWeightInstanceSubsystem.cpp
--- a/Engine/Source/Runtime/Engine/Private/LightWeightInstanceSubsystem.cpp
+++ b/Engine/Source/Runtime/Engine/Private/LightWeightInstanceSubsystem.cpp
@@ -1,6 +1,7 @@
 // Copyright Epic Games, Inc. All Rights Reserved.
 
 #include "GameFramework/LightWeightInstanceSubsystem.h"
+#include "Engine/Engine.h"
 #include "Engine/World.h"
 #include "UObject/UObjectIterator.h"
 #include "WorldPartition/DataLayer/DataLayer.h"
