Fix Linux build

#rnx

diff --git a/Engine/Source/Editor/UnrealEd/Private/Commandlets/GatherShaderStatsFromMaterialsCommandlet.cpp b/Engine/Source/Editor/UnrealEd/Private/Commandlets/GatherShaderStatsFromMaterialsCommandlet.cpp
--- a/Engine/Source/Editor/UnrealEd/Private/Commandlets/GatherShaderStatsFromMaterialsCommandlet.cpp
+++ b/Engine/Source/Editor/UnrealEd/Private/Commandlets/GatherShaderStatsFromMaterialsCommandlet.cpp
@@ -5,7 +5,7 @@
 #include "AssetRegistryModule.h"
 #include "HAL/FileManager.h"
 #include "Interfaces/ITargetPlatform.h"
-#include "Interfaces/ItargetPlatformManagerModule.h"
+#include "Interfaces/ITargetPlatformManagerModule.h"
 #include "Materials/Material.h"
 #include "Materials/MaterialInstance.h"
 #include "Materials/MaterialInstanceConstant.h"
