U5 - Do not stop the RenderingThread as with modern Processors the performance saved is marginal and we do want to move away from stopping and starting the Rendering thread as it causes a multitude of issues.

diff --git a/Engine/Source/Editor/UnrealEd/Private/EditorBuildUtils.cpp b/Engine/Source/Editor/UnrealEd/Private/EditorBuildUtils.cpp
--- a/Engine/Source/Editor/UnrealEd/Private/EditorBuildUtils.cpp
+++ b/Engine/Source/Editor/UnrealEd/Private/EditorBuildUtils.cpp
@@ -252,9 +252,6 @@ bool FEditorBuildUtils::EditorBuild( UWorld* InWorld, FName Id, const bool bAllo
 	// Indicates whether the persistent level should be dirtied at the end of a build.
 	bool bDirtyPersistentLevel = true;
 
-	// Stop rendering thread so we're not wasting CPU cycles.
-	StopRenderingThread();
-
 	// Hack: These don't initialize properly and if you pick BuildAll right off the
 	// bat when opening a map you will get incorrect values in them.
 	GSwarmDebugOptions.Touch();
@@ -463,12 +460,6 @@ bool FEditorBuildUtils::EditorBuild( UWorld* InWorld, FName Id, const bool bAllo
 		GUnrealEd->Exec( InWorld, TEXT("MAP CHECK DONTDISPLAYDIALOG") );
 	}
 
-	// Re-start the rendering thread after build operations completed.
-	if (GUseThreadedRendering)
-	{
-		StartRenderingThread();
-	}	
-
 	if ( bDoBuild )
 	{
 		// Display elapsed build time.
