Replaced ProcessRootMotion option in the Animation Editor with multiple modes (Ignore, Loop, and LoopAndReset)
Ignore: Preview mesh will not consume root motion
Loop: Preview mesh will consume root motion continually (same behavior that we had)
LoopAndReset: Preview mesh will consume root motion resetting the position back to the origin every time the animation loops

diff --git a/Engine/Source/Editor/Persona/Private/AnimViewportShowCommands.cpp b/Engine/Source/Editor/Persona/Private/AnimViewportShowCommands.cpp
--- a/Engine/Source/Editor/Persona/Private/AnimViewportShowCommands.cpp
+++ b/Engine/Source/Editor/Persona/Private/AnimViewportShowCommands.cpp
@@ -9,8 +9,11 @@ void FAnimViewportShowCommands::RegisterCommands()
 	UI_COMMAND( AutoAlignFloorToMesh, "Auto Align Floor to Mesh", "Auto align floor to mesh bounds", EUserInterfaceActionType::ToggleButton, FInputChord() );
 	UI_COMMAND( MuteAudio, "Mute Audio", "Mute audio from the preview", EUserInterfaceActionType::ToggleButton, FInputChord() );
 	UI_COMMAND( UseAudioAttenuation, "Use Audio Attenuation", "Use audio attenuation when playing back audio in the preview", EUserInterfaceActionType::ToggleButton, FInputChord() );
-	
-	UI_COMMAND(ProcessRootMotion, "Process Root Motion", "Move preview based on animation root motion", EUserInterfaceActionType::ToggleButton, FInputChord());
+
+	UI_COMMAND(DoNotProcessRootMotion, "Ignore", "Preview mesh will not consume root motion", EUserInterfaceActionType::RadioButton, FInputChord());
+	UI_COMMAND(ProcessRootMotionLoopAndReset, "Loop and Reset", "Preview mesh will consume root motion resetting the position back to the origin every time the animation loops", EUserInterfaceActionType::RadioButton, FInputChord());
+	UI_COMMAND(ProcessRootMotionLoop, "Loop", "Preview mesh will consume root motion continually", EUserInterfaceActionType::RadioButton, FInputChord());
+
 	UI_COMMAND(DisablePostProcessBlueprint, "Disable Post Process", "Disable the evaluation of post process animation blueprints on the preview instance.", EUserInterfaceActionType::ToggleButton, FInputChord());
 
 	UI_COMMAND( ShowRetargetBasePose, "Retarget Base Pose", "Show retarget Base pose on preview mesh", EUserInterfaceActionType::ToggleButton, FInputChord() );
diff --git a/Engine/Source/Editor/Persona/Private/AnimViewportShowCommands.h b/Engine/Source/Editor/Persona/Private/AnimViewportShowCommands.h
--- a/Engine/Source/Editor/Persona/Private/AnimViewportShowCommands.h
+++ b/Engine/Source/Editor/Persona/Private/AnimViewportShowCommands.h
@@ -33,8 +33,14 @@ public:
 	/** Option to use audio attenuation in the viewport */
 	TSharedPtr< FUICommandInfo > UseAudioAttenuation;
 
-	/** Option to show root motion in viewport */
-	TSharedPtr< FUICommandInfo > ProcessRootMotion;
+	/** Option to set ProcessRootMotionMode to Ignore (Preview mesh will not consume root motion) */
+	TSharedPtr< FUICommandInfo > DoNotProcessRootMotion;
+
+	/** Option to set ProcessRootMotionMode to LoopAndReset (Preview mesh will consume root motion resetting the position back to the origin every time the animation loops) */
+	TSharedPtr< FUICommandInfo > ProcessRootMotionLoopAndReset;
+
+	/** Option to set ProcessRootMotionMode to Loop (Preview mesh will consume root motion continually) */
+	TSharedPtr< FUICommandInfo > ProcessRootMotionLoop;
 
 	/** Option to enable/disable post process anim blueprint evaluation */
 	TSharedPtr< FUICommandInfo > DisablePostProcessBlueprint;
diff --git a/Engine/Source/Editor/Persona/Private/SAnimViewportToolBar.cpp b/Engine/Source/Editor/Persona/Private/SAnimViewportToolBar.cpp
--- a/Engine/Source/Editor/Persona/Private/SAnimViewportToolBar.cpp
+++ b/Engine/Source/Editor/Persona/Private/SAnimViewportToolBar.cpp
@@ -638,7 +638,9 @@ TSharedRef<SWidget> SAnimViewportToolBar::GenerateCharacterMenu() const
 				{
 					SubMenuBuilder.BeginSection("AnimViewportRootMotion", LOCTEXT("CharacterMenu_RootMotionLabel", "Root Motion"));
 					{
-						SubMenuBuilder.AddMenuEntry(FAnimViewportShowCommands::Get().ProcessRootMotion);
+						SubMenuBuilder.AddMenuEntry(FAnimViewportShowCommands::Get().DoNotProcessRootMotion);
+						SubMenuBuilder.AddMenuEntry(FAnimViewportShowCommands::Get().ProcessRootMotionLoop);
+						SubMenuBuilder.AddMenuEntry(FAnimViewportShowCommands::Get().ProcessRootMotionLoopAndReset);
 					}
 					SubMenuBuilder.EndSection();
 
diff --git a/Engine/Source/Editor/Persona/Private/SAnimationEditorViewport.cpp b/Engine/Source/Editor/Persona/Private/SAnimationEditorViewport.cpp
--- a/Engine/Source/Editor/Persona/Private/SAnimationEditorViewport.cpp
+++ b/Engine/Source/Editor/Persona/Private/SAnimationEditorViewport.cpp
@@ -1007,11 +1007,27 @@ void SAnimationEditorViewportTabBody::BindCommands()
 		FCanExecuteAction(),
 		FIsActionChecked::CreateSP(this, &SAnimationEditorViewportTabBody::IsAudioAttenuationEnabled));
 
+	CommandList.BeginGroup(TEXT("RootMotion"));
+
 	CommandList.MapAction(
-		ViewportShowMenuCommands.ProcessRootMotion,
-		FExecuteAction::CreateSP(this, &SAnimationEditorViewportTabBody::OnTogglePreviewRootMotion),
-		FCanExecuteAction(),
-		FIsActionChecked::CreateSP(this, &SAnimationEditorViewportTabBody::IsPreviewingRootMotion));
+		ViewportShowMenuCommands.DoNotProcessRootMotion,
+		FExecuteAction::CreateSP(this, &SAnimationEditorViewportTabBody::SetProcessRootMotionMode, EProcessRootMotionMode::Ignore),
+		FIsActionChecked::CreateSP(this, &SAnimationEditorViewportTabBody::CanUseProcessRootMotionMode, EProcessRootMotionMode::Ignore),
+		FIsActionChecked::CreateSP(this, &SAnimationEditorViewportTabBody::IsProcessRootMotionModeSet, EProcessRootMotionMode::Ignore));
+
+	CommandList.MapAction(
+		ViewportShowMenuCommands.ProcessRootMotionLoopAndReset,
+		FExecuteAction::CreateSP(this, &SAnimationEditorViewportTabBody::SetProcessRootMotionMode, EProcessRootMotionMode::LoopAndReset),
+		FIsActionChecked::CreateSP(this, &SAnimationEditorViewportTabBody::CanUseProcessRootMotionMode, EProcessRootMotionMode::LoopAndReset),
+		FIsActionChecked::CreateSP(this, &SAnimationEditorViewportTabBody::IsProcessRootMotionModeSet, EProcessRootMotionMode::LoopAndReset));
+
+	CommandList.MapAction(
+		ViewportShowMenuCommands.ProcessRootMotionLoop,
+		FExecuteAction::CreateSP(this, &SAnimationEditorViewportTabBody::SetProcessRootMotionMode, EProcessRootMotionMode::Loop),
+		FIsActionChecked::CreateSP(this, &SAnimationEditorViewportTabBody::CanUseProcessRootMotionMode, EProcessRootMotionMode::Loop),
+		FIsActionChecked::CreateSP(this, &SAnimationEditorViewportTabBody::IsProcessRootMotionModeSet, EProcessRootMotionMode::Loop));
+
+	CommandList.EndGroup();
 
 	CommandList.MapAction(
 		ViewportShowMenuCommands.DisablePostProcessBlueprint,
@@ -1899,24 +1915,27 @@ bool SAnimationEditorViewportTabBody::IsAudioAttenuationEnabled() const
 	return GetAnimationViewportClient()->IsUsingAudioAttenuation();
 }
 
-void SAnimationEditorViewportTabBody::OnTogglePreviewRootMotion()
+void SAnimationEditorViewportTabBody::SetProcessRootMotionMode(EProcessRootMotionMode Mode)
 {
-	UDebugSkelMeshComponent* PreviewComponent = GetPreviewScene()->GetPreviewMeshComponent();
-
-	if (PreviewComponent)
+	if (UDebugSkelMeshComponent* PreviewComponent = GetPreviewScene()->GetPreviewMeshComponent())
 	{
-		PreviewComponent->SetPreviewRootMotion(!PreviewComponent->GetPreviewRootMotion());
+		PreviewComponent->SetProcessRootMotionMode(Mode);
 	}
 }
 
-bool SAnimationEditorViewportTabBody::IsPreviewingRootMotion() const
+bool SAnimationEditorViewportTabBody::IsProcessRootMotionModeSet(EProcessRootMotionMode Mode) const
 {
-	UDebugSkelMeshComponent* PreviewComponent = GetPreviewScene()->GetPreviewMeshComponent();
+	const UDebugSkelMeshComponent* PreviewComponent = GetPreviewScene()->GetPreviewMeshComponent();
+	return PreviewComponent ? (PreviewComponent->GetProcessRootMotionMode() == Mode) : false;
+}
 
-	if (PreviewComponent)
+bool SAnimationEditorViewportTabBody::CanUseProcessRootMotionMode(EProcessRootMotionMode Mode) const
+{
+	if(const UDebugSkelMeshComponent* PreviewComponent = GetPreviewScene()->GetPreviewMeshComponent())
 	{
-		return PreviewComponent->GetPreviewRootMotion();
+		return PreviewComponent->CanUseProcessRootMotionMode(Mode);
 	}
+
 	return false;
 }
 
diff --git a/Engine/Source/Editor/Persona/Private/SAnimationEditorViewport.h b/Engine/Source/Editor/Persona/Private/SAnimationEditorViewport.h
--- a/Engine/Source/Editor/Persona/Private/SAnimationEditorViewport.h
+++ b/Engine/Source/Editor/Persona/Private/SAnimationEditorViewport.h
@@ -478,11 +478,15 @@ private:
 	/** Whether audio from the viewport is attenuated */
 	bool IsAudioAttenuationEnabled() const;
 
-	/** Function to set whether we are previewing root motion */
-	void OnTogglePreviewRootMotion();
-	
-	/** Whether or not we are previewing root motion */
-	bool IsPreviewingRootMotion() const;
+	/** Sets process root motion mode on the debug mesh */
+	void SetProcessRootMotionMode(EProcessRootMotionMode Mode);
+
+	/** Checks whether the supplied mode is set on the debug mesh */
+	bool IsProcessRootMotionModeSet(EProcessRootMotionMode Mode) const;
+
+	/** Whether the supplied mode can be used */
+	bool CanUseProcessRootMotionMode(EProcessRootMotionMode Mode) const;
+
 private:
 	/** Selected Turn Table speed  */
 	EAnimationPlaybackSpeeds::Type SelectedTurnTableSpeed;
diff --git a/Engine/Source/Editor/SkeletonEditor/Private/BoneProxy.cpp b/Engine/Source/Editor/SkeletonEditor/Private/BoneProxy.cpp
--- a/Engine/Source/Editor/SkeletonEditor/Private/BoneProxy.cpp
+++ b/Engine/Source/Editor/SkeletonEditor/Private/BoneProxy.cpp
@@ -59,6 +59,15 @@ void UBoneProxy::Tick(float DeltaTime)
 				ReferenceRotation = ReferenceTransform.GetRotation().Rotator();
 				ReferenceScale = ReferenceTransform.GetScale3D();
 			}
+
+			// Show mesh relative transform on the details panel so we have a way to visualize the root transform when processing root motion
+			// Note that this doesn't always represent the actual transform of the root in the animation at current time but where root motion has taken us so far
+			// It will not match the root transform at the current time in the animation after lopping multiple times if we are using ProcessRootMotion::Loop
+			// or if we are visualizing a complex section from a montage, for example
+			const FTransform MeshRelativeTransform = Component->GetRelativeTransform();
+			MeshLocation = MeshRelativeTransform.GetLocation();
+			MeshRotation = MeshRelativeTransform.GetRotation().Rotator();
+			MeshScale = MeshRelativeTransform.GetScale3D();
 		}
 	}
 }
diff --git a/Engine/Source/Editor/SkeletonEditor/Public/BoneProxy.h b/Engine/Source/Editor/SkeletonEditor/Public/BoneProxy.h
--- a/Engine/Source/Editor/SkeletonEditor/Public/BoneProxy.h
+++ b/Engine/Source/Editor/SkeletonEditor/Public/BoneProxy.h
@@ -58,6 +58,18 @@ public:
 	UPROPERTY(VisibleAnywhere, Category = "Reference Transform")
 	FVector ReferenceScale;
 
+	/** Mesh Location. Non zero when processing root motion */
+	UPROPERTY(VisibleAnywhere, Category = "Mesh Relative Transform")
+	FVector MeshLocation;
+
+	/** Mesh Rotation. Non zero when processing root motion */
+	UPROPERTY(VisibleAnywhere, Category = "Mesh Relative Transform")
+	FRotator MeshRotation;
+
+	/** Mesh Scale */
+	UPROPERTY(VisibleAnywhere, Category = "Mesh Relative Transform")
+	FVector MeshScale;
+
 	/** The skeletal mesh component we glean our transform data from */
 	UPROPERTY()
 	TWeakObjectPtr<UDebugSkelMeshComponent> SkelMeshComponent;
diff --git a/Engine/Source/Editor/UnrealEd/Classes/Animation/DebugSkelMeshComponent.h b/Engine/Source/Editor/UnrealEd/Classes/Animation/DebugSkelMeshComponent.h
--- a/Engine/Source/Editor/UnrealEd/Classes/Animation/DebugSkelMeshComponent.h
+++ b/Engine/Source/Editor/UnrealEd/Classes/Animation/DebugSkelMeshComponent.h
@@ -64,6 +64,20 @@ namespace EPersonaTurnTableMode
 	};
 };
 
+/** Different modes for when processing root motion */
+UENUM()
+enum class EProcessRootMotionMode : uint8
+{
+	/** Preview mesh will not consume root motion */
+	Ignore,
+
+	/** Preview mesh will consume root motion continually */
+	Loop,
+
+	/** Preview mesh will consume root motion resetting the position back to the origin every time the animation loops */
+	LoopAndReset
+};
+
 //////////////////////////////////////////////////////////////////////////
 // FDebugSkelMeshSceneProxy
 
@@ -213,8 +227,17 @@ class UNREALED_API UDebugSkelMeshComponent : public USkeletalMeshComponent
 	UPROPERTY(transient)
 	bool bDisplayVertexColors;
 
+	UE_DEPRECATED(5.0, "This variable is no longer used. Use ProcessRootMotionMode instead.")
+	UPROPERTY()
+	uint32 bPreviewRootMotion_DEPRECATED : 1;
+
+	/** Process root motion mode */
+	UPROPERTY(transient)
+	EProcessRootMotionMode ProcessRootMotionMode;
+
+	/** Playback time last time ConsumeRootmotion was called */
 	UPROPERTY(transient)
-	uint32 bPreviewRootMotion:1;
+	float ConsumeRootMotionPreviousPlaybackTime;
 
 	UPROPERTY(transient)
 	uint32 bShowClothData : 1;
@@ -389,10 +412,27 @@ class UNREALED_API UDebugSkelMeshComponent : public USkeletalMeshComponent
 	void SetShowClothProperty(bool bState);
 
 	/** Get whether we should be previewing root motion */
-	bool GetPreviewRootMotion() const;
+	UE_DEPRECATED(5.0, "Please use IsProcessingRootMotion or GetProcessRootMotionMode")
+	bool GetPreviewRootMotion() const { return IsProcessingRootMotion(); }
 
 	/** Set whether we should be previewing root motion. Note: disabling root motion preview resets transform. */
-	void SetPreviewRootMotion(bool bInPreviewRootMotion);
+	UE_DEPRECATED(5.0, "Please use SetProcessRootMotionMode")
+	void SetPreviewRootMotion(bool bInPreviewRootMotion) { SetProcessRootMotionMode(bInPreviewRootMotion ? EProcessRootMotionMode::Loop : EProcessRootMotionMode::Ignore); }
+
+	/** Whether we are processing root motion or not */
+	bool IsProcessingRootMotion() const;
+
+	/** Gets process root motion mode */
+	EProcessRootMotionMode GetProcessRootMotionMode() const;
+
+	/** Sets process root motion mode. Note: disabling root motion preview resets transform. */
+	void SetProcessRootMotionMode(EProcessRootMotionMode Mode);
+
+	/** Whether the supplied root motion mode can be used for the current asset */
+	bool CanUseProcessRootMotionMode(EProcessRootMotionMode Mode) const;
+
+	/** Whether the current asset is using root motion */
+	bool DoesCurrentAssetHaveRootMotion() const;
 
 	/** Whether the current LOD of the debug mesh is being synced with the attached (preview) mesh instance. */
 	bool IsTrackingAttachedLOD() const;
@@ -509,7 +549,7 @@ public:
 	EPersonaTurnTableMode::Type TurnTableMode;
 	/** Current turn table speed scaling */
 	float TurnTableSpeedScaling;
-	
+
 	virtual void TickComponent(float DeltaTime, enum ELevelTick TickType, FActorComponentTickFunction *ThisTickFunction) override;
 
 	void RefreshSelectedClothingSkinnedPositions();
diff --git a/Engine/Source/Editor/UnrealEd/Private/Animation/DebugSkelMeshComponent.cpp b/Engine/Source/Editor/UnrealEd/Private/Animation/DebugSkelMeshComponent.cpp
--- a/Engine/Source/Editor/UnrealEd/Private/Animation/DebugSkelMeshComponent.cpp
+++ b/Engine/Source/Editor/UnrealEd/Private/Animation/DebugSkelMeshComponent.cpp
@@ -48,6 +48,9 @@ UDebugSkelMeshComponent::UDebugSkelMeshComponent(const FObjectInitializer& Objec
 	BoneRadiusMultiplier = 1.f;
 	
 	CachedClothBounds = FBoxSphereBounds(ForceInit);
+
+	ProcessRootMotionMode = EProcessRootMotionMode::Ignore;
+	ConsumeRootMotionPreviousPlaybackTime = 0.f;
 }
 
 void UDebugSkelMeshComponent::SetDebugForcedLOD(int32 InNewForcedLOD)
@@ -168,45 +171,224 @@ float WrapInRange(float StartVal, float MinVal, float MaxVal)
 
 void UDebugSkelMeshComponent::ConsumeRootMotion(const FVector& FloorMin, const FVector& FloorMax)
 {
-	//Extract root motion regardless of where we use it so that we don't hit
-	//problems with it building up in the instance
+	if (PreviewInstance == nullptr)
+	{
+		return;
+	}
+
+	// Helper to extract root motion manually from the current asset
+	auto ExtractRootMotionFromAnimationAsset = [](const UAnimationAsset* Animation, float StartPosition, float EndPosition) -> FTransform
+	{
+		if (const UAnimMontage* Anim = Cast<UAnimMontage>(Animation))
+		{
+			return Anim->ExtractRootMotionFromTrackRange(StartPosition, EndPosition);
+		}
+		
+		if (const UAnimComposite* Anim = Cast<UAnimComposite>(Animation))
+		{
+			FRootMotionMovementParams RootMotion;
+			Anim->ExtractRootMotionFromTrack(Anim->AnimationTrack, StartPosition, EndPosition, RootMotion);
+			return RootMotion.GetRootMotionTransform();
+		}
+		
+		if (const UAnimSequence* Anim = Cast<UAnimSequence>(Animation))
+		{
+			return Anim->ExtractRootMotionFromRange(StartPosition, EndPosition);
+		}
 
+		return FTransform::Identity;
+	};
+
+	// Force ProcessRootMotionMode to Ignore if the current asset is not using root motion. 
+	if(ProcessRootMotionMode != EProcessRootMotionMode::Ignore && DoesCurrentAssetHaveRootMotion() == false)
+	{
+		SetProcessRootMotionMode(EProcessRootMotionMode::Ignore);
+	}
+
+	//Extract root motion regardless of where we use it so that we don't hit problems with it building up in the instance
 	FRootMotionMovementParams ExtractedRootMotion = ConsumeRootMotion_Internal(1.0f);
+	
+	const float CurrentTime = PreviewInstance->GetCurrentTime();
+	const float PreviousTime = ConsumeRootMotionPreviousPlaybackTime;
+	ConsumeRootMotionPreviousPlaybackTime = CurrentTime;
 
-	if (bPreviewRootMotion)
+	// Ignore Mode: Preview mesh remains in place at the origin regardless of the root motion
+	if (ProcessRootMotionMode == EProcessRootMotionMode::Ignore)
+	{
+		if(!GetRelativeTransform().Equals(FTransform::Identity))
+		{
+			SetRelativeTransform(FTransform::Identity);
+		}
+	}
+	else 
 	{
-		if (ExtractedRootMotion.bHasRootMotion)
+		if(PreviewInstance->IsPlaying())
+		{
+			// Loop Mode: Preview mesh will consume root motion continually
+			if (ProcessRootMotionMode == EProcessRootMotionMode::Loop)
+			{
+				AddLocalTransform(ExtractedRootMotion.GetRootMotionTransform());
+
+				//Handle moving component so that it stays within the editor floor
+				FTransform CurrentTransform = GetRelativeTransform();
+				FVector Trans = CurrentTransform.GetTranslation();
+				Trans.X = WrapInRange(Trans.X, FloorMin.X, FloorMax.X);
+				Trans.Y = WrapInRange(Trans.Y, FloorMin.Y, FloorMax.Y);
+				CurrentTransform.SetTranslation(Trans);
+				SetRelativeTransform(CurrentTransform);
+			}
+			// Loop and Reset Mode: Preview mesh will consume root motion resetting the position back to the origin every time the animation loops
+			else if (ProcessRootMotionMode == EProcessRootMotionMode::LoopAndReset)
+			{
+				// In this mode we have to deal with montage explicitly because we can have multiple sections and we want to reset the position when the section loops
+				// and depending on the composition, CurrentTime < PreviousTime (or CurrentTime > PreviousTime when playing in reverse) is not enough
+				if (const UAnimMontage* Montage = Cast<UAnimMontage>(PreviewInstance->CurrentAsset))
+				{
+					const int32 PreviewStartSectionIdx = Montage->CompositeSections.IsValidIndex(PreviewInstance->MontagePreviewStartSectionIdx) ? PreviewInstance->MontagePreviewStartSectionIdx : Montage->GetSectionIndexFromPosition(CurrentTime);
+					const int32 FirstSectionIdx = PreviewInstance->MontagePreview_FindFirstSectionAsInMontage(PreviewStartSectionIdx);
+					const int32 LastSectionIdx = PreviewInstance->MontagePreview_FindLastSection(FirstSectionIdx);
+
+					bool bLooped = false;
+
+					// If FirstSection == LastSection we are previewing a single section
+					// In this case to know if we have looped we just need to check if CurrentTime < PreviousTime (or the oposite if we are playing the montage in reverse)
+					if (FirstSectionIdx == LastSectionIdx)
+					{
+						bLooped = PreviewInstance->IsReverse() ? (CurrentTime > PreviousTime) : (CurrentTime < PreviousTime);
+					}
+					// Otherwise, we are previewing a montage with multiple section. In this case we check if section at CurrentTime is the FirstSection and the section at PreviewTime is the LastSection (or the opposite if we are playing the montage in reverse)
+					else
+					{
+						const int32 SectionIndexPrevTime = Montage->GetSectionIndexFromPosition(PreviousTime);
+						const int32 SectionIndexCurrentTime = Montage->GetSectionIndexFromPosition(CurrentTime);
+						bLooped = PreviewInstance->IsReverse() ? (SectionIndexPrevTime == FirstSectionIdx && SectionIndexCurrentTime == LastSectionIdx) : (SectionIndexPrevTime == LastSectionIdx && SectionIndexCurrentTime == FirstSectionIdx);
+					}
+
+					// If we have looped...
+					if (bLooped)
+					{
+						//... and we are playing in reverse, reset the position of the mesh back to the origin and apply root motion from the end of the LastSection to CurrentTime
+						if (PreviewInstance->IsReverse())
+						{
+							float StartTime, EndTime;
+							Montage->GetSectionStartAndEndTime(LastSectionIdx, StartTime, EndTime);
+							const FTransform RootMotionDelta = ExtractRootMotionFromAnimationAsset(PreviewInstance->CurrentAsset, EndTime, CurrentTime);
+							SetRelativeTransform(RootMotionDelta);
+						}
+						// Otherwise, reset the position of the mesh back to the origin and apply root motion from the beginning of the FirstSection to CurrentTime
+						else
+						{
+							float StartTime, EndTime;
+							Montage->GetSectionStartAndEndTime(FirstSectionIdx, StartTime, EndTime);
+							const FTransform RootMotionDelta = ExtractRootMotionFromAnimationAsset(PreviewInstance->CurrentAsset, StartTime, CurrentTime);
+							SetRelativeTransform(RootMotionDelta);
+						}
+
+					}
+					// we haven't looped, just keep applying the extracted root motion
+					else
+					{
+						AddLocalTransform(ExtractedRootMotion.GetRootMotionTransform());
+					}
+				}
+				else // CurrentAsset is not a Montage
+				{
+					bool bLooped = PreviewInstance->IsReverse() ? (CurrentTime > PreviousTime) : (CurrentTime < PreviousTime);
+					if (bLooped)
+					{
+						if (PreviewInstance->IsReverse())
+						{
+							const FTransform RootMotionDelta = ExtractRootMotionFromAnimationAsset(PreviewInstance->CurrentAsset, PreviewInstance->CurrentAsset->GetPlayLength(), CurrentTime);
+							SetRelativeTransform(RootMotionDelta);
+						}
+						else
+						{
+							const FTransform RootMotionDelta = ExtractRootMotionFromAnimationAsset(PreviewInstance->CurrentAsset, 0.f, CurrentTime);
+							SetRelativeTransform(RootMotionDelta);
+						}
+					}
+					else
+					{
+						AddLocalTransform(ExtractedRootMotion.GetRootMotionTransform());
+					}
+				}
+			}
+		}
+		else // Not Playing. When not playing user can still scrub the time line but animation is not ticking so we have to extract and apply root motion manually
 		{
-			AddLocalTransform(ExtractedRootMotion.GetRootMotionTransform());
-
-			//Handle moving component so that it stays within the editor floor
-			FTransform CurrentTransform = GetRelativeTransform();
-			FVector Trans = CurrentTransform.GetTranslation();
-			Trans.X = WrapInRange(Trans.X, FloorMin.X, FloorMax.X);
-			Trans.Y = WrapInRange(Trans.Y, FloorMin.Y, FloorMax.Y);
-			CurrentTransform.SetTranslation(Trans);
-			SetRelativeTransform(CurrentTransform);
+			const FTransform RootMotion = ExtractRootMotionFromAnimationAsset(PreviewInstance->CurrentAsset, PreviousTime, CurrentTime);
+			AddLocalTransform(RootMotion);
 		}
 	}
 }
 
-bool UDebugSkelMeshComponent::GetPreviewRootMotion() const
+bool UDebugSkelMeshComponent::IsProcessingRootMotion() const 
+{ 
+	return GetProcessRootMotionMode() != EProcessRootMotionMode::Ignore;
+}
+
+EProcessRootMotionMode UDebugSkelMeshComponent::GetProcessRootMotionMode() const
 {
-	return bPreviewRootMotion;
+	return ProcessRootMotionMode;
 }
 
-void UDebugSkelMeshComponent::SetPreviewRootMotion(bool bInPreviewRootMotion)
+bool UDebugSkelMeshComponent::DoesCurrentAssetHaveRootMotion() const
 {
-	bPreviewRootMotion = bInPreviewRootMotion;
-	if (!bPreviewRootMotion)
+	if(PreviewInstance)
 	{
-		if (TurnTableMode == EPersonaTurnTableMode::Stopped)
+		if(UAnimSequenceBase* AnimSequenceBase = Cast<UAnimSequenceBase>(PreviewInstance->GetCurrentAsset()))
 		{
-			SetWorldTransform(FTransform());
+			return AnimSequenceBase->HasRootMotion();
 		}
-		else
+	}
+
+	return false;
+}
+
+bool UDebugSkelMeshComponent::CanUseProcessRootMotionMode(EProcessRootMotionMode Mode) const
+{
+	if (PreviewInstance == nullptr)
+	{
+		return false;
+	}
+	
+	// Disable Loop modes if the current asset doesn't have root motion
+	if(Mode != EProcessRootMotionMode::Ignore)
+	{
+		if(!DoesCurrentAssetHaveRootMotion())
+		{
+			return false;
+		}
+	}
+
+	// Disable Loop and Reset mode for blend spaces
+	if (Mode == EProcessRootMotionMode::LoopAndReset)
+	{
+		if (UBlendSpace* BlendSpace = Cast<UBlendSpace>(PreviewInstance->GetCurrentAsset()))
 		{
-			SetRelativeLocation(FVector::ZeroVector);
+			return false;
+		}
+	}
+	
+	return true;
+}
+
+void UDebugSkelMeshComponent::SetProcessRootMotionMode(EProcessRootMotionMode Mode)
+{
+	if(CanUseProcessRootMotionMode(Mode))
+	{
+		ProcessRootMotionMode = Mode;
+
+		if (ProcessRootMotionMode == EProcessRootMotionMode::Ignore)
+		{
+			if (TurnTableMode == EPersonaTurnTableMode::Stopped)
+			{
+				SetWorldTransform(FTransform());
+			}
+			else
+			{
+				SetRelativeLocation(FVector::ZeroVector);
+			}
 		}
 	}
 }
@@ -357,7 +539,14 @@ void UDebugSkelMeshComponent::EnablePreview(bool bEnable, UAnimationAsset* Previ
 		    // restore previous state
 		    bDisableClothSimulation = bPrevDisableClothSimulation;
     
-			PreviewInstance->SetAnimationAsset(PreviewAsset);
+			PreviewInstance->SetAnimationAsset(PreviewAsset); 
+
+			// Reset mesh relative transform. It could be offset from the origin if we are processing root motion
+			if (!GetRelativeTransform().Equals(FTransform::Identity))
+			{
+				SetRelativeTransform(FTransform::Identity);
+			}
+
 		}
 		else if (IsPreviewOn())
 		{
