CookOnTheFlyServer: sort includes; sorted includes are now preferred for Core files.
#rnx

diff --git a/Engine/Source/Editor/UnrealEd/Classes/CookOnTheSide/CookOnTheFlyServer.h b/Engine/Source/Editor/UnrealEd/Classes/CookOnTheSide/CookOnTheFlyServer.h
--- a/Engine/Source/Editor/UnrealEd/Classes/CookOnTheSide/CookOnTheFlyServer.h
+++ b/Engine/Source/Editor/UnrealEd/Classes/CookOnTheSide/CookOnTheFlyServer.h
@@ -5,17 +5,18 @@
 #include "CoreMinimal.h"
 
 #include "Containers/ArrayView.h"
+#include "CookOnTheSide/CookLog.h"
+#include "INetworkFileSystemModule.h"
+#include "IPlatformFileSandboxWrapper.h"
 #include "Misc/EnumClassFlags.h"
+#include "Templates/UniquePtr.h"
+#include "TickableEditorObject.h"
+#include "UObject/Object.h"
 #include "UObject/ObjectMacros.h"
+#include "UObject/Package.h"
 #include "UObject/UObjectGlobals.h"
-#include "UObject/Object.h"
 #include "UObject/WeakObjectPtr.h"
-#include "UObject/Package.h"
-#include "Templates/UniquePtr.h"
-#include "TickableEditorObject.h"
-#include "IPlatformFileSandboxWrapper.h"
-#include "INetworkFileSystemModule.h"
-#include "CookOnTheSide/CookLog.h"
+
 #include "CookOnTheFlyServer.generated.h"
 
 class FAssetRegistryGenerator;
@@ -114,11 +115,11 @@ namespace UE::Cook
 	class FSaveCookedPackageContext;
 	class ICookOnTheFlyRequestManager;
 	struct FCookerTimer;
+	struct FGeneratorPackage;
 	struct FPackageData;
 	struct FPackageDatas;
 	struct FPackageTracker;
 	struct FPendingCookedPlatformData;
-	struct FGeneratorPackage;
 	struct FPlatformManager;
 	struct FPopulatePackageContext;
 	struct FTickStackData;
diff --git a/Engine/Source/Editor/UnrealEd/Private/CookOnTheFlyServer.cpp b/Engine/Source/Editor/UnrealEd/Private/CookOnTheFlyServer.cpp
--- a/Engine/Source/Editor/UnrealEd/Private/CookOnTheFlyServer.cpp
+++ b/Engine/Source/Editor/UnrealEd/Private/CookOnTheFlyServer.cpp
@@ -5,130 +5,118 @@
 =============================================================================*/
 
 #include "CookOnTheSide/CookOnTheFlyServer.h"
-#include "Cooker/PackageNameCache.h"
 
+#include "Algo/AnyOf.h"
+#include "Algo/Find.h"
 #include "AssetCompilingManager.h"
+#include "AssetRegistryModule.h"
+#include "AssetRegistryState.h"
+#include "Async/Async.h"
+#include "Async/ParallelFor.h"
+#include "Commandlets/AssetRegistryGenerator.h"
+#include "Commandlets/ShaderPipelineCacheToolsCommandlet.h"
+#include "Containers/RingBuffer.h"
+#include "Cooker/AsyncIODelete.h"
+#include "Cooker/CookOnTheFlyServerInterface.h"
 #include "Cooker/CookPackageData.h"
 #include "Cooker/CookPlatformManager.h"
 #include "Cooker/CookProfiling.h"
 #include "Cooker/CookRequests.h"
 #include "Cooker/CookTypes.h"
-#include "Cooker/PackageTracker.h"
-#include "Cooker/CookOnTheFlyServerInterface.h"
-#include "Cooker/NetworkFileCookOnTheFlyRequestManager.h"
 #include "Cooker/IoStoreCookOnTheFlyRequestManager.h"
+#include "Cooker/NetworkFileCookOnTheFlyRequestManager.h"
+#include "Cooker/PackageNameCache.h"
+#include "Cooker/PackageTracker.h"
+#include "CookerSettings.h"
 #include "CookPackageSplitter.h"
-#include "Containers/RingBuffer.h"
+#include "DerivedDataCacheInterface.h"
+#include "DistanceFieldAtlas.h"
+#include "Editor.h"
+#include "Editor/UnrealEdEngine.h"
 #include "EditorDomain/EditorDomain.h"
-#include "HAL/PlatformFileManager.h"
+#include "Engine/AssetManager.h"
+#include "Engine/Level.h"
+#include "Engine/LevelStreaming.h"
+#include "Engine/Texture.h"
+#include "Engine/TextureLODSettings.h"
+#include "Engine/WorldComposition.h"
+#include "EngineGlobals.h"
+#include "FilePackageStoreWriter.h"
+#include "FileServerMessages.h"
+#include "GameDelegates.h"
+#include "GlobalShader.h"
 #include "HAL/FileManager.h"
-#include "Misc/CommandLine.h"
-#include "Misc/FileHelper.h"
-#include "HAL/Runnable.h"
-#include "HAL/RunnableThread.h"
-#include "Stats/Stats.h"
-#include "Stats/StatsMisc.h"
-#include "Misc/ConfigCacheIni.h"
 #include "HAL/IConsoleManager.h"
+#include "HAL/PlatformFileManager.h"
 #include "HAL/PlatformProcess.h"
-#include "Misc/LocalTimestampDirectoryVisitor.h"
-#include "Serialization/CustomVersion.h"
+#include "HAL/Runnable.h"
+#include "HAL/RunnableThread.h"
+#include "IMessageContext.h"
+#include "INetworkFileServer.h"
+#include "INetworkFileSystemModule.h"
+#include "Interfaces/IAudioFormat.h"
+#include "Interfaces/IPluginManager.h"
+#include "Interfaces/IProjectManager.h"
+#include "Interfaces/IShaderFormat.h"
+#include "Interfaces/ITargetPlatform.h"
+#include "Interfaces/ITargetPlatformManagerModule.h"
+#include "Interfaces/ITextureFormat.h"
+#include "Internationalization/Culture.h"
+#include "IPAddress.h"
+#include "LocalizationChunkDataGenerator.h"
+#include "Logging/MessageLog.h"
+#include "Logging/TokenizedMessage.h"
+#include "Materials/Material.h"
+#include "Materials/MaterialInterface.h"
+#include "MeshCardRepresentation.h"
+#include "MessageEndpoint.h"
+#include "MessageEndpointBuilder.h"
 #include "Misc/App.h"
+#include "Misc/CommandLine.h"
+#include "Misc/ConfigCacheIni.h"
 #include "Misc/CoreDelegates.h"
+#include "Misc/FileHelper.h"
+#include "Misc/LocalTimestampDirectoryVisitor.h"
+#include "Misc/NetworkVersion.h"
+#include "Misc/PackageName.h"
+#include "Misc/PathViews.h"
+#include "Misc/RedirectCollector.h"
 #include "Misc/ScopeExit.h"
 #include "Misc/ScopeLock.h"
 #include "Modules/ModuleManager.h"
+#include "PackageHelperFunctions.h"
+#include "PlatformInfo.h"
+#include "ProfilingDebugging/CookStats.h"
+#include "ProfilingDebugging/PlatformFileTrace.h"
+#include "ProjectDescriptor.h"
+#include "SceneUtils.h"
+#include "Serialization/ArchiveStackTrace.h"
 #include "Serialization/ArchiveUObject.h"
-#include "UObject/GarbageCollection.h"
+#include "Serialization/ArrayReader.h"
+#include "Serialization/ArrayWriter.h"
+#include "Serialization/CustomVersion.h"
+#include "Settings/LevelEditorPlaySettings.h"
+#include "Settings/ProjectPackagingSettings.h"
+#include "ShaderCodeLibrary.h"
+#include "ShaderCompiler.h"
+#include "ShaderLibraryChunkDataGenerator.h"
+#include "Stats/Stats.h"
+#include "Stats/StatsMisc.h"
+#include "String/Find.h"
+#include "TargetDomain/TargetDomainUtils.h"
+#include "UnrealEdGlobals.h"
 #include "UObject/Class.h"
-#include "UObject/UObjectIterator.h"
+#include "UObject/ConstructorHelpers.h"
+#include "UObject/GarbageCollection.h"
 #include "UObject/LinkerDiff.h"
 #include "UObject/LinkerSave.h"
-#include "UObject/Package.h"
-#include "UObject/SavePackage.h"
 #include "UObject/MetaData.h"
-#include "UObject/ConstructorHelpers.h"
-#include "UObject/ReferenceChainSearch.h"
 #include "UObject/ObjectSaveContext.h"
+#include "UObject/Package.h"
+#include "UObject/ReferenceChainSearch.h"
+#include "UObject/SavePackage.h"
 #include "UObject/UObjectArray.h"
-#include "Misc/PackageName.h"
-#include "Misc/RedirectCollector.h"
-#include "Engine/Level.h"
-#include "Materials/MaterialInterface.h"
-#include "Materials/Material.h"
-#include "Settings/LevelEditorPlaySettings.h"
-#include "Engine/AssetManager.h"
-#include "Editor/UnrealEdEngine.h"
-#include "Engine/Texture.h"
-#include "SceneUtils.h"
-#include "Settings/ProjectPackagingSettings.h"
-#include "EngineGlobals.h"
-#include "Editor.h"
-#include "UnrealEdGlobals.h"
-#include "FileServerMessages.h"
-#include "LocalizationChunkDataGenerator.h"
-#include "Internationalization/Culture.h"
-#include "Serialization/ArrayReader.h"
-#include "Serialization/ArrayWriter.h"
-#include "PackageHelperFunctions.h"
-#include "DerivedDataCacheInterface.h"
-#include "GlobalShader.h"
-#include "Interfaces/ITargetPlatform.h"
-#include "Interfaces/ITargetPlatformManagerModule.h"
-#include "Interfaces/IAudioFormat.h"
-#include "Interfaces/IShaderFormat.h"
-#include "Interfaces/ITextureFormat.h"
-#include "IMessageContext.h"
-#include "MessageEndpoint.h"
-#include "MessageEndpointBuilder.h"
-#include "INetworkFileServer.h"
-#include "INetworkFileSystemModule.h"
-#include "PlatformInfo.h"
-#include "Serialization/ArchiveStackTrace.h"
-#include "DistanceFieldAtlas.h"
-#include "Cooker/AsyncIODelete.h"
-#include "MeshCardRepresentation.h"
-#include "Misc/PathViews.h"
-#include "String/Find.h"
-#include "TargetDomain/TargetDomainUtils.h"
-#include "Async/Async.h"
-
-#include "AssetRegistryModule.h"
-#include "AssetRegistryState.h"
-#include "CookerSettings.h"
-
-#include "GameDelegates.h"
-#include "IPAddress.h"
-
-#include "Interfaces/IPluginManager.h"
-#include "ProjectDescriptor.h"
-#include "Interfaces/IProjectManager.h"
-
-// cook by the book requirements
-#include "Commandlets/AssetRegistryGenerator.h"
-#include "Engine/WorldComposition.h"
-
-// error message log
-#include "Logging/TokenizedMessage.h"
-#include "Logging/MessageLog.h"
-
-// shader compiler processAsyncResults
-#include "ShaderCompiler.h"
-#include "ShaderCodeLibrary.h"
-#include "ShaderLibraryChunkDataGenerator.h"
-#include "Engine/LevelStreaming.h"
-#include "Engine/TextureLODSettings.h"
-#include "ProfilingDebugging/CookStats.h"
-#include "ProfilingDebugging/PlatformFileTrace.h"
-
-#include "Misc/NetworkVersion.h"
-
-#include "Algo/Find.h"
-#include "Algo/AnyOf.h"
-#include "Async/ParallelFor.h"
-
-#include "Commandlets/ShaderPipelineCacheToolsCommandlet.h"
-#include "FilePackageStoreWriter.h"
+#include "UObject/UObjectIterator.h"
 #include "ZenStoreWriter.h"
 
 #define LOCTEXT_NAMESPACE "Cooker"
