Remove extra call to SaveBreakpoints; it is already called by RemoveBreakpointsByPredicate.
#rnx

diff --git a/Engine/Source/Editor/UnrealEd/Private/Kismet2/KismetDebugUtilities.cpp b/Engine/Source/Editor/UnrealEd/Private/Kismet2/KismetDebugUtilities.cpp
--- a/Engine/Source/Editor/UnrealEd/Private/Kismet2/KismetDebugUtilities.cpp
+++ b/Engine/Source/Editor/UnrealEd/Private/Kismet2/KismetDebugUtilities.cpp
@@ -907,7 +907,6 @@ void FKismetDebugUtilities::RemoveBreakpointFromNode(const UEdGraphNode* OwnerNo
 			return Breakpoint.GetLocation() == OwnerNode;
 		}
 	);
-	SaveBreakpoints();
 #endif	//#if WITH_EDITORONLY_DATA
 }
 
