DDC: Split the new API from the old API

The new API uses exported functions and cannot be included with the old API without compile errors in existing code that has an include-only dependency on DDC.

#rnx

diff --git a/Engine/Source/Developer/DerivedDataCache/Private/DerivedDataCache.cpp b/Engine/Source/Developer/DerivedDataCache/Private/DerivedDataCache.cpp
--- a/Engine/Source/Developer/DerivedDataCache/Private/DerivedDataCache.cpp
+++ b/Engine/Source/Developer/DerivedDataCache/Private/DerivedDataCache.cpp
@@ -9,6 +9,7 @@
 #include "Async/TaskGraphInterfaces.h"
 #include "DDCCleanup.h"
 #include "DerivedDataBackendInterface.h"
+#include "DerivedDataCache.h"
 #include "DerivedDataCachePrivate.h"
 #include "DerivedDataCacheRecord.h"
 #include "DerivedDataCacheUsageStats.h"
@@ -163,7 +164,7 @@ namespace UE::DerivedData::Private
  * Implementation of the derived data cache
  * This API is fully threadsafe
 **/
-class FDerivedDataCache final : public FDerivedDataCacheInterface
+class FDerivedDataCache final : public FDerivedDataCacheInterface, public ICache
 {
 
 	/** 
@@ -735,9 +736,14 @@ public:
 	}
 };
 
-FDerivedDataCacheInterface* CreateCache()
+ICache* CreateCache(FDerivedDataCacheInterface** OutLegacyCache)
 {
-	return new FDerivedDataCache();
+	FDerivedDataCache* Cache = new FDerivedDataCache;
+	if (OutLegacyCache)
+	{
+		*OutLegacyCache = Cache;
+	}
+	return Cache;
 }
 
 } // UE::DerivedData::Private
diff --git a/Engine/Source/Developer/DerivedDataCache/Private/DerivedDataCacheModule.cpp b/Engine/Source/Developer/DerivedDataCache/Private/DerivedDataCacheModule.cpp
--- a/Engine/Source/Developer/DerivedDataCache/Private/DerivedDataCacheModule.cpp
+++ b/Engine/Source/Developer/DerivedDataCache/Private/DerivedDataCacheModule.cpp
@@ -2,6 +2,7 @@
 
 #include "DerivedDataCacheModule.h"
 
+#include "CoreGlobals.h"
 #include "DerivedDataBuild.h"
 #include "DerivedDataBuildPrivate.h"
 #include "DerivedDataCache.h"
@@ -13,8 +14,9 @@
 namespace UE::DerivedData::Private
 {
 
-static FDerivedDataCacheInterface* GDerivedDataCacheInstance;
-static IBuild* GDerivedDataBuildInstance;
+static FDerivedDataCacheInterface* GDerivedDataLegacyCache;
+static ICache* GDerivedDataCache;
+static IBuild* GDerivedDataBuild;
 
 PRAGMA_DISABLE_DEPRECATION_WARNINGS
 
@@ -27,33 +29,40 @@ public:
 	}
 
 	FDerivedDataCacheInterface* const* CreateOrGetCache() final
+	{
+		CreateCacheOnce();
+		check(GDerivedDataLegacyCache);
+		return &GDerivedDataLegacyCache;
+	}
+
+	void CreateCacheOnce()
 	{
 		FScopeLock Lock(&CreateLock);
-		if (!GDerivedDataCacheInstance)
+		if (!GDerivedDataCache)
 		{
-			GDerivedDataCacheInstance = CreateCache();
-			check(GDerivedDataCacheInstance);
+			GDerivedDataCache = CreateCache(&GDerivedDataLegacyCache);
+			check(GDerivedDataCache);
 		}
-		return &GDerivedDataCacheInstance;
 	}
 
-	IBuild* const* CreateOrGetBuild() final
+	void CreateBuildOnce()
 	{
+		CreateCacheOnce();
 		FScopeLock Lock(&CreateLock);
-		if (!GDerivedDataBuildInstance)
+		if (!GDerivedDataBuild)
 		{
-			GDerivedDataBuildInstance = CreateBuild(**CreateOrGetCache());
-			check(GDerivedDataBuildInstance);
+			GDerivedDataBuild = CreateBuild(*GDerivedDataCache);
+			check(GDerivedDataBuild);
 		}
-		return &GDerivedDataBuildInstance;
 	}
 
 	void ShutdownModule() final
 	{
-		delete GDerivedDataBuildInstance;
-		GDerivedDataBuildInstance = nullptr;
-		delete GDerivedDataCacheInstance;
-		GDerivedDataCacheInstance = nullptr;
+		delete GDerivedDataBuild;
+		GDerivedDataBuild = nullptr;
+		delete GDerivedDataCache;
+		GDerivedDataLegacyCache = nullptr;
+		GDerivedDataCache = nullptr;
 	}
 
 private:
@@ -62,6 +71,54 @@ private:
 
 PRAGMA_ENABLE_DEPRECATION_WARNINGS
 
+static FDerivedDataCacheModule* GetModule()
+{
+	if (!FPlatformProperties::RequiresCookedData())
+	{
+		check(IsInGameThread());
+		if (IDerivedDataCacheModule* Module = FModuleManager::LoadModulePtr<IDerivedDataCacheModule>("DerivedDataCache"))
+		{
+			return static_cast<FDerivedDataCacheModule*>(Module);
+		}
+	}
+	return nullptr;
+}
+
 } // UE::DerivedData::Private
 
+namespace UE::DerivedData
+{
+
+ICache& GetCache()
+{
+	if (ICache* Cache = Private::GDerivedDataCache)
+	{
+		return *Cache;
+	}
+	if (Private::FDerivedDataCacheModule* Module = Private::GetModule())
+	{
+		Module->CreateCacheOnce();
+	}
+	ICache* Cache = Private::GDerivedDataCache;
+	check(Cache);
+	return *Cache;
+}
+
+IBuild& GetBuild()
+{
+	if (IBuild* Build = Private::GDerivedDataBuild)
+	{
+		return *Build;
+	}
+	if (Private::FDerivedDataCacheModule* Module = Private::GetModule())
+	{
+		Module->CreateBuildOnce();
+	}
+	IBuild* Build = Private::GDerivedDataBuild;
+	check(Build);
+	return *Build;
+}
+
+} // UE::DerivedData
+
 IMPLEMENT_MODULE(UE::DerivedData::Private::FDerivedDataCacheModule, DerivedDataCache);
diff --git a/Engine/Source/Developer/DerivedDataCache/Private/DerivedDataCachePrivate.h b/Engine/Source/Developer/DerivedDataCache/Private/DerivedDataCachePrivate.h
--- a/Engine/Source/Developer/DerivedDataCache/Private/DerivedDataCachePrivate.h
+++ b/Engine/Source/Developer/DerivedDataCache/Private/DerivedDataCachePrivate.h
@@ -13,13 +13,14 @@ namespace UE::DerivedData { class FCacheBucket; }
 namespace UE::DerivedData { class FCacheRecord; }
 namespace UE::DerivedData { class FCacheRecordBuilder; }
 namespace UE::DerivedData { class FOptionalCacheRecord; }
+namespace UE::DerivedData { class ICache; }
 namespace UE::DerivedData { struct FCacheKey; }
 
 namespace UE::DerivedData::Private
 {
 
 // Implemented in DerivedDataCache.cpp
-FDerivedDataCacheInterface* CreateCache();
+ICache* CreateCache(FDerivedDataCacheInterface** OutLegacyCache);
 
 // Implemented in DerivedDataCacheKey.cpp
 FCacheBucket CreateCacheBucket(FStringView Name);
diff --git a/Engine/Source/Developer/DerivedDataCache/Public/DerivedDataBuild.h b/Engine/Source/Developer/DerivedDataCache/Public/DerivedDataBuild.h
--- a/Engine/Source/Developer/DerivedDataCache/Public/DerivedDataBuild.h
+++ b/Engine/Source/Developer/DerivedDataCache/Public/DerivedDataBuild.h
@@ -4,6 +4,8 @@
 
 #include "Containers/StringFwd.h"
 
+#define UE_API DERIVEDDATACACHE_API
+
 class FCbObject;
 struct FGuid;
 
@@ -145,4 +147,8 @@ public:
 	virtual IBuildWorkerRegistry& GetWorkerRegistry() const = 0;
 };
 
+UE_API IBuild& GetBuild();
+
 } // UE::DerivedData
+
+#undef UE_API
diff --git a/Engine/Source/Developer/DerivedDataCache/Public/DerivedDataCache.h b/Engine/Source/Developer/DerivedDataCache/Public/DerivedDataCache.h
--- a/Engine/Source/Developer/DerivedDataCache/Public/DerivedDataCache.h
+++ b/Engine/Source/Developer/DerivedDataCache/Public/DerivedDataCache.h
@@ -10,6 +10,8 @@
 
 class FCbPackage;
 
+#define UE_API DERIVEDDATACACHE_API
+
 namespace UE::DerivedData
 {
 
@@ -281,4 +283,10 @@ struct FCacheGetPayloadCompleteParams
 
 ///////////////////////////////////////////////////////////////////////////////////////////////////////////////////////
 
+UE_API ICache& GetCache();
+
+///////////////////////////////////////////////////////////////////////////////////////////////////////////////////////
+
 } // UE::DerivedData
+
+#undef UE_API
diff --git a/Engine/Source/Developer/DerivedDataCache/Public/DerivedDataCacheInterface.h b/Engine/Source/Developer/DerivedDataCache/Public/DerivedDataCacheInterface.h
--- a/Engine/Source/Developer/DerivedDataCache/Public/DerivedDataCacheInterface.h
+++ b/Engine/Source/Developer/DerivedDataCache/Public/DerivedDataCacheInterface.h
@@ -6,7 +6,6 @@
 #include "CoreMinimal.h"
 #include "Containers/BitArray.h"
 #include "Containers/StringView.h"
-#include "DerivedDataCache.h"
 #include "DerivedDataCacheModule.h"
 
 class FDerivedDataCacheUsageStats;
@@ -19,13 +18,11 @@ struct FDerivedDataCacheSummaryStats;
  * Interface for the derived data cache
  * This API is fully threadsafe (with the possible exception of the system interface: NotfiyBootComplete, etc).
  */
-class FDerivedDataCacheInterface : public UE::DerivedData::ICache
+class FDerivedDataCacheInterface
 {
 public:
 	virtual ~FDerivedDataCacheInterface() = default;
 
-	using UE::DerivedData::ICache::Put;
-
 	//--------------------
 	// High Level Interface
 	//--------------------
diff --git a/Engine/Source/Developer/DerivedDataCache/Public/DerivedDataCacheModule.h b/Engine/Source/Developer/DerivedDataCache/Public/DerivedDataCacheModule.h
--- a/Engine/Source/Developer/DerivedDataCache/Public/DerivedDataCacheModule.h
+++ b/Engine/Source/Developer/DerivedDataCache/Public/DerivedDataCacheModule.h
@@ -6,8 +6,6 @@
 
 class FDerivedDataCacheInterface;
 
-namespace UE::DerivedData { class IBuild; }
-
 /**
  * Module for the Derived Data Cache and Derived Data Build.
  */
@@ -26,13 +24,4 @@ public:
 	 * the destruction of the cache without polling this function or monitoring the module lifetime.
 	 */
 	virtual FDerivedDataCacheInterface* const* CreateOrGetCache() = 0;
-
-	/**
-	 * Returns the build system, which is created by the first call to this function.
-	 *
-	 * This always returns a pointer to a valid build, but that pointer becomes null when the module
-	 * shuts down and destroys the build. This extra level of indirection allows a caller to observe
-	 * the destruction of the build without polling this function or monitoring the module lifetime.
-	 */
-	virtual UE::DerivedData::IBuild* const* CreateOrGetBuild() = 0;
 };
diff --git a/Engine/Source/Developer/TextureBuild/TextureBuild.Build.cs b/Engine/Source/Developer/TextureBuild/TextureBuild.Build.cs
--- a/Engine/Source/Developer/TextureBuild/TextureBuild.Build.cs
+++ b/Engine/Source/Developer/TextureBuild/TextureBuild.Build.cs
@@ -9,17 +9,17 @@ public class TextureBuild : ModuleRules
 		PrivateIncludePathModuleNames.AddRange(new string[]
 		{
 			"CoreUObject",
-			"DerivedDataCache",
 			"Engine",
 		});
 
-		PrivateDependencyModuleNames.AddRange(
-			new string[] {
-				"Core",
-				"ImageCore",
-				"ImageWrapper",
-				"TextureCompressor",
-				"TextureFormat",
-			});
+		PrivateDependencyModuleNames.AddRange(new string[]
+		{
+			"Core",
+			"DerivedDataCache",
+			"ImageCore",
+			"ImageWrapper",
+			"TextureCompressor",
+			"TextureFormat",
+		});
 	}
 }
diff --git a/Engine/Source/Editor/UnrealEd/Private/EditorDomain/BulkDataRegistryEditorDomain.cpp b/Engine/Source/Editor/UnrealEd/Private/EditorDomain/BulkDataRegistryEditorDomain.cpp
--- a/Engine/Source/Editor/UnrealEd/Private/EditorDomain/BulkDataRegistryEditorDomain.cpp
+++ b/Engine/Source/Editor/UnrealEd/Private/EditorDomain/BulkDataRegistryEditorDomain.cpp
@@ -4,7 +4,7 @@
 
 #include "Serialization/BulkDataRegistry.h"
 #include "Compression/CompressedBuffer.h"
-#include "DerivedDataCacheInterface.h"
+#include "DerivedDataCache.h"
 #include "DerivedDataCacheRecord.h"
 #include "EditorBuildInputResolver.h"
 #include "EditorDomain/EditorDomainUtils.h"
@@ -484,7 +484,7 @@ void FBulkDataRegistryEditorDomain::ReadPayloadIdsFromCache(FName PackageName, T
 
 FPendingPackage::FPendingPackage(FName InPackageName, FBulkDataRegistryEditorDomain* InOwner)
 	: PackageName(InPackageName)
-	, BulkDataListCacheRequest(GetDerivedDataCacheRef().CreateGroup(UE::DerivedData::EPriority::Low))
+	, BulkDataListCacheRequest(UE::DerivedData::GetCache().CreateGroup(UE::DerivedData::EPriority::Low))
 	, Owner(InOwner)
 {
 	PendingOperations = Flag_EndLoad | Flag_BulkDataListResults;
@@ -732,7 +732,7 @@ void FUpdatePayloadWorker::DoWork()
 
 FPendingPayloadId::FPendingPayloadId(const FGuid& InBulkDataId)
 	: BulkDataId(InBulkDataId)
-	, Request(GetDerivedDataCacheRef().CreateGroup(UE::DerivedData::EPriority::Low))
+	, Request(UE::DerivedData::GetCache().CreateGroup(UE::DerivedData::EPriority::Low))
 {
 }
 
diff --git a/Engine/Source/Editor/UnrealEd/Private/EditorDomain/EditorDomain.cpp b/Engine/Source/Editor/UnrealEd/Private/EditorDomain/EditorDomain.cpp
--- a/Engine/Source/Editor/UnrealEd/Private/EditorDomain/EditorDomain.cpp
+++ b/Engine/Source/Editor/UnrealEd/Private/EditorDomain/EditorDomain.cpp
@@ -7,7 +7,6 @@
 #include "Async/AsyncFileHandleNull.h"
 #include "Containers/UnrealString.h"
 #include "DerivedDataCache.h"
-#include "DerivedDataCacheInterface.h"
 #include "DerivedDataCacheRecord.h"
 #include "EditorDomain/EditorDomainArchive.h"
 #include "EditorDomain/EditorDomainSave.h"
@@ -277,7 +276,7 @@ int64 FEditorDomain::FileSize(const FPackagePath& PackagePath, EPackageSegment P
 		};
 		// Fetch meta-data only
 		ECachePolicy SkipFlags = ECachePolicy::SkipData & ~ECachePolicy::SkipMeta;
-		RequestGroup = GetDerivedDataCacheRef().CreateGroup(EPriority::Highest);
+		RequestGroup = GetCache().CreateGroup(EPriority::Highest);
 		RequestEditorDomainPackage(PackagePath, PackageSource->Digest, SkipFlags, RequestGroup.Get(), MoveTemp(MetaDataGetComplete));
 	}
 	RequestGroup.Get().Wait();
diff --git a/Engine/Source/Editor/UnrealEd/Private/EditorDomain/EditorDomainArchive.cpp b/Engine/Source/Editor/UnrealEd/Private/EditorDomain/EditorDomainArchive.cpp
--- a/Engine/Source/Editor/UnrealEd/Private/EditorDomain/EditorDomainArchive.cpp
+++ b/Engine/Source/Editor/UnrealEd/Private/EditorDomain/EditorDomainArchive.cpp
@@ -5,7 +5,6 @@
 #include "AssetRegistry/AssetData.h"
 #include "Async/AsyncFileHandleNull.h"
 #include "DerivedDataCache.h"
-#include "DerivedDataCacheInterface.h"
 #include "DerivedDataCacheKey.h"
 #include "DerivedDataCacheRecord.h"
 #include "EditorDomain/EditorDomainSave.h"
@@ -20,7 +19,7 @@
 FEditorDomainPackageSegments::FEditorDomainPackageSegments(const TRefCountPtr<FEditorDomain::FLocks>& InLocks,
 	const FPackagePath& InPackagePath, const TRefCountPtr<FEditorDomain::FPackageSource>& InPackageSource)
 	: EditorDomainLocks(InLocks)
-	, RequestGroup(GetDerivedDataCacheRef().CreateGroup(UE::DerivedData::EPriority::Normal))
+	, RequestGroup(UE::DerivedData::GetCache().CreateGroup(UE::DerivedData::EPriority::Normal))
 	, PackagePath(InPackagePath)
 	, PackageSource(InPackageSource)
 	, PackageDigest(InPackageSource->Digest)
@@ -374,7 +373,7 @@ void FEditorDomainPackageSegments::SendSegmentRequest(FSegment& Segment)
 	using namespace UE::DerivedData;
 
 	// Called from Callback-only until AsyncSource is set, then from Interface-only
-	ICache& Cache = GetDerivedDataCacheRef();
+	ICache& Cache = GetCache();
 
 	// Note that Segment.RequestGroup is Interface-only and so we can write it outside the lock
 	Segment.RequestGroup = Cache.CreateGroup(UE::DerivedData::EPriority::Normal);
diff --git a/Engine/Source/Editor/UnrealEd/Private/EditorDomain/EditorDomainUtils.cpp b/Engine/Source/Editor/UnrealEd/Private/EditorDomain/EditorDomainUtils.cpp
--- a/Engine/Source/Editor/UnrealEd/Private/EditorDomain/EditorDomainUtils.cpp
+++ b/Engine/Source/Editor/UnrealEd/Private/EditorDomain/EditorDomainUtils.cpp
@@ -7,7 +7,6 @@
 #include "AssetRegistry/IAssetRegistry.h"
 #include "Containers/Array.h"
 #include "DerivedDataCache.h"
-#include "DerivedDataCacheInterface.h"
 #include "DerivedDataCacheKey.h"
 #include "DerivedDataCacheRecord.h"
 #include "Editor.h"
@@ -204,21 +203,21 @@ EPackageDigestResult AppendPackageDigest(IAssetRegistry& AssetRegistry, FName Pa
 UE::DerivedData::FCacheKey GetEditorDomainPackageKey(const FPackageDigest& PackageDigest)
 {
 	static UE::DerivedData::FCacheBucket EditorDomainPackageCacheBucket =
-		GetDerivedDataCacheRef().CreateBucket(EditorDomainPackageBucketName);
+		UE::DerivedData::GetCache().CreateBucket(EditorDomainPackageBucketName);
 	return UE::DerivedData::FCacheKey{EditorDomainPackageCacheBucket, PackageDigest};
 }
 
 UE::DerivedData::FCacheKey GetBulkDataListKey(const FPackageDigest& PackageDigest)
 {
 	static UE::DerivedData::FCacheBucket EditorDomainBulkDataListBucket =
-		GetDerivedDataCacheRef().CreateBucket(EditorDomainBulkDataListBucketName);
+		UE::DerivedData::GetCache().CreateBucket(EditorDomainBulkDataListBucketName);
 	return UE::DerivedData::FCacheKey{ EditorDomainBulkDataListBucket, PackageDigest };
 }
 
 UE::DerivedData::FCacheKey GetBulkDataPayloadIdKey(const FIoHash& PackageAndGuidDigest)
 {
 	static UE::DerivedData::FCacheBucket EditorDomainBulkDataPayloadIdBucket =
-		GetDerivedDataCacheRef().CreateBucket(EditorDomainBulkDataPayloadIdBucketName);
+		UE::DerivedData::GetCache().CreateBucket(EditorDomainBulkDataPayloadIdBucketName);
 
 	return UE::DerivedData::FCacheKey{ EditorDomainBulkDataPayloadIdBucket, PackageAndGuidDigest };
 }
@@ -227,7 +226,7 @@ void RequestEditorDomainPackage(const FPackagePath& PackagePath,
 	const FPackageDigest& PackageDigest, UE::DerivedData::ECachePolicy SkipFlags, UE::DerivedData::IRequestOwner& Owner,
 	UE::DerivedData::FOnCacheGetComplete&& Callback)
 {
-	UE::DerivedData::ICache& Cache = GetDerivedDataCacheRef();
+	UE::DerivedData::ICache& Cache = UE::DerivedData::GetCache();
 	checkf((SkipFlags & (~UE::DerivedData::ECachePolicy::SkipData)) == UE::DerivedData::ECachePolicy::None,
 		TEXT("SkipFlags should only contain ECachePolicy::Skip* flags"));
 	Cache.Get({ GetEditorDomainPackageKey(PackageDigest) },
@@ -437,7 +436,7 @@ bool TrySavePackage(UPackage* Package)
 		return false;
 	}
 
-	ICache& Cache = GetDerivedDataCacheRef();
+	ICache& Cache = GetCache();
 	FCacheRecordBuilder RecordBuilder = Cache.CreateRecord(GetEditorDomainPackageKey(PackageDigest));
 
 	// We use a counter for PayloadIds rather than hashes of the Attachments. We do this because
@@ -497,7 +496,7 @@ void GetBulkDataList(FName PackageName, UE::DerivedData::IRequestOwner& Owner, T
 	}
 
 	using namespace UE::DerivedData;
-	ICache& Cache = GetDerivedDataCacheRef();
+	ICache& Cache = GetCache();
 	Cache.Get({ GetBulkDataListKey(PackageDigest) },
 		WriteToString<128>(PackageName), ECachePolicy::Default, Owner,
 		[InnerCallback = MoveTemp(Callback)](FCacheGetCompleteParams&& Params)
@@ -524,7 +523,7 @@ void PutBulkDataList(FName PackageName, FSharedBuffer Buffer)
 	}
 
 	using namespace UE::DerivedData;
-	ICache& Cache = GetDerivedDataCacheRef();
+	ICache& Cache = GetCache();
 	FRequestGroup Group = Cache.CreateGroup(EPriority::Normal);
 	FCacheRecordBuilder RecordBuilder = Cache.CreateRecord(GetBulkDataListKey(PackageDigest));
 	RecordBuilder.SetValue(Buffer);
@@ -557,7 +556,7 @@ void GetBulkDataPayloadId(FName PackageName, const FGuid& BulkDataId, UE::Derive
 	FIoHash PackageAndGuidDigest = GetPackageAndGuidDigest(Builder, BulkDataId);
 
 	using namespace UE::DerivedData;
-	ICache& Cache = GetDerivedDataCacheRef();
+	ICache& Cache = GetCache();
 	Cache.Get({ GetBulkDataPayloadIdKey(PackageAndGuidDigest) },
 		WriteToString<192>(PackageName, TEXT("/"), BulkDataId), ECachePolicy::Default, Owner,
 		[InnerCallback = MoveTemp(Callback)](FCacheGetCompleteParams&& Params)
@@ -584,7 +583,7 @@ void PutBulkDataPayloadId(FName PackageName, const FGuid& BulkDataId, FSharedBuf
 	FIoHash PackageAndGuidDigest = GetPackageAndGuidDigest(Builder, BulkDataId);
 
 	using namespace UE::DerivedData;
-	ICache& Cache = GetDerivedDataCacheRef();
+	ICache& Cache = GetCache();
 	FRequestGroup Group = Cache.CreateGroup(EPriority::Normal);
 	FCacheRecordBuilder RecordBuilder = Cache.CreateRecord(GetBulkDataPayloadIdKey(PackageAndGuidDigest));
 	RecordBuilder.SetValue(Buffer);
diff --git a/Engine/Source/Runtime/Core/Private/Misc/CoreMisc.cpp b/Engine/Source/Runtime/Core/Private/Misc/CoreMisc.cpp
--- a/Engine/Source/Runtime/Core/Private/Misc/CoreMisc.cpp
+++ b/Engine/Source/Runtime/Core/Private/Misc/CoreMisc.cpp
@@ -120,36 +120,6 @@ FDerivedDataCacheInterface& GetDerivedDataCacheRef()
 	return *Cache;
 }
 
-UE::DerivedData::IBuild* GetDerivedDataBuild()
-{
-	static UE::DerivedData::IBuild* const* Build;
-	static bool bInitialized = false;
-	if (!bInitialized)
-	{
-		if (!FPlatformProperties::RequiresCookedData())
-		{
-			check(IsInGameThread());
-			bInitialized = true;
-			if (IDerivedDataCacheModule* Module = FModuleManager::LoadModulePtr<IDerivedDataCacheModule>("DerivedDataCache"))
-			{
-				Build = Module->CreateOrGetBuild();
-			}
-		}
-	}
-	return Build ? *Build : nullptr;
-}
-
-UE::DerivedData::IBuild& GetDerivedDataBuildRef()
-{
-	UE::DerivedData::IBuild* Build = GetDerivedDataBuild();
-	if (!Build)
-	{
-		UE_LOG(LogInit, Fatal, TEXT("Derived Data Build was requested, but not available."));
-		CA_ASSUME(Build); // Suppress static analysis warning in unreachable code (fatal error)
-	}
-	return *Build;
-}
-
 class ITargetPlatformManagerModule* GetTargetPlatformManager(bool bFailOnInitErrors)
 {
 	static class ITargetPlatformManagerModule* SingletonInterface = NULL;
diff --git a/Engine/Source/Runtime/Core/Public/Misc/CoreMisc.h b/Engine/Source/Runtime/Core/Public/Misc/CoreMisc.h
--- a/Engine/Source/Runtime/Core/Public/Misc/CoreMisc.h
+++ b/Engine/Source/Runtime/Core/Public/Misc/CoreMisc.h
@@ -13,8 +13,6 @@
 #include "CoreGlobals.h"
 #include "HAL/ThreadSingleton.h"
 
-namespace UE::DerivedData { class IBuild; }
-
 /**
  * Exec handler that registers itself and is being routed via StaticExec.
  * Note: Not intended for use with UObjects!
@@ -73,12 +71,6 @@ CORE_API class FDerivedDataCacheInterface* GetDerivedDataCache();
 /** Returns the derived data cache interface, or fatal error if it is not available. */
 CORE_API class FDerivedDataCacheInterface& GetDerivedDataCacheRef();
 
-/** Return the derived data build interface if it is available, otherwise null. */
-CORE_API UE::DerivedData::IBuild* GetDerivedDataBuild();
-
-/** Return the derived data build interface, or fatal error if it is not available. */
-CORE_API UE::DerivedData::IBuild& GetDerivedDataBuildRef();
-
 /**
  * Return the Target Platform Manager interface, if it is available, otherwise return nullptr.
  *
diff --git a/Engine/Source/Runtime/Engine/Engine.Build.cs b/Engine/Source/Runtime/Engine/Engine.Build.cs
--- a/Engine/Source/Runtime/Engine/Engine.Build.cs
+++ b/Engine/Source/Runtime/Engine/Engine.Build.cs
@@ -334,6 +334,8 @@ public class Engine : ModuleRules
 				}
 			);
 
+			PrivateDependencyModuleNames.Add("DerivedDataCache");
+
 			PrivateIncludePathModuleNames.Add("TextureCompressor");
 			PrivateIncludePaths.Add("Developer/TextureCompressor/Public");
 
diff --git a/Engine/Source/Runtime/Engine/Private/TextureDerivedDataBuildUtils.cpp b/Engine/Source/Runtime/Engine/Private/TextureDerivedDataBuildUtils.cpp
--- a/Engine/Source/Runtime/Engine/Private/TextureDerivedDataBuildUtils.cpp
+++ b/Engine/Source/Runtime/Engine/Private/TextureDerivedDataBuildUtils.cpp
@@ -231,7 +231,7 @@ bool TryFindTextureBuildFunction(FStringBuilderBase& OutFunctionName, const FTex
 	}
 
 	const FStringView FunctionName = OutFunctionName.ToView().RightChop(NameIndex);
-	return GetDerivedDataBuildRef().GetFunctionRegistry().FindFunctionVersion(FunctionName).IsValid();
+	return UE::DerivedData::GetBuild().GetFunctionRegistry().FindFunctionVersion(FunctionName).IsValid();
 }
 
 FCbObject SaveTextureBuildSettings(const FString& KeySuffix, const UTexture& Texture, const FTextureBuildSettings& BuildSettings, int32 LayerIndex, int32 NumInlineMips)
diff --git a/Engine/Source/Runtime/Engine/Private/TextureDerivedDataTask.cpp b/Engine/Source/Runtime/Engine/Private/TextureDerivedDataTask.cpp
--- a/Engine/Source/Runtime/Engine/Private/TextureDerivedDataTask.cpp
+++ b/Engine/Source/Runtime/Engine/Private/TextureDerivedDataTask.cpp
@@ -543,7 +543,7 @@ void FTextureCacheDerivedDataWorker::BuildTexture(bool bReplaceExistingDDC)
 			using namespace UE::DerivedData;
 			// Compress the texture by calling the DDC2 build function.  May be executed locally or remotely.
 			// The FBuildDefinition in the future will replace both the cache get/put as well as execution.
-			IBuild& Build = GetDerivedDataBuildRef();
+			IBuild& Build = GetBuild();
 			FString TexturePath = Texture.GetPathName();
 			FBuildDefinitionBuilder DefinitionBuilder = Build.CreateDefinition(TexturePath, BuildFunctionName);
 
@@ -1015,7 +1015,7 @@ public:
 		TStringBuilder<256> TexturePath;
 		Texture.GetPathName(nullptr, TexturePath);
 
-		IBuild& Build = GetDerivedDataBuildRef();
+		IBuild& Build = GetBuild();
 		IBuildInputResolver* GlobalResolver = GetGlobalBuildInputResolver();
 		BuildSession = Build.CreateSession(TexturePath, GlobalResolver ? GlobalResolver : &InputResolver);
 
diff --git a/Engine/Source/Runtime/Launch/Launch.Build.cs b/Engine/Source/Runtime/Launch/Launch.Build.cs
--- a/Engine/Source/Runtime/Launch/Launch.Build.cs
+++ b/Engine/Source/Runtime/Launch/Launch.Build.cs
@@ -137,7 +137,15 @@ public class Launch : ModuleRules
 			});
 
 			PublicDependencyModuleNames.Add("SessionServices");
-			PrivateIncludePathModuleNames.Add("DerivedDataCache");
+
+			if (Target.bBuildWithEditorOnlyData)
+			{
+				PrivateDependencyModuleNames.Add("DerivedDataCache");
+			}
+			else
+			{
+				PrivateIncludePathModuleNames.Add("DerivedDataCache");
+			}
 
 			// LaunchEngineLoop.cpp will still attempt to load XMPP but not all projects require it so it will silently fail unless referenced by the project's build.cs file.
 			// DynamicallyLoadedModuleNames.Add("XMPP");
diff --git a/Engine/Source/Runtime/Launch/Private/LaunchEngineLoop.cpp b/Engine/Source/Runtime/Launch/Private/LaunchEngineLoop.cpp
--- a/Engine/Source/Runtime/Launch/Private/LaunchEngineLoop.cpp
+++ b/Engine/Source/Runtime/Launch/Private/LaunchEngineLoop.cpp
@@ -133,6 +133,10 @@
 #if WITH_ENGINE && !UE_BUILD_SHIPPING
 	#include "IAutomationControllerModule.h"
 #endif // WITH_ENGINE && !UE_BUILD_SHIPPING
+#if WITH_EDITORONLY_DATA
+	#include "DerivedDataBuild.h"
+	#include "DerivedDataCache.h"
+#endif // WITH_EDITORONLY_DATA
 	#include "DerivedDataCacheInterface.h"
 	#include "ShaderCompiler.h"
 	#include "DistanceFieldAtlas.h"
@@ -2850,9 +2854,12 @@ PRAGMA_ENABLE_DEPRECATION_WARNINGS
 
 	if (!FPlatformProperties::RequiresCookedData())
 	{
+#if WITH_EDITORONLY_DATA
 		// Ensure that DDC is initialized from the game thread.
+		UE::DerivedData::GetBuild();
+		UE::DerivedData::GetCache();
 		GetDerivedDataCacheRef();
-		GetDerivedDataBuildRef();
+#endif
 
 #if WITH_EDITOR
 		// Make sure that FVirtualizationManager is first called from the game thread
