Cooker: Fix warning "Ambiguous path search" when loading legacy assets that have an internal Package object of the same name as an existing .uasset file.
#rnx

diff --git a/Engine/Source/Editor/UnrealEd/Private/CookOnTheFlyServer.cpp b/Engine/Source/Editor/UnrealEd/Private/CookOnTheFlyServer.cpp
--- a/Engine/Source/Editor/UnrealEd/Private/CookOnTheFlyServer.cpp
+++ b/Engine/Source/Editor/UnrealEd/Private/CookOnTheFlyServer.cpp
@@ -3170,7 +3170,7 @@ bool UCookOnTheFlyServer::LoadPackageForCooking(UE::Cook::FPackageData& PackageD
 	};
 
 	FString PackageName = PackageData.GetPackageName().ToString();
-	OutPackage = FindObject<UPackage>(ANY_PACKAGE, *PackageName);
+	OutPackage = FindObject<UPackage>(nullptr, *PackageName);
 
 	FString FileName(LoadFromFilename ? *LoadFromFilename : PackageData.GetFileName().ToString());
 	FString ReportingFileName(ReportingPackageData ? ReportingPackageData->GetFileName().ToString() : FileName);
