Updating documentation links/parsing to include versioning.

diff --git a/Engine/Config/BaseEditor.ini b/Engine/Config/BaseEditor.ini
--- a/Engine/Config/BaseEditor.ini
+++ b/Engine/Config/BaseEditor.ini
@@ -273,31 +273,28 @@ bShouldPromptOnCurveRemoveAll=true
 
 [UnrealEd.URLs]
 UDNURL="http://udn.unrealengine.com"
-UDNDocsURL="http://udn.unrealengine.com/docs/ue4"
+UDNDocsURL="https://docs.unrealengine.com/{VERSION}/{I18N}/"
 OnlineLearningURL="https://www.unrealengine.com/onlinelearning-courses"
 ForumsURL="http://forums.unrealengine.com"
-APIDocsURL="https://docs.unrealengine.com/latest/INT/API/index.html"
-MapErrorURL="http://udn.unrealengine.com/docs/ue4/INT/Editor/LevelEditing/MapErrors/index.html"
-DocumentationURL="http://udn.unrealengine.com/docs/ue4/INT/index.html"
+APIDocsURL="https://docs.unrealengine.com/{VERSION}/{I18N}/API/"
+MapErrorURL="https://docs.unrealengine.com/{VERSION}/{I18N}/"
+DocumentationURL="http://docs.unrealengine.com/"
 AskAQuestionURL="http://udn.unrealengine.com/questions/ask.html"
 SearchForAnswersURL="http://udn.unrealengine.com/search.html"
 SupportWebsiteURL="http://udn.unrealengine.com/index.html"
 SourceCodeIDEURL_Windows="https://aka.ms/vs/16/release/vs_Community.exe"
 SourceCodeIDEURL_Mac="https://developer.apple.com/xcode"
 SourceCodeIDEURL_Other="http://en.wikipedia.org/wiki/Integrated_development_environment"
-APIDocsURL="https://docs.unrealengine.com/latest/INT/API/index.html"
 EpicGamesURL="http://UnrealEngine.com"
 FacebookURL="http://www.facebook.com/UnrealEngine"
-ViewportControlsURL="http://docs.unrealengine.com/latest/INT/Engine/UI/LevelEditor/Viewports/ViewportControls/index.html"
+ViewportControlsURL="https://docs.unrealengine.com/{VERSION}/{I18N}/BuildingWorlds/LevelEditor/Viewports/ViewportControls/"
 MeshSimplificationPluginsURL="/ue/marketplace/content-cat/assets/codeplugins"
 ReportABugURL="https://epicsupport.force.com/unrealengine/s/"
 IssueTrackerURL="https://issues.unrealengine.com/"
 
 
 [UnrealEd.URLOverrides]
-UDNURL="http://docs.unrealengine.com/latest/INT/index.html"
-UDNDocsURL="http://docs.unrealengine.com/latest"
-MapErrorURL="http://docs.unrealengine.com/latest/INT/index.html"
+UDNURL="http://docs.unrealengine.com/"
 DocumentationURL="http://docs.unrealengine.com/"
 AskAQuestionURL="http://answers.unrealengine.com"
 SearchForAnswersURL="http://answers.unrealengine.com"
diff --git a/Engine/Source/Editor/Documentation/Private/Documentation.cpp b/Engine/Source/Editor/Documentation/Private/Documentation.cpp
--- a/Engine/Source/Editor/Documentation/Private/Documentation.cpp
+++ b/Engine/Source/Editor/Documentation/Private/Documentation.cpp
@@ -50,8 +50,7 @@ bool FDocumentation::OpenAPIHome(FDocumentationSourceInfo Source) const
 
 	if (!Url.IsEmpty())
 	{
-		Url.ReplaceInline(TEXT("/INT/"), *FString::Printf(TEXT("/%s/"), *(FInternationalization::Get().GetCurrentCulture()->GetUnrealLegacyThreeLetterISOLanguageName())));
-
+		FUnrealEdMisc::Get().ReplaceDocumentationURLWildcards(Url, FInternationalization::Get().GetCurrentCulture());
 		FPlatformProcess::LaunchURL(*Url, nullptr, nullptr);
 
 		return true;
diff --git a/Engine/Source/Editor/Documentation/Private/DocumentationLink.h b/Engine/Source/Editor/Documentation/Private/DocumentationLink.h
--- a/Engine/Source/Editor/Documentation/Private/DocumentationLink.h
+++ b/Engine/Source/Editor/Documentation/Private/DocumentationLink.h
@@ -34,8 +34,8 @@ public:
 	{
 		FString Url;
 		FUnrealEdMisc::Get().GetURL( TEXT("UDNURL"), Url, true );
+		FUnrealEdMisc::Get().ReplaceDocumentationURLWildcards(Url, Culture);
 
-		Url.ReplaceInline(TEXT("/INT/"), *FString::Printf(TEXT("/%s/"), *(Culture->GetUnrealLegacyThreeLetterISOLanguageName())));
 		return Url;
 	}
 
@@ -50,12 +50,13 @@ public:
 		FString Anchor;
 		FString QueryString;
 		SplitLink(Link, Path, QueryString, Anchor);
-
-		const FString PartialPath = FString::Printf(TEXT("%s%s/index.html"), *(Culture->GetUnrealLegacyThreeLetterISOLanguageName()), *Path);
+		const FString PartialPath = FString::Printf(TEXT("%sindex.html"), *Path);
 		
 		AddSourceInfoToQueryString(QueryString, Source);
-
-		return GetUrlRoot() + PartialPath + QueryString + Anchor;
+		FString URLRoot = GetUrlRoot();
+		FInternationalization& I18N = FInternationalization::Get();
+		FUnrealEdMisc::Get().ReplaceDocumentationURLWildcards(URLRoot, I18N.GetCurrentCulture());
+		return URLRoot + PartialPath + QueryString + Anchor;
 	}
 
 	static FString ToFilePath( const FString& Link )
@@ -87,7 +88,8 @@ public:
 		FString QueryString;
 		SplitLink(Link, Path, QueryString, Anchor);
 
-		const FString PartialPath = FString::Printf(TEXT("%s%s/index.html"), *(Culture->GetUnrealLegacyThreeLetterISOLanguageName()), *Path);
+		static FString Version = FString::FromInt(FEngineVersion::Current().GetMajor()) + TEXT(".") + FString::FromInt(FEngineVersion::Current().GetMinor());
+		const FString PartialPath = FString::Printf(TEXT("%s/%s%s/index.html"), *Version, *(Culture->GetName()), *Path);
 		return FString::Printf(TEXT("%sDocumentation/HTML/%s"), *FPaths::ConvertRelativePathToFull(FPaths::EngineDir()), *PartialPath);
 	}
 
diff --git a/Engine/Source/Editor/LevelEditor/Private/LevelEditorActions.cpp b/Engine/Source/Editor/LevelEditor/Private/LevelEditorActions.cpp
--- a/Engine/Source/Editor/LevelEditor/Private/LevelEditorActions.cpp
+++ b/Engine/Source/Editor/LevelEditor/Private/LevelEditorActions.cpp
@@ -104,6 +104,8 @@
 #include "Misc/CoreMisc.h"
 #include "Misc/ScopeExit.h"
 
+#include "Internationalization/Culture.h"
+
 #if WITH_LIVE_CODING
 #include "ILiveCodingModule.h"
 #endif
@@ -190,6 +192,8 @@ void FLevelEditorActionCallbacks::BrowseViewportControls()
 	FString URL;
 	if (FUnrealEdMisc::Get().GetURL(TEXT("ViewportControlsURL"), URL))
 	{
+		// documentation URLS are now stored with placeholders for internationalization and version
+		FUnrealEdMisc::Get().ReplaceDocumentationURLWildcards(URL, FInternationalization::Get().GetCurrentCulture());
 		FPlatformProcess::LaunchURL(*URL, NULL, NULL);
 	}
 }
diff --git a/Engine/Source/Editor/LevelEditor/Private/SLevelViewportControlsPopup.cpp b/Engine/Source/Editor/LevelEditor/Private/SLevelViewportControlsPopup.cpp
--- a/Engine/Source/Editor/LevelEditor/Private/SLevelViewportControlsPopup.cpp
+++ b/Engine/Source/Editor/LevelEditor/Private/SLevelViewportControlsPopup.cpp
@@ -8,6 +8,7 @@
 #include "Widgets/Input/SButton.h"
 #include "EditorStyleSet.h"
 #include "UnrealEdMisc.h"
+#include "Internationalization/Culture.h"
 
 void SLevelViewportControlsPopup::Construct(const FArguments& InArgs)
 {
@@ -53,6 +54,8 @@ FReply SLevelViewportControlsPopup::OnClicked() const
 	FString URL;
 	if (FUnrealEdMisc::Get().GetURL(TEXT("ViewportControlsURL"), URL))
 	{
+		// documentation URLS are now stored with placeholders for internationalization and version
+		FUnrealEdMisc::Get().ReplaceDocumentationURLWildcards(URL, FInternationalization::Get().GetCurrentCulture());
 		FPlatformProcess::LaunchURL(*URL, NULL, NULL);
 	}
 
diff --git a/Engine/Source/Editor/UnrealEd/Private/UnrealEdMisc.cpp b/Engine/Source/Editor/UnrealEd/Private/UnrealEdMisc.cpp
--- a/Engine/Source/Editor/UnrealEd/Private/UnrealEdMisc.cpp
+++ b/Engine/Source/Editor/UnrealEd/Private/UnrealEdMisc.cpp
@@ -1633,12 +1633,16 @@ FString FUnrealEdMisc::GenerateURL(const FString& InUDNPage)
 	if( InUDNPage.Len() > 0 )
 	{
 		FInternationalization& I18N = FInternationalization::Get();
-
-		const FString PageURL = FString::Printf( TEXT( "%s/Editor/LevelEditing/MapErrors/index.html" ), *I18N.GetCurrentCulture()->GetUnrealLegacyThreeLetterISOLanguageName() );
+		// The I18N info and version is now stored in MapErrorURL in the ini
+		const FString PageURL =  TEXT("/Editor/LevelEditing/MapErrors/index.html");
 		const FString BookmarkURL = FString::Printf( TEXT( "#%s" ), *InUDNPage );
 
 		// Developers can browse documentation included with the engine distribution, check for file presence...
 		FString MapErrorURL = FString::Printf( TEXT( "%sDocumentation/HTML/%s" ), *FPaths::ConvertRelativePathToFull( FPaths::EngineDir() ), *PageURL );
+
+		static FString Version = FString::FromInt(FEngineVersion::Current().GetMajor()) + TEXT(".") + FString::FromInt(FEngineVersion::Current().GetMinor());
+		const FString PartialPath = FString::Printf(TEXT("%s/%s%s/index.html"), *Version, *(*I18N.GetCurrentCulture()->GetName()), *PageURL);
+		return FString::Printf(TEXT("%sDocumentation/HTML/%s"), *FPaths::ConvertRelativePathToFull(FPaths::EngineDir()), *PartialPath);
 		if (IFileManager::Get().FileSize(*MapErrorURL) != INDEX_NONE)
 		{
 			MapErrorURL = FString::Printf( TEXT( "file://%s%s" ), *MapErrorURL, *BookmarkURL );
@@ -1646,7 +1650,7 @@ FString FUnrealEdMisc::GenerateURL(const FString& InUDNPage)
 		// ... if it's not present, fallback to using the online version, if the full URL is provided...
 		else if(FUnrealEdMisc::Get().GetURL( TEXT("MapErrorURL"), MapErrorURL, true ) && MapErrorURL.EndsWith( TEXT( ".html" ) ))
 		{	
-			MapErrorURL.ReplaceInline( TEXT( "/INT/" ), *FString::Printf( TEXT( "/%s/" ), *I18N.GetCurrentCulture()->GetUnrealLegacyThreeLetterISOLanguageName() ) );
+			FUnrealEdMisc::Get().ReplaceDocumentationURLWildcards(MapErrorURL, I18N.GetCurrentCulture());
 			MapErrorURL += BookmarkURL;
 		}
 		// ...otherwise, attempt to create the URL from what we know here...
@@ -1656,6 +1660,9 @@ FString FUnrealEdMisc::GenerateURL(const FString& InUDNPage)
 			{
 				MapErrorURL += TEXT( "/" );
 			}
+
+			// UDNDocsURL is now stored with placeholders for internalization and version to be replaced
+			FUnrealEdMisc::Get().ReplaceDocumentationURLWildcards(MapErrorURL, I18N.GetCurrentCulture());
 			MapErrorURL += PageURL;
 			MapErrorURL += BookmarkURL;
 		}
@@ -1953,6 +1960,13 @@ bool FUnrealEdMisc::GetURL( const TCHAR* InKey, FString& OutURL, const bool bChe
 	return bFound;
 }
 
+void FUnrealEdMisc::ReplaceDocumentationURLWildcards(FString& Url, const FCultureRef& Culture)
+{
+	static FString Version = FString::FromInt(FEngineVersion::Current().GetMajor()) + TEXT(".") + FString::FromInt(FEngineVersion::Current().GetMinor());
+	Url.ReplaceInline(TEXT("{VERSION}"), *Version);
+	Url.ReplaceInline(TEXT("{I18N}"), *(Culture->GetName()));
+}
+
 FString FUnrealEdMisc::GetExecutableForCommandlets() const
 {
 	ILauncherServicesModule& LauncherServicesModule = FModuleManager::LoadModuleChecked<ILauncherServicesModule>(TEXT("LauncherServices"));
diff --git a/Engine/Source/Editor/UnrealEd/Public/UnrealEdMisc.h b/Engine/Source/Editor/UnrealEd/Public/UnrealEdMisc.h
--- a/Engine/Source/Editor/UnrealEd/Public/UnrealEdMisc.h
+++ b/Engine/Source/Editor/UnrealEd/Public/UnrealEdMisc.h
@@ -212,6 +212,8 @@ public:
 	 */
 	bool GetURL( const TCHAR* InKey, FString& OutURL, const bool bCheckRocket = false ) const;
 
+	void ReplaceDocumentationURLWildcards(FString& Url, const FCultureRef& Culture);
+
 	/** Returns the editor executable to use to execute commandlets */
 	FString GetExecutableForCommandlets() const;
 
