EditorDomain - move OnEndLoadPackage to private; it's not referenced outside the class.
#rnx

diff --git a/Engine/Source/Editor/UnrealEd/Private/EditorDomain/EditorDomain.h b/Engine/Source/Editor/UnrealEd/Private/EditorDomain/EditorDomain.h
--- a/Engine/Source/Editor/UnrealEd/Private/EditorDomain/EditorDomain.h
+++ b/Engine/Source/Editor/UnrealEd/Private/EditorDomain/EditorDomain.h
@@ -83,7 +83,6 @@ public:
 
 	// EditorDomain interface
 	void PrecachePackageDigest(FName PackageName);
-	void OnEndLoadPackage(TConstArrayView<UPackage*> LoadedPackages);
 
 private:
 	/**
@@ -131,6 +130,8 @@ private:
 	void MarkNeedsLoadFromWorkspace(const FPackagePath& PackagePath, TRefCountPtr<FPackageSource>& PackageSource);
 	/** Callback for PostEngineInit, to handle saving of packages which we could not save before then. */
 	void OnPostEngineInit();
+	/** EndLoad callback to handle saving the EditorDomain version of the package. */
+	void OnEndLoadPackage(TConstArrayView<UPackage*> LoadedPackages);
 	/** For each of the now-loaded packages, if we had to load from workspace domain, save into the editor domain. */
 	void FilterKeepPackagesToSave(TArray<UPackage*>& InOutLoadedPackages);
 
