IterationProfile: Add Trace Bookmark support

* Add bookmark pair to create telemetry scope for LoadAssetAndPIE, whitelist it in project-specific statistics files
* Split Stats.csv files into global Engine/Build/IterationProfile.csv and per-project Build/IterationProfile.csv directives (for project privacy)
* Global statistics file is now assumed, not necessary to pass into the <Commandlet>.
* Run <Commandlet> with Project specified so it can load project-specific statistic entries
* Remove -NoEnginePlugins parameter (projects require plugins!)
* Always do many hundreds of frames after we request PIE (vs. a time span which can be taken up mostly with preparing to PIE)
* Remove GC from the automated test function, it was just legacy and it betrays how realistic the test is
* Disable sleep in Editor when it doesn't have focus only in Deterministic mode

#rnx

diff --git a/Engine/Build/IterationProfile.xml b/Engine/Build/IterationProfile.xml
--- a/Engine/Build/IterationProfile.xml
+++ b/Engine/Build/IterationProfile.xml
@@ -74,7 +74,7 @@
 	<Property Name="ProfileArgs" Value="$(ProfileArgs) -NoShaderDistrib -VerboseShaderLogging"/>
 
 	<!-- Gauntlet arguments for summarize commandlet runs -->
-	<Property Name="SummarizeArgs" Value="-statsfile=&quot;$(RootDir)/Engine/Build/IterationProfileStats.csv&quot; -notraceserver"/>
+	<Property Name="SummarizeArgs" Value="-notraceserver"/>
 	<Property Name="SummarizeArgs" Value="$(SummarizeArgs) -skipbaseline" If="$(RunAndBaseline)"/>
 
 	<!-- The main test running macro logic that is expanded with platform and dependency node names -->
@@ -96,7 +96,7 @@
 		<Command Name="RunUnreal" Arguments="$(AllGauntletArgs) $(ProfileArgs) $(Args) -tracefile=&quot;$(TraceFilePath)&quot;"/>
 
 		<!-- Translate the utrace file into csv data -->
-		<Commandlet Name="SummarizeTrace" Arguments="-NoEnginePlugins -inputfile=&quot;$(TraceFilePath)&quot; -testname=&quot;$(TestName)&quot; $(SummarizeArgs)"/>
+		<Commandlet Project="$(Project)" Name="SummarizeTrace" Arguments="-inputfile=&quot;$(TraceFilePath)&quot; -testname=&quot;$(TestName)&quot; $(SummarizeArgs)"/>
 
 		<!-- Save all files named for this run in the trace directory -->
 		<Copy From="$(TraceDir)/$(Project)$(Platform)$(TestName)*.*" To="$(AutomationOutputDir)"/>
diff --git a/Engine/Plugins/Tests/EditorTests/Source/EditorTests/Private/UnrealEd/EditorIterationTests.cpp b/Engine/Plugins/Tests/EditorTests/Source/EditorTests/Private/UnrealEd/EditorIterationTests.cpp
--- a/Engine/Plugins/Tests/EditorTests/Source/EditorTests/Private/UnrealEd/EditorIterationTests.cpp
+++ b/Engine/Plugins/Tests/EditorTests/Source/EditorTests/Private/UnrealEd/EditorIterationTests.cpp
@@ -54,22 +54,41 @@ void FIterationOpenAssets::GetTests(TArray<FString>& OutBeautifiedNames, TArray<
 
 bool FIterationOpenAssets::RunTest(const FString& LongAssetPath)
 {
+	static uint64 FrameNumber = 0;
+
+	// Setup
 	AddCommand(new FCloseAllAssetEditorsCommand());
+	AddCommand(new FFunctionLatentCommand([LongAssetPath] {
+		TRACE_BOOKMARK(TEXT("LoadAssetAndPIE - %s"), *LongAssetPath);
+		return true;
+		}));
+
+	// Issue Load request
 	AddCommand(new FOpenEditorForAssetCommand(LongAssetPath));
 	AddCommand(new FWaitLatentCommand(0.5f));
+
+	// Wait on all async asset processing
 	AddCommand(new FFunctionLatentCommand([] {
 		return FAssetCompilingManager::Get().GetNumRemainingAssets() == 0;
 	}));
 	AddCommand(new FWaitLatentCommand(0.5f));
-	AddCommand(new FCloseAllAssetEditorsCommand());
-	AddCommand(new FDelayedFunctionLatentCommand([] {
-		CollectGarbage(GARBAGE_COLLECTION_KEEPFLAGS);
-	}));
 
-	bool bSimulateInEditor = false;
-	AddCommand(new FStartPIECommand(bSimulateInEditor));
-	AddCommand(new FWaitLatentCommand(10.0f));
+	// Do many frames of PIE (not a time span since PIE can take a while to init)
+	AddCommand(new FDelayedFunctionLatentCommand([] {
+		FrameNumber = GFrameCounter;
+		}));
+	AddCommand(new FStartPIECommand(false));
+	AddCommand(new FFunctionLatentCommand([LongAssetPath] {
+		return GFrameCounter > (FrameNumber + 600);
+		}));
 	AddCommand(new FEndPlayMapCommand());
+	AddCommand(new FWaitLatentCommand(0.5f));
+
+	// Teardown
+	AddCommand(new FFunctionLatentCommand([LongAssetPath] {
+		TRACE_BOOKMARK(TEXT("LoadAssetAndPIEComplete - %s"), *LongAssetPath);
+		return true;
+		}));
 
 	return true;
 }
diff --git a/Engine/Source/Developer/TraceServices/Private/Model/Log.cpp b/Engine/Source/Developer/TraceServices/Private/Model/Log.cpp
--- a/Engine/Source/Developer/TraceServices/Private/Model/Log.cpp
+++ b/Engine/Source/Developer/TraceServices/Private/Model/Log.cpp
@@ -174,4 +174,10 @@ const ILogProvider& ReadLogProvider(const IAnalysisSession& Session)
 	return *Session.ReadProvider<ILogProvider>(FLogProvider::ProviderName);
 }
 
+void FormatString(TCHAR* OutputString, uint32 OutputStringCount, const TCHAR* FormatString, const uint8* FormatArgs)
+{
+	TCHAR* TempBuffer = (TCHAR*)FMemory_Alloca(OutputStringCount * sizeof(TCHAR));
+	FFormatArgsHelper::Format(OutputString, OutputStringCount - 1, TempBuffer, OutputStringCount - 1, FormatString, FormatArgs);
+}
+
 } // namespace TraceServices
diff --git a/Engine/Source/Developer/TraceServices/Public/TraceServices/Model/Log.h b/Engine/Source/Developer/TraceServices/Public/TraceServices/Model/Log.h
--- a/Engine/Source/Developer/TraceServices/Public/TraceServices/Model/Log.h
+++ b/Engine/Source/Developer/TraceServices/Public/TraceServices/Model/Log.h
@@ -42,5 +42,6 @@ public:
 };
 
 TRACESERVICES_API const ILogProvider& ReadLogProvider(const IAnalysisSession& Session);
+TRACESERVICES_API void FormatString(TCHAR* OutputString, uint32 OutputStringCount, const TCHAR* FormatString, const uint8* FormatArgs);
 
 } // namespace TraceServices
diff --git a/Engine/Source/Editor/UnrealEd/Private/Commandlets/SummarizeTraceCommandlet.cpp b/Engine/Source/Editor/UnrealEd/Private/Commandlets/SummarizeTraceCommandlet.cpp
--- a/Engine/Source/Editor/UnrealEd/Private/Commandlets/SummarizeTraceCommandlet.cpp
+++ b/Engine/Source/Editor/UnrealEd/Private/Commandlets/SummarizeTraceCommandlet.cpp
@@ -16,6 +16,7 @@
 #include "Trace/Analysis.h"
 #include "Trace/Analyzer.h"
 #include "Trace/DataStream.h"
+#include "TraceServices/Model/Log.h"
 #include "TraceServices/Utils.h"
 #include "ProfilingDebugging/CountersTrace.h"
 
@@ -242,6 +243,86 @@ void FCountersAnalyzer::OnCountersSetValueFloat(const FOnEventContext& Context)
 	OnCounterFloatValue({ uint16(CounterId - 1), Value });
 }
 
+class FBookmarksAnalyzer
+	: public UE::Trace::IAnalyzer
+{
+public:
+	struct FBookmarkSpecEvent
+	{
+		uint64			Id;
+		const TCHAR*	FileName;
+		int32			Line;
+		const TCHAR*	FormatString;
+	};
+
+	struct FBookmarkEvent
+	{
+		uint64					Id;
+		double					Timestamp;
+		TArrayView<const uint8>	FormatArgs;
+	};
+
+	virtual void		OnBookmarkSpecEvent(const FBookmarkSpecEvent& BookmarkSpecEvent) = 0;
+	virtual void		OnBookmarkEvent(const FBookmarkEvent& BookmarkEvent) = 0;
+
+private:
+	virtual void		OnAnalysisBegin(const FOnAnalysisContext& Context) override;
+	virtual bool		OnEvent(uint16 RouteId, EStyle Style, const FOnEventContext& Context) override;
+	void				OnBookmarksSpec(const FOnEventContext& Context);
+	void				OnBookmarksBookmark(const FOnEventContext& Context);
+};
+
+enum
+{
+	// MiscTrace.cpp
+	RouteId_Bookmarks_BookmarkSpec,
+	RouteId_Bookmarks_Bookmark,
+};
+
+void FBookmarksAnalyzer::OnAnalysisBegin(const FOnAnalysisContext& Context)
+{
+	Context.InterfaceBuilder.RouteEvent(RouteId_Bookmarks_BookmarkSpec, "Misc", "BookmarkSpec");
+	Context.InterfaceBuilder.RouteEvent(RouteId_Bookmarks_Bookmark, "Misc", "Bookmark");
+}
+
+bool FBookmarksAnalyzer::OnEvent(uint16 RouteId, EStyle Style, const FOnEventContext& Context)
+{
+	switch (RouteId)
+	{
+	case RouteId_Bookmarks_BookmarkSpec:
+		OnBookmarksSpec(Context);
+		break;
+
+	case RouteId_Bookmarks_Bookmark:
+		OnBookmarksBookmark(Context);
+		break;
+	}
+
+	return true;
+}
+
+void FBookmarksAnalyzer::OnBookmarksSpec(const FOnEventContext& Context)
+{
+	const FEventData& EventData = Context.EventData;
+	uint64 Id = EventData.GetValue<uint64>("BookmarkPoint");
+	int32 Line = EventData.GetValue<int32>("Line");
+	FString FileName;
+	EventData.GetString("FileName", FileName);
+	FString FormatString;
+	EventData.GetString("FormatString", FormatString);
+	OnBookmarkSpecEvent({ Id, *FileName, Line, *FormatString });
+}
+
+void FBookmarksAnalyzer::OnBookmarksBookmark(const FOnEventContext& Context)
+{
+	const FEventData& EventData = Context.EventData;
+	uint64 Id = EventData.GetValue<uint64>("BookmarkPoint");
+	uint64 Cycle = EventData.GetValue<uint64>("Cycle");
+	double Timestamp = Context.EventTime.AsSeconds(Cycle);
+	TArrayView<const uint8> FormatArgsView = TraceServices::FTraceAnalyzerUtils::LegacyAttachmentArray("FormatArgs", Context);
+	OnBookmarkEvent({ Id, Timestamp, FormatArgsView });
+}
+
 /*
  * This too could be housed elsewhere, along with an API to make it easier to
  * run analysis on trace files. The current model is influenced a little too much
@@ -292,178 +373,242 @@ public:
  * Helper classes for the SummarizeTrace commandlet. Aggregates statistics about a trace.
  */
 
-class FSummarizeCpuAnalyzer
-	: public FCpuAnalyzer
+struct FSummarizeScope
 {
-public:
-	virtual void OnCpuScopeName(const FScopeName& ScopeName) override;
-	virtual void OnCpuScopeEnter(const FScopeEnter& ScopeEnter) override;
-	virtual void OnCpuScopeExit(const FScopeExit& ScopeExit) override;
+	FString Name;
+	uint64 Count = 0;
+	double TotalDurationSeconds = 0.0;
 
-	struct FThread
-	{
-		TArray<FScopeEnter> ScopeStack;
-	};
+	double FirstStartSeconds = 0.0;
+	double FirstFinishSeconds = 0.0;
+	double FirstDurationSeconds = 0.0;
 
-	struct FScope
-	{
-		FString Name;
-		uint64 Count = 0;
-		double TotalDurationSeconds = 0.0;
+	double LastStartSeconds = 0.0;
+	double LastFinishSeconds = 0.0;
+	double LastDurationSeconds = 0.0;
 
-		double FirstStartSeconds = 0.0;
-		double FirstEndSeconds = 0.0;
-		double FirstDurationSeconds = 0.0;
+	double MinDurationSeconds = 1e10;
+	double MaxDurationSeconds = -1e10;
+	double MeanDurationSeconds = 0.0;
+	double VarianceAcc = 0.0; // Accumulator for Welford's
 
-		double LastStartSeconds = 0.0;
-		double LastEndSeconds = 0.0;
-		double LastDurationSeconds = 0.0;
+	void AddDuration(double StartSeconds, double FinishSeconds)
+	{
+		Count += 1;
 
-		double MinDurationSeconds = 1e10;
-		double MaxDurationSeconds = -1e10;
-		double MeanDurationSeconds = 0.0;
-		double VarianceAcc = 0.0; // Accumulator for Welford's
+		// compute the duration
+		double DurationSeconds = FinishSeconds - StartSeconds;
 
-		void AddDuration(double StartSeconds, double EndSeconds)
+		// only set first for the first sample, compare exact zero
+		if (FirstStartSeconds == 0.0)
 		{
-			Count += 1;
+			FirstStartSeconds = StartSeconds;
+			FirstFinishSeconds = FinishSeconds;
+			FirstDurationSeconds = DurationSeconds;
+		}
 
-			// compute the duration
-			double DurationSeconds = EndSeconds - StartSeconds;
+		LastStartSeconds = StartSeconds;
+		LastFinishSeconds = FinishSeconds;
+		LastDurationSeconds = DurationSeconds;
 
-			// only set first for the first sample, compare exact zero
-			if (FirstStartSeconds == 0.0)
-			{
-				FirstStartSeconds = StartSeconds;
-				FirstEndSeconds = EndSeconds;
-				FirstDurationSeconds = DurationSeconds;
-			}
+		// set duration statistics
+		TotalDurationSeconds += DurationSeconds;
+		MinDurationSeconds = FMath::Min(MinDurationSeconds, DurationSeconds);
+		MaxDurationSeconds = FMath::Max(MaxDurationSeconds, DurationSeconds);
+		UpdateVariance(DurationSeconds);
+	}
 
-			LastStartSeconds = StartSeconds;
-			LastEndSeconds = EndSeconds;
-			LastDurationSeconds = DurationSeconds;
+	void UpdateVariance(double DurationSeconds)
+	{
+		ensure(Count);
 
-			// set duration statistics
-			TotalDurationSeconds += DurationSeconds;
-			MinDurationSeconds = FMath::Min(MinDurationSeconds, DurationSeconds);
-			MaxDurationSeconds = FMath::Max(MaxDurationSeconds, DurationSeconds);
-			UpdateVariance(DurationSeconds);
-		}
+		// Welford's increment
+		double OldMeanDurationSeconds = MeanDurationSeconds;
+		MeanDurationSeconds = MeanDurationSeconds + ((DurationSeconds - MeanDurationSeconds) / double(Count));
+		VarianceAcc = VarianceAcc + ((DurationSeconds - MeanDurationSeconds) * (DurationSeconds - OldMeanDurationSeconds));
+	}
 
-		void UpdateVariance(double DurationSeconds)
+	double GetDeviationDurationSeconds() const
+	{
+		if (Count > 1)
 		{
-			ensure(Count);
+			// Welford's final step, dependent on sample count
+			double VarianceSecondsSquared = VarianceAcc / double(Count - 1);
 
-			// Welford's increment
-			double OldMeanDurationSeconds = MeanDurationSeconds;
-			MeanDurationSeconds = MeanDurationSeconds + ((DurationSeconds - MeanDurationSeconds) / double(Count));
-			VarianceAcc = VarianceAcc + ((DurationSeconds - MeanDurationSeconds) * (DurationSeconds - OldMeanDurationSeconds));
+			// stddev is sqrt of variance, to restore to units of seconds (vs. seconds squared)
+			return sqrt(VarianceSecondsSquared);
 		}
-
-		double GetDeviationDurationSeconds() const
+		else
 		{
-			if (Count > 1)
-			{
-				// Welford's final step, dependent on sample count
-				double VarianceSecondsSquared = VarianceAcc / double(Count - 1);
-
-				// stddev is sqrt of variance, to restore to units of seconds (vs. seconds squared)
-				return sqrt(VarianceSecondsSquared);
-			}
-			else
-			{
-				return 0.0;
-			}
+			return 0.0;
 		}
+	}
 
-		void Merge(const FScope& Scope)
+	void Merge(const FSummarizeScope& Scope)
+	{
+		check(Name == Scope.Name);
+		TotalDurationSeconds += Scope.TotalDurationSeconds;
+		MinDurationSeconds = FMath::Min(MinDurationSeconds, Scope.MinDurationSeconds);
+		MaxDurationSeconds = FMath::Max(MaxDurationSeconds, Scope.MaxDurationSeconds);
+		Count += Scope.Count;
+	}
+
+	FString GetValue(const FStringView& Statistic) const
+	{
+		if (Statistic == TEXT("Name"))
 		{
-			check(Name == Scope.Name);
-			TotalDurationSeconds += Scope.TotalDurationSeconds;
-			MinDurationSeconds = FMath::Min(MinDurationSeconds, Scope.MinDurationSeconds);
-			MaxDurationSeconds = FMath::Max(MaxDurationSeconds, Scope.MaxDurationSeconds);
-			Count += Scope.Count;
+			return Name;
 		}
-
-		FString GetValue(const FStringView& Statistic) const
+		else if (Statistic == TEXT("Count"))
 		{
-			if (Statistic == TEXT("Name"))
-			{
-				return Name;
-			}
-			else if (Statistic == TEXT("Count"))
-			{
-				return FString::Printf(TEXT("%llu"), Count);
-			}
-			else if (Statistic == TEXT("TotalDurationSeconds"))
-			{
-				return FString::Printf(TEXT("%f"), TotalDurationSeconds);
-			}
-			else if (Statistic == TEXT("FirstStartSeconds"))
-			{
-				return FString::Printf(TEXT("%f"), FirstStartSeconds);
-			}
-			else if (Statistic == TEXT("FirstEndSeconds"))
-			{
-				return FString::Printf(TEXT("%f"), FirstEndSeconds);
-			}
-			else if (Statistic == TEXT("FirstDurationSeconds"))
-			{
-				return FString::Printf(TEXT("%f"), FirstDurationSeconds);
-			}
-			else if (Statistic == TEXT("LastStartSeconds"))
-			{
-				return FString::Printf(TEXT("%f"), LastStartSeconds);
-			}
-			else if (Statistic == TEXT("LastEndSeconds"))
-			{
-				return FString::Printf(TEXT("%f"), LastEndSeconds);
-			}
-			else if (Statistic == TEXT("LastDurationSeconds"))
-			{
-				return FString::Printf(TEXT("%f"), LastDurationSeconds);
-			}
-			else if (Statistic == TEXT("MinDurationSeconds"))
-			{
-				return FString::Printf(TEXT("%f"), MinDurationSeconds);
-			}
-			else if (Statistic == TEXT("MaxDurationSeconds"))
-			{
-				return FString::Printf(TEXT("%f"), MaxDurationSeconds);
-			}
-			else if (Statistic == TEXT("MeanDurationSeconds"))
-			{
-				return FString::Printf(TEXT("%f"), MeanDurationSeconds);
-			}
-			else if (Statistic == TEXT("DeviationDurationSeconds"))
-			{
-				return FString::Printf(TEXT("%f"), GetDeviationDurationSeconds());
-			}
-			return FString();
+			return FString::Printf(TEXT("%llu"), Count);
 		}
+		else if (Statistic == TEXT("TotalDurationSeconds"))
+		{
+			return FString::Printf(TEXT("%f"), TotalDurationSeconds);
+		}
+		else if (Statistic == TEXT("FirstStartSeconds"))
+		{
+			return FString::Printf(TEXT("%f"), FirstStartSeconds);
+		}
+		else if (Statistic == TEXT("FirstFinishSeconds"))
+		{
+			return FString::Printf(TEXT("%f"), FirstFinishSeconds);
+		}
+		else if (Statistic == TEXT("FirstDurationSeconds"))
+		{
+			return FString::Printf(TEXT("%f"), FirstDurationSeconds);
+		}
+		else if (Statistic == TEXT("LastStartSeconds"))
+		{
+			return FString::Printf(TEXT("%f"), LastStartSeconds);
+		}
+		else if (Statistic == TEXT("LastFinishSeconds"))
+		{
+			return FString::Printf(TEXT("%f"), LastFinishSeconds);
+		}
+		else if (Statistic == TEXT("LastDurationSeconds"))
+		{
+			return FString::Printf(TEXT("%f"), LastDurationSeconds);
+		}
+		else if (Statistic == TEXT("MinDurationSeconds"))
+		{
+			return FString::Printf(TEXT("%f"), MinDurationSeconds);
+		}
+		else if (Statistic == TEXT("MaxDurationSeconds"))
+		{
+			return FString::Printf(TEXT("%f"), MaxDurationSeconds);
+		}
+		else if (Statistic == TEXT("MeanDurationSeconds"))
+		{
+			return FString::Printf(TEXT("%f"), MeanDurationSeconds);
+		}
+		else if (Statistic == TEXT("DeviationDurationSeconds"))
+		{
+			return FString::Printf(TEXT("%f"), GetDeviationDurationSeconds());
+		}
+		return FString();
+	}
+
+	// for deduplication
+	bool operator==(const FSummarizeScope& Scope) const
+	{
+		return Name == Scope.Name;
+	}
 
-		// for deduplication
-		bool operator==(const FScope& Scope) const
+	// for sorting descending
+	bool operator<(const FSummarizeScope& Scope) const
+	{
+		return TotalDurationSeconds > Scope.TotalDurationSeconds;
+	}
+};
+
+static uint32 GetTypeHash(const FSummarizeScope& Scope)
+{
+	return FCrc::StrCrc32(*Scope.Name);
+}
+
+struct FSummarizeBookmark
+{
+	FString Name;
+	uint64 Count = 0;
+
+	double FirstSeconds = 0.0;
+	double LastSeconds = 0.0;
+
+	void AddTimestamp(double Seconds)
+	{
+		Count += 1;
+
+		// only set first for the first sample, compare exact zero
+		if (FirstSeconds == 0.0)
 		{
-			return Name == Scope.Name;
+			FirstSeconds = Seconds;
 		}
 
-		// for sorting descending
-		bool operator<(const FScope& Scope) const
+		LastSeconds = Seconds;
+	}
+
+	FString GetValue(const FStringView& Statistic) const
+	{
+		if (Statistic == TEXT("Name"))
 		{
-			return TotalDurationSeconds > Scope.TotalDurationSeconds;
+			return Name;
 		}
-	};
+		else if (Statistic == TEXT("Count"))
+		{
+			return FString::Printf(TEXT("%llu"), Count);
+		}
+		else if (Statistic == TEXT("FirstSeconds"))
+		{
+			return FString::Printf(TEXT("%f"), FirstSeconds);
+		}
+		else if (Statistic == TEXT("LastSeconds"))
+		{
+			return FString::Printf(TEXT("%f"), LastSeconds);
+		}
+		return FString();
+	}
 
-	TArray<FScope> Scopes;
-	TArray<FThread> Threads;
+	// for deduplication
+	bool operator==(const FSummarizeBookmark& Bookmark) const
+	{
+		return Name == Bookmark.Name;
+	}
 };
 
-static uint32 GetTypeHash(const FSummarizeCpuAnalyzer::FScope Key)
+static uint32 GetTypeHash(const FSummarizeBookmark& Bookmark)
 {
-	return FCrc::StrCrc32(*Key.Name);
+	return FCrc::StrCrc32(*Bookmark.Name);
 }
 
+//
+// FSummarizeCpuAnalyzer - Generate Scopes from cpu channel scope enter/exit events
+//
+
+class FSummarizeCpuAnalyzer
+	: public FCpuAnalyzer
+{
+public:
+	virtual void OnCpuScopeName(const FScopeName& ScopeName) override;
+	virtual void OnCpuScopeEnter(const FScopeEnter& ScopeEnter) override;
+	virtual void OnCpuScopeExit(const FScopeExit& ScopeExit) override;
+
+	// Scopes is an array only because indexes are doled out on the process-side
+	//  As such, there could be different scopes with the same name
+	//  We merge these together later to ensure name is a pkey in the csv
+	TArray<FSummarizeScope> Scopes;
+
+	// For each thread we track what the stack of scopes are, for matching end-to-start
+	struct FThread
+	{
+		TArray<FScopeEnter> ScopeStack;
+	};
+
+	// The state at any moment of the threads, indexes are doled out on the process-side
+	TArray<FThread> Threads;
+};
+
 void FSummarizeCpuAnalyzer::OnCpuScopeName(const FScopeName& ScopeName)
 {
 	if (ScopeName.Id >= uint32(Scopes.Num()))
@@ -504,6 +649,10 @@ void FSummarizeCpuAnalyzer::OnCpuScopeExit(const FScopeExit& ScopeExit)
 	}
 }
 
+//
+// FSummarizeCountersAnalyzer - Tally Counters from counter set/increment events
+//
+
 class FSummarizeCountersAnalyzer
 	: public FCountersAnalyzer
 {
@@ -606,10 +755,146 @@ void FSummarizeCountersAnalyzer::OnCounterFloatValue(const FCounterFloatValue& N
 	}
 }
 
+//
+// FSummarizeBookmarksAnalyzer - Tally Bookmarks from bookmark events
+//
+
+class FSummarizeBookmarksAnalyzer
+	: public FBookmarksAnalyzer
+{
+	virtual void OnBookmarkSpecEvent(const FBookmarkSpecEvent& BookmarkSpecEvent) override;
+	virtual void OnBookmarkEvent(const FBookmarkEvent& BookmarkEvent) override;
+
+	FSummarizeBookmark* FindStartBookmarkForEndBookmark(const FString& Name);
+
+public:
+	struct FBookmarkSpec
+	{
+		uint64	Id;
+		FString	FileName;
+		int32	Line;
+		FString	FormatString;
+	};
+
+	// Keyed by a unique memory address
+	TMap<uint64, FBookmarkSpec> BookmarkSpecs;
+
+	// Keyed by name
+	TMap<FString, FSummarizeBookmark> Bookmarks;
+
+	// Bookmarks named formed to scopes, see FindStartBookmarkForEndBookmark
+	TMap<FString, FSummarizeScope> Scopes;
+};
+
+void FSummarizeBookmarksAnalyzer::OnBookmarkSpecEvent(const FBookmarkSpecEvent& BookmarkSpecEvent)
+{
+	BookmarkSpecs.Add(BookmarkSpecEvent.Id, {
+		BookmarkSpecEvent.Id,
+		BookmarkSpecEvent.FileName,
+		BookmarkSpecEvent.Line,
+		BookmarkSpecEvent.FormatString
+		});
+}
+
+void FSummarizeBookmarksAnalyzer::OnBookmarkEvent(const FBookmarkEvent& BookmarkEvent)
+{
+	FBookmarkSpec* Spec = BookmarkSpecs.Find(BookmarkEvent.Id);
+	if (Spec)
+	{
+		TCHAR FormattedString[65535];
+		TraceServices::FormatString(FormattedString, sizeof(FormattedString) / sizeof(FormattedString[0]), *Spec->FormatString, BookmarkEvent.FormatArgs.GetData());
+
+		FString Name (FormattedString);
+
+		FSummarizeBookmark* FoundBookmark = Bookmarks.Find(Name);
+		if (!FoundBookmark)
+		{
+			FoundBookmark = &Bookmarks.Add(Name, FSummarizeBookmark());
+			FoundBookmark->Name = Name;
+		}
+
+		FoundBookmark->AddTimestamp(BookmarkEvent.Timestamp);
+
+		FSummarizeBookmark* StartBookmark = FindStartBookmarkForEndBookmark(Name);
+		if (StartBookmark)
+		{
+			FString ScopeName = FString(TEXT("Generated Scope for ")) + StartBookmark->Name;
+			FSummarizeScope* FoundScope = Scopes.Find(ScopeName);
+			if (!FoundScope)
+			{
+				FoundScope = &Scopes.Add(ScopeName, FSummarizeScope());
+				FoundScope->Name = ScopeName;
+			}
+
+			FoundScope->AddDuration(StartBookmark->LastSeconds, BookmarkEvent.Timestamp);
+		}
+	}
+}
+
+FSummarizeBookmark* FSummarizeBookmarksAnalyzer::FindStartBookmarkForEndBookmark(const FString& Name)
+{
+	int32 Index = Name.Find(TEXT("Complete"));
+	if (Index != -1)
+	{
+		FString StartName = Name;
+		StartName.RemoveAt(Index, TCString<TCHAR>::Strlen(TEXT("Complete")));
+		return Bookmarks.Find(StartName);
+	}
+
+	return nullptr;
+}
+
+/*
+* Begin SummarizeTrace commandlet implementation
+*/
+
+DEFINE_LOG_CATEGORY_STATIC(LogSummarizeTrace, Log, All);
+
 /*
-* Helper class for the stats csv file
+* Helpers for the csv files
 */
 
+static bool IsCsvSafeString(const FString& String)
+{
+	static struct DisallowedCharacter
+	{
+		const TCHAR Character;
+		bool First;
+	}
+	DisallowedCharacters[] =
+	{
+		// breaks simple csv files
+		{ TEXT('\n'), true },
+		{ TEXT('\r'), true },
+		{ TEXT(','), true },
+	};
+
+	// sanitize strings for a bog-simple csv file
+	bool bDisallowed = false;
+	int32 Index = 0;
+	for (struct DisallowedCharacter& DisallowedCharacter : DisallowedCharacters)
+	{
+		if (String.FindChar(DisallowedCharacter.Character, Index))
+		{
+			if (DisallowedCharacter.First)
+			{
+				UE_LOG(LogSummarizeTrace, Display, TEXT("A string contains disallowed character '%c'. See log for full list."), DisallowedCharacter.Character);
+				DisallowedCharacter.First = false;
+			}
+
+			UE_LOG(LogSummarizeTrace, Verbose, TEXT("String '%s' contains disallowed character '%c', skipping..."), *String, DisallowedCharacter.Character);
+			bDisallowed = true;
+		}
+
+		if (bDisallowed)
+		{
+			break;
+		}
+	}
+
+	return !bDisallowed;
+}
+
 struct StatisticDefinition
 {
 	StatisticDefinition()
@@ -952,8 +1237,6 @@ FString TelemetryDefinition::SignFlipThreshold(const FString& Threshold)
  * events and what statistics you would like to track.
  */
 
-DEFINE_LOG_CATEGORY_STATIC(LogSummarizeTrace, Log, All);
-
 USummarizeTraceCommandlet::USummarizeTraceCommandlet(const FObjectInitializer& ObjectInitializer)
 	: Super(ObjectInitializer)
 {
@@ -973,7 +1256,6 @@ int32 USummarizeTraceCommandlet::Main(const FString& CmdLineParams)
 		UE_LOG(LogSummarizeTrace, Log, TEXT("This commandlet will summarize a utrace into something more easily ingestable by a reporting tool (csv)."));
 		UE_LOG(LogSummarizeTrace, Log, TEXT("Options:"));
 		UE_LOG(LogSummarizeTrace, Log, TEXT(" Required: -inputfile=<utrace path>   (The utrace you wish to process)"));
-		UE_LOG(LogSummarizeTrace, Log, TEXT(" Optional: -statsfile=<csv path>      (The csv of statistics to generate)"));
 		UE_LOG(LogSummarizeTrace, Log, TEXT(" Optional: -testname=<string>         (Test name to use in telemetry csv)"));
 		return 0;
 	}
@@ -991,12 +1273,19 @@ int32 USummarizeTraceCommandlet::Main(const FString& CmdLineParams)
 
 	// load the stats file to know which event name and statistic name to generate in the telementry csv
 	// the telemetry csv is ingested completely, so this just whitelists specific data elements we want to track
-	FString StatisticsFileName;
 	TMultiMap<FString, StatisticDefinition> NameToDefinitionMap;
-	if (FParse::Value(*CmdLineParams, TEXT("statsfile="), StatisticsFileName, true))
+	FString GlobalStatisticsFileName = FPaths::RootDir() / TEXT("Engine") / TEXT("Build") / TEXT("IterationProfileStats.csv");
+	if (FPaths::FileExists(GlobalStatisticsFileName))
 	{
-		UE_LOG(LogSummarizeTrace, Display, TEXT("Loading statistics from %s"), *StatisticsFileName);
-		bool bCSVOk = StatisticDefinition::LoadFromCSV(StatisticsFileName, NameToDefinitionMap);
+		UE_LOG(LogSummarizeTrace, Display, TEXT("Loading global statistics from %s"), *GlobalStatisticsFileName);
+		bool bCSVOk = StatisticDefinition::LoadFromCSV(GlobalStatisticsFileName, NameToDefinitionMap);
+		check(bCSVOk);
+	}
+	FString ProjectStatisticsFileName = FPaths::ProjectDir() / TEXT("Build") / TEXT("IterationProfileStats.csv");
+	if (FPaths::FileExists(ProjectStatisticsFileName))
+	{
+		UE_LOG(LogSummarizeTrace, Display, TEXT("Loading project statistics from %s"), *ProjectStatisticsFileName);
+		bool bCSVOk = StatisticDefinition::LoadFromCSV(ProjectStatisticsFileName, NameToDefinitionMap);
 		check(bCSVOk);
 	}
 
@@ -1043,6 +1332,8 @@ int32 USummarizeTraceCommandlet::Main(const FString& CmdLineParams)
 	AnalysisContext.AddAnalyzer(CpuAnalyzer);
 	FSummarizeCountersAnalyzer CountersAnalyzer;
 	AnalysisContext.AddAnalyzer(CountersAnalyzer);
+	FSummarizeBookmarksAnalyzer BookmarksAnalyzer;
+	AnalysisContext.AddAnalyzer(BookmarksAnalyzer);
 
 	// kick processing on a thread
 	UE::Trace::FAnalysisProcessor AnalysisProcessor = AnalysisContext.Process(DataStream);
@@ -1050,67 +1341,20 @@ int32 USummarizeTraceCommandlet::Main(const FString& CmdLineParams)
 	// sync on completion
 	AnalysisProcessor.Wait();
 
-	TSet<FSummarizeCpuAnalyzer::FScope> DeduplicatedScopes;
-
-	struct DisallowedString
-	{
-		const TCHAR* String;
-		bool First;
-	}
-	DisallowedStrings[] =
-	{
-		// breaks simple csv files
-		{ TEXT("\n"), true },
-		{ TEXT("\r"), true },
-		{ TEXT(","), true },
-
-		// likely a path to an asset
-		//  this is broad but unfortunate with current instrumentation
-		//  at the current moment names can start with /Relpath/To/Asset
-		{ TEXT("/"), true }, 
-	};
-
-	bool bDisallowed = false;
-	for (const FSummarizeCpuAnalyzer::FScope& Scope : CpuAnalyzer.Scopes)
+	TSet<FSummarizeScope> DeduplicatedScopes;
+	auto IngestScope = [](TSet<FSummarizeScope>& DeduplicatedScopes, const FSummarizeScope& Scope)
 	{
 		if (Scope.Name.IsEmpty())
 		{
-			continue;
+			return;
 		}
 
 		if (Scope.Count == 0)
 		{
-			continue;
+			return;
 		}
 
-		// sanitize strings for a bog-simple csv file
-		bDisallowed = false;
-		for (struct DisallowedString& DisallowedString : DisallowedStrings)
-		{
-			if (Scope.Name.Contains(DisallowedString.String))
-			{
-				if (DisallowedString.First)
-				{
-					UE_LOG(LogSummarizeTrace, Display, TEXT("A scope contains disallowed string '%s'. See log for full list."), DisallowedString.String);
-					DisallowedString.First = false;
-				}
-
-				UE_LOG(LogSummarizeTrace, Verbose, TEXT("Scope '%s' contains disallowed string '%s', skipping..."), *Scope.Name, DisallowedString.String);
-				bDisallowed = true;
-			}
-
-			if (bDisallowed)
-			{
-				break;
-			}
-		}
-
-		if (bDisallowed)
-		{
-			continue;
-		}
-
-		FSummarizeCpuAnalyzer::FScope* FoundScope = DeduplicatedScopes.Find(Scope);
+		FSummarizeScope* FoundScope = DeduplicatedScopes.Find(Scope);
 		if (FoundScope)
 		{
 			FoundScope->Merge(Scope);
@@ -1119,11 +1363,19 @@ int32 USummarizeTraceCommandlet::Main(const FString& CmdLineParams)
 		{
 			DeduplicatedScopes.Add(Scope);
 		}
+	};
+	for (const FSummarizeScope& Scope : CpuAnalyzer.Scopes)
+	{
+		IngestScope(DeduplicatedScopes, Scope);
+	}
+	for (const TMap<FString, FSummarizeScope>::ElementType& ScopeItem : BookmarksAnalyzer.Scopes)
+	{
+		IngestScope(DeduplicatedScopes, ScopeItem.Value);
 	}
 
 	UE_LOG(LogSummarizeTrace, Display, TEXT("Sorting %d events by total time accumulated..."), DeduplicatedScopes.Num());
-	TArray<FSummarizeCpuAnalyzer::FScope> SortedScopes;
-	for (const FSummarizeCpuAnalyzer::FScope& Scope : DeduplicatedScopes)
+	TArray<FSummarizeScope> SortedScopes;
+	for (const FSummarizeScope& Scope : DeduplicatedScopes)
 	{
 		SortedScopes.Add(Scope);
 	}
@@ -1136,39 +1388,98 @@ int32 USummarizeTraceCommandlet::Main(const FString& CmdLineParams)
 		Handle->Write(reinterpret_cast<const uint8*>(UTF8String.Get()), UTF8String.Length());
 	};
 
-	// generate a summary csv, always
-	FString CsvFileName = FPaths::SetExtension(TraceFileName, "csv");
-	UE_LOG(LogSummarizeTrace, Display, TEXT("Writing summary to %s..."), *CsvFileName);
+	// some locals to help with all the derived files we are about to generate
+	const FString TracePath = FPaths::GetPath(TraceFileName);
+	const FString TraceFileBasename = FPaths::GetBaseFilename(TraceFileName);
+
+	// generate a summary csv files, always
+	FString CsvFileName = TraceFileBasename + TEXT("Scopes");
+	CsvFileName = FPaths::Combine(TracePath, FPaths::SetExtension(CsvFileName, "csv"));
+	UE_LOG(LogSummarizeTrace, Display, TEXT("Writing %s..."), *CsvFileName);
 	IFileHandle* CsvHandle = FPlatformFileManager::Get().GetPlatformFile().OpenWrite(*CsvFileName);
-	if (CsvHandle)
+	if (!CsvHandle)
+	{
+		UE_LOG(LogSummarizeTrace, Error, TEXT("Unable to open csv '%s' for write"), *CsvFileName);
+		return 1;
+	}
+	else
 	{
 		// no newline, see row printfs
-		WriteUTF8(CsvHandle, FString::Printf(TEXT("Name,Count,TotalDurationSeconds,FirstStartSeconds,FirstEndSeconds,FirstDurationSeconds,LastStartSeconds,LastEndSeconds,LastDurationSeconds,MinDurationSeconds,MaxDurationSeconds,MeanDurationSeconds,DeviationDurationSeconds,")));
-		for (const FSummarizeCpuAnalyzer::FScope& Scope : SortedScopes)
+		WriteUTF8(CsvHandle, FString::Printf(TEXT("Name,Count,TotalDurationSeconds,FirstStartSeconds,FirstFinishSeconds,FirstDurationSeconds,LastStartSeconds,LastFinishSeconds,LastDurationSeconds,MinDurationSeconds,MaxDurationSeconds,MeanDurationSeconds,DeviationDurationSeconds,")));
+		for (const FSummarizeScope& Scope : SortedScopes)
 		{
+			if (!IsCsvSafeString(Scope.Name))
+			{
+				continue;
+			}
+
 			// note newline is at the front of every data line to prevent final extraneous newline, per customary for csv
-			WriteUTF8(CsvHandle, FString::Printf(TEXT("\n%s,%llu,%f,%f,%f,%f,%f,%f,%f,%f,%f,%f,%f,"), *Scope.Name, Scope.Count, Scope.TotalDurationSeconds, Scope.FirstStartSeconds, Scope.FirstEndSeconds, Scope.FirstDurationSeconds, Scope.FirstStartSeconds, Scope.FirstEndSeconds, Scope.FirstDurationSeconds, Scope.MinDurationSeconds, Scope.MaxDurationSeconds, Scope.MeanDurationSeconds, Scope.GetDeviationDurationSeconds()));
+			WriteUTF8(CsvHandle, FString::Printf(TEXT("\n%s,%llu,%f,%f,%f,%f,%f,%f,%f,%f,%f,%f,%f,"), *Scope.Name, Scope.Count, Scope.TotalDurationSeconds, Scope.FirstStartSeconds, Scope.FirstFinishSeconds, Scope.FirstDurationSeconds, Scope.FirstStartSeconds, Scope.FirstFinishSeconds, Scope.FirstDurationSeconds, Scope.MinDurationSeconds, Scope.MaxDurationSeconds, Scope.MeanDurationSeconds, Scope.GetDeviationDurationSeconds()));
 		}
+		CsvHandle->Flush();
+		delete CsvHandle;
+		CsvHandle = nullptr;
+	}
+
+	CsvFileName = TraceFileBasename + TEXT("Counters");
+	CsvFileName = FPaths::Combine(TracePath, FPaths::SetExtension(CsvFileName, "csv"));
+	UE_LOG(LogSummarizeTrace, Display, TEXT("Writing %s..."), *CsvFileName);
+	CsvHandle = FPlatformFileManager::Get().GetPlatformFile().OpenWrite(*CsvFileName);
+	if (!CsvHandle)
+	{
+		UE_LOG(LogSummarizeTrace, Error, TEXT("Unable to open csv '%s' for write"), *CsvFileName);
+		return 1;
+	}
+	else
+	{
+		// no newline, see row printfs
+		WriteUTF8(CsvHandle, FString::Printf(TEXT("Name,Value,")));
 		for (const TMap<uint16, FSummarizeCountersAnalyzer::FCounter>::ElementType& Counter : CountersAnalyzer.Counters)
 		{
+			if (!IsCsvSafeString(Counter.Value.Name))
+			{
+				continue;
+			}
+
 			// note newline is at the front of every data line to prevent final extraneous newline, per customary for csv
-			WriteUTF8(CsvHandle, FString::Printf(TEXT("\n%s,%s,,,,,,,,,,,,"), *Counter.Value.Name, *Counter.Value.GetValue()));
+			WriteUTF8(CsvHandle, FString::Printf(TEXT("\n%s,%s,"), *Counter.Value.Name, *Counter.Value.GetValue()));
 		}
 		CsvHandle->Flush();
 		delete CsvHandle;
 		CsvHandle = nullptr;
 	}
-	else
+
+	CsvFileName = TraceFileBasename + TEXT("Bookmarks");
+	CsvFileName = FPaths::Combine(TracePath, FPaths::SetExtension(CsvFileName, "csv"));
+	UE_LOG(LogSummarizeTrace, Display, TEXT("Writing %s..."), *CsvFileName);
+	CsvHandle = FPlatformFileManager::Get().GetPlatformFile().OpenWrite(*CsvFileName);
+	if (!CsvHandle)
 	{
 		UE_LOG(LogSummarizeTrace, Error, TEXT("Unable to open csv '%s' for write"), *CsvFileName);
 		return 1;
 	}
+	else
+	{
+		// no newline, see row printfs
+		WriteUTF8(CsvHandle, FString::Printf(TEXT("Name,Count,FirstSeconds,LastSeconds,")));
+		for (const TMap<FString, FSummarizeBookmark>::ElementType& Bookmark : BookmarksAnalyzer.Bookmarks)
+		{
+			if (!IsCsvSafeString(Bookmark.Value.Name))
+			{
+				continue;
+			}
+
+			// note newline is at the front of every data line to prevent final extraneous newline, per customary for csv
+			WriteUTF8(CsvHandle, FString::Printf(TEXT("\n%s,%d,%f,%f,"), *Bookmark.Value.Name, Bookmark.Value.Count, Bookmark.Value.FirstSeconds, Bookmark.Value.LastSeconds));
+		}
+		CsvHandle->Flush();
+		delete CsvHandle;
+		CsvHandle = nullptr;
+	}
 
 	// if we were asked to generate a telememtry file, generate it
 	if (!NameToDefinitionMap.IsEmpty())
 	{
-		FString TracePath = FPaths::GetPath(TraceFileName);
-		FString TraceFileBasename = FPaths::GetBaseFilename(TraceFileName);
 		FString TelemetryCsvFileName = TraceFileBasename + TEXT("Telemetry");
 		TelemetryCsvFileName = FPaths::Combine(TracePath, FPaths::SetExtension(TelemetryCsvFileName, "csv"));
 
@@ -1177,27 +1488,57 @@ int32 USummarizeTraceCommandlet::Main(const FString& CmdLineParams)
 		FParse::Value(*CmdLineParams, TEXT("testname="), TestName, true);
 
 		TArray<TelemetryDefinition> TelemetryData;
-
-		// resolve scopes to telemetry
-		for (const FSummarizeCpuAnalyzer::FScope& Scope : SortedScopes)
 		{
 			TArray<StatisticDefinition> Statistics;
-			NameToDefinitionMap.MultiFind(Scope.Name, Statistics, true);
-			for (const StatisticDefinition& Statistic : Statistics)
+
+			// resolve scopes to telemetry
+			for (const FSummarizeScope& Scope : SortedScopes)
+			{
+				if (!IsCsvSafeString(Scope.Name))
+				{
+					continue;
+				}
+
+				NameToDefinitionMap.MultiFind(Scope.Name, Statistics, true);
+				for (const StatisticDefinition& Statistic : Statistics)
+				{
+					TelemetryData.Add(TelemetryDefinition(TestName, Statistic.TelemetryContext, Statistic.TelemetryDataPoint, Statistic.TelemetryUnit, Scope.GetValue(Statistic.Statistic)));
+				}
+				Statistics.Reset();
+			}
+
+			// resolve counters to telemetry
+			for (const TMap<uint16, FSummarizeCountersAnalyzer::FCounter>::ElementType& Counter : CountersAnalyzer.Counters)
 			{
-				TelemetryData.Add(TelemetryDefinition(TestName, Statistic.TelemetryContext, Statistic.TelemetryDataPoint, Statistic.TelemetryUnit, Scope.GetValue(Statistic.Statistic)));
+				if (!IsCsvSafeString(Counter.Value.Name))
+				{
+					continue;
+				}
+
+				NameToDefinitionMap.MultiFind(Counter.Value.Name, Statistics, true);
+				ensure(Statistics.Num() <= 1); // there should only be one, the counter value
+				for (const StatisticDefinition& Statistic : Statistics)
+				{
+					TelemetryData.Add(TelemetryDefinition(TestName, Statistic.TelemetryContext, Statistic.TelemetryDataPoint, Statistic.TelemetryUnit, Counter.Value.GetValue()));
+				}
+				Statistics.Reset();
 			}
-		}
 
-		// resolve counters to telemetry
-		for (const TMap<uint16, FSummarizeCountersAnalyzer::FCounter>::ElementType& Counter : CountersAnalyzer.Counters)
-		{
-			TArray<StatisticDefinition> Statistics;
-			NameToDefinitionMap.MultiFind(Counter.Value.Name, Statistics, true);
-			ensure(Statistics.Num() <= 1); // there should only be one, the counter value
-			for (const StatisticDefinition& Statistic : Statistics)
+			// resolve bookmarks to telemetry
+			for (const TMap<FString, FSummarizeBookmark>::ElementType& Bookmark : BookmarksAnalyzer.Bookmarks)
 			{
-				TelemetryData.Add(TelemetryDefinition(TestName, Statistic.TelemetryContext, Statistic.TelemetryDataPoint, Statistic.TelemetryUnit, Counter.Value.GetValue()));
+				if (!IsCsvSafeString(Bookmark.Value.Name))
+				{
+					continue;
+				}
+
+				NameToDefinitionMap.MultiFind(Bookmark.Value.Name, Statistics, true);
+				ensure(Statistics.Num() <= 1); // there should only be one, the bookmark itself
+				for (const StatisticDefinition& Statistic : Statistics)
+				{
+					TelemetryData.Add(TelemetryDefinition(TestName, Statistic.TelemetryContext, Statistic.TelemetryDataPoint, Statistic.TelemetryUnit, Bookmark.Value.GetValue(Statistic.Statistic)));
+				}
+				Statistics.Reset();
 			}
 		}
 
@@ -1260,13 +1601,10 @@ int32 USummarizeTraceCommandlet::Main(const FString& CmdLineParams)
 									FString BaselineRelPath = FPaths::ConvertRelativePathToFull(BaselineTelemetryCsvFilePath);
 									FPaths::MakePathRelativeTo(BaselineRelPath, *FPaths::RootDir());
 
-									FString StatisticsRelPath = FPaths::ConvertRelativePathToFull(StatisticsFileName);
-									FPaths::MakePathRelativeTo(StatisticsRelPath, *FPaths::RootDir());
-
-									UE_LOG(LogSummarizeTrace, Warning, TEXT("Telemetry %s,%s,%s,%s significantly within baseline value %s using warning threshold %s. Please submit a new baseline to %s or adjust the threshold in %s."),
+									UE_LOG(LogSummarizeTrace, Warning, TEXT("Telemetry %s,%s,%s,%s significantly within baseline value %s using warning threshold %s. Please submit a new baseline to %s or adjust the threshold in the statistics file."),
 										*Telemetry.TestName, *Telemetry.Context, *Telemetry.DataPoint, *Telemetry.Measurement,
 										*BaselineTelemetry->Measurement, *RelatedStatistic->BaselineWarningThreshold,
-										*BaselineRelPath, *StatisticsRelPath);
+										*BaselineRelPath);
 								}
 								else // it's within tolerance, just report that it's ok
 								{
@@ -1325,4 +1663,4 @@ int32 USummarizeTraceCommandlet::Main(const FString& CmdLineParams)
 	}
 
 	return 0;
-}
+}
\ No newline at end of file
diff --git a/Engine/Source/Programs/AutomationTool/Gauntlet/Unreal/Editor/EditorTests.cs b/Engine/Source/Programs/AutomationTool/Gauntlet/Unreal/Editor/EditorTests.cs
--- a/Engine/Source/Programs/AutomationTool/Gauntlet/Unreal/Editor/EditorTests.cs
+++ b/Engine/Source/Programs/AutomationTool/Gauntlet/Unreal/Editor/EditorTests.cs
@@ -96,7 +96,7 @@ namespace UnrealEditor
 			{
 				AppConfig.CommandLineParams.Add(string.Format("tracefile={0}", TraceFile));
 				AppConfig.CommandLineParams.Add("tracefiletrunc"); // replace existing
-				AppConfig.CommandLineParams.Add("trace=cpu,counters");
+				AppConfig.CommandLineParams.Add("trace=cpu,counters,bookmarks");
 				AppConfig.CommandLineParams.Add("statnamedevents");
 			}
 
diff --git a/Engine/Source/Runtime/ApplicationCore/Private/Windows/WindowsPlatformApplicationMisc.cpp b/Engine/Source/Runtime/ApplicationCore/Private/Windows/WindowsPlatformApplicationMisc.cpp
--- a/Engine/Source/Runtime/ApplicationCore/Private/Windows/WindowsPlatformApplicationMisc.cpp
+++ b/Engine/Source/Runtime/ApplicationCore/Private/Windows/WindowsPlatformApplicationMisc.cpp
@@ -134,35 +134,36 @@ void FWindowsPlatformApplicationMisc::PumpMessages(bool bFromMainLoop)
 	WinPumpMessages();
 
 	// Determine if application has focus
-	bool HasFocus = FApp::UseVRFocus() ? FApp::HasVRFocus() : FWindowsPlatformApplicationMisc::IsThisApplicationForeground();
-	static bool HadFocus = false;
+	bool bHasFocus = FApp::UseVRFocus() ? FApp::HasVRFocus() : FWindowsPlatformApplicationMisc::IsThisApplicationForeground();
+	static bool bHadFocus = false;
 
 #if WITH_EDITOR
 	// If editor thread doesn't have the focus, don't suck up too much CPU time.
 	if( GIsEditor )
 	{
-		if( !HasFocus )
+		bool bBenchmarking = FApp::IsBenchmarking();
+		if( !bHasFocus && !bBenchmarking )
 		{
 			TRACE_CPUPROFILER_EVENT_SCOPE("EditorIsInBackgroundSleep");
 			// Sleep for a bit to not eat up all CPU time.
 			FPlatformProcess::Sleep(0.005f);
 		}
-		HadFocus = HasFocus;
+		bHadFocus = bHasFocus;
 	}
 #endif
 
 #if !UE_SERVER
 	// For non-editor clients, record if the active window is in focus
-	if( HadFocus != HasFocus )
+	if( bHadFocus != bHasFocus )
 	{
-		FGenericCrashContext::SetEngineData(TEXT("Platform.AppHasFocus"), HasFocus ? TEXT("true") : TEXT("false"));
+		FGenericCrashContext::SetEngineData(TEXT("Platform.AppHasFocus"), bHasFocus ? TEXT("true") : TEXT("false"));
 	}
 #endif
 
-	HadFocus = HasFocus;
+	bHadFocus = bHasFocus;
 
 	// if its our window, allow sound, otherwise apply multiplier
-	FApp::SetVolumeMultiplier( HasFocus ? 1.0f : FApp::GetUnfocusedVolumeMultiplier() );
+	FApp::SetVolumeMultiplier( bHasFocus ? 1.0f : FApp::GetUnfocusedVolumeMultiplier() );
 }
 
 void FWindowsPlatformApplicationMisc::PreventScreenSaver()
