Configuration adjustments and logging improvements to DDC2 remote executor.

#rnx

diff --git a/Engine/Source/Editor/UnrealEd/Private/DerivedDataBuildRemoteExecutor.cpp b/Engine/Source/Editor/UnrealEd/Private/DerivedDataBuildRemoteExecutor.cpp
--- a/Engine/Source/Editor/UnrealEd/Private/DerivedDataBuildRemoteExecutor.cpp
+++ b/Engine/Source/Editor/UnrealEd/Private/DerivedDataBuildRemoteExecutor.cpp
@@ -529,6 +529,8 @@ int32 FRemoteBuildExecutionRequest::AddMerkleTreeFile(FStringView Path, const FI
 	NewNode.Digest.SizeBytes = RawSize;
 	NewNode.IsExecutable = bIsExecutable;
 
+	UE_LOG(LogDerivedDataBuildRemoteExecutor, Verbose, TEXT("Remote execution: added merkle tree file '%.*s' (hash: %s, size: %u)"), Path.Len(), Path.GetData(), *::LexToString(RawHash), RawSize);
+
 	int32 NewFileIndex = State.Files.Num();
 	State.DigestFilesystemIndex.Add(NewNode.Digest, FVariantIndex(ENodeType::File, NewFileIndex));
 	FMerkleTreeFileBuilder& FileBuilder = State.Files.AddDefaulted_GetRef();
@@ -915,7 +917,7 @@ TFuture<TPair<FStatus, FBatchUpdateBlobsResponse>> FRemoteBuildExecutionRequest:
 						NewRequest.Data.Append((const char *)Segment.GetData(), Segment.GetSize());
 					}
 					FileBuffer.Reset();
-					UE_LOG(LogDerivedDataBuildRemoteExecutor, Verbose, TEXT("Uploading file '%s' (hash: %s, type: %s) of upload size %d."), *State.Files[VariantIndex.Index].File.Name, *::LexToString(NewRequest.Digest.Hash), *LexToString(State.Files[VariantIndex.Index].Type), NewRequest.Data.Num());
+					UE_LOG(LogDerivedDataBuildRemoteExecutor, Verbose, TEXT("Uploading file '%s' (hash: %s, type: %s) of upload size %d."), *State.Files[VariantIndex.Index].File.Name, *::LexToString(NewRequest.Digest.Hash), LexToString(State.Files[VariantIndex.Index].Type), NewRequest.Data.Num());
 				}
 				break;
 			default:
@@ -963,6 +965,12 @@ void FRemoteBuildExecutionRequest::OnMissingBlobsDetermined(FFindMissingBlobsRes
 	}
 
 	State.FindMissingBlobsResponse = MoveTemp(Result);
+	constexpr bool bForceUploads = false;
+	if (bForceUploads)
+	{
+		State.FindMissingBlobsResponse.MissingBlobDigests = State.FindMissingBlobsRequest.BlobDigests;
+	}
+
 	TArray<FStringView> MissingInputViews;
 	GatherMissingInputFileBlobs(MissingInputViews);
 	if (!MissingInputViews.IsEmpty())
@@ -1043,7 +1051,7 @@ void FRemoteBuildExecutionRequest::OnExecutionCompleted(FExecuteResponse&& Resul
 	}
 	else
 	{
-		UE_LOG(LogDerivedDataBuildRemoteExecutor, Warning, TEXT("Remote execution system error: Failed to execute build operation!"));
+		UE_LOG(LogDerivedDataBuildRemoteExecutor, Warning, TEXT("Remote execution system error: Failed to execute build operation!  Response message: %s"), *State.ExecuteResponse.Status.Message);
 		State.Owner.End(this, [this]
 		{
 			CompletionCallback({State.BuildAction.GetKey(), {}, {}, EStatus::Error});
