EditorDomain: Optimize EditorDomainArchive for repeated contiguous serialize calls.
BulkDataRegistryEditorDomain: Fix unsynchronized access to PendingPackage during WriteCache
BulkDataRegistryEditorDomain: Fix OnEndLoadPackage to only call OnEndLoad for the specified packages.
BulkDataRegistryEditorDomain: Fix OnEndLoad bug that was making it never succeed.
#rnx

diff --git a/Engine/Source/Editor/UnrealEd/Private/EditorDomain/BulkDataRegistryEditorDomain.cpp b/Engine/Source/Editor/UnrealEd/Private/EditorDomain/BulkDataRegistryEditorDomain.cpp
--- a/Engine/Source/Editor/UnrealEd/Private/EditorDomain/BulkDataRegistryEditorDomain.cpp
+++ b/Engine/Source/Editor/UnrealEd/Private/EditorDomain/BulkDataRegistryEditorDomain.cpp
@@ -118,24 +118,34 @@ FBulkDataRegistryEditorDomain::~FBulkDataRegistryEditorDomain()
 
 void FBulkDataRegistryEditorDomain::OnEndLoadPackage(TConstArrayView<UPackage*> LoadedPackages)
 {
-	FScopeLock PendingPackageScopeLock(&PendingPackageLock);
-
 	TArray<TUniquePtr<FPendingPackage>> PackagesToWrite;
-	for (TMap<FName, TUniquePtr<FPendingPackage>>::TIterator Iter(PendingPackages); Iter; ++Iter)
-	{
-		bool bShouldRemove;
-		bool bShouldWriteCache;
-		Iter.Value()->OnEndLoad(bShouldRemove, bShouldWriteCache);
-		// We currently do not hold a lock when writing the cache, so we require that it not be used
-		// by any other thread if we need to write the cache
-		check(!bShouldWriteCache || bShouldRemove);
-		if (bShouldRemove)
+	{
+		FScopeLock PendingPackageScopeLock(&PendingPackageLock);
+
+		for (UPackage* LoadedPackage : LoadedPackages)
 		{
-			if (bShouldWriteCache)
+			FName PackageName = LoadedPackage->GetFName();
+			uint32 KeyHash = GetTypeHash(PackageName);
+			TUniquePtr<FPendingPackage>* PendingPackage = PendingPackages.FindByHash(KeyHash, PackageName);
+			if (PendingPackage)
 			{
-				PackagesToWrite.Add(MoveTemp(Iter.Value()));
+				bool bShouldRemove;
+				bool bShouldWriteCache;
+				check(PendingPackage->IsValid());
+				(*PendingPackage)->OnEndLoad(bShouldRemove, bShouldWriteCache);
+				// We do not hold a lock when calling WriteCache, so we require that the package no longer
+				// be accessible to other threads through PendingPackages if we need to write the cache,
+				// so bShouldRemove must be true if bShouldWriteCache is
+				check(!bShouldWriteCache || bShouldRemove);
+				if (bShouldRemove)
+				{
+					if (bShouldWriteCache)
+					{
+						PackagesToWrite.Add(MoveTemp(*PendingPackage));
+					}
+					PendingPackages.RemoveByHash(KeyHash, PackageName);
+				}
 			}
-			Iter.RemoveCurrent();
 		}
 	}
 
@@ -516,8 +526,11 @@ void FPendingPackage::OnEndLoad(bool& bOutShouldRemove, bool& bOutShouldWriteCac
 		bOutShouldWriteCache = true;
 		bOutShouldRemove = true;
 	}
-	bOutShouldWriteCache = false;
-	bOutShouldRemove = false;
+	else
+	{
+		bOutShouldWriteCache = false;
+		bOutShouldRemove = false;
+	}
 }
 
 void FPendingPackage::OnBulkDataListResults(FSharedBuffer Buffer)
@@ -536,14 +549,20 @@ void FPendingPackage::OnBulkDataListResults(FSharedBuffer Buffer)
 
 	if (PendingOperations.fetch_and(~Flag_BulkDataListResults) == Flag_BulkDataListResults)
 	{
+		// We do not hold a lock when writing the cache, so we need to remove this from PendingPackages 
+		// before calling WriteCache to avoid other threads being able to access it.
+		TUniquePtr<FPendingPackage> ThisPointer;
+		{
+			FScopeLock PendingPackageScopeLock(&Owner->PendingPackageLock);
+			ThisPointer = Owner->PendingPackages.FindAndRemoveChecked(PackageName);
+		}
 		WriteCache();
 
-		// Removing *this from its owner will delete it, which will attempt to Cancel this request.
-		// Direct the owner to keep requests alive to avoid a deadlock.
+		// Deleting *this will destruct BulkDataListCacheRequest, which by
+		// default calls Cancel, but Cancel will block on this callback we are
+		// currently in. Direct the owner to keep requests alive to avoid that deadlock.
 		BulkDataListCacheRequest.KeepAlive();
-
-		FScopeLock PendingPackageScopeLock(&Owner->PendingPackageLock);
-		Owner->PendingPackages.Remove(PackageName);
+		ThisPointer.Reset();
 		// *this has been deleted and can no longer be accessed
 	}
 }
diff --git a/Engine/Source/Editor/UnrealEd/Private/EditorDomain/EditorDomainArchive.cpp b/Engine/Source/Editor/UnrealEd/Private/EditorDomain/EditorDomainArchive.cpp
--- a/Engine/Source/Editor/UnrealEd/Private/EditorDomain/EditorDomainArchive.cpp
+++ b/Engine/Source/Editor/UnrealEd/Private/EditorDomain/EditorDomainArchive.cpp
@@ -106,16 +106,28 @@ bool FEditorDomainPackageSegments::Serialize(void* V, int64 Length)
 			Length, Size - Pos, *PackagePath.GetDebugName(), Size);
 		return false;
 	}
-	bool bIsReady;
-	FSegment* Segment = EnsureSegmentRange(Pos, End, true, bIsReady);
 	while (Pos < End)
 	{
-		int64 SegmentEnd = Segment < &Segments.Last() ? (Segment + 1)->Start : Size;
-		int64 LengthInSegment = FMath::Min(SegmentEnd, End) - Pos;
-		int64 SegmentPos = Pos - Segment->Start;
-		FMemory::Memcpy(V, static_cast<const uint8*>(Segment->Data.GetData()) + SegmentPos, LengthInSegment);
-		Pos += LengthInSegment;
-		++Segment;
+		if (MRUSegment)
+		{
+			int64 SegmentEnd = MRUSegment < &Segments.Last() ? (MRUSegment + 1)->Start : Size;
+			int64 LengthInSegment = FMath::Min(SegmentEnd, End) - Pos;
+			int64 SegmentPos = Pos - MRUSegment->Start;
+			if (SegmentPos >= 0 && LengthInSegment > 0)
+			{
+				FMemory::Memcpy(V, static_cast<const uint8*>(MRUSegment->Data.GetData()) + SegmentPos, LengthInSegment);
+				Pos += LengthInSegment;
+			}
+			else
+			{
+				MRUSegment = nullptr;
+			}
+		}
+		else
+		{
+			bool bIsReady;
+			MRUSegment = EnsureSegmentRange(Pos, End, true, bIsReady);
+		}
 	}
 	return true;
 }
diff --git a/Engine/Source/Editor/UnrealEd/Private/EditorDomain/EditorDomainArchive.h b/Engine/Source/Editor/UnrealEd/Private/EditorDomain/EditorDomainArchive.h
--- a/Engine/Source/Editor/UnrealEd/Private/EditorDomain/EditorDomainArchive.h
+++ b/Engine/Source/Editor/UnrealEd/Private/EditorDomain/EditorDomainArchive.h
@@ -140,13 +140,15 @@ private:
 
 	/**
 	 * Lock from the EditorDomain to allow shutting down this archive if callback outlives editor domain.
-	 * Pointer is read - only, * pointer has internal lock.
+	 * Pointer is read-only, *pointer has internal lock.
 	 */
 	TRefCountPtr<FEditorDomain::FLocks> EditorDomainLocks;
 	/** Initial Request for the list of segments. Read-only. */
 	UE::DerivedData::FRequestOwner RequestOwner;
 	/** PackagePath for fallback to WorkspaceDomain. Read-only. */
 	FPackagePath PackagePath;
+	/** Most recently used Segment in Serialize. Interface-only. */
+	FSegment* MRUSegment = nullptr;
 	/** Offset to next use in Serialize. Interface-only. */
 	int64 Pos = 0;
 	/**
