EditorDomain: Add PackageBlacklist and ClassBlacklist, for working around packages and classes that are temporarily not working when loaded from EditorDomain.
Move all Editor domain variables under the [EditorDomain] setting.

diff --git a/Engine/Config/BaseEditor.ini b/Engine/Config/BaseEditor.ini
--- a/Engine/Config/BaseEditor.ini
+++ b/Engine/Config/BaseEditor.ini
@@ -324,8 +324,6 @@ MinFreeMemory=1024
 MaxUsedMemory=16384
 
 [CookSettings]
-EditorDomainEnabled=false
-EditorDomainSaveUnversioned=false
 BulkDataRegistryEnabled=false
 TargetDomainEnabled=false
 PackagesPerGC=0
@@ -371,6 +369,12 @@ bUseFullGCAssetClassNames=true
 +ConfigSettingBlacklist=*.Engine:/Script/UnrealEd.CookerSettings:bIterativeCookingForFileCookContent
 +ConfigSettingBlacklist=*.Engine:OnlineSubsystemMcp
 
+[EditorDomain]
+EditorDomainEnabled=false
+SaveUnversioned=false
+; +ClassBlacklist=/Long/PackageName.ClassName
+; +PackageBlacklist=/Long/PackageName OR ../../../Standard/Filename/NoExtension OR D:\Absolute\Filename\NoExtension
+
 
 [CookPlatformDataCacheSettings]
 MaterialInstance=100
diff --git a/Engine/Source/Editor/UnrealEd/Private/EditorDomain/EditorDomain.cpp b/Engine/Source/Editor/UnrealEd/Private/EditorDomain/EditorDomain.cpp
--- a/Engine/Source/Editor/UnrealEd/Private/EditorDomain/EditorDomain.cpp
+++ b/Engine/Source/Editor/UnrealEd/Private/EditorDomain/EditorDomain.cpp
@@ -43,7 +43,11 @@ public:
 		bool bEditorDomainEnabled = false;
 		if (GIsEditor && (!IsRunningCommandlet() || IsRunningCookCommandlet()))
 		{
-			GConfig->GetBool(TEXT("CookSettings"), TEXT("EditorDomainEnabled"), bEditorDomainEnabled, GEditorIni);
+			GConfig->GetBool(TEXT("EditorDomain"), TEXT("EditorDomainEnabled"), bEditorDomainEnabled, GEditorIni);
+			if (GConfig->GetBool(TEXT("CookSettings"), TEXT("EditorDomainEnabled"), bEditorDomainEnabled, GEditorIni))
+			{
+				UE_LOG(LogEditorDomain, Error, TEXT("Editor.ini:[CookSettings]:EditorDomainEnabled is deprecated, use Editor.ini:[EditorDomain]:EditorDomainEnabled instead."));
+			}
 			if (bEditorDomainEnabled)
 			{
 				// Set values for config settings EditorDomain depends on
@@ -154,14 +158,15 @@ bool FEditorDomain::TryFindOrAddPackageSource(const FPackagePath& PackagePath, T
 
 	FString ErrorMessage;
 	FPackageDigest PackageDigest;
-	EPackageDigestResult Result = GetPackageDigest(*AssetRegistry, PackageName, PackageDigest, ErrorMessage);
+	bool bIsBlacklisted;
+	EPackageDigestResult Result = GetPackageDigest(*AssetRegistry, PackageName, PackageDigest, bIsBlacklisted, ErrorMessage);
 	switch (Result)
 	{
 	case EPackageDigestResult::Success:
 		PackageSource = new FPackageSource();
 		PackageSource->Digest = PackageDigest;
 		OutSource = PackageSource;
-		if (!bEditorDomainReadEnabled)
+		if (!bEditorDomainReadEnabled || bIsBlacklisted)
 		{
 			PackageSource->Source = EPackageSource::Workspace;
 		}
diff --git a/Engine/Source/Editor/UnrealEd/Private/EditorDomain/EditorDomainUtils.cpp b/Engine/Source/Editor/UnrealEd/Private/EditorDomain/EditorDomainUtils.cpp
--- a/Engine/Source/Editor/UnrealEd/Private/EditorDomain/EditorDomainUtils.cpp
+++ b/Engine/Source/Editor/UnrealEd/Private/EditorDomain/EditorDomainUtils.cpp
@@ -3,6 +3,7 @@
 #include "EditorDomain/EditorDomainUtils.h"
 
 #include "Algo/IsSorted.h"
+#include "Algo/Unique.h"
 #include "AssetRegistry/AssetData.h"
 #include "AssetRegistry/IAssetRegistry.h"
 #include "Containers/Array.h"
@@ -45,14 +46,25 @@ const TCHAR* EditorDomainBulkDataPayloadIdBucketName = TEXT("EditorDomainBulkDat
 
 static bool GetEditorDomainSaveUnversioned()
 {
-	bool bParsedValue;
-	static bool bEditorDomainSaveUnversioned = GConfig->GetBool(TEXT("CookSettings"), TEXT("EditorDomainSaveUnversioned"), bParsedValue, GEditorIni) ? bParsedValue : true;
+	auto Initialize = []()
+	{
+		bool bParsedValue;
+		bool bResult = GConfig->GetBool(TEXT("EditorDomain"), TEXT("SaveUnversioned"), bParsedValue, GEditorIni) ? bParsedValue : true;
+		if (GConfig->GetBool(TEXT("CookSettings"), TEXT("EditorDomainSaveUnversioned"), bResult, GEditorIni))
+		{
+			UE_LOG(LogEditorDomain, Error, TEXT("Editor.ini:[CookSettings]:EditorDomainSaveUnversioned is deprecated, use Editor.ini:[EditorDomain]:SaveUnversioned instead."));
+		}
+		return bResult;
+	};
+	static bool bEditorDomainSaveUnversioned = Initialize();
 	return bEditorDomainSaveUnversioned;
 }
 
-EPackageDigestResult AppendPackageDigest(FCbWriter& Writer, FString& OutErrorMessage,
+EPackageDigestResult AppendPackageDigest(FCbWriter& Writer, bool& bOutIsBlacklisted, FString& OutErrorMessage,
 	const FAssetPackageData& PackageData, FName PackageName)
 {
+	bOutIsBlacklisted = false;
+
 	int32 CurrentFileVersionUE = GPackageFileUEVersion;
 	int32 CurrentFileVersionLicenseeUE = GPackageFileLicenseeUEVersion;
 	Writer << EditorDomainVersion;
@@ -109,77 +121,298 @@ EPackageDigestResult AppendPackageDigest(FCbWriter& Writer, FString& OutErrorMes
 			{
 				Writer << ExistingData->SchemaHash;
 			}
+			bOutIsBlacklisted |= ExistingData->bBlacklisted;
 		}
 	}
 	return EPackageDigestResult::Success;
 }
 
-void PrecacheClassDigests(TConstArrayView<FName> ClassNames)
+void PrecacheClassDigests(TConstArrayView<FName> ClassNames, TMap<FName, FClassDigestData>* OutDatas)
 {
 	FClassDigestMap& ClassDigests = GetClassDigests();
-	TArray<TPair<FName,FClassDigestData>> ClassesToAdd;
 	FString ClassNameStr;
+	TArray<FName, TInlineAllocator<10>> ClassesToAdd;
+	ClassesToAdd.Reserve(ClassNames.Num());
 	{
 		FScopeLock ClassDigestsScopeLock(&ClassDigests.Lock);
 		for (FName ClassName : ClassNames)
 		{
-			if (!ClassDigests.Map.Find(ClassName))
+			FClassDigestData* DigestData = ClassDigests.Map.Find(ClassName);
+			if (DigestData)
+			{
+				if (OutDatas)
+				{
+					OutDatas->Add(ClassName, *DigestData);
+				}
+			}
+			else
 			{
-				ClassesToAdd.Emplace_GetRef().Get<0>() = ClassName;
+				ClassesToAdd.Add(ClassName);
 			}
 		}
 	}
-	if (ClassesToAdd.Num())
+	if (ClassesToAdd.Num() == 0)
+	{
+		return;
+	}
+
+	struct FClassData
 	{
-		FString TargetClassNameString;
-		for (int32 Index = 0; Index < ClassesToAdd.Num();)
+		FName Name;
+		FName ParentName;
+		UStruct* ParentStruct = nullptr;
+		FClassDigestData DigestData;
+	};
+	TArray<FClassData, TInlineAllocator<10>> ClassDatas;
+	FString NameStringBuffer;
+	IAssetRegistry& AssetRegistry = *IAssetRegistry::Get();
+	TArray<FName> AncestorShortNames;
+
+	ClassDatas.Reserve(ClassesToAdd.Num());
+	for (FName ClassName : ClassesToAdd)
+	{
+		FName LookupName = ClassName;
+		ClassName.ToString(NameStringBuffer);
+		FCoreRedirectObjectName ClassNameRedirect(NameStringBuffer);
+		FCoreRedirectObjectName RedirectedClassNameRedirect = FCoreRedirects::GetRedirectedName(ECoreRedirectFlags::Type_Class, ClassNameRedirect);
+		if (ClassNameRedirect != RedirectedClassNameRedirect)
+		{
+			NameStringBuffer = RedirectedClassNameRedirect.ToString();
+			LookupName = FName(NameStringBuffer);
+		}
+		UStruct* Struct = nullptr;
+		if (FPackageName::IsScriptPackage(NameStringBuffer))
+		{
+			Struct = FindObject<UStruct>(nullptr, *NameStringBuffer);
+			if (!Struct)
+			{
+				// If a native class is not found we do not put it in our results
+				continue;
+			}
+		}
+		
+		FClassData& ClassData = ClassDatas.Emplace_GetRef();
+		ClassData.Name = ClassName;
+		ClassData.DigestData.bBlacklisted = GetClassBlacklist().Contains(ClassName);
+		if (LookupName != ClassName)
+		{
+			ClassData.DigestData.bBlacklisted |= GetClassBlacklist().Contains(LookupName);
+		}
+
+		if (Struct)
+		{
+			ClassData.DigestData.bNative = true;
+			ClassData.DigestData.SchemaHash = Struct->GetSchemaHash(false /* bSkipEditorOnly */);
+			ClassData.ParentStruct = Struct->GetSuperStruct();
+			if (ClassData.ParentStruct)
+			{
+				NameStringBuffer.Reset();
+				ClassData.ParentStruct->GetPathName(nullptr, NameStringBuffer);
+				ClassData.ParentName = FName(*NameStringBuffer);
+			}
+		}
+		else
 		{
-			FName OldClassFName = ClassesToAdd[Index].Get<0>();
-			FClassDigestData& Data = ClassesToAdd[Index].Get<1>();
-			OldClassFName.ToString(TargetClassNameString);
-			FCoreRedirectObjectName OldClassName(TargetClassNameString);
-			FCoreRedirectObjectName NewClassName = FCoreRedirects::GetRedirectedName(ECoreRedirectFlags::Type_Class, OldClassName);
-			if (OldClassName != NewClassName)
+			ClassData.DigestData.bNative = false;
+			ClassData.DigestData.SchemaHash.Reset();
+			FStringView UnusedClassOfClassName;
+			FStringView ClassPackageName;
+			FStringView ClassObjectName;
+			FStringView ClassSubObjectName;
+			FPackageName::SplitFullObjectPath(NameStringBuffer, UnusedClassOfClassName, ClassPackageName, ClassObjectName, ClassSubObjectName);
+			FName ClassObjectFName(ClassObjectName);
+			// TODO_EDITORDOMAIN: If the class is not yet present in the assetregistry, or
+			// if its parent classes are not, then we will not be able to propagate information from the parent classes; wait on the class to be parsed
+			AncestorShortNames.Reset();
+			IAssetRegistry::Get()->GetAncestorClassNames(ClassObjectFName, AncestorShortNames);
+			for (FName ShortName : AncestorShortNames)
 			{
-				TargetClassNameString = NewClassName.ToString();
+				// TODO_EDITORDOMAIN: For robustness and performance, we need the AssetRegistry to return FullPathNames rather than ShortNames
+				// For now, we lookup each shortname using FindObject, and do not handle propagating data from blueprint classes to child classes
+				if (UStruct* ParentStruct = FindObjectFast<UStruct>(nullptr, ShortName, false /* ExactClass */, true /* AnyPackage */))
+				{
+					NameStringBuffer.Reset();
+					ParentStruct->GetPathName(nullptr, NameStringBuffer);
+					if (FPackageName::IsScriptPackage(NameStringBuffer))
+					{
+						ClassData.ParentStruct = ParentStruct;
+						ClassData.ParentName = FName(*NameStringBuffer);
+						break;
+					}
+				}
 			}
+		}
+	}
+
+	TMap<FName, FClassData> RemainingBatch;
+	{
+		FScopeLock ClassDigestsScopeLock(&ClassDigests.Lock);
 
-			if (FPackageName::IsScriptPackage(TargetClassNameString))
+		// Look up the data for the parent of each class, so we can propagate bBlacklisted from the parent,
+		// once parent data is propagated, add it to the ClassDigests map.
+		// For any parents missing data, keep the class for a second pass that adds the parent
+		for (FClassData& ClassData : ClassDatas)
+		{
+			bool bNeedsParent = false;
+			if (!ClassData.ParentName.IsNone())
 			{
-				UStruct* Struct = FindObject<UStruct>(nullptr, *TargetClassNameString);
-				if (Struct)
+				FClassDigestData* ParentDigest = ClassDigests.Map.Find(ClassData.ParentName);
+				if (ParentDigest)
 				{
-					Data.SchemaHash = Struct->GetSchemaHash(false /* bSkipEditorOnly */);
-					Data.bNative = true;
-					++Index;
+					ClassData.DigestData.bBlacklisted |= ParentDigest->bBlacklisted;
 				}
 				else
 				{
-					ClassesToAdd.RemoveAtSwap(Index);
+					bNeedsParent = true;
 				}
 			}
+			if (!bNeedsParent)
+			{
+				if (OutDatas)
+				{
+					OutDatas->Add(ClassData.Name, ClassData.DigestData);
+				}
+				ClassDigests.Map.Add(ClassData.Name, MoveTemp(ClassData.DigestData));
+			}
 			else
 			{
-				Data.SchemaHash.Reset();
-				Data.bNative = false;
-				++Index;
+				RemainingBatch.Add(ClassData.Name, MoveTemp(ClassData));
 			}
 		}
+	}
+
+	if (RemainingBatch.Num() == 0)
+	{
+		return;
+	}
+
+	// Get all unique ancestors (skipping those that are already in the batch) and recursively cache them
+	TSet<FName> Parents;
+	for (const TPair<FName, FClassData>& RemainingPair : RemainingBatch)
+	{
+		const FClassData& ClassData = RemainingPair.Value;
+		if (ClassData.ParentName.IsNone() || RemainingBatch.Contains(ClassData.ParentName))
+		{
+			continue;
+		}
+		checkf(ClassData.ParentStruct, TEXT("If the ClassData has a parent, it should have come either from the ParentStruct."));
+		FName ParentName = ClassData.ParentName;
+		UStruct* ParentStruct = ClassData.ParentStruct;
+		do
 		{
-			FScopeLock ClassDigestsScopeLock(&ClassDigests.Lock);
-			for (TPair<FName,FClassDigestData>& Pair: ClassesToAdd)
+			bool bAlreadyExists;
+			Parents.Add(ParentName, &bAlreadyExists);
+			if (bAlreadyExists)
+			{
+				break;
+			}
+			ParentStruct = ParentStruct->GetSuperStruct();
+			if (ParentStruct)
 			{
-				ClassDigests.Map.Add(Pair.Get<0>(), MoveTemp(Pair.Get<1>()));
+				NameStringBuffer.Reset();
+				ParentStruct->GetPathName(nullptr, NameStringBuffer);
+				ParentName = FName(*NameStringBuffer);
 			}
 		}
+		while (ParentStruct);
 	}
+	TMap<FName, FClassDigestData> ParentDigests;
+	PrecacheClassDigests(Parents.Array(), &ParentDigests);
+
+	// Propagate parent values to children, pulling parentdata from ParentDigests or RemainingBatch
+	TSet<FName> Visited;
+	auto RecursivePropagate = [&RemainingBatch, &ParentDigests, &Visited](FClassData& ClassData, auto& RecursivePropagateReference)
+	{
+		bool bAlreadyInSet;
+		Visited.Add(ClassData.Name, &bAlreadyInSet);
+		if (bAlreadyInSet)
+		{
+			return;
+		}
+		FClassDigestData* ParentDigest = ParentDigests.Find(ClassData.ParentName);
+		if (!ParentDigest)
+		{
+			FClassData* ParentData = RemainingBatch.Find(ClassData.ParentName);
+			if (ParentData)
+			{
+				RecursivePropagateReference(*ParentData, RecursivePropagateReference);
+				ParentDigest = &ParentData->DigestData;
+			}
+		}
+		// If the superclass was not found, due to a bad redirect or a missing blueprint assetregistry entry, then give up and treat the class as having no parent
+		if (ParentDigest)
+		{
+			ClassData.DigestData.bBlacklisted |= ParentDigest->bBlacklisted;
+		}
+	};
+	for (TPair<FName, FClassData>& RemainingPair : RemainingBatch)
+	{
+		RecursivePropagate(RemainingPair.Value, RecursivePropagate);
+	}
+
+	// Add the now-complete RemainingBatch digests to ClassDigests
+	{
+		FScopeLock ClassDigestsScopeLock(&ClassDigests.Lock);
+		for (TPair<FName, FClassData>& RemainingPair : RemainingBatch)
+		{
+			if (OutDatas)
+			{
+				OutDatas->Add(RemainingPair.Key, RemainingPair.Value.DigestData);
+			}
+			ClassDigests.Map.Add(RemainingPair.Key, MoveTemp(RemainingPair.Value.DigestData));
+		}
+	}
+}
+
+TSet<FName> ConstructClassBlacklist()
+{
+	TSet<FName> Result;
+	TArray<FString> BlacklistArray;
+	GConfig->GetArray(TEXT("EditorDomain"), TEXT("ClassBlacklist"), BlacklistArray, GEditorIni);
+	for (const FString& ClassPathName : BlacklistArray)
+	{
+		Result.Add(FName(*ClassPathName));
+	}
+	return Result;
+}
+
+const TSet<FName>& GetClassBlacklist()
+{
+	static TSet<FName> ClassBlacklist = ConstructClassBlacklist();
+	return ClassBlacklist;
+}
+
+TSet<FName> ConstructPackageNameBlacklist()
+{
+	TSet<FName> Result;
+	TArray<FString> BlacklistArray;
+	GConfig->GetArray(TEXT("EditorDomain"), TEXT("PackageBlacklist"), BlacklistArray, GEditorIni);
+	FString PackageName;
+	FString ErrorReason;
+	for (const FString& PackageNameOrFilename : BlacklistArray)
+	{
+		if (!FPackageName::TryConvertFilenameToLongPackageName(PackageNameOrFilename, PackageName, &ErrorReason))
+		{
+			UE_LOG(LogEditorDomain, Warning, TEXT("Editor.ini:[EditorDomain]:PackageBlacklist: Could not convert %s to a LongPackageName: %s"),
+				*PackageNameOrFilename, *ErrorReason);
+			continue;
+		}
+		Result.Add(FName(*PackageName));
+	}
+	return Result;
+}
+
+const TSet<FName>& GetPackageNameBlacklist()
+{
+	static TSet<FName> PackageNameBlacklist = ConstructPackageNameBlacklist();
+	return PackageNameBlacklist;
 }
 
 EPackageDigestResult GetPackageDigest(IAssetRegistry& AssetRegistry, FName PackageName,
-	FPackageDigest& OutPackageDigest, FString& OutErrorMessage)
+	FPackageDigest& OutPackageDigest, bool& bOutIsBlacklisted, FString& OutErrorMessage)
 {
 	FCbWriter Builder;
-	EPackageDigestResult Result = AppendPackageDigest(AssetRegistry, PackageName, Builder, OutErrorMessage);
+	EPackageDigestResult Result = AppendPackageDigest(AssetRegistry, PackageName, Builder, bOutIsBlacklisted, OutErrorMessage);
 	if (Result == EPackageDigestResult::Success)
 	{
 		OutPackageDigest = Builder.Save().GetRangeHash();
@@ -188,7 +421,7 @@ EPackageDigestResult GetPackageDigest(IAssetRegistry& AssetRegistry, FName Packa
 }
 
 EPackageDigestResult AppendPackageDigest(IAssetRegistry& AssetRegistry, FName PackageName,
-	FCbWriter& Builder, FString& OutErrorMessage)
+	FCbWriter& Builder, bool& bOutIsBlacklisted, FString& OutErrorMessage)
 {
 	AssetRegistry.WaitForPackage(PackageName.ToString());
 	TOptional<FAssetPackageData> PackageData = AssetRegistry.GetAssetPackageDataCopy(PackageName);
@@ -196,9 +429,12 @@ EPackageDigestResult AppendPackageDigest(IAssetRegistry& AssetRegistry, FName Pa
 	{
 		OutErrorMessage = FString::Printf(TEXT("Package %s does not exist in the AssetRegistry"),
 			*PackageName.ToString());
+		bOutIsBlacklisted = false;
 		return EPackageDigestResult::FileDoesNotExist;
 	}
-	return AppendPackageDigest(Builder, OutErrorMessage, *PackageData, PackageName);
+	EPackageDigestResult Result = AppendPackageDigest(Builder, bOutIsBlacklisted, OutErrorMessage, *PackageData, PackageName);
+	bOutIsBlacklisted |= GetPackageNameBlacklist().Contains(PackageName);
+	return Result;
 }
 
 UE::DerivedData::FCacheKey GetEditorDomainPackageKey(const FPackageDigest& PackageDigest)
@@ -382,8 +618,9 @@ bool TrySavePackage(UPackage* Package)
 
 	FString ErrorMessage;
 	FPackageDigest PackageDigest;
+	bool bIsBlacklisted;
 	EPackageDigestResult FindHashResult = GetPackageDigest(*IAssetRegistry::Get(), Package->GetFName(), PackageDigest,
-		ErrorMessage);
+		bIsBlacklisted, ErrorMessage);
 	switch (FindHashResult)
 	{
 	case EPackageDigestResult::Success:
@@ -392,6 +629,11 @@ bool TrySavePackage(UPackage* Package)
 		UE_LOG(LogEditorDomain, Warning, TEXT("Could not save package to EditorDomain: %s."), *ErrorMessage)
 		return false;
 	}
+	if (bIsBlacklisted)
+	{
+		UE_LOG(LogEditorDomain, Verbose, TEXT("Skipping save of blacklisted package to EditorDomain: %s."), *Package->GetName())
+		return false;
+	}
 
 	uint32 SaveFlags = SAVE_NoError // Do not crash the SaveServer on an error
 		| SAVE_BulkDataByReference	// EditorDomain saves reference bulkdata from the WorkspaceDomain rather than duplicating it
@@ -478,8 +720,9 @@ void GetBulkDataList(FName PackageName, UE::DerivedData::IRequestOwner& Owner, T
 {
 	FString ErrorMessage;
 	FPackageDigest PackageDigest;
+	bool bIsBlacklisted;
 	EPackageDigestResult FindHashResult = GetPackageDigest(*IAssetRegistry::Get(), PackageName, PackageDigest,
-		ErrorMessage);
+		bIsBlacklisted, ErrorMessage);
 	switch (FindHashResult)
 	{
 	case EPackageDigestResult::Success:
@@ -490,6 +733,11 @@ void GetBulkDataList(FName PackageName, UE::DerivedData::IRequestOwner& Owner, T
 		return;
 	}
 	}
+	if (bIsBlacklisted)
+	{
+		Callback(FSharedBuffer());
+		return;
+	}
 
 	using namespace UE::DerivedData;
 	ICache& Cache = GetCache();
@@ -506,8 +754,9 @@ void PutBulkDataList(FName PackageName, FSharedBuffer Buffer)
 {
 	FString ErrorMessage;
 	FPackageDigest PackageDigest;
+	bool bIsBlacklisted;
 	EPackageDigestResult FindHashResult = GetPackageDigest(*IAssetRegistry::Get(), PackageName, PackageDigest,
-		ErrorMessage);
+		bIsBlacklisted, ErrorMessage);
 	switch (FindHashResult)
 	{
 	case EPackageDigestResult::Success:
@@ -517,6 +766,10 @@ void PutBulkDataList(FName PackageName, FSharedBuffer Buffer)
 		return;
 	}
 	}
+	if (bIsBlacklisted)
+	{
+		return;
+	}
 
 	using namespace UE::DerivedData;
 	ICache& Cache = GetCache();
@@ -538,7 +791,8 @@ void GetBulkDataPayloadId(FName PackageName, const FGuid& BulkDataId, UE::Derive
 {
 	FString ErrorMessage;
 	FCbWriter Builder;
-	EPackageDigestResult FindHashResult = AppendPackageDigest(*IAssetRegistry::Get(), PackageName, Builder, ErrorMessage);
+	bool bIsBlacklisted;
+	EPackageDigestResult FindHashResult = AppendPackageDigest(*IAssetRegistry::Get(), PackageName, Builder, bIsBlacklisted, ErrorMessage);
 	switch (FindHashResult)
 	{
 	case EPackageDigestResult::Success:
@@ -549,6 +803,11 @@ void GetBulkDataPayloadId(FName PackageName, const FGuid& BulkDataId, UE::Derive
 		return;
 	}
 	}
+	if (bIsBlacklisted)
+	{
+		Callback(FSharedBuffer());
+		return;
+	}
 	FIoHash PackageAndGuidDigest = GetPackageAndGuidDigest(Builder, BulkDataId);
 
 	using namespace UE::DerivedData;
@@ -566,7 +825,8 @@ void PutBulkDataPayloadId(FName PackageName, const FGuid& BulkDataId, FSharedBuf
 {
 	FString ErrorMessage;
 	FCbWriter Builder;
-	EPackageDigestResult FindHashResult = AppendPackageDigest(*IAssetRegistry::Get(), PackageName, Builder, ErrorMessage);
+	bool bIsBlacklisted;
+	EPackageDigestResult FindHashResult = AppendPackageDigest(*IAssetRegistry::Get(), PackageName, Builder, bIsBlacklisted, ErrorMessage);
 	switch (FindHashResult)
 	{
 	case EPackageDigestResult::Success:
@@ -576,6 +836,10 @@ void PutBulkDataPayloadId(FName PackageName, const FGuid& BulkDataId, FSharedBuf
 		return;
 	}
 	}
+	if (bIsBlacklisted)
+	{
+		return;
+	}
 	FIoHash PackageAndGuidDigest = GetPackageAndGuidDigest(Builder, BulkDataId);
 
 	using namespace UE::DerivedData;
diff --git a/Engine/Source/Editor/UnrealEd/Private/EditorDomain/EditorDomainUtils.h b/Engine/Source/Editor/UnrealEd/Private/EditorDomain/EditorDomainUtils.h
--- a/Engine/Source/Editor/UnrealEd/Private/EditorDomain/EditorDomainUtils.h
+++ b/Engine/Source/Editor/UnrealEd/Private/EditorDomain/EditorDomainUtils.h
@@ -3,6 +3,8 @@
 #pragma once
 
 #include "Containers/ArrayView.h"
+#include "Containers/Map.h"
+#include "Containers/Set.h"
 #include "Containers/UnrealString.h"
 #include "DerivedDataCache.h"
 #include "EditorDomain/EditorDomain.h"
@@ -32,11 +34,12 @@ enum class EPackageDigestResult
 	MissingClass,
 };
 
-/** A UClass's data that is used in the EditorDomain Digest */
+/** A UClass's data that is used in the EditorDomain Digest, and holds other information about Classes the EditorDomain needs. */
 struct FClassDigestData
 {
 	FBlake3Hash SchemaHash;
-	bool bNative;
+	bool bNative = false;
+	bool bBlacklisted = false;
 };
 
 /** Threadsafe cache of ClassName -> Digest data for calculating EditorDomain Digests */
@@ -51,13 +54,17 @@ struct FClassDigestMap
  * Reads information from the AssetRegistry to compute the digest.
  */
 EPackageDigestResult GetPackageDigest(IAssetRegistry& AssetRegistry, FName PackageName,
-	FPackageDigest& OutPackageDigest, FString& OutErrorMessage);
+	FPackageDigest& OutPackageDigest, bool& bOutIsBlacklisted, FString& OutErrorMessage);
 /** Appends the fields to calculate the packagedigest; call Builder.Save().GetRangeHash() to get digest. */
 EPackageDigestResult AppendPackageDigest(IAssetRegistry& AssetRegistry, FName PackageName,
-	FCbWriter& Builder, FString& OutErrorMessage);
+	FCbWriter& Builder, bool& bOutIsBlacklisted, FString& OutErrorMessage);
 
 /** For any ClassNames not already in ClassDigests, look up their UStruct and add them. */
-void PrecacheClassDigests(TConstArrayView<FName> ClassNames);
+void PrecacheClassDigests(TConstArrayView<FName> ClassNames, TMap<FName, FClassDigestData>* OutDatas = nullptr);
+
+/** Return the TSet of class fullpath names that was parsed from the Editor.ini:[EditorDomain]:ClassBlacklist array. */
+const TSet<FName>& GetClassBlacklist();
+const TSet<FName>& GetPackageNameBlacklist();
 
 /** Get the CacheRequest for the given package from the EditorDomain cache bucket. */
 void RequestEditorDomainPackage(const FPackagePath& PackagePath,
