Change the DDC remote execution timeout to use the action timeout instead of the grpc timeout with the expectation that the action timeout will actually be enforced by the executing agent in the case of a hung process.

diff --git a/Engine/Source/Editor/UnrealEd/Private/DerivedDataBuildRemoteExecutor.cpp b/Engine/Source/Editor/UnrealEd/Private/DerivedDataBuildRemoteExecutor.cpp
--- a/Engine/Source/Editor/UnrealEd/Private/DerivedDataBuildRemoteExecutor.cpp
+++ b/Engine/Source/Editor/UnrealEd/Private/DerivedDataBuildRemoteExecutor.cpp
@@ -666,6 +666,10 @@ void FRemoteBuildExecutionRequest::BuildMerkleTreeNodes()
 		});
 	Executor.ContentAddressableStorage->ToBlob(State.Command, State.CommandContentBytes, State.Action.CommandDigest);
 
+	if (Executor.GlobalExecutionTimeoutSeconds > 0)
+	{
+		State.Action.Timeout = FTimespan::FromSeconds(Executor.GlobalExecutionTimeoutSeconds);
+	}
 	State.Action.Salt = Executor.Salt;
 	Executor.ContentAddressableStorage->ToBlob(State.Action, State.ActionContentBytes, State.ActionDigest);
 }
@@ -935,12 +939,7 @@ TFuture<FExecuteResponse> FRemoteBuildExecutionRequest::ExecuteBuildAsync()
 	State.ExecuteRequest.InstanceName = Executor.InstanceName;
 	State.ExecuteRequest.ActionDigest = State.ActionDigest;
 	State.ExecuteRequest.SkipCacheLookup = true;
-	int64 TimeoutMs = 0;
-	if (Executor.GlobalExecutionTimeoutSeconds > 0)
-	{
-		TimeoutMs = Executor.GlobalExecutionTimeoutSeconds * 1000LL;
-	}
-	return Executor.RemoteExecutor->GetExecution()->ExecuteAsync(State.ExecuteRequest, nullptr, TimeoutMs);
+	return Executor.RemoteExecutor->GetExecution()->ExecuteAsync(State.ExecuteRequest);
 }
 
 TFuture<TPair<FStatus, FBatchReadBlobsResponse>> FRemoteBuildExecutionRequest::DownloadResultsAsync()
