PackageDependencyTracker: Add the current TargetPlatform to package reference data for references made during cooking saves.
#rnx

diff --git a/Engine/Source/Editor/UnrealEd/Private/CookOnTheFlyServer.cpp b/Engine/Source/Editor/UnrealEd/Private/CookOnTheFlyServer.cpp
--- a/Engine/Source/Editor/UnrealEd/Private/CookOnTheFlyServer.cpp
+++ b/Engine/Source/Editor/UnrealEd/Private/CookOnTheFlyServer.cpp
@@ -4552,6 +4552,7 @@ void UCookOnTheFlyServer::SaveCookedPackage(UE::Cook::FSaveCookedPackageContext&
 		if (Context.bPlatformSetupSuccessful)
 		{
 			UE_SCOPED_HIERARCHICAL_COOKTIMER(GEditorSavePackage);
+			UE_TRACK_REFERENCING_PACKAGE_SCOPED_PLATFORM(TargetPlatform);
 
 			try
 			{
