TargetDomain: Move TargetDomain utils implementation into a new cpp file. Rename GetDependenciesObject to CollectDependenciesObject to contrast with FetchDependencies. Create stubs for the fetch functions.
#rnx

diff --git a/Engine/Source/Editor/UnrealEd/Private/CookOnTheFlyServer.cpp b/Engine/Source/Editor/UnrealEd/Private/CookOnTheFlyServer.cpp
--- a/Engine/Source/Editor/UnrealEd/Private/CookOnTheFlyServer.cpp
+++ b/Engine/Source/Editor/UnrealEd/Private/CookOnTheFlyServer.cpp
@@ -4777,7 +4777,7 @@ void FSaveCookedPackageContext::FinishPlatform()
 		if (COTFS.bTargetDomainEnabled)
 		{
 			UE_SCOPED_HIERARCHICAL_COOKTIMER(TargetDomainDependencies);
-			TargetDomainDependencies = UE::TargetDomain::GetDependenciesObject(Package, TargetPlatform, nullptr /* ErrorMessage */);
+			TargetDomainDependencies = UE::TargetDomain::CollectDependenciesObject(Package, TargetPlatform, nullptr /* ErrorMessage */);
 		}
 
 		IPackageStoreWriter::FCommitPackageInfo Info;
diff --git a/Engine/Source/Editor/UnrealEd/Private/TargetDomain/TargetDomainUtils.cpp b/Engine/Source/Editor/UnrealEd/Private/TargetDomain/TargetDomainUtils.cpp
new file mode 100644
--- /dev/null
+++ b/Engine/Source/Editor/UnrealEd/Private/TargetDomain/TargetDomainUtils.cpp
@@ -0,0 +1,77 @@
+// Copyright Epic Games, Inc. All Rights Reserved.
+
+#include "TargetDomain/TargetDomainUtils.h"
+
+#include "AssetRegistry/IAssetRegistry.h"
+#include "Containers/Set.h"
+#include "Containers/UnrealString.h"
+#include "Cooker/PackageBuildDependencyTracker.h"
+#include "EditorDomain/EditorDomainUtils.h"
+#include "IO/IoHash.h"
+#include "Misc/StringBuilder.h"
+#include "Serialization/CompactBinary.h"
+#include "Serialization/CompactBinaryWriter.h"
+
+namespace UE::TargetDomain
+{
+
+bool TryCollectKeyAndDependencies(UPackage* Package, const ITargetPlatform* TargetPlatform, FIoHash* OutHash, TArray<FName>* OutBuildDependencies,
+	TArray<FName>* OutRuntimeOnlyDependencies, FString* OutErrorMessage)
+{
+	// Implementation is coming in a separate changelist that mostly involves the PackageBuildDependencyTracker
+	if (OutErrorMessage) *OutErrorMessage = TEXT("Not yet implemented.");
+	return false;
+}
+
+FCbObject CollectDependenciesObject(UPackage* Package, const ITargetPlatform* TargetPlatform, FString* ErrorMessage)
+{
+	FIoHash TargetDomainKey;
+	TArray<FName> BuildDependencies;
+	TArray<FName> RuntimeOnlyDependencies;
+	if (!UE::TargetDomain::TryCollectKeyAndDependencies(Package, TargetPlatform, &TargetDomainKey, &BuildDependencies, &RuntimeOnlyDependencies, ErrorMessage))
+	{
+		return FCbObject();
+	}
+
+	FCbWriter Writer;
+	Writer.BeginObject();
+	Writer << "targetdomainkey" << TargetDomainKey;
+	TStringBuilder<128> PackageNameBuffer;
+	if (!BuildDependencies.IsEmpty())
+	{
+		Writer.BeginArray("builddependencies");
+		for (FName DependencyName : BuildDependencies)
+		{
+			DependencyName.ToString(PackageNameBuffer);
+			Writer << PackageNameBuffer;
+		}
+		Writer.EndArray();
+	}
+	if (!RuntimeOnlyDependencies.IsEmpty())
+	{
+		Writer.BeginArray("runtimeonlydependencies");
+		for (FName DependencyName : RuntimeOnlyDependencies)
+		{
+			DependencyName.ToString(PackageNameBuffer);
+			Writer << PackageNameBuffer;
+		}
+		Writer.EndArray();
+	}
+	return Writer.Save().AsObject();
+}
+
+
+bool TryFetchKeyAndDependencies(FName PackageName, const ITargetPlatform* TargetPlatform, FIoHash* OutHash, TArray<FName>* OutDependencies)
+{
+	// TODO: Read from oplog for the packagename, return true if and only if it has a targetdomainkey and the generated key from the collection of dependencies matches
+	return false;
+}
+
+bool TryFetchDependencies(FName PackageName, const ITargetPlatform* TargetPlatform,
+	TArray<FName>& OutBuildDependencies, TArray<FName>& OutRuntimeOnlyDependencies)
+{
+	// TODO: Read from oplog for the packagename, return true if and only if it has a targetdomainkey and the generated key from the collection of dependencies matches
+	return false;
+}
+
+}
diff --git a/Engine/Source/Editor/UnrealEd/Private/TargetDomain/TargetDomainUtils.h b/Engine/Source/Editor/UnrealEd/Private/TargetDomain/TargetDomainUtils.h
--- a/Engine/Source/Editor/UnrealEd/Private/TargetDomain/TargetDomainUtils.h
+++ b/Engine/Source/Editor/UnrealEd/Private/TargetDomain/TargetDomainUtils.h
@@ -3,62 +3,24 @@
 #pragma once
 
 #include "Containers/Array.h"
-#include "IO/IoHash.h"
-#include "Misc/StringBuilder.h"
-#include "Serialization/CompactBinary.h"
-#include "Serialization/CompactBinaryWriter.h"
 #include "UObject/NameTypes.h"
 
 class ITargetPlatform;
+class FCbObject;
+class FString;
 class UPackage;
 struct FIoHash;
 
 namespace UE::TargetDomain
 {
 
-inline bool TryCollectKeyAndDependencies(UPackage* Package, const ITargetPlatform* TargetPlatform,
-	FIoHash* OutHash, TArray<FName>* OutBuildDependencies, TArray<FName>* OutRuntimeOnlyDependencies, FString* OutErrorMessage)
-{
-	// Implementation is coming in a separate changelist that mostly involves the PackageBuildDependencyTracker
-	if (OutErrorMessage) *OutErrorMessage = TEXT("Not yet implemented.");
-	return false;
-}
+bool TryCollectKeyAndDependencies(UPackage* Package, const ITargetPlatform* TargetPlatform,
+	FIoHash* OutHash, TArray<FName>* OutBuildDependencies, TArray<FName>* OutRuntimeOnlyDependencies, FString* OutErrorMessage);
+FCbObject CollectDependenciesObject(UPackage* Package, const ITargetPlatform* TargetPlatform, FString* ErrorMessage);
 
-inline FCbObject GetDependenciesObject(UPackage* Package, const ITargetPlatform* TargetPlatform, FString* ErrorMessage)
-{
-	FIoHash TargetDomainKey;
-	TArray<FName> BuildDependencies;
-	TArray<FName> RuntimeOnlyDependencies;
-	if (!UE::TargetDomain::TryCollectKeyAndDependencies(Package, TargetPlatform, &TargetDomainKey, &BuildDependencies, &RuntimeOnlyDependencies, ErrorMessage))
-	{
-		return FCbObject();
-	}
-
-	FCbWriter Writer;
-	Writer.BeginObject();
-	Writer << "targetdomainkey" << TargetDomainKey;
-	TStringBuilder<128> PackageNameBuffer;
-	if (!BuildDependencies.IsEmpty())
-	{
-		Writer.BeginArray("builddependencies");
-		for (FName DependencyName : BuildDependencies)
-		{
-			DependencyName.ToString(PackageNameBuffer);
-			Writer << PackageNameBuffer;
-		}
-		Writer.EndArray();
-	}
-	if (!RuntimeOnlyDependencies.IsEmpty())
-	{
-		Writer.BeginArray("runtimeonlydependencies");
-		for (FName DependencyName : RuntimeOnlyDependencies)
-		{
-			DependencyName.ToString(PackageNameBuffer);
-			Writer << PackageNameBuffer;
-		}
-		Writer.EndArray();
-	}
-	return Writer.Save().AsObject();
-}
+bool TryFetchKeyAndDependencies(FName PackageName, const ITargetPlatform* TargetPlatform,
+	FIoHash* OutHash, TArray<FName>* OutBuildDependencies, TArray<FName>* OutRuntimeOnlyDependencies, FString* OutErrorMessage);
+bool TryFetchDependencies(FName PackageName, const ITargetPlatform* TargetPlatform,
+	TArray<FName>& OutBuildDependencies, TArray<FName>& OutRuntimeOnlyDependencies);
 
 }
