Log DerivedDataCache when registry overrides dirs
make editor dialog update when string is empty

diff --git a/Engine/Source/Developer/DerivedDataCache/Private/DerivedDataBackends.cpp b/Engine/Source/Developer/DerivedDataCache/Private/DerivedDataBackends.cpp
--- a/Engine/Source/Developer/DerivedDataCache/Private/DerivedDataBackends.cpp
+++ b/Engine/Source/Developer/DerivedDataCache/Private/DerivedDataBackends.cpp
@@ -561,6 +561,7 @@ public:
 				if (DDCPath.Len() > 0)
 				{
 					Path = DDCPath;
+					UE_LOG( LogDerivedDataCache, Log, TEXT("Found registry key GlobalDataCachePath %s=%s"), *EnvPathOverride, *Path );
 				}
 			}
 		}
@@ -732,7 +733,11 @@ public:
 				FString DDCPath;
 				if (FPlatformMisc::GetStoredValue(TEXT("Epic Games"), TEXT("GlobalDataCachePath"), *EnvPathOverride, DDCPath))
 				{
-					CachePath = DDCPath;
+					if ( DDCPath.Len() > 0 )
+					{
+						CachePath = DDCPath;			
+						UE_LOG( LogDerivedDataCache, Log, TEXT("Found registry key GlobalDataCachePath %s=%s"), *EnvPathOverride, *CachePath );
+					}
 				}
 			}
 		}
diff --git a/Engine/Source/Editor/UnrealEd/Private/Settings/EditorSettings.cpp b/Engine/Source/Editor/UnrealEd/Private/Settings/EditorSettings.cpp
--- a/Engine/Source/Editor/UnrealEd/Private/Settings/EditorSettings.cpp
+++ b/Engine/Source/Editor/UnrealEd/Private/Settings/EditorSettings.cpp
@@ -67,6 +67,9 @@ void UEditorSettings::PostEditChangeProperty( struct FPropertyChangedEvent& Prop
 		if (GlobalLocalDDCPath.Path.IsEmpty())
 		{
 			FPlatformMisc::DeleteStoredValue(TEXT("Epic Games"), TEXT("GlobalDataCachePath"), TEXT("UE-LocalDataCachePath"));
+			
+			// empty registry key means use environment, so re-fetch it now
+			GlobalLocalDDCPath.Path = FPlatformMisc::GetEnvironmentVariable(TEXT("UE-LocalDataCachePath"));
 		}
 		else
 		{
@@ -78,6 +81,9 @@ void UEditorSettings::PostEditChangeProperty( struct FPropertyChangedEvent& Prop
 		if (GlobalSharedDDCPath.Path.IsEmpty())
 		{
 			FPlatformMisc::DeleteStoredValue(TEXT("Epic Games"), TEXT("GlobalDataCachePath"), TEXT("UE-SharedDataCachePath"));
+			
+			// empty registry key means use environment, so re-fetch it now
+			GlobalSharedDDCPath.Path = FPlatformMisc::GetEnvironmentVariable(TEXT("UE-SharedDataCachePath"));
 		}
 		else
 		{
@@ -89,6 +95,9 @@ void UEditorSettings::PostEditChangeProperty( struct FPropertyChangedEvent& Prop
 		if (GlobalS3DDCPath.Path.IsEmpty())
 		{
 			FPlatformMisc::DeleteStoredValue(TEXT("Epic Games"), TEXT("GlobalDataCachePath"), TEXT("UE-S3DataCachePath"));
+			
+			// empty registry key means use environment, so re-fetch it now
+			GlobalS3DDCPath.Path = FPlatformMisc::GetEnvironmentVariable(TEXT("UE-S3DataCachePath"));
 		}
 		else
 		{
