Use correct validation method when consuming DDC2 remote executed build output.

diff --git a/Engine/Source/Editor/UnrealEd/Private/DerivedDataBuildRemoteExecutor.cpp b/Engine/Source/Editor/UnrealEd/Private/DerivedDataBuildRemoteExecutor.cpp
--- a/Engine/Source/Editor/UnrealEd/Private/DerivedDataBuildRemoteExecutor.cpp
+++ b/Engine/Source/Editor/UnrealEd/Private/DerivedDataBuildRemoteExecutor.cpp
@@ -752,7 +752,7 @@ FOptionalBuildOutput FRemoteBuildExecutionRequest::ComposeBuildOutput(FBatchRead
 		{
 			FSharedBuffer BuildOutputBuffer = MakeSharedBufferFromArray(MoveTemp(OutputReadResponse.Data));
 
-			if (ValidateCompactBinaryRange(BuildOutputBuffer, ECbValidateMode::Default) != ECbValidateError::None)
+			if (ValidateCompactBinary(BuildOutputBuffer, ECbValidateMode::Default) != ECbValidateError::None)
 			{
 				UE_LOG(LogDerivedDataBuildRemoteExecutor, Warning, TEXT("Worker error: build output structure not valid!"));
 				return FOptionalBuildOutput();
