Removing option in editor settings for controlling if audio device is created for every PIE.

diff --git a/Engine/Source/Editor/UnrealEd/Classes/Settings/LevelEditorPlaySettings.h b/Engine/Source/Editor/UnrealEd/Classes/Settings/LevelEditorPlaySettings.h
--- a/Engine/Source/Editor/UnrealEd/Classes/Settings/LevelEditorPlaySettings.h
+++ b/Engine/Source/Editor/UnrealEd/Classes/Settings/LevelEditorPlaySettings.h
@@ -388,16 +388,6 @@ private:
 	UPROPERTY(config, EditAnywhere, Category="Multiplayer Options|Client", meta=(EditCondition = "RunUnderOneProcess"))
 	bool RouteGamepadToSecondWindow;
 
-	/** 
-	* If checked, a separate audio device is created for every player. 
-	
-	* If unchecked, a separate audio device is created for only the first two players and uses the main audio device for more than 2 players.
-	*
-	* Enabling this will allow rendering accurate audio from every player's perspective but will use more CPU. Keep this disabled on lower-perf machines.
-	*/
-	UPROPERTY(config, EditAnywhere, Category = "Multiplayer Options|Client", meta=(EditCondition = "EnableGameSound && RunUnderOneProcess"))
-	bool CreateAudioDeviceForEveryPlayer;
-
 	/** Height to use when spawning additional windows. */
 	UPROPERTY(config, EditAnywhere, Category="Multiplayer Options|Client", meta=(ClampMin = 0))
 	int32 ClientWindowHeight;
@@ -522,7 +512,6 @@ public:
 	bool IsClientWindowSizeActive() const { return PlayNumberOfClients > 1; }
 	bool GetClientWindowSize( FIntPoint &OutClientWindowSize ) const { OutClientWindowSize = FIntPoint(ClientWindowWidth, ClientWindowHeight); return IsClientWindowSizeActive(); }
 	EVisibility GetClientWindowSizeVisibility() const { return (RunUnderOneProcess ? EVisibility::Hidden : EVisibility::Visible); }
-	bool IsCreateAudioDeviceForEveryPlayer() const { return CreateAudioDeviceForEveryPlayer; }
 
 	EBuildConfiguration GetLaunchBuildConfiguration() const
 	{
diff --git a/Engine/Source/Editor/UnrealEd/Private/PlayLevel.cpp b/Engine/Source/Editor/UnrealEd/Private/PlayLevel.cpp
--- a/Engine/Source/Editor/UnrealEd/Private/PlayLevel.cpp
+++ b/Engine/Source/Editor/UnrealEd/Private/PlayLevel.cpp
@@ -2822,11 +2822,9 @@ UGameInstance* UEditorEngine::CreateInnerProcessPIEGameInstance(FRequestPlaySess
 
 	if (!InPIEParameters.bRunAsDedicated)
 	{
-		bool bCreateNewAudioDevice = InParams.EditorPlaySettings->IsCreateAudioDeviceForEveryPlayer();
-
 		// Create an instance of the Game Viewport Client, with the class specified by the Engine.
 		ViewportClient = NewObject<UGameViewportClient>(this, GameViewportClientClass);
-		ViewportClient->Init(*PieWorldContext, GameInstance, bCreateNewAudioDevice);
+		ViewportClient->Init(*PieWorldContext, GameInstance);
 
 		ULevelEditorPlaySettings* PlayInSettings = GetMutableDefault<ULevelEditorPlaySettings>();
 		ViewportClient->EngineShowFlags.SetServerDrawDebug(PlayInSettings->ShowServerDebugDrawingByDefault());
diff --git a/Engine/Source/Runtime/Engine/Classes/Engine/GameViewportClient.h b/Engine/Source/Runtime/Engine/Classes/Engine/GameViewportClient.h
--- a/Engine/Source/Runtime/Engine/Classes/Engine/GameViewportClient.h
+++ b/Engine/Source/Runtime/Engine/Classes/Engine/GameViewportClient.h
@@ -148,7 +148,7 @@ public:
 	/* Returns the relevant game instance for this viewport */
 	UGameInstance* GetGameInstance() const;
 
-	virtual void Init(struct FWorldContext& WorldContext, UGameInstance* OwningGameInstance, bool bCreateNewAudioDevice = true);
+	virtual void Init(struct FWorldContext& WorldContext, UGameInstance* OwningGameInstance);
 
 public:
 	//~ Begin UObject Interface
diff --git a/Engine/Source/Runtime/Engine/Private/GameViewportClient.cpp b/Engine/Source/Runtime/Engine/Private/GameViewportClient.cpp
--- a/Engine/Source/Runtime/Engine/Private/GameViewportClient.cpp
+++ b/Engine/Source/Runtime/Engine/Private/GameViewportClient.cpp
@@ -465,7 +465,7 @@ void UGameViewportClient::SetEnabledStats(const TArray<FString>& InEnabledStats)
 }
 
 
-void UGameViewportClient::Init(struct FWorldContext& WorldContext, UGameInstance* OwningGameInstance, bool bCreateNewAudioDevice)
+void UGameViewportClient::Init(struct FWorldContext& WorldContext, UGameInstance* OwningGameInstance)
 {
 	// set reference to world context
 	WorldContext.AddRef(static_cast<UWorld*&>(World));
