TargetDomain: Fix crash when TargetDomain is enabled caused by missing EndObject.
#rnx

diff --git a/Engine/Source/Editor/UnrealEd/Private/TargetDomain/TargetDomainUtils.cpp b/Engine/Source/Editor/UnrealEd/Private/TargetDomain/TargetDomainUtils.cpp
--- a/Engine/Source/Editor/UnrealEd/Private/TargetDomain/TargetDomainUtils.cpp
+++ b/Engine/Source/Editor/UnrealEd/Private/TargetDomain/TargetDomainUtils.cpp
@@ -57,6 +57,7 @@ FCbObject CollectDependenciesObject(UPackage* Package, const ITargetPlatform* Ta
 		}
 		Writer.EndArray();
 	}
+	Writer.EndObject();
 	return Writer.Save().AsObject();
 }
 
