CookOnTheFlyServer.cpp factoring for use in multiple cpp files.
Move profiling code into CookProfiling.h/.cpp.
Remove FScopeAssign struct that duplicates TGuardValue.
#rnx

diff --git a/Engine/Source/Editor/UnrealEd/Private/CookOnTheFlyServer.cpp b/Engine/Source/Editor/UnrealEd/Private/CookOnTheFlyServer.cpp
--- a/Engine/Source/Editor/UnrealEd/Private/CookOnTheFlyServer.cpp
+++ b/Engine/Source/Editor/UnrealEd/Private/CookOnTheFlyServer.cpp
@@ -100,8 +100,6 @@
 #include "ShaderCodeLibrary.h"
 #include "ShaderCompiler.h"
 #include "ShaderLibraryChunkDataGenerator.h"
-#include "Stats/Stats.h"
-#include "Stats/StatsMisc.h"
 #include "String/Find.h"
 #include "TargetDomain/TargetDomainUtils.h"
 #include "UnrealEdGlobals.h"
@@ -170,55 +168,10 @@ static FAutoConsoleVariableRef CVarCookDisplayWarnBusyTime(
 	TEXT("Controls the time before the cooker will issue a warning that there is a deadlock in a busy queue.\n"),
 	ECVF_Default);
 
-
-#define PROFILE_NETWORK 0
-
-#if PROFILE_NETWORK
-double TimeTillRequestStarted = 0.0;
-double TimeTillRequestForfilled = 0.0;
-double TimeTillRequestForfilledError = 0.0;
-double WaitForAsyncFilesWrites = 0.0;
-FEvent *NetworkRequestEvent = nullptr;
-#endif
-
-#if ENABLE_COOK_STATS
-namespace DetailedCookStats
-{
-	// These times are externable so CookCommandlet can pick them up and merge them with its cook stats
-	double TickCookOnTheSideTimeSec = 0.0;
-	double TickCookOnTheSideLoadPackagesTimeSec = 0.0;
-	double TickCookOnTheSideResolveRedirectorsTimeSec = 0.0;
-	double TickCookOnTheSideSaveCookedPackageTimeSec = 0.0;
-	double TickCookOnTheSideBeginPrepareSaveTimeSec = 0.0;
-	double TickCookOnTheSideFinishPrepareSaveTimeSec = 0.0;
-	double GameCookModificationDelegateTimeSec = 0.0;
-	int32 PeakRequestQueueSize = 0;
-	int32 PeakLoadQueueSize = 0;
-	int32 PeakSaveQueueSize = 0;
-
-	// Stats tracked through FAutoRegisterCallback
-	uint32 NumPreloadedDependencies = 0;
-	FCookStatsManager::FAutoRegisterCallback RegisterCookOnTheFlyServerStats([](FCookStatsManager::AddStatFuncRef AddStat)
-		{
-			AddStat(TEXT("Package.Load"), FCookStatsManager::CreateKeyValueArray(TEXT("NumPreloadedDependencies"), NumPreloadedDependencies));
-			AddStat(TEXT("CookOnTheFlyServer"), FCookStatsManager::CreateKeyValueArray(TEXT("PeakRequestQueueSize"), PeakRequestQueueSize));
-			AddStat(TEXT("CookOnTheFlyServer"), FCookStatsManager::CreateKeyValueArray(TEXT("PeakLoadQueueSize"), PeakLoadQueueSize));
-			AddStat(TEXT("CookOnTheFlyServer"), FCookStatsManager::CreateKeyValueArray(TEXT("PeakSaveQueueSize"), PeakSaveQueueSize));
-		});
-}
-#endif
-
-
-
 ////////////////////////////////////////////////////////////////
 /// Cook on the fly server
 ///////////////////////////////////////////////////////////////
 
-DECLARE_STATS_GROUP(TEXT("Cooking"), STATGROUP_Cooking, STATCAT_Advanced);
-
-DECLARE_CYCLE_STAT(TEXT("Precache Derived data for platform"), STAT_TickPrecacheCooking, STATGROUP_Cooking);
-DECLARE_CYCLE_STAT(TEXT("Tick cooking"), STAT_TickCooker, STATGROUP_Cooking);
-
 constexpr uint32 ExpectedMaxNumPlatforms = 32;
 constexpr float TickCookableObjectsFrameTime = .100f;
 
@@ -233,28 +186,6 @@ const TCHAR* GeneratedPackageSubPath = TEXT("_Generated_");
 /* helper structs functions
  *****************************************************************************/
 
-
-/** Helper to assign to any variable for a scope period */
-template<class T>
-struct FScopeAssign
-{
-private:
-	T* Setting;
-	T OriginalValue;
-public:
-	FScopeAssign(T& InSetting, const T NewValue)
-	{
-		Setting = &InSetting;
-		OriginalValue = *Setting;
-		*Setting = NewValue;
-	}
-	~FScopeAssign()
-	{
-		*Setting = OriginalValue;
-	}
-};
-
-
 class FPackageSearchVisitor : public IPlatformFile::FDirectoryVisitor
 {
 	TArray<FString>& FoundFiles;
@@ -5720,7 +5651,7 @@ void UCookOnTheFlyServer::ProcessAccessedIniSettings(const FConfigFile* Config,
 
 bool UCookOnTheFlyServer::IniSettingsOutOfDate(const ITargetPlatform* TargetPlatform) const
 {
-	FScopeAssign<bool> A = FScopeAssign<bool>(IniSettingRecurse, true);
+	TGuardValue<bool> A(IniSettingRecurse, true);
 
 	FIniSettingContainer OldIniSettings;
 	TMap<FString, FString> OldAdditionalSettings;
@@ -5861,7 +5792,7 @@ bool UCookOnTheFlyServer::IniSettingsOutOfDate(const ITargetPlatform* TargetPlat
 
 bool UCookOnTheFlyServer::SaveCurrentIniSettings(const ITargetPlatform* TargetPlatform) const
 {
-	FScopeAssign<bool> S = FScopeAssign<bool>(IniSettingRecurse, true);
+	TGuardValue<bool> S(IniSettingRecurse, true);
 
 	TMap<FString, FString> AdditionalIniSettings;
 	GetAdditionalCurrentIniVersionStrings(TargetPlatform, AdditionalIniSettings);
diff --git a/Engine/Source/Editor/UnrealEd/Private/Cooker/CookProfiling.cpp b/Engine/Source/Editor/UnrealEd/Private/Cooker/CookProfiling.cpp
--- a/Engine/Source/Editor/UnrealEd/Private/Cooker/CookProfiling.cpp
+++ b/Engine/Source/Editor/UnrealEd/Private/Cooker/CookProfiling.cpp
@@ -170,3 +170,29 @@ void ClearHierarchyTimers()
 	RootTimerInfo.ClearChildren();
 }
 #endif
+
+#if ENABLE_COOK_STATS
+namespace DetailedCookStats
+{
+	double TickCookOnTheSideTimeSec = 0.0;
+	double TickCookOnTheSideLoadPackagesTimeSec = 0.0;
+	double TickCookOnTheSideResolveRedirectorsTimeSec = 0.0;
+	double TickCookOnTheSideSaveCookedPackageTimeSec = 0.0;
+	double TickCookOnTheSideBeginPrepareSaveTimeSec = 0.0;
+	double TickCookOnTheSideFinishPrepareSaveTimeSec = 0.0;
+	double GameCookModificationDelegateTimeSec = 0.0;
+
+	// Stats tracked through FAutoRegisterCallback
+	int32 PeakRequestQueueSize = 0;
+	int32 PeakLoadQueueSize = 0;
+	int32 PeakSaveQueueSize = 0;
+	uint32 NumPreloadedDependencies = 0;
+	FCookStatsManager::FAutoRegisterCallback RegisterCookOnTheFlyServerStats([](FCookStatsManager::AddStatFuncRef AddStat)
+		{
+			AddStat(TEXT("Package.Load"), FCookStatsManager::CreateKeyValueArray(TEXT("NumPreloadedDependencies"), NumPreloadedDependencies));
+			AddStat(TEXT("CookOnTheFlyServer"), FCookStatsManager::CreateKeyValueArray(TEXT("PeakRequestQueueSize"), PeakRequestQueueSize));
+			AddStat(TEXT("CookOnTheFlyServer"), FCookStatsManager::CreateKeyValueArray(TEXT("PeakLoadQueueSize"), PeakLoadQueueSize));
+			AddStat(TEXT("CookOnTheFlyServer"), FCookStatsManager::CreateKeyValueArray(TEXT("PeakSaveQueueSize"), PeakSaveQueueSize));
+		});
+}
+#endif
diff --git a/Engine/Source/Editor/UnrealEd/Private/Cooker/CookProfiling.h b/Engine/Source/Editor/UnrealEd/Private/Cooker/CookProfiling.h
--- a/Engine/Source/Editor/UnrealEd/Private/Cooker/CookProfiling.h
+++ b/Engine/Source/Editor/UnrealEd/Private/Cooker/CookProfiling.h
@@ -4,8 +4,10 @@
 
 #include "CoreTypes.h"
 #include "ProfilingDebugging/CookStats.h"
+#include "Stats/Stats.h"
 
 #define OUTPUT_COOKTIMING ENABLE_COOK_STATS
+#define PROFILE_NETWORK 0
 
 #if OUTPUT_COOKTIMING
 #include "Trace/Trace.inl"
@@ -78,4 +80,36 @@ UE_TRACE_CHANNEL_EXTERN(CookChannel)
 
 void OutputHierarchyTimers() {}
 void ClearHierarchyTimers() {}
-#endif
\ No newline at end of file
+#endif
+
+#if ENABLE_COOK_STATS
+namespace DetailedCookStats
+{
+	extern double TickCookOnTheSideTimeSec;
+	extern double TickCookOnTheSideLoadPackagesTimeSec;
+	extern double TickCookOnTheSideResolveRedirectorsTimeSec;
+	extern double TickCookOnTheSideSaveCookedPackageTimeSec;
+	extern double TickCookOnTheSideBeginPrepareSaveTimeSec;
+	extern double TickCookOnTheSideFinishPrepareSaveTimeSec;
+	extern double GameCookModificationDelegateTimeSec;
+
+	// Stats tracked through FAutoRegisterCallback
+	extern uint32 NumPreloadedDependencies;
+	extern int32 PeakRequestQueueSize;
+	extern int32 PeakLoadQueueSize;
+	extern int32 PeakSaveQueueSize;
+}
+#endif
+
+#if PROFILE_NETWORK
+double TimeTillRequestStarted = 0.0;
+double TimeTillRequestForfilled = 0.0;
+double TimeTillRequestForfilledError = 0.0;
+double WaitForAsyncFilesWrites = 0.0;
+FEvent* NetworkRequestEvent = nullptr;
+#endif
+
+DECLARE_STATS_GROUP(TEXT("Cooking"), STATGROUP_Cooking, STATCAT_Advanced);
+
+DECLARE_CYCLE_STAT(TEXT("Precache Derived data for platform"), STAT_TickPrecacheCooking, STATGROUP_Cooking);
+DECLARE_CYCLE_STAT(TEXT("Tick cooking"), STAT_TickCooker, STATGROUP_Cooking);
diff --git a/Engine/Source/Editor/UnrealEd/Private/Cooker/CookTypes.h b/Engine/Source/Editor/UnrealEd/Private/Cooker/CookTypes.h
--- a/Engine/Source/Editor/UnrealEd/Private/Cooker/CookTypes.h
+++ b/Engine/Source/Editor/UnrealEd/Private/Cooker/CookTypes.h
@@ -6,8 +6,11 @@
 #include "Containers/Set.h"
 #include "CookOnTheSide/CookOnTheFlyServer.h" // ECookTickFlags
 #include "HAL/Platform.h"
+#include "Logging/TokenizedMessage.h"
 #include "Templates/Function.h"
 
+class FString;
+
 #define COOK_CHECKSLOW_PACKAGEDATA 1
 #define DEBUG_COOKONTHEFLY 0
 
@@ -163,3 +166,5 @@ namespace Cook
 	void SetIsSchedulerThread(bool bValue);
 }
 }
+
+void LogCookerMessage(const FString& MessageText, EMessageSeverity::Type Severity);
