Fix local space gizmos not working with static mesh instance selections.

diff --git a/Engine/Source/Editor/UnrealEd/Private/EditorModeManager.cpp b/Engine/Source/Editor/UnrealEd/Private/EditorModeManager.cpp
--- a/Engine/Source/Editor/UnrealEd/Private/EditorModeManager.cpp
+++ b/Engine/Source/Editor/UnrealEd/Private/EditorModeManager.cpp
@@ -43,6 +43,8 @@
 #include "Tools/AssetEditorContextObject.h"
 #include "ContextObjectStore.h"
 
+#include "Elements/Interfaces/TypedElementWorldInterface.h"
+
 /*------------------------------------------------------------------------------
 	FEditorModeTools.
 
@@ -1155,7 +1157,6 @@ FMatrix FEditorModeTools::GetLocalCoordinateSystem()
 {
 	FMatrix Matrix = FMatrix::Identity;
 	// Let the current mode have a shot at setting the local coordinate system.
-	// If it doesn't want to, create it by looking at the currently selected actors list.
 
 	bool CustomCoordinateSystemProvided = false;
 	ForEachEdMode<ILegacyEdModeWidgetInterface>([&Matrix, &CustomCoordinateSystemProvided](ILegacyEdModeWidgetInterface* LegacyMode)
@@ -1164,21 +1165,14 @@ FMatrix FEditorModeTools::GetLocalCoordinateSystem()
 			return !CustomCoordinateSystemProvided;
 		});
 
-	if (!CustomCoordinateSystemProvided)
+	// If there isn't an active mode overriding the local coordinate system, create it by looking at the current selection.
+	if (!CustomCoordinateSystemProvided && GetEditorSelectionSet())
 	{
-		if (USceneComponent* SceneComponent = GetSelectedComponents()->GetBottom<USceneComponent>())
-		{
-			Matrix = FQuatRotationMatrix(SceneComponent->GetComponentQuat());
-		}
-		else
+		if (TTypedElement<UTypedElementWorldInterface> BottomSelected = GetEditorSelectionSet()->GetBottomSelectedElement<UTypedElementWorldInterface>())
 		{
-			const int32 Num = GetSelectedActors()->CountSelections<AActor>();
-
-			// Coordinate system needs to come from the last actor selected
-			if (Num > 0)
-			{
-				Matrix = FQuatRotationMatrix(GetSelectedActors()->GetBottom<AActor>()->GetActorQuat());
-			}
+			FTransform RelativeTransform;
+			BottomSelected.GetRelativeTransform(RelativeTransform);
+			Matrix = FQuatRotationMatrix(RelativeTransform.GetRotation());
 		}
 	}
 
