Fix crash when opening material.

diff --git a/Engine/Source/Editor/UnrealEd/Private/MaterialGraph.cpp b/Engine/Source/Editor/UnrealEd/Private/MaterialGraph.cpp
--- a/Engine/Source/Editor/UnrealEd/Private/MaterialGraph.cpp
+++ b/Engine/Source/Editor/UnrealEd/Private/MaterialGraph.cpp
@@ -581,7 +581,8 @@ bool UMaterialGraph::IsInputActive(UEdGraphPin* GraphPin) const
 	if (Material &&
 		RootNode && // No inputs without a root node
 		GraphPin->Direction == EGPD_Input &&
-		GraphPin->PinType.PinCategory != UMaterialGraphSchema::PC_Exec)
+		GraphPin->PinType.PinCategory != UMaterialGraphSchema::PC_Exec &&
+		GraphPin->SourceIndex < MaterialInputs.Num())
 	{
 		
 		return Material->IsPropertyActiveInEditor(MaterialInputs[GraphPin->SourceIndex].GetProperty());
