Fix Actors made writable not deleted on disk when Source Control is enabled.

#changelist validated

diff --git a/Engine/Source/Editor/UnrealEd/Private/ObjectTools.cpp b/Engine/Source/Editor/UnrealEd/Private/ObjectTools.cpp
--- a/Engine/Source/Editor/UnrealEd/Private/ObjectTools.cpp
+++ b/Engine/Source/Editor/UnrealEd/Private/ObjectTools.cpp
@@ -2196,7 +2196,8 @@ namespace ObjectTools
 
 					// Revert the file if it is checked out
 					const bool bIsAdded = SourceControlState->IsAdded();
-					if ( SourceControlState->IsCheckedOut() || bIsAdded || SourceControlState->IsDeleted() )
+					const bool bIsCheckedOut = SourceControlState->IsCheckedOut();
+					if ( bIsCheckedOut || bIsAdded || SourceControlState->IsDeleted() )
 					{
 						// Batch the revert operation so that we only make one request to the source control module.
 						SCCFilesToRevert.Add(FullPackageFilename);
@@ -2207,6 +2208,11 @@ namespace ObjectTools
 						// The file was open for add and reverted, this leaves the file on disk so here we delete it
 						IFileManager::Get().Delete(*PackageFilename);
 					}
+					else if ((!bIsCheckedOut) && (!IFileManager::Get().IsReadOnly(*PackageFilename)))
+					{
+						// The file was made writable and deleted, so we delete it on disk.
+						IFileManager::Get().Delete(*PackageFilename);
+					}
 					else
 					{
 						// Batch this file for deletion so that we only send one deletion request to the source control module.
