Log time in SaveAssetRegistry

diff --git a/Engine/Source/Editor/UnrealEd/Private/Commandlets/AssetRegistryGenerator.cpp b/Engine/Source/Editor/UnrealEd/Private/Commandlets/AssetRegistryGenerator.cpp
--- a/Engine/Source/Editor/UnrealEd/Private/Commandlets/AssetRegistryGenerator.cpp
+++ b/Engine/Source/Editor/UnrealEd/Private/Commandlets/AssetRegistryGenerator.cpp
@@ -1396,6 +1396,8 @@ void FAssetRegistryGenerator::AddAssetToFileOrderRecursive(const FName& InPackag
 bool FAssetRegistryGenerator::SaveAssetRegistry(const FString& SandboxPath, bool bSerializeDevelopmentAssetRegistry, bool bForceNoFilter)
 {
 	UE_LOG(LogAssetRegistryGenerator, Display, TEXT("Saving asset registry v%d."), FAssetRegistryVersion::Type::LatestVersion);
+	double StartTime = FPlatformTime::Seconds();
+
 	const TMap<FName, const FAssetData*>& ObjectToDataMap = State.GetObjectPathToAssetDataMap();
 	
 	// Write development first, this will always write
@@ -1536,8 +1538,9 @@ bool FAssetRegistryGenerator::SaveAssetRegistry(const FString& SandboxPath, bool
 			UE_LOG(LogAssetRegistryGenerator, Display, TEXT("Generated asset registry num assets %d, size is %5.2fkb"), ObjectToDataMap.Num(), (float)SerializedAssetRegistry.Num() / 1024.f);
 		}
 	}
-
-	UE_LOG(LogAssetRegistryGenerator, Display, TEXT("Done saving asset registry."));
+	
+	double StopTime = FPlatformTime::Seconds();
+	UE_LOG(LogAssetRegistryGenerator, Display, TEXT("Done saving asset registry. Took %.1f seconds."), StopTime-StartTime);
 
 	return true;
 }
