Default Linux to Visual Studio Code over Null Source Code

diff --git a/Engine/Config/BaseEditor.ini b/Engine/Config/BaseEditor.ini
--- a/Engine/Config/BaseEditor.ini
+++ b/Engine/Config/BaseEditor.ini
@@ -284,6 +284,7 @@ SearchForAnswersURL="http://udn.unrealengine.com/search.html"
 SupportWebsiteURL="http://udn.unrealengine.com/index.html"
 SourceCodeIDEURL_Windows="https://aka.ms/vs/16/release/vs_Community.exe"
 SourceCodeIDEURL_Mac="https://developer.apple.com/xcode"
+SourceCodeIDEURL_Linux="https://code.visualstudio.com/docs/setup/linux"
 SourceCodeIDEURL_Other="http://en.wikipedia.org/wiki/Integrated_development_environment"
 EpicGamesURL="http://UnrealEngine.com"
 FacebookURL="http://www.facebook.com/UnrealEngine"
diff --git a/Engine/Config/Linux/LinuxEngine.ini b/Engine/Config/Linux/LinuxEngine.ini
--- a/Engine/Config/Linux/LinuxEngine.ini
+++ b/Engine/Config/Linux/LinuxEngine.ini
@@ -10,7 +10,7 @@ fx.NiagaraAllowRuntimeScalabilityChanges=1
 ;r.RHICmdBypass=1
 
 [/Script/SourceCodeAccess.SourceCodeAccessSettings]
-PreferredAccessor=NullSourceCodeAccessor
+PreferredAccessor=VisualStudioCode
 
 [TextureStreaming]
 ; PoolSizeVRAMPercentage is how much percentage of GPU Dedicated VRAM should be used as a TexturePool cache for streaming textures (0 - unlimited streaming)
diff --git a/Engine/Config/LinuxArm64/LinuxArm64Engine.ini b/Engine/Config/LinuxArm64/LinuxArm64Engine.ini
--- a/Engine/Config/LinuxArm64/LinuxArm64Engine.ini
+++ b/Engine/Config/LinuxArm64/LinuxArm64Engine.ini
@@ -9,7 +9,7 @@ r.setres=1280x720
 ;r.RHICmdBypass=1
 
 [/Script/SourceCodeAccess.SourceCodeAccessSettings]
-PreferredAccessor=NullSourceCodeAccessor
+PreferredAccessor=VisualStudioCode
 
 [TextureStreaming]
 ; PoolSizeVRAMPercentage is how much percentage of GPU Dedicated VRAM should be used as a TexturePool cache for streaming textures (0 - unlimited streaming)
@@ -28,4 +28,4 @@ g.TimeoutForBlockOnRenderFence=60000
 PlatformRequiresDataCrypto=True
 
 [CrashReportClient]
-bAgreeToCrashUpload=true
\ No newline at end of file
+bAgreeToCrashUpload=true
diff --git a/Engine/Source/Editor/UnrealEd/Private/SourceCodeNavigation.cpp b/Engine/Source/Editor/UnrealEd/Private/SourceCodeNavigation.cpp
--- a/Engine/Source/Editor/UnrealEd/Private/SourceCodeNavigation.cpp
+++ b/Engine/Source/Editor/UnrealEd/Private/SourceCodeNavigation.cpp
@@ -1810,7 +1810,7 @@ FText FSourceCodeNavigation::GetSuggestedSourceCodeIDE(bool bShortIDEName)
 #elif PLATFORM_MAC
 	return LOCTEXT("SuggestedCodeIDE_Mac", "Xcode");
 #elif PLATFORM_LINUX
-	return LOCTEXT("SuggestedCodeIDE_Linux", "NullSourceCodeAccessor");
+	return LOCTEXT("SuggestedCodeIDE_Linux", "Visual Studio Code");
 #else
 	return LOCTEXT("SuggestedCodeIDE_Generic", "an IDE to edit source code");
 #endif
@@ -1845,6 +1845,9 @@ FString FSourceCodeNavigation::GetSuggestedSourceCodeIDEDownloadURL()
 #elif PLATFORM_MAC
 	// Xcode
 	FUnrealEdMisc::Get().GetURL( TEXT("SourceCodeIDEURL_Mac"), SourceCodeIDEURL );
+#elif PLATFORM_LINUX
+	// VSCode
+	FUnrealEdMisc::Get().GetURL( TEXT("SourceCodeIDEURL_Linux"), SourceCodeIDEURL );
 #else
 	// Unknown platform, just link to wikipedia page on IDEs
 	FUnrealEdMisc::Get().GetURL( TEXT("SourceCodeIDEURL_Other"), SourceCodeIDEURL );
