Extend viewport focusing to work with typed elements

diff --git a/Engine/Source/Editor/EditorFramework/Private/Elements/Framework/TypedElementViewportInteraction.cpp b/Engine/Source/Editor/EditorFramework/Private/Elements/Framework/TypedElementViewportInteraction.cpp
--- a/Engine/Source/Editor/EditorFramework/Private/Elements/Framework/TypedElementViewportInteraction.cpp
+++ b/Engine/Source/Editor/EditorFramework/Private/Elements/Framework/TypedElementViewportInteraction.cpp
@@ -130,6 +130,11 @@ void FTypedElementViewportInteractionCustomization::MirrorElement(const TTypedEl
 	}
 }
 
+bool FTypedElementViewportInteractionCustomization::GetFocusBounds(const TTypedElement<UTypedElementWorldInterface>& InElementWorldHandle, FBoxSphereBounds& OutBounds)
+{
+	return InElementWorldHandle.GetBounds(OutBounds);
+}
+
 
 void UTypedElementViewportInteraction::BeginGizmoManipulation(FTypedElementListConstRef InElementsToMove, const UE::Widget::EWidgetMode InWidgetMode)
 {
@@ -218,6 +223,21 @@ void UTypedElementViewportInteraction::MirrorElement(const FTypedElementHandle&
 	}
 }
 
+bool UTypedElementViewportInteraction::GetFocusBounds(FTypedElementListConstRef InElements, FBoxSphereBounds& OutBounds)
+{
+	bool bAnyHaveBounds = false;
+	InElements->ForEachElementHandle([this, &OutBounds, &bAnyHaveBounds](const FTypedElementHandle& InElementToMove)
+	{
+		FTypedElementViewportInteractionElement ViewportInteractionElement = ResolveViewportInteractionElement(InElementToMove);
+		if (ViewportInteractionElement)
+		{
+			bAnyHaveBounds = ViewportInteractionElement.GetFocusBounds(OutBounds) || bAnyHaveBounds;
+		}
+		return true;
+	});
+	return bAnyHaveBounds;
+}
+
 FTypedElementViewportInteractionElement UTypedElementViewportInteraction::ResolveViewportInteractionElement(const FTypedElementHandle& InElementHandle) const
 {
 	return InElementHandle
diff --git a/Engine/Source/Editor/EditorFramework/Public/Elements/Framework/TypedElementViewportInteraction.h b/Engine/Source/Editor/EditorFramework/Public/Elements/Framework/TypedElementViewportInteraction.h
--- a/Engine/Source/Editor/EditorFramework/Public/Elements/Framework/TypedElementViewportInteraction.h
+++ b/Engine/Source/Editor/EditorFramework/Public/Elements/Framework/TypedElementViewportInteraction.h
@@ -35,6 +35,7 @@ public:
 	virtual void GizmoManipulationStopped(const TTypedElement<UTypedElementWorldInterface>& InElementWorldHandle, const UE::Widget::EWidgetMode InWidgetMode, const ETypedElementViewportInteractionGizmoManipulationType InManipulationType);
 	virtual void PostGizmoManipulationStopped(TArrayView<const FTypedElementHandle> InElementHandles, const UE::Widget::EWidgetMode InWidgetMode);
 	virtual void MirrorElement(const TTypedElement<UTypedElementWorldInterface>& InElementWorldHandle, const FVector& InMirrorScale, const FVector& InPivotLocation);
+	virtual bool GetFocusBounds(const TTypedElement<UTypedElementWorldInterface>& InElementWorldHandle, FBoxSphereBounds& OutBounds);
 };
 
 /**
@@ -74,7 +75,8 @@ public:
 	void GizmoManipulationDeltaUpdate(const UE::Widget::EWidgetMode InWidgetMode, const EAxisList::Type InDragAxis, const FInputDeviceState& InInputState, const FTransform& InDeltaTransform, const FVector& InPivotLocation) const { ViewportInteractionCustomization->GizmoManipulationDeltaUpdate(ElementWorldHandle, InWidgetMode, InDragAxis, InInputState, InDeltaTransform, InPivotLocation); }
 	void GizmoManipulationStopped(const UE::Widget::EWidgetMode InWidgetMode, const ETypedElementViewportInteractionGizmoManipulationType InManipulationType) const { ViewportInteractionCustomization->GizmoManipulationStopped(ElementWorldHandle, InWidgetMode, InManipulationType); }
 	void MirrorElement(const FVector& InMirrorScale, const FVector& InPivotLocation) const { ViewportInteractionCustomization->MirrorElement(ElementWorldHandle, InMirrorScale, InPivotLocation); }
-
+	bool GetFocusBounds(FBoxSphereBounds& OutBounds) const { return ViewportInteractionCustomization->GetFocusBounds(ElementWorldHandle, OutBounds); }
+	
 private:
 	TTypedElement<UTypedElementWorldInterface> ElementWorldHandle;
 	FTypedElementViewportInteractionCustomization* ViewportInteractionCustomization = nullptr;
@@ -119,6 +121,11 @@ public:
 	 */
 	void MirrorElement(const FTypedElementHandle& InElementHandle, const FVector& InMirrorScale);
 
+	/**
+	 *	Calculate the bounds of all elements in the list, to be focused on in the viewport
+	 *	@note This list should have been pre-normalized via UTypedElementSelectionSet::GetNormalizedSelection or UTypedElementSelectionSet::GetNormalizedElementList, or FLevelEditorViewportClient::GetElementsToManipulate.
+	 */
+	bool GetFocusBounds(FTypedElementListConstRef InElements, FBoxSphereBounds& OutBox);
 private:
 	/**
 	 * Attempt to resolve the selection interface and viewport interaction customization for the given element, if any.
diff --git a/Engine/Source/Editor/UnrealEd/Classes/Editor/EditorEngine.h b/Engine/Source/Editor/UnrealEd/Classes/Editor/EditorEngine.h
--- a/Engine/Source/Editor/UnrealEd/Classes/Editor/EditorEngine.h
+++ b/Engine/Source/Editor/UnrealEd/Classes/Editor/EditorEngine.h
@@ -70,6 +70,7 @@ struct FAnalyticsEventAttribute;
 class UEditorWorldExtensionManager;
 class ITargetDevice;
 class ULevelEditorDragDropHandler;
+class UTypedElementSelectionSet;
 
 //
 // Things to set in mapSetBrush.
@@ -1003,6 +1004,13 @@ public:
 	*/
 	void MoveViewportCamerasToComponent(const USceneComponent* Component, bool bActiveViewportOnly);
 
+	/**
+	* Moves all viewport cameras to focus on the provided set of elements.
+	* @param	SelectionSet			Target elements
+	* @param	bActiveViewportOnly		If true, move/reorient only the active viewport.
+	*/
+	void MoveViewportCamerasToElement(const UTypedElementSelectionSet* SelectionSet, bool bActiveViewportOnly) const;
+
 	/**
 	 * Moves all viewport cameras to focus on the provided bounding box.
 	 * @param	BoundingBox				Target box
diff --git a/Engine/Source/Editor/UnrealEd/Private/EditorServer.cpp b/Engine/Source/Editor/UnrealEd/Private/EditorServer.cpp
--- a/Engine/Source/Editor/UnrealEd/Private/EditorServer.cpp
+++ b/Engine/Source/Editor/UnrealEd/Private/EditorServer.cpp
@@ -166,6 +166,7 @@
 #include "EditorBuildUtils.h"
 #include "Subsystems/BrushEditingSubsystem.h"
 #include "EdMode.h"
+#include "LevelEditorSubsystem.h"
 
 #include "Serialization/StructuredArchive.h"
 #include "Serialization/Formatters/JsonArchiveInputFormatter.h"
@@ -4823,20 +4824,6 @@ void UEditorEngine::MoveViewportCamerasToActor(const TArray<AActor*> &Actors, co
 		}
 	}
 
-	struct ComponentTypeMatcher
-	{
-		ComponentTypeMatcher(UPrimitiveComponent* InComponentToMatch)
-			: ComponentToMatch(InComponentToMatch)
-		{}
-
-		bool operator()(const UClass* ComponentClass) const
-		{
-			return ComponentToMatch->IsA(ComponentClass);
-		}
-
-		UPrimitiveComponent* ComponentToMatch;
-	};
-
 	TArray<AActor*> InvisLevelActors;
 
 	// Create a bounding volume of all of the selected actors.
@@ -4983,6 +4970,48 @@ void UEditorEngine::MoveViewportCamerasToComponent(const USceneComponent* Compon
 	}
 }
 
+void UEditorEngine::MoveViewportCamerasToElement(const UTypedElementSelectionSet* SelectionSet, bool bActiveViewportOnly) const
+{
+	if (!SelectionSet || !SelectionSet->HasSelectedElements(UTypedElementWorldInterface::StaticClass()))
+	{
+		return;
+	}
+
+	UTypedElementRegistry* Registry = UTypedElementRegistry::GetInstance();
+	if (!Registry)
+	{
+		return;
+	}
+	
+	// Create a bounding volume of all of the selected elements.
+	FBox BoundingBox(ForceInit);
+
+	FTypedElementListConstRef ElementsToView = SelectionSet->GetNormalizedElementList(SelectionSet->GetElementList(), FTypedElementSelectionNormalizationOptions());
+
+	if (bActiveViewportOnly)
+	{
+		FBoxSphereBounds Bounds(ForceInit);
+		if (GCurrentLevelEditingViewportClient->GetFocusBounds(ElementsToView, Bounds))
+		{
+			BoundingBox += Bounds.GetBox();
+		}
+	}
+	else
+	{
+		// Use all viewports.
+		for (FLevelEditorViewportClient* LinkedViewportClient : GetLevelViewportClients())
+		{
+			FBoxSphereBounds Bounds(ForceInit);
+			if (LinkedViewportClient->GetFocusBounds(ElementsToView, Bounds))
+			{
+				BoundingBox += Bounds.GetBox();
+			}
+		}
+	}
+
+	MoveViewportCamerasToBox(BoundingBox, bActiveViewportOnly);
+}
+
 void UEditorEngine::MoveViewportCamerasToBox(const FBox& BoundingBox, bool bActiveViewportOnly) const
 {
 	// Make sure we had at least one non-null actor in the array passed in.
@@ -5299,32 +5328,21 @@ bool UEditorEngine::Exec_Camera( const TCHAR* Str, FOutputDevice& Ar )
 			}
 			else
 			{
-				TArray<AActor*> Actors;
-				for (FSelectionIterator It(GetSelectedActorIterator()); It; ++It)
+				const UTypedElementSelectionSet* SelectionSet = GetSelectedActors()->GetElementSelectionSet();
+				if (SelectionSet && SelectionSet->HasSelectedElements())
 				{
-					AActor* Actor = static_cast<AActor*>(*It);
-					checkSlow(Actor->IsA(AActor::StaticClass()));
-					Actors.Add(Actor);
+					MoveViewportCamerasToElement(SelectionSet, bActiveViewportOnly);
 				}
-
-				TArray<UPrimitiveComponent*> SelectedComponents;
-				for (FSelectionIterator It(GetSelectedComponentIterator()); It; ++It)
-				{
-					UPrimitiveComponent* PrimitiveComp = Cast<UPrimitiveComponent>(*It);
-					if (PrimitiveComp)
-					{
-						SelectedComponents.Add(PrimitiveComp);
-					}
-				}
-
+				/*
 				if (Actors.Num() || SelectedComponents.Num())
 				{
 					MoveViewportCamerasToActor(Actors, SelectedComponents, bActiveViewportOnly);
 					return true;
 				}
+				*/
 				else
-				{
-					Ar.Log(TEXT("Can't find target actor or component."));
+				{					
+					Ar.Log(TEXT("Can't find target element."));
 					return false;
 				}
 			}
diff --git a/Engine/Source/Editor/UnrealEd/Private/Elements/Actor/ActorElementEditorViewportInteractionCustomization.cpp b/Engine/Source/Editor/UnrealEd/Private/Elements/Actor/ActorElementEditorViewportInteractionCustomization.cpp
--- a/Engine/Source/Editor/UnrealEd/Private/Elements/Actor/ActorElementEditorViewportInteractionCustomization.cpp
+++ b/Engine/Source/Editor/UnrealEd/Private/Elements/Actor/ActorElementEditorViewportInteractionCustomization.cpp
@@ -9,6 +9,10 @@
 #include "Toolkits/IToolkitHost.h"
 #include "AI/NavigationSystemBase.h"
 #include "Components/BrushComponent.h"
+#include "Engine/DocumentationActor.h"
+#include "Particles/Emitter.h"
+#include "Elements/Interfaces/TypedElementHierarchyInterface.h"
+#include "Elements/Framework/TypedElementRegistry.h"
 
 bool FActorElementEditorViewportInteractionCustomization::GetGizmoPivotLocation(const TTypedElement<UTypedElementWorldInterface>& InElementWorldHandle, const UE::Widget::EWidgetMode InWidgetMode, FVector& OutPivotLocation)
 {
@@ -53,6 +57,45 @@ void FActorElementEditorViewportInteractionCustomization::MirrorElement(const TT
 	Actor->MarkPackageDirty();
 }
 
+bool FActorElementEditorViewportInteractionCustomization::GetFocusBounds( const TTypedElement<UTypedElementWorldInterface>& InElementWorldHandle, FBoxSphereBounds& OutBounds)
+{
+	AActor* Actor = ActorElementDataUtil::GetActorFromHandle(InElementWorldHandle);
+	if (!Actor)
+	{
+		return false;
+	}
+
+	// open the documentation of documentation actors
+	if (ADocumentationActor* DocActor = Cast<ADocumentationActor>(Actor))
+	{
+		DocActor->OpenDocumentLink();
+	}
+
+	// Create a bounding volume of all of the sub-elements
+	FBox BoundingBox(ForceInit);
+
+	const bool bActorIsEmitter = (Cast<AEmitter>(Actor) != nullptr);
+
+	if (bActorIsEmitter && GEditor && GEditor->bCustomCameraAlignEmitter)
+	{
+		const float CustomCameraAlignEmitterDistance = GEditor->CustomCameraAlignEmitterDistance;
+		const FVector DefaultExtent(CustomCameraAlignEmitterDistance, CustomCameraAlignEmitterDistance,
+		                            CustomCameraAlignEmitterDistance);
+		const FBox DefaultSizeBox(Actor->GetActorLocation() - DefaultExtent, Actor->GetActorLocation() + DefaultExtent);
+		BoundingBox += DefaultSizeBox;
+	}
+
+	BoundingBox += Actor->GetComponentsBoundingBox(/*bNonColliding*/true, /*bIncludeFromChildActors*/true);
+	
+	if (USceneComponent* RootComponent = Actor->GetRootComponent())
+	{
+		BoundingBox += RootComponent->GetComponentLocation();
+	}
+	
+	OutBounds = BoundingBox;
+	return true;
+}
+
 void FActorElementEditorViewportInteractionCustomization::ApplyDeltaToActor(AActor* InActor, const bool InIsDelta, const FVector* InDeltaTranslationPtr, const FRotator* InDeltaRotationPtr, const FVector* InDeltaScalePtr, const FVector& InPivotLocation, const FInputDeviceState& InInputState)
 {
 	const bool bIsSimulatingInEditor = GEditor->IsSimulatingInEditor();
diff --git a/Engine/Source/Editor/UnrealEd/Private/Elements/Component/ComponentElementEditorViewportInteractionCustomization.cpp b/Engine/Source/Editor/UnrealEd/Private/Elements/Component/ComponentElementEditorViewportInteractionCustomization.cpp
--- a/Engine/Source/Editor/UnrealEd/Private/Elements/Component/ComponentElementEditorViewportInteractionCustomization.cpp
+++ b/Engine/Source/Editor/UnrealEd/Private/Elements/Component/ComponentElementEditorViewportInteractionCustomization.cpp
@@ -48,6 +48,67 @@ void FComponentElementEditorViewportInteractionCustomization::GizmoManipulationD
 	}
 }
 
+namespace EditorEngineDefs
+{
+/** Limit the minimum size of the bounding box when centering cameras on individual components to avoid extreme zooming */
+static const float MinComponentBoundsForZoom = 50.0f;
+}
+
+bool FComponentElementEditorViewportInteractionCustomization::GetFocusBounds(const TTypedElement<UTypedElementWorldInterface>& InElementWorldHandle, FBoxSphereBounds& OutBounds)
+{
+	if (UActorComponent* Component = ComponentElementDataUtil::GetComponentFromHandle(InElementWorldHandle))
+	{
+		if (UPrimitiveComponent* PrimitiveComponent = Cast<UPrimitiveComponent>(Component))
+		{
+			if (!PrimitiveComponent->IsRegistered())
+			{
+				return false;
+			}
+
+			AActor* AlignActor = Component->GetOwner();
+			if (!AlignActor)
+			{
+				return false;
+			}
+
+			if (USceneComponent* RootComponent = AlignActor->GetRootComponent())
+			{
+				TArray<USceneComponent*> SceneComponents;
+				RootComponent->GetChildrenComponents(true, SceneComponents);
+				SceneComponents.Add(RootComponent);
+				// Some components can have huge bounds but are not visible.  Ignore these components unless it is the only component on the actor 
+				const bool bIgnore = SceneComponents.Num() > 1 && PrimitiveComponent->IgnoreBoundsForEditorFocus();
+
+				if (!bIgnore)
+				{
+					FBox LocalBox(ForceInit);
+					if (!GLevelEditorModeTools().ComputeBoundingBoxForViewportFocus(
+						AlignActor, PrimitiveComponent, LocalBox))
+					{
+						LocalBox = PrimitiveComponent->Bounds.GetBox();
+						FVector Center;
+						FVector Extents;
+						LocalBox.GetCenterAndExtents(Center, Extents);
+
+						// Apply a minimum size to the extents of the component's box to avoid the camera's zooming too close to small or zero-sized components
+						if (Extents.SizeSquared() < EditorEngineDefs::MinComponentBoundsForZoom *
+							EditorEngineDefs::MinComponentBoundsForZoom)
+						{
+							const FVector NewExtents(EditorEngineDefs::MinComponentBoundsForZoom, SMALL_NUMBER,
+							                         SMALL_NUMBER);
+							LocalBox = FBox(Center - NewExtents, Center + NewExtents);
+						}
+					}
+					OutBounds = LocalBox;
+
+					return true;
+				}
+			}
+		}
+	}
+	return false;
+}
+
 void FComponentElementEditorViewportInteractionCustomization::ApplyDeltaToComponent(USceneComponent* InComponent, const bool InIsDelta, const FVector* InDeltaTranslationPtr, const FRotator* InDeltaRotationPtr, const FVector* InDeltaScalePtr, const FVector& InPivotLocation, const FInputDeviceState& InInputState)
 {
 	if (GEditor->IsDeltaModificationEnabled())
diff --git a/Engine/Source/Editor/UnrealEd/Private/LevelEditorViewport.cpp b/Engine/Source/Editor/UnrealEd/Private/LevelEditorViewport.cpp
--- a/Engine/Source/Editor/UnrealEd/Private/LevelEditorViewport.cpp
+++ b/Engine/Source/Editor/UnrealEd/Private/LevelEditorViewport.cpp
@@ -3845,6 +3845,11 @@ void FLevelEditorViewportClient::MirrorSelectedElements(const FVector& InMirrorS
 	RedrawAllViewportsIntoThisScene();
 }
 
+bool FLevelEditorViewportClient::GetFocusBounds(FTypedElementListConstRef InElements, FBoxSphereBounds& OutBounds)
+{
+	return ViewportInteraction->GetFocusBounds(InElements, OutBounds);
+}
+
 FTransform FLevelEditorViewportClient::CachePreDragActorTransform(const AActor* InActor)
 {
 	if (const FTransform* PreDragTransform = PreDragActorTransforms.Find(InActor))
diff --git a/Engine/Source/Editor/UnrealEd/Public/EditorViewportClient.h b/Engine/Source/Editor/UnrealEd/Public/EditorViewportClient.h
--- a/Engine/Source/Editor/UnrealEd/Public/EditorViewportClient.h
+++ b/Engine/Source/Editor/UnrealEd/Public/EditorViewportClient.h
@@ -973,7 +973,6 @@ public:
 	 * Focuses the viewport to the center of the bounding box ensuring that the entire box is in view 
 	 *
 	 * @param BoundingBox			The box to focus
-	 * @param bUpdateLinkedOrthoViewports	Whether or not to updated linked viewports when this viewport changes
 	 * @param bInstant			Whether or not to focus the viewport instantly or over time
 	 */
 	void FocusViewportOnBox( const FBox& BoundingBox, bool bInstant = false );
diff --git a/Engine/Source/Editor/UnrealEd/Public/Elements/Actor/ActorElementEditorViewportInteractionCustomization.h b/Engine/Source/Editor/UnrealEd/Public/Elements/Actor/ActorElementEditorViewportInteractionCustomization.h
--- a/Engine/Source/Editor/UnrealEd/Public/Elements/Actor/ActorElementEditorViewportInteractionCustomization.h
+++ b/Engine/Source/Editor/UnrealEd/Public/Elements/Actor/ActorElementEditorViewportInteractionCustomization.h
@@ -13,6 +13,7 @@ public:
 	virtual bool GetGizmoPivotLocation(const TTypedElement<UTypedElementWorldInterface>& InElementWorldHandle, const UE::Widget::EWidgetMode InWidgetMode, FVector& OutPivotLocation) override;
 	virtual void GizmoManipulationDeltaUpdate(const TTypedElement<UTypedElementWorldInterface>& InElementWorldHandle, const UE::Widget::EWidgetMode InWidgetMode, const EAxisList::Type InDragAxis, const FInputDeviceState& InInputState, const FTransform& InDeltaTransform, const FVector& InPivotLocation) override;
 	virtual void MirrorElement(const TTypedElement<UTypedElementWorldInterface>& InElementWorldHandle, const FVector& InMirrorScale, const FVector& InPivotLocation) override;
-
+	virtual bool GetFocusBounds(const TTypedElement<UTypedElementWorldInterface>& InElementWorldHandle, FBoxSphereBounds& OutBounds) override;
+	
 	static void ApplyDeltaToActor(AActor* InActor, const bool InIsDelta, const FVector* InDeltaTranslationPtr, const FRotator* InDeltaRotationPtr, const FVector* InDeltaScalePtr, const FVector& InPivotLocation, const FInputDeviceState& InInputState);
 };
diff --git a/Engine/Source/Editor/UnrealEd/Public/Elements/Component/ComponentElementEditorViewportInteractionCustomization.h b/Engine/Source/Editor/UnrealEd/Public/Elements/Component/ComponentElementEditorViewportInteractionCustomization.h
--- a/Engine/Source/Editor/UnrealEd/Public/Elements/Component/ComponentElementEditorViewportInteractionCustomization.h
+++ b/Engine/Source/Editor/UnrealEd/Public/Elements/Component/ComponentElementEditorViewportInteractionCustomization.h
@@ -13,6 +13,7 @@ class UNREALED_API FComponentElementEditorViewportInteractionCustomization : pub
 public:
 	virtual bool GetGizmoPivotLocation(const TTypedElement<UTypedElementWorldInterface>& InElementWorldHandle, const UE::Widget::EWidgetMode InWidgetMode, FVector& OutPivotLocation) override;
 	virtual void GizmoManipulationDeltaUpdate(const TTypedElement<UTypedElementWorldInterface>& InElementWorldHandle, const UE::Widget::EWidgetMode InWidgetMode, const EAxisList::Type InDragAxis, const FInputDeviceState& InInputState, const FTransform& InDeltaTransform, const FVector& InPivotLocation) override;
-
+	virtual bool GetFocusBounds(const TTypedElement<UTypedElementWorldInterface>& InElementWorldHandle, FBoxSphereBounds& OutBounds) override;
+	
 	static void ApplyDeltaToComponent(USceneComponent* InComponent, const bool InIsDelta, const FVector* InDeltaTranslationPtr, const FRotator* InDeltaRotationPtr, const FVector* InDeltaScalePtr, const FVector& InPivotLocation, const FInputDeviceState& InInputState);
 };
diff --git a/Engine/Source/Editor/UnrealEd/Public/LevelEditorViewport.h b/Engine/Source/Editor/UnrealEd/Public/LevelEditorViewport.h
--- a/Engine/Source/Editor/UnrealEd/Public/LevelEditorViewport.h
+++ b/Engine/Source/Editor/UnrealEd/Public/LevelEditorViewport.h
@@ -328,6 +328,8 @@ public:
 	void MirrorSelectedActors(const FVector& InMirrorScale);
 	void MirrorSelectedElements(const FVector& InMirrorScale);
 
+	bool GetFocusBounds(FTypedElementListConstRef InElements, FBoxSphereBounds& OutBounds);
+
 	/**
 	 * Get the elements (from the current selection set) that this viewport can manipulate (eg, via the transform gizmo).
 	 */
