Editor static mesh factory uses placement partition actor instead of foliage partition actor.

diff --git a/Engine/Source/Editor/AssetPlacementEdMode/Private/AssetPlacementEdMode.cpp b/Engine/Source/Editor/AssetPlacementEdMode/Private/AssetPlacementEdMode.cpp
--- a/Engine/Source/Editor/AssetPlacementEdMode/Private/AssetPlacementEdMode.cpp
+++ b/Engine/Source/Editor/AssetPlacementEdMode/Private/AssetPlacementEdMode.cpp
@@ -8,7 +8,6 @@
 #include "AssetPlacementSettings.h"
 #include "EdModeInteractiveToolsContext.h"
 #include "EditorModeManager.h"
-#include "InstancedFoliageActor.h"
 #include "Selection.h"
 #include "EngineUtils.h"
 
@@ -20,16 +19,14 @@
 #include "Tools/PlacementEraseTool.h"
 
 #include "Factories/AssetFactoryInterface.h"
-#include "Elements/Actor/ActorElementData.h"
 #include "Elements/Framework/EngineElementsLibrary.h"
 #include "Elements/Framework/TypedElementCommonActions.h"
 #include "Elements/Framework/TypedElementRegistry.h"
+#include "Elements/SMInstance/SMInstanceElementData.h" // For SMInstanceElementDataUtil::SMInstanceElementsEnabled
 
 #include "Settings/LevelEditorMiscSettings.h"
 #include "Modes/PlacementModeSubsystem.h"
 
-#include "InstancedFoliageActor.h"
-
 #define LOCTEXT_NAMESPACE "AssetPlacementEdMode"
 
 constexpr TCHAR UAssetPlacementEdMode::AssetPlacementEdModeID[];
@@ -74,15 +71,6 @@ void UAssetPlacementEdMode::Enter()
 
 void UAssetPlacementEdMode::Exit()
 {
-	// Todo: remove with HISM element handles
-	for (TActorIterator<AInstancedFoliageActor> It(GetWorld()); It; ++It)
-	{
-		AInstancedFoliageActor* IFA = *It;
-		// null component will clear all selection
-		constexpr bool bAppendSelection = false;
-		IFA->SelectInstance(nullptr, 0, bAppendSelection);
-	}
-
 	Super::Exit();
 
 	// Restore the selection to the original state after all the tools have shutdown in UEdMode::Exit()
@@ -163,13 +151,6 @@ void UAssetPlacementEdMode::OnToolStarted(UInteractiveToolManager* Manager, UInt
 	if (bRestoreSelectionState)
 	{
 		Owner->RestoreSelection(UPlacementModeSelectTool::ToolName);
-
-		// Todo: remove with HISM element handles
-		for (TActorIterator<AInstancedFoliageActor> It(GetWorld()); It; ++It)
-		{
-			AInstancedFoliageActor* IFA = *It;
-			IFA->ApplySelection(bIsInSelectionTool);
-		}
 	}
 	else if (!bIsInSelectionTool)
 	{
@@ -207,7 +188,7 @@ bool UAssetPlacementEdMode::ShouldDrawWidget() const
 
 bool UAssetPlacementEdMode::IsEnabled()
 {
-	return GetDefault<ULevelEditorMiscSettings>()->bEnableAssetPlacementMode;
+	return SMInstanceElementDataUtil::SMInstanceElementsEnabled() && GetDefault<ULevelEditorMiscSettings>()->bEnableAssetPlacementMode;
 }
 
 void UAssetPlacementEdMode::DeleteSelection()
@@ -219,19 +200,6 @@ void UAssetPlacementEdMode::DeleteSelection()
 		CommonActions->DeleteSelectedElements(Owner->GetEditorSelectionSet(), GetWorld(), FTypedElementDeletionOptions());
 	}
 
-	for (TActorIterator<AInstancedFoliageActor> It(GetWorld()); It; ++It)
-	{
-		if (AInstancedFoliageActor* FoliageActor = *It)
-		{
-			FoliageActor->ForEachFoliageInfo([](UFoliageType* FoliageType, FFoliageInfo& FoliageInfo)
-			{
-				TArray<int32> SelectedIndices = FoliageInfo.SelectedIndices.Array();
-				FoliageInfo.RemoveInstances(SelectedIndices, true);
-				return true; // continue iteration
-			});
-		}
-	}
-
 	GetToolManager()->EndUndoTransaction();
 }
 
@@ -242,19 +210,6 @@ void UAssetPlacementEdMode::ClearSelection()
 	{
 		SelectionSet->ClearSelection(FTypedElementSelectionOptions());
 	}
-
-	// Todo - remove with instanced handles
-	for (TActorIterator<AInstancedFoliageActor> It(GetWorld()); It; ++It)
-	{
-		if (AInstancedFoliageActor* FoliageActor = *It)
-		{
-			FoliageActor->ForEachFoliageInfo([](UFoliageType* FoliageType, FFoliageInfo& FoliageInfo)
-			{
-				FoliageInfo.ClearSelection();
-				return true; // continue iteration
-			});
-		}
-	}
 	GetToolManager()->EndUndoTransaction();
 }
 
@@ -275,24 +230,6 @@ bool UAssetPlacementEdMode::HasActiveSelection() const
 		return true;
 	}
 
-	for (TActorIterator<AInstancedFoliageActor> It(GetWorld()); It; ++It)
-	{
-		if (AInstancedFoliageActor* FoliageActor = *It)
-		{
-			bool bHasSelectedFoliage = false;
-			FoliageActor->ForEachFoliageInfo([&bHasSelectedFoliage](UFoliageType* FoliageType, FFoliageInfo& FoliageInfo)
-				{
-					bHasSelectedFoliage = (FoliageInfo.SelectedIndices.Num() > 0);
-					return !bHasSelectedFoliage;
-				});
-
-			if (bHasSelectedFoliage)
-			{
-				return true;
-			}
-		}
-	}
-
 	return false;
 }
 
diff --git a/Engine/Source/Editor/AssetPlacementEdMode/Private/Tools/PlacementBrushToolBase.cpp b/Engine/Source/Editor/AssetPlacementEdMode/Private/Tools/PlacementBrushToolBase.cpp
--- a/Engine/Source/Editor/AssetPlacementEdMode/Private/Tools/PlacementBrushToolBase.cpp
+++ b/Engine/Source/Editor/AssetPlacementEdMode/Private/Tools/PlacementBrushToolBase.cpp
@@ -315,37 +315,5 @@ FTypedElementListRef UPlacementBrushToolBase::GetElementsInBrushRadius(const FIn
 		return true;
 	});
 
-	// Handle the IFA for the brush stroke level if needed for ISMs
-	if (!FoliageElementUtil::FoliageInstanceElementsEnabled())
-	{
-		if (UActorPartitionSubsystem* PartitionSubsystem = UWorld::GetSubsystem<UActorPartitionSubsystem>(GEditor->GetEditorWorldContext().World()))
-		{
-			constexpr bool bCreatePartitionActorIfMissing = false;
-			FActorPartitionGetParams PartitionActorFindParams(AInstancedFoliageActor::StaticClass(), bCreatePartitionActorIfMissing, GEditor->GetEditorWorldContext().World()->GetCurrentLevel(), LastBrushStamp.WorldPosition);
-			if (AInstancedFoliageActor* FoliageActor = Cast<AInstancedFoliageActor>(PartitionSubsystem->GetActor(PartitionActorFindParams)))
-			{
-				for (const auto& FoliageInfo : FoliageActor->GetFoliageInfos())
-				{
-					// The foliage info needs a valid component for ISM placement from our systems.
-					if (FoliageInfo.Value->GetComponent() != nullptr)
-					{
-						continue;
-					}
-
-					FTypedElementHandle SourceObjectHandle = UEngineElementsLibrary::AcquireEditorObjectElementHandle(FoliageInfo.Key->GetSource());
-					TArray<int32> Instances;
-					FSphere SphereToCheck(LastBrushStamp.WorldPosition, LastBrushStamp.Radius);
-					FoliageInfo.Value->GetInstancesInsideSphere(SphereToCheck, Instances);
-					if (Instances.Num())
-					{
-						// For now, return the whole foliage actor, and allow the calling code to drill down, since we do not have element handles at the instance level just yet
-						ElementHandles->Add(UEngineElementsLibrary::AcquireEditorActorElementHandle(FoliageActor));
-						break;
-					}
-				}
-			}
-		}
-	}
-
 	return ElementHandles;
 }
diff --git a/Engine/Source/Editor/AssetPlacementEdMode/Private/Tools/PlacementEraseTool.cpp b/Engine/Source/Editor/AssetPlacementEdMode/Private/Tools/PlacementEraseTool.cpp
--- a/Engine/Source/Editor/AssetPlacementEdMode/Private/Tools/PlacementEraseTool.cpp
+++ b/Engine/Source/Editor/AssetPlacementEdMode/Private/Tools/PlacementEraseTool.cpp
@@ -11,7 +11,6 @@
 #include "Elements/Interfaces/TypedElementObjectInterface.h"
 #include "Tools/AssetEditorContextInterface.h"
 #include "Editor.h"
-#include "InstancedFoliageActor.h"
 #include "AssetPlacementEdMode.h"
 #include "AssetPlacementSettings.h"
 #include "Modes/PlacementModeSubsystem.h"
@@ -67,31 +66,6 @@ void UPlacementModeEraseTool::OnTick(float DeltaTime)
 	}
 
 	FTypedElementListRef ElementsToDelete = GetElementsInBrushRadius(LastDeviceInputRay);
-	if (!FoliageElementUtil::FoliageInstanceElementsEnabled())
-	{
-		ElementsToDelete->RemoveAll<UTypedElementObjectInterface>([this](const TTypedElement<UTypedElementObjectInterface>& ObjectInterface)
-		{
-			// Since the foliage static mesh instances do not currently operate with element handles, we have to drill in manually here.
-			if (AInstancedFoliageActor* FoliageActor = ObjectInterface.GetObjectAs<AInstancedFoliageActor>())
-			{
-				FoliageActor->ForEachFoliageInfo([this](UFoliageType* FoliageType, FFoliageInfo& FoliageInfo)
-				{
-					FTypedElementHandle SourceObjectHandle = UEngineElementsLibrary::AcquireEditorObjectElementHandle(FoliageType->GetSource());
-					if (GEditor->GetEditorSubsystem<UPlacementModeSubsystem>()->DoesCurrentPaletteSupportElement(SourceObjectHandle))
-					{
-						TArray<int32> Instances;
-						FSphere SphereToCheck(LastBrushStamp.WorldPosition, LastBrushStamp.Radius);
-						FoliageInfo.GetInstancesInsideSphere(SphereToCheck, Instances);
-						FoliageInfo.RemoveInstances(Instances, true);
-					}
-					return true; // continue iteration
-				});
-				return true; // Foliage - remove from the normal element delete
-			}
-			return false; // Not foliage - will be processed via the normal element delete
-		});
-	}
-
 	if (ElementsToDelete->HasElements())
 	{
 		FTypedElementListRef NormalizedElementsToDelete = SelectionSet->GetNormalizedElementList(ElementsToDelete, FTypedElementSelectionNormalizationOptions());
diff --git a/Engine/Source/Editor/AssetPlacementEdMode/Private/Tools/PlacementLassoSelectTool.cpp b/Engine/Source/Editor/AssetPlacementEdMode/Private/Tools/PlacementLassoSelectTool.cpp
--- a/Engine/Source/Editor/AssetPlacementEdMode/Private/Tools/PlacementLassoSelectTool.cpp
+++ b/Engine/Source/Editor/AssetPlacementEdMode/Private/Tools/PlacementLassoSelectTool.cpp
@@ -10,7 +10,6 @@
 #include "Elements/Framework/TypedElementRegistry.h"
 #include "Elements/Interfaces/TypedElementObjectInterface.h"
 #include "InteractiveToolManager.h"
-#include "InstancedFoliageActor.h"
 #include "ToolContextInterfaces.h"
 #include "UObject/Object.h"
 #include "BaseBehaviors/KeyAsModifierInputBehavior.h"
@@ -47,33 +46,7 @@ void UPlacementModeLassoSelectTool::OnEndDrag(const FRay& Ray)
 		UTypedElementSelectionSet* SelectionSet = AssetEditorContext->GetMutableSelectionSet();
 		if (SelectionSet && ElementsFromDrag)
 		{
-			bool bSelectElements = !bCtrlToggle;
-
-			if (!FoliageElementUtil::FoliageInstanceElementsEnabled())
-			{
-				ElementsFromDrag->RemoveAll<UTypedElementObjectInterface>([this, bSelectElements](const TTypedElement<UTypedElementObjectInterface>& ObjectInterface)
-				{
-					// Since the foliage static mesh instances do not currently operate with element handles, we have to drill in manually here.
-					if (AInstancedFoliageActor* FoliageActor = ObjectInterface.GetObjectAs<AInstancedFoliageActor>())
-					{
-						FoliageActor->ForEachFoliageInfo([this, bSelectElements](UFoliageType* FoliageType, FFoliageInfo& FoliageInfo)
-						{
-							FTypedElementHandle SourceObjectHandle = UEngineElementsLibrary::AcquireEditorObjectElementHandle(FoliageType->GetSource());
-							if (GEditor->GetEditorSubsystem<UPlacementModeSubsystem>()->DoesCurrentPaletteSupportElement(SourceObjectHandle))
-							{
-								TArray<int32> Instances;
-								FSphere SphereToCheck(LastBrushStamp.WorldPosition, LastBrushStamp.Radius);
-								FoliageInfo.GetInstancesInsideSphere(SphereToCheck, Instances);
-								FoliageInfo.SelectInstances(bSelectElements, Instances);
-							}
-							return true; // continue iteration
-						});
-						return true; // Foliage - remove from the normal element select
-					}
-					return false; // Not foliage - will be processed via the normal element select
-				});
-			}
-
+			const bool bSelectElements = !bCtrlToggle;
 			if (bSelectElements)
 			{
 				SelectionSet->SelectElements(ElementsFromDrag.ToSharedRef(), PlacementModeLassoToolInternal::SelectionOptions);
diff --git a/Engine/Source/Editor/AssetPlacementEdMode/Private/Tools/PlacementPlaceSingleTool.cpp b/Engine/Source/Editor/AssetPlacementEdMode/Private/Tools/PlacementPlaceSingleTool.cpp
--- a/Engine/Source/Editor/AssetPlacementEdMode/Private/Tools/PlacementPlaceSingleTool.cpp
+++ b/Engine/Source/Editor/AssetPlacementEdMode/Private/Tools/PlacementPlaceSingleTool.cpp
@@ -72,7 +72,7 @@ void UPlacementModePlaceSingleTool::OnClickPress(const FInputDeviceRay& PressPos
 	// Place the Preview data if we managed to get to a valid handled click.
 	FPlacementOptions PlacementOptions;
 	PlacementOptions.bPreferBatchPlacement = true;
-	PlacementOptions.bPreferInstancedPlacement = SMInstanceElementDataUtil::SMInstanceElementsEnabled();
+	PlacementOptions.InstancedPlacementGridGuid = UPlacementSubsystem::GetUserGridGuid();
 
 	LastBrushStamp.Radius = 100.0f * LastBrushStampWorldToPixelScale;
 	
diff --git a/Engine/Source/Editor/AssetPlacementEdMode/Private/Tools/PlacementPlaceTool.cpp b/Engine/Source/Editor/AssetPlacementEdMode/Private/Tools/PlacementPlaceTool.cpp
--- a/Engine/Source/Editor/AssetPlacementEdMode/Private/Tools/PlacementPlaceTool.cpp
+++ b/Engine/Source/Editor/AssetPlacementEdMode/Private/Tools/PlacementPlaceTool.cpp
@@ -106,7 +106,7 @@ void UPlacementModePlacementTool::OnTick(float DeltaTime)
 		{
 			FPlacementOptions PlacementOptions;
 			PlacementOptions.bPreferBatchPlacement = true;
-			PlacementOptions.bPreferInstancedPlacement = true;
+			PlacementOptions.InstancedPlacementGridGuid = UPlacementSubsystem::GetUserGridGuid();
 			PlacementSubsystem->PlaceAssets(DesiredInfoSpawnLocations, PlacementOptions);
 		}
 	}
diff --git a/Engine/Source/Editor/EditorFramework/Private/Subsystems/PlacementSubsystem.cpp b/Engine/Source/Editor/EditorFramework/Private/Subsystems/PlacementSubsystem.cpp
--- a/Engine/Source/Editor/EditorFramework/Private/Subsystems/PlacementSubsystem.cpp
+++ b/Engine/Source/Editor/EditorFramework/Private/Subsystems/PlacementSubsystem.cpp
@@ -8,6 +8,7 @@
 #include "Elements/Framework/TypedElementRegistry.h"
 #include "Elements/Framework/TypedElementSelectionSet.h"
 
+#include "GenericPlatform/GenericPlatformMisc.h"
 #include "Misc/CoreDelegates.h"
 #include "UObject/UObjectIterator.h"
 
@@ -95,6 +96,13 @@ bool UPlacementSubsystem::IsCreatingPreviewElements() const
 	return bIsCreatingPreviewElements;
 }
 
+FGuid UPlacementSubsystem::GetUserGridGuid()
+{
+	FGuid UserGridGuid;
+	FGuid::Parse(FPlatformMisc::GetLoginId(), UserGridGuid);
+	return UserGridGuid;
+}
+
 void UPlacementSubsystem::RegisterPlacementFactories()
 {
 	for (TObjectIterator<UClass> ObjectIt; ObjectIt; ++ObjectIt)
diff --git a/Engine/Source/Editor/EditorFramework/Public/Subsystems/PlacementSubsystem.h b/Engine/Source/Editor/EditorFramework/Public/Subsystems/PlacementSubsystem.h
--- a/Engine/Source/Editor/EditorFramework/Public/Subsystems/PlacementSubsystem.h
+++ b/Engine/Source/Editor/EditorFramework/Public/Subsystems/PlacementSubsystem.h
@@ -35,16 +35,27 @@ struct EDITORFRAMEWORK_API FAssetPlacementInfo
 	// If set, will use the given factory to place the asset, instead of allowing the placement subsystem to determine which factory to use.
 	UPROPERTY()
 	TScriptInterface<IAssetFactoryInterface> FactoryOverride;
+
+	/**
+	 * The Guid which corresponds to the item that should be placed.
+	 * If unset, the asset package's persistent guid will be used.
+	 * Factories should use this to tie any decomposed assets together. For example, the ItemGuid would correspond to the client in an AISMPartitionActor for tracking all static meshes which make up a decomposed actor.
+	 */
+	UPROPERTY()
+	FGuid ItemGuid;
 };
 
 USTRUCT()
 struct EDITORFRAMEWORK_API FPlacementOptions
 {
 	GENERATED_BODY()
-
-	// If true, asset factory implementations should defer to placing instanced items (i.e. instanced static mesh instead of individual static mesh actors).
+	
+	/**
+	 * The guid to use by factories for instanced placement. If unset, factories will not use instanced placement.
+	 * This is used to reduce contention within one file per actor within a partition.
+	 */
 	UPROPERTY()
-	bool bPreferInstancedPlacement = true;
+	FGuid InstancedPlacementGridGuid;
 
 	// If true, asset factory implementations should prefer a batch placement algorithm (like duplicating an object) over a single placement algorithm.
 	UPROPERTY()
@@ -73,7 +84,7 @@ public:
 	TArray<FTypedElementHandle> PlaceAsset(const FAssetPlacementInfo& InPlacementInfo, const FPlacementOptions& InPlacementOptions);
 
 	/**
-	 * Places a multiple asset based on the given FAssetPlacementInfos and FPlacementOptions.
+	 * Places multiple assets based on the given FAssetPlacementInfos and FPlacementOptions.
 	 * @returns an array of FTypedElementHandles corresponding to any successfully placed elements.
 	 */
 	TArray<FTypedElementHandle> PlaceAssets(TArrayView<const FAssetPlacementInfo> InPlacementInfos, const FPlacementOptions& InPlacementOptions);
@@ -90,6 +101,11 @@ public:
 	 */
 	bool IsCreatingPreviewElements() const;
 
+	/**
+	 * Returns the guid which the placement system should use to tie instanced placement to the current user for partitions.
+	 */
+	static FGuid GetUserGridGuid();
+
 private:
 	void RegisterPlacementFactories();
 	void UnregisterPlacementFactories();
diff --git a/Engine/Source/Editor/UnrealEd/Private/AssetSelection.cpp b/Engine/Source/Editor/UnrealEd/Private/AssetSelection.cpp
--- a/Engine/Source/Editor/UnrealEd/Private/AssetSelection.cpp
+++ b/Engine/Source/Editor/UnrealEd/Private/AssetSelection.cpp
@@ -615,7 +615,6 @@ static AActor* PrivateAddActor( UObject* Asset, UActorFactory* Factory, bool Sel
 
 			FPlacementOptions PlacementOptions;
 			PlacementOptions.bIsCreatingPreviewElements = FLevelEditorViewportClient::IsDroppingPreviewActor();
-			PlacementOptions.bPreferInstancedPlacement = false;
 
 			TArray<FTypedElementHandle> PlacedElements = PlacementSubsystem->PlaceAsset(PlacementInfo, PlacementOptions);
 			if (PlacedElements.Num())
diff --git a/Engine/Source/Editor/UnrealEd/Private/Factories/EditorStaticMeshFactory.cpp b/Engine/Source/Editor/UnrealEd/Private/Factories/EditorStaticMeshFactory.cpp
--- a/Engine/Source/Editor/UnrealEd/Private/Factories/EditorStaticMeshFactory.cpp
+++ b/Engine/Source/Editor/UnrealEd/Private/Factories/EditorStaticMeshFactory.cpp
@@ -6,14 +6,14 @@
 #include "Elements/Framework/TypedElementRegistry.h"
 #include "Elements/Interfaces/TypedElementAssetDataInterface.h"
 #include "Elements/Interfaces/TypedElementObjectInterface.h"
+#include "Instances/InstancedPlacementPartitionActor.h"
 
 #include "ActorPartition/ActorPartitionSubsystem.h"
-#include "InstancedFoliageActor.h"
-#include "InstancedFoliage.h"
-#include "FoliageType_InstancedStaticMesh.h"
 
 #include "Components/InstancedStaticMeshComponent.h"
 #include "Engine/StaticMesh.h"
+#include "Engine/StaticMeshActor.h"
+#include "UObject/Class.h"
 
 #include "Subsystems/PlacementSubsystem.h"
 
@@ -25,108 +25,114 @@ TArray<FTypedElementHandle> UEditorStaticMeshFactory::PlaceAsset(const FAssetPla
 		return Super::PlaceAsset(InPlacementInfo, InPlacementOptions);
 	}
 
+	TArray<FTypedElementHandle> PlacedInstanceHandles;
 	if (InPlacementInfo.PreferredLevel.IsValid())
 	{
+		UObject* AssetToPlaceAsObject = InPlacementInfo.AssetToPlace.GetAsset();
+		FISMComponentDescriptor ComponentDescriptor;
+		if (UStaticMesh* StaticMeshObject = Cast<UStaticMesh>(AssetToPlaceAsObject))
+		{
+			// If this is a Nanite mesh, prefer to use ISM over HISM, as HISM duplicates many features/bookkeeping that Nanite already handles for us.
+			if (StaticMeshObject->HasValidNaniteData())
+			{
+				ComponentDescriptor.InitFrom(UInstancedStaticMeshComponent::StaticClass()->GetDefaultObject<UInstancedStaticMeshComponent>());
+			}
+			ComponentDescriptor.StaticMesh = StaticMeshObject;
+		}
+		else if (AStaticMeshActor* StaticMeshActor = Cast<AStaticMeshActor>(AssetToPlaceAsObject))
+		{
+			if (UStaticMeshComponent* StaticMeshComponent = StaticMeshActor->GetStaticMeshComponent())
+			{
+				ComponentDescriptor.StaticMesh = StaticMeshComponent->GetStaticMesh();
+			}
+		}
+
+		if (!ComponentDescriptor.StaticMesh)
+		{
+			return TArray<FTypedElementHandle>();
+		}
+
+		ComponentDescriptor.ComputeHash();
+
 		if (UActorPartitionSubsystem* PartitionSubsystem = UWorld::GetSubsystem<UActorPartitionSubsystem>(InPlacementInfo.PreferredLevel->GetWorld()))
 		{
-			// Create or find the foliage partition actor
-			constexpr bool bCreatePartitionActorIfMissing = true;
-			FActorPartitionGetParams PartitionActorFindParams(AInstancedFoliageActor::StaticClass(), bCreatePartitionActorIfMissing, InPlacementInfo.PreferredLevel.Get(), InPlacementInfo.FinalizedTransform.GetLocation());
-			AInstancedFoliageActor* FoliageActor = Cast<AInstancedFoliageActor>(PartitionSubsystem->GetActor(PartitionActorFindParams));
-			check(FoliageActor);
-
-			// Create or find the instanced static mesh foliage type for the given mesh
-			UObject* AssetDataObject = InPlacementInfo.AssetToPlace.GetAsset();
-			UFoliageType_InstancedStaticMesh* FoundOrCreatedType = nullptr;
-			for (const auto& FoliageTypePair : FoliageActor->GetFoliageInfos())
+			// Create or find the placement partition actor
+			auto OnActorCreated = [InPlacementOptions](APartitionActor* CreatedPartitionActor)
 			{
-				FAssetData FoliageTypeSource(FoliageTypePair.Key->GetSource());
-				if (FoliageTypeSource == InPlacementInfo.AssetToPlace)
+				if (AInstancedPlacementPartitionActor* ElementPartitionActor = Cast<AInstancedPlacementPartitionActor>(CreatedPartitionActor))
 				{
-					FoundOrCreatedType = Cast<UFoliageType_InstancedStaticMesh>(FoliageTypePair.Key);
-					break;
+					ElementPartitionActor->SetGridGuid(InPlacementOptions.InstancedPlacementGridGuid);
 				}
-			}
-			if (!FoundOrCreatedType)
+			};
+
+			// Make a good known client GUID out of the placed asset's package if one was not given to us.
+			FGuid ItemGuidToUse = InPlacementInfo.ItemGuid;
+			if (!ItemGuidToUse.IsValid())
 			{
-				FoundOrCreatedType = NewObject<UFoliageType_InstancedStaticMesh>(FoliageActor);
-				FoundOrCreatedType->SetSource(InPlacementInfo.AssetToPlace.GetAsset());
+				ItemGuidToUse = InPlacementInfo.AssetToPlace.GetAsset()->GetPackage()->GetPersistentGuid();
 			}
-			FFoliageInfo* FoliageInfo = FoliageActor->FindOrAddMesh(FoundOrCreatedType);
-			check(FoliageInfo);
-
-			FDesiredFoliageInstance InstancePlaceInfo;
-			FFoliageInstance FoliagePlacementInfo;
-			FoliagePlacementInfo.Location = InPlacementInfo.FinalizedTransform.GetLocation();
-			FoliagePlacementInfo.Rotation = InPlacementInfo.FinalizedTransform.Rotator();
-			FoliagePlacementInfo.DrawScale3D = InPlacementInfo.FinalizedTransform.GetScale3D();
-			FoliageInfo->AddInstance(FoundOrCreatedType, FoliagePlacementInfo);
-			CurrentPlacementScopedFoliageInfos.Add(FoliageInfo);
-
-			FTypedElementHandle PlacedElement = UEngineElementsLibrary::AcquireEditorSMInstanceElementHandle(FoliageInfo->GetComponent(), FoliageInfo->GetPlacedInstanceCount() - 1);
-			if (!PlacedElement)
+
+			constexpr bool bCreatePartitionActorIfMissing = true;
+			FActorPartitionGetParams PartitionActorFindParams(AInstancedPlacementPartitionActor::StaticClass(), bCreatePartitionActorIfMissing, InPlacementInfo.PreferredLevel.Get(), InPlacementInfo.FinalizedTransform.GetLocation());
+			AInstancedPlacementPartitionActor* PlacedElementsActor = Cast<AInstancedPlacementPartitionActor>(PartitionSubsystem->GetActor(PartitionActorFindParams));
+
+			FISMClientHandle ClientHandle = PlacedElementsActor->RegisterClient(ItemGuidToUse);
+			TSortedMap<int32, TArray<FTransform>> InstanceMap;
+			InstanceMap.Emplace(PlacedElementsActor->RegisterISMComponentDescriptor(ComponentDescriptor), TArray({ FTransform() }));
+
+			TArray<FSMInstanceId> PlacedInstances = PlacedElementsActor->AddISMInstance(ClientHandle, InPlacementInfo.FinalizedTransform, InstanceMap);
+			for (const FSMInstanceId PlacedInstanceID : PlacedInstances)
 			{
-				PlacedElement = UEngineElementsLibrary::AcquireEditorComponentElementHandle(FoliageInfo->GetComponent());
+				if (FTypedElementHandle PlacedHandle = UEngineElementsLibrary::AcquireEditorSMInstanceElementHandle(PlacedInstanceID))
+				{
+					PlacedInstanceHandles.Emplace(PlacedHandle);
+				}
 			}
-			check(PlacedElement);
-			return TArray<FTypedElementHandle>({ MoveTemp(PlacedElement) });
 		}
 	}
 
-	return TArray<FTypedElementHandle>();
+	return PlacedInstanceHandles;
 }
 
 FAssetData UEditorStaticMeshFactory::GetAssetDataFromElementHandle(const FTypedElementHandle& InHandle)
 {
-	UInstancedStaticMeshComponent* ISMComponent = nullptr;
+	FAssetData FoundAssetData;
+	if (TTypedElement<UTypedElementAssetDataInterface> AssetDataInterface = UTypedElementRegistry::GetInstance()->GetElement<UTypedElementAssetDataInterface>(InHandle))
+	{
+		FoundAssetData = AssetDataInterface.GetAssetData();
+	}
 
-	if (TTypedElement<UTypedElementObjectInterface> ObjectInterface = UTypedElementRegistry::GetInstance()->GetElement<UTypedElementObjectInterface>(InHandle))
+	if (!FoundAssetData.IsValid())
 	{
-		// Try to pull from component handle
-		if (UInstancedStaticMeshComponent* RawComponentPtr = ObjectInterface.GetObjectAs<UInstancedStaticMeshComponent>())
+		UInstancedStaticMeshComponent* ISMComponent = nullptr;
+		if (TTypedElement<UTypedElementObjectInterface> ObjectInterface = UTypedElementRegistry::GetInstance()->GetElement<UTypedElementObjectInterface>(InHandle))
 		{
-			ISMComponent = RawComponentPtr;
+			// Try to pull from component handle
+			if (UInstancedStaticMeshComponent* RawComponentPtr = ObjectInterface.GetObjectAs<UInstancedStaticMeshComponent>())
+			{
+				ISMComponent = RawComponentPtr;
+			}
+			else if (AActor* RawActorPtr = ObjectInterface.GetObjectAs<AActor>())
+			{
+				ISMComponent = RawActorPtr->FindComponentByClass<UInstancedStaticMeshComponent>();
+			}
 		}
-		else if (AActor* RawActorPtr = ObjectInterface.GetObjectAs<AActor>())
+
+		if (ISMComponent)
 		{
-			ISMComponent = RawActorPtr->FindComponentByClass<UInstancedStaticMeshComponent>();
+			FoundAssetData = FAssetData(ISMComponent->GetStaticMesh());
 		}
 	}
 
-	FAssetData FoundAssetData;
-	if (ISMComponent)
-	{
-		FoundAssetData = FAssetData(ISMComponent->GetStaticMesh());
-	}
-	else if (TTypedElement<UTypedElementAssetDataInterface> AssetDataInterface = UTypedElementRegistry::GetInstance()->GetElement<UTypedElementAssetDataInterface>(InHandle))
-	{
-		FoundAssetData = AssetDataInterface.GetAssetData();
-	}
-
 	if (CanPlaceElementsFromAssetData(FoundAssetData))
 	{
 		return FoundAssetData;
 	}
 
-	// Todo: deal with instanced handles
-
 	return Super::GetAssetDataFromElementHandle(InHandle);
 }
 
-void UEditorStaticMeshFactory::EndPlacement(TArrayView<const FTypedElementHandle> InPlacedElements, const FPlacementOptions& InPlacementOptions)
-{
-	if (ShouldPlaceInstancedStaticMeshes(InPlacementOptions))
-	{
-		for (FFoliageInfo* FoliageInfo : CurrentPlacementScopedFoliageInfos)
-		{
-			FoliageInfo->Refresh(true, false);
-		}
-	}
-
-	CurrentPlacementScopedFoliageInfos.Empty();
-}
-
 bool UEditorStaticMeshFactory::ShouldPlaceInstancedStaticMeshes(const FPlacementOptions& InPlacementOptions) const
 {
-	return !InPlacementOptions.bIsCreatingPreviewElements && InPlacementOptions.bPreferInstancedPlacement;
+	return !InPlacementOptions.bIsCreatingPreviewElements && InPlacementOptions.InstancedPlacementGridGuid.IsValid();
 }
diff --git a/Engine/Source/Editor/UnrealEd/Private/Factories/EditorStaticMeshFactory.h b/Engine/Source/Editor/UnrealEd/Private/Factories/EditorStaticMeshFactory.h
--- a/Engine/Source/Editor/UnrealEd/Private/Factories/EditorStaticMeshFactory.h
+++ b/Engine/Source/Editor/UnrealEd/Private/Factories/EditorStaticMeshFactory.h
@@ -18,10 +18,8 @@ public:
 	// Begin IAssetFactoryInterface
 	virtual TArray<FTypedElementHandle> PlaceAsset(const FAssetPlacementInfo& InPlacementInfo, const FPlacementOptions& InPlacementOptions) override;
 	virtual FAssetData GetAssetDataFromElementHandle(const FTypedElementHandle& InHandle) override;
-	virtual void EndPlacement(TArrayView<const FTypedElementHandle> InPlacedElements, const FPlacementOptions& InPlacementOptions) override;
 	// End IAssetFactoryInterface
 
 protected:
 	bool ShouldPlaceInstancedStaticMeshes(const FPlacementOptions& InPlacementOptions) const;
-	TSet<FFoliageInfo*> CurrentPlacementScopedFoliageInfos;
 };
\ No newline at end of file
diff --git a/Engine/Source/Runtime/Engine/Classes/GameFramework/WorldSettings.h b/Engine/Source/Runtime/Engine/Classes/GameFramework/WorldSettings.h
--- a/Engine/Source/Runtime/Engine/Classes/GameFramework/WorldSettings.h
+++ b/Engine/Source/Runtime/Engine/Classes/GameFramework/WorldSettings.h
@@ -553,6 +553,10 @@ public:
 	UPROPERTY(EditAnywhere, Category = Foliage)
 	uint32 InstancedFoliageGridSize;
 
+	/** Default size of the grid for placed elements from the editor */
+	UPROPERTY(EditAnywhere, Category=Editor)
+	uint32 DefaultPlacementGridSize;
+
 	/** scale of 1uu to 1m in real world measurements, for HMD and other physically tracked devices (e.g. 1uu = 1cm would be 100.0) */
 	UPROPERTY(EditAnywhere, BlueprintReadOnly, Category=VR)
 	float WorldToMeters;
diff --git a/Engine/Source/Runtime/Engine/Private/ISMPartition/ISMPartitionActor.cpp b/Engine/Source/Runtime/Engine/Private/ISMPartition/ISMPartitionActor.cpp
--- a/Engine/Source/Runtime/Engine/Private/ISMPartition/ISMPartitionActor.cpp
+++ b/Engine/Source/Runtime/Engine/Private/ISMPartition/ISMPartitionActor.cpp
@@ -19,7 +19,9 @@ AISMPartitionActor::AISMPartitionActor(const FObjectInitializer& ObjectInitializ
 	, bWasModifyCalled(false)
 #endif
 {
-
+#if WITH_EDITORONLY_DATA
+	bListedInSceneOutliner = false;
+#endif
 }
 
 #if WITH_EDITOR
@@ -200,13 +202,14 @@ void AISMPartitionActor::ReserveISMInstances(const FISMClientHandle& Handle, int
 	EndUpdate();
 }
 
-void AISMPartitionActor::AddISMInstance(const FISMClientHandle& Handle, const FTransform& InstanceTransform, const TSortedMap<int32, TArray<FTransform>>& InstanceDefinition)
+TArray<FSMInstanceId> AISMPartitionActor::AddISMInstance(const FISMClientHandle& Handle, const FTransform& InstanceTransform, const TSortedMap<int32, TArray<FTransform>>& InstanceDefinition)
 {
 	check(Handle.Guid == Clients[Handle.Index]);
 	// Create Descriptor entry
 	BeginUpdate();
 	ModifyActor();
 
+	TArray<FSMInstanceId> AddedInstances;
 	for (const auto& Pair : InstanceDefinition)
 	{
 		const int32 DescriptorIndex = Pair.Key;
@@ -236,11 +239,18 @@ void AISMPartitionActor::AddISMInstance(const FISMClientHandle& Handle, const FT
 			ClientInstance.ComponentIndices.Add(ComponentInstanceIndex);
 			
 			FTransform NewInstanceTransform = LocalTransform * InstanceTransform;
-			AddInstanceToComponent(ComponentData, NewInstanceTransform);
+
+			int32 NewInstanceIdx = AddInstanceToComponent(ComponentData, NewInstanceTransform);
+			if (NewInstanceIdx != INDEX_NONE)
+			{
+				AddedInstances.Emplace(FSMInstanceId{ ComponentData.Component, NewInstanceIdx });
+			}
 		}
 	}
 
 	EndUpdate();
+
+	return AddedInstances;
 }
 
 void AISMPartitionActor::RemoveISMInstances(const FISMClientHandle& Handle)
@@ -313,12 +323,13 @@ void AISMPartitionActor::RemoveInstanceFromComponent(FISMComponentData& Componen
 	}
 }
 
-void AISMPartitionActor::AddInstanceToComponent(FISMComponentData& ComponentData, const FTransform& WorldTransform)
+int32 AISMPartitionActor::AddInstanceToComponent(FISMComponentData& ComponentData, const FTransform& WorldTransform)
 {
 	check(ComponentData.Component);
 	ModifyComponent(ComponentData);
-	ComponentData.Component->AddInstance(WorldTransform, /*bWorldSpace*/true);
+	int32 NewInstanceIdx = ComponentData.Component->AddInstance(WorldTransform, /*bWorldSpace*/true);
 	InvalidateComponentLightingCache(ComponentData);
+	return NewInstanceIdx;
 }
 
 void AISMPartitionActor::UpdateInstanceTransform(FISMComponentData& ComponentData, int32 ComponentInstanceIndex, const FTransform& WorldTransform, bool bTeleport)
diff --git a/Engine/Source/Runtime/Engine/Private/Instances/InstancedPlacementPartitionActor.cpp b/Engine/Source/Runtime/Engine/Private/Instances/InstancedPlacementPartitionActor.cpp
new file mode 100644
--- /dev/null
+++ b/Engine/Source/Runtime/Engine/Private/Instances/InstancedPlacementPartitionActor.cpp
@@ -0,0 +1,50 @@
+// Copyright Epic Games, Inc. All Rights Reserved.
+
+#include "Instances/InstancedPlacementPartitionActor.h"
+
+#include "Components/InstancedStaticMeshComponent.h"
+#include "Engine/World.h"
+#include "GameFramework/WorldSettings.h"
+
+AInstancedPlacementPartitionActor::AInstancedPlacementPartitionActor(const FObjectInitializer& ObjectInitializer)
+	: Super(ObjectInitializer)
+{
+	PrimaryActorTick.bCanEverTick = false;
+	SetActorEnableCollision(true);
+}
+
+#if WITH_EDITOR	
+uint32 AInstancedPlacementPartitionActor::GetDefaultGridSize(UWorld* InWorld) const
+{
+	return InWorld->GetWorldSettings()->DefaultPlacementGridSize;
+}
+
+FGuid AInstancedPlacementPartitionActor::GetGridGuid() const
+{
+	return PlacementGridGuid;
+}
+
+void AInstancedPlacementPartitionActor::SetGridGuid(const FGuid& InGuid)
+{
+	PlacementGridGuid = InGuid;
+}
+#endif
+
+ISMInstanceManager* AInstancedPlacementPartitionActor::GetSMInstanceManager(const FSMInstanceId& InstanceId)
+{
+	if (ISMInstanceManager* ParentSMInstanceManager = Super::GetSMInstanceManager(InstanceId))
+	{
+		return ParentSMInstanceManager;
+	}
+
+#if WITH_EDITOR
+	// For now, assume that if we didn't have some manager registered, it is safe to edit the ISM directly.
+	// This should be removed after palette items are set up properly in the placement API
+	if (IsISMComponent(InstanceId.ISMComponent))
+	{
+		return InstanceId.ISMComponent;
+	}
+#endif
+
+	return nullptr;
+}
diff --git a/Engine/Source/Runtime/Engine/Private/WorldSettings.cpp b/Engine/Source/Runtime/Engine/Private/WorldSettings.cpp
--- a/Engine/Source/Runtime/Engine/Private/WorldSettings.cpp
+++ b/Engine/Source/Runtime/Engine/Private/WorldSettings.cpp
@@ -90,7 +90,7 @@ PRAGMA_ENABLE_DEPRECATION_WARNINGS
 	KillZ = -HALF_WORLD_MAX1;
 	KillZDamageType = ConstructorStatics.DmgType_Environmental_Object.Object;
 
-	InstancedFoliageGridSize = 25600;
+	InstancedFoliageGridSize = DefaultPlacementGridSize = 25600;
 
 	WorldToMeters = 100.f;
 
diff --git a/Engine/Source/Runtime/Engine/Public/ISMPartition/ISMPartitionActor.h b/Engine/Source/Runtime/Engine/Public/ISMPartition/ISMPartitionActor.h
--- a/Engine/Source/Runtime/Engine/Public/ISMPartition/ISMPartitionActor.h
+++ b/Engine/Source/Runtime/Engine/Public/ISMPartition/ISMPartitionActor.h
@@ -37,7 +37,7 @@ public:
 	int32 RegisterISMComponentDescriptor(const FISMComponentDescriptor& Descriptor);
 	const FISMComponentDescriptor& GetISMComponentDescriptor(int32 DescriptorIndex) const { return Descriptors[DescriptorIndex]; }
 
-	void AddISMInstance(const FISMClientHandle& Handle, const FTransform& InstanceTransform, const TSortedMap<int32, TArray<FTransform>>& InstanceDefinition);
+	TArray<FSMInstanceId> AddISMInstance(const FISMClientHandle& Handle, const FTransform& InstanceTransform, const TSortedMap<int32, TArray<FTransform>>& InstanceDefinition);
 	void RemoveISMInstance(const FISMClientHandle& Handle, int32 InstanceIndex, bool* bOutIsEmpty = nullptr);
 	void RemoveISMInstances(const FISMClientHandle& Handle);
 	void SelectISMInstances(const FISMClientHandle& Handle, bool bSelect, const TSet<int32>& Indices);
@@ -61,7 +61,7 @@ private:
 	void RemoveISMInstancesInternal(FISMComponentData& ComponentData, FISMClientData& OwnerData, int32 InstanceIndex);
 
 	void InvalidateComponentLightingCache(FISMComponentData& ComponentData);
-	void AddInstanceToComponent(FISMComponentData& ComponentData, const FTransform& WorldTransform);
+	int32 AddInstanceToComponent(FISMComponentData& ComponentData, const FTransform& WorldTransform);
 	void UpdateInstanceTransform(FISMComponentData& ComponentData, int32 ComponentInstanceIndex, const FTransform& WorldTransform, bool bTeleport);
 	void RemoveInstanceFromComponent(FISMComponentData& ComponentData, int32 ComponentInstanceIndex);
 	bool DestroyComponentIfEmpty(FISMComponentDescriptor& Descriptor, FISMComponentData& ComponentData);
diff --git a/Engine/Source/Runtime/Engine/Public/Instances/InstancedPlacementPartitionActor.h b/Engine/Source/Runtime/Engine/Public/Instances/InstancedPlacementPartitionActor.h
new file mode 100644
--- /dev/null
+++ b/Engine/Source/Runtime/Engine/Public/Instances/InstancedPlacementPartitionActor.h
@@ -0,0 +1,32 @@
+// Copyright Epic Games, Inc. All Rights Reserved.
+
+#pragma once
+
+#include "CoreMinimal.h"
+#include "ISMPartition/ISMPartitionActor.h"
+#include "Components/InstancedStaticMeshComponent.h"
+#include "InstancedPlacementPartitionActor.generated.h"
+
+/**
+ * The base class used by any editor placement of instanced objects, which holds any relevant runtime data for the placed instances.
+ */
+UCLASS()
+class ENGINE_API AInstancedPlacementPartitionActor : public AISMPartitionActor
+{
+	GENERATED_UCLASS_BODY()
+
+public:
+#if WITH_EDITOR	
+	virtual uint32 GetDefaultGridSize(UWorld* InWorld) const override;
+	virtual FGuid GetGridGuid() const override;
+
+	void SetGridGuid(const FGuid& InGuid);
+#endif
+
+	virtual ISMInstanceManager* GetSMInstanceManager(const FSMInstanceId& InstanceId) override;
+
+protected:
+#if WITH_EDITORONLY_DATA
+	FGuid PlacementGridGuid;
+#endif
+};
diff --git a/Engine/Source/Runtime/Foliage/Private/InstancedFoliage.cpp b/Engine/Source/Runtime/Foliage/Private/InstancedFoliage.cpp
--- a/Engine/Source/Runtime/Foliage/Private/InstancedFoliage.cpp
+++ b/Engine/Source/Runtime/Foliage/Private/InstancedFoliage.cpp
@@ -2811,9 +2811,6 @@ AInstancedFoliageActor::AInstancedFoliageActor(const FObjectInitializer& ObjectI
 	: Super(ObjectInitializer)
 {
 	SetActorEnableCollision(true);
-#if WITH_EDITORONLY_DATA
-	bListedInSceneOutliner = false;
-#endif // WITH_EDITORONLY_DATA
 	PrimaryActorTick.bCanEverTick = false;
 }
 
