CIS for Issue 92786

diff --git a/Engine/Source/Editor/UnrealEd/Private/EditorServer.cpp b/Engine/Source/Editor/UnrealEd/Private/EditorServer.cpp
--- a/Engine/Source/Editor/UnrealEd/Private/EditorServer.cpp
+++ b/Engine/Source/Editor/UnrealEd/Private/EditorServer.cpp
@@ -176,12 +176,6 @@ DEFINE_LOG_CATEGORY_STATIC(LogEditorServer, Log, All);
 /** Used for the "tagsounds" and "checksounds" commands only			*/
 static FUObjectAnnotationSparseBool DebugSoundAnnotation;
 
-namespace EditorEngineDefs
-{
-	/** Limit the minimum size of the bounding box when centering cameras on individual components to avoid extreme zooming */
-	static const float MinComponentBoundsForZoom = 50.0f;
-}
-
 namespace 
 {
 	/**
@@ -4949,6 +4943,9 @@ void UEditorEngine::MoveViewportCamerasToActor(const TArray<AActor*> &Actors, co
 
 void UEditorEngine::MoveViewportCamerasToComponent(const USceneComponent* Component, bool bActiveViewportOnly)
 {
+	/** Limit the minimum size of the bounding box when centering cameras on individual components to avoid extreme zooming */
+	constexpr float MinComponentBoundsForZoom = 50.0f; 
+	
 	if (Component != nullptr)
 	{
 		if (FLevelUtils::IsLevelVisible(Component->GetComponentLevel()) && Component->IsRegistered())
@@ -4959,9 +4956,9 @@ void UEditorEngine::MoveViewportCamerasToComponent(const USceneComponent* Compon
 			Box.GetCenterAndExtents(Center, Extents);
 
 			// Apply a minimum size to the extents of the component's box to avoid the camera's zooming too close to small or zero-sized components
-			if (Extents.SizeSquared() < EditorEngineDefs::MinComponentBoundsForZoom * EditorEngineDefs::MinComponentBoundsForZoom)
+			if (Extents.SizeSquared() < MinComponentBoundsForZoom * MinComponentBoundsForZoom)
 			{
-				FVector NewExtents(EditorEngineDefs::MinComponentBoundsForZoom, SMALL_NUMBER, SMALL_NUMBER);
+				FVector NewExtents(MinComponentBoundsForZoom, SMALL_NUMBER, SMALL_NUMBER);
 				Box = FBox(Center - NewExtents, Center + NewExtents);
 			}
 
