Small changes for debugger

diff --git a/Engine/Source/Editor/Kismet/Private/Debugging/SKismetDebuggingView.cpp b/Engine/Source/Editor/Kismet/Private/Debugging/SKismetDebuggingView.cpp
--- a/Engine/Source/Editor/Kismet/Private/Debugging/SKismetDebuggingView.cpp
+++ b/Engine/Source/Editor/Kismet/Private/Debugging/SKismetDebuggingView.cpp
@@ -97,16 +97,6 @@ TSharedRef<SWidget> FDebugLineItem::GetValueIcon()
 	return SNew(SImage).Image(CachedBrush);
 }
 
-FMargin FDebugLineItem::GetInnerMargin()
-{
-	return {0.f};
-}
-
-FMargin FDebugLineItem::GetOuterMargin()
-{
-	return {0.5f, 0.5f};
-}
-
 UBlueprint* FDebugLineItem::GetBlueprintForObject(UObject* ParentObject)
 {
 	if(ParentObject == nullptr)
@@ -1372,10 +1362,9 @@ public:
 		}
 		
 		return SNew(SBox)
-			.Padding(ItemToEdit->GetOuterMargin())
+			.Padding(FMargin(0.5f, 0.5f))
 			[
 				SNew(SBorder)
-					.Padding(ItemToEdit->GetInnerMargin())
 					.BorderImage(FAppStyle::Get().GetBrush("DetailsView.CategoryMiddle"))
 					.BorderBackgroundColor_Static(
 						PropertyInfoViewStyle::GetRowBackgroundColor,
diff --git a/Engine/Source/Editor/Kismet/Private/Debugging/SKismetDebuggingView.h b/Engine/Source/Editor/Kismet/Private/Debugging/SKismetDebuggingView.h
--- a/Engine/Source/Editor/Kismet/Private/Debugging/SKismetDebuggingView.h
+++ b/Engine/Source/Editor/Kismet/Private/Debugging/SKismetDebuggingView.h
@@ -71,12 +71,6 @@ public:
 	// returns a widget that will go to the left of the Value Widget.
 	virtual TSharedRef<SWidget> GetValueIcon();
 
-	// determines the margin between the Name/Value widgets and the boarder
-	virtual FMargin GetInnerMargin();
-	
-	// determines the space between multiple rows/colums
-	virtual FMargin GetOuterMargin();
-
 	// Helper function to try to get the blueprint for a given object;
 	//   Returns the blueprint that was used to create the instance if there was one
 	//   Returns the object itself if it is already a blueprint
@@ -103,7 +97,7 @@ protected:
 
 	// Used to update the state of a line item rather than replace it.
 	// called after Compare returns true
-	virtual void UpdateData(const FDebugLineItem& NewerData) {}
+	[[maybe_unused]] virtual void UpdateData(const FDebugLineItem& NewerData) {}
 
 	// @return The text to display in the name column, unless GenerateNameWidget is overridden
 	virtual FText GetDisplayName() const;
diff --git a/Engine/Source/Editor/Kismet/Private/WatchPointViewer.cpp b/Engine/Source/Editor/Kismet/Private/WatchPointViewer.cpp
--- a/Engine/Source/Editor/Kismet/Private/WatchPointViewer.cpp
+++ b/Engine/Source/Editor/Kismet/Private/WatchPointViewer.cpp
@@ -78,7 +78,7 @@ namespace
 			FText InBlueprintName,
 			FText InGraphName,
 			FText InNodeName,
-			FPropertyInstanceInfo Info
+			const FPropertyInstanceInfo& Info
 		)
 			: BP(InBP)
 			, Node(InNode)
@@ -87,16 +87,16 @@ namespace
 			, BlueprintName(MoveTemp(InBlueprintName))
 			, GraphName(MoveTemp(InGraphName))
 			, NodeName(MoveTemp(InNodeName))
-			, DisplayName(MoveTemp(Info.DisplayName))
-			, Value(MoveTemp(Info.Value))
-			, Type(MoveTemp(Info.Type))
+			, DisplayName(Info.DisplayName)
+			, Value(Info.Value)
+			, Type(Info.Type)
 		{
 			SetObjectBeingDebuggedName();
 
 			UPackage* Package = Cast<UPackage>(BP.IsValid() ? BP->GetOuter() : nullptr);
 			BlueprintPackageName = Package ? Package->GetFName() : FName();
 
-			for (TSharedPtr<FPropertyInstanceInfo>& ChildInfo : Info.Children)
+			for (const TSharedPtr<FPropertyInstanceInfo>& ChildInfo : Info.Children)
 			{
 				Children.Add(MakeShared<FWatchRow>(InBP, InNode, InPin, InObjectBeingDebugged, BlueprintName, GraphName, NodeName, MoveTemp(*ChildInfo)));
 			}
diff --git a/Engine/Source/Editor/UnrealEd/Private/PropertyInfoViewStyle.cpp b/Engine/Source/Editor/UnrealEd/Private/PropertyInfoViewStyle.cpp
--- a/Engine/Source/Editor/UnrealEd/Private/PropertyInfoViewStyle.cpp
+++ b/Engine/Source/Editor/UnrealEd/Private/PropertyInfoViewStyle.cpp
@@ -20,7 +20,7 @@ namespace PropertyInfoViewStyle
 		int32 ColorIndex = 0;
 		int32 Increment = 1;
 
-		for (int i = 0; i < IndentLevel; ++i)
+		for (int IndentCount = 0; IndentCount < IndentLevel; ++IndentCount)
 		{
 			ColorIndex += Increment;
 
@@ -51,46 +51,10 @@ namespace PropertyInfoViewStyle
 		return GetIndentBackgroundColor(Row->GetIndentLevel(), Row->AsWidget()->IsHovered());
 	}
 
-	void SConstrainedBox::Construct(const FArguments& InArgs)
-	{
-		MinWidth = InArgs._MinWidth;
-		MaxWidth = InArgs._MaxWidth;
-
-		ChildSlot
-		[
-			InArgs._Content.Widget
-		];
-	}
-
-	FVector2D SConstrainedBox::ComputeDesiredSize(float LayoutScaleMultiplier) const
-	{
-		const float MinWidthVal = MinWidth.Get().Get(0.0f);
-		const float MaxWidthVal = MaxWidth.Get().Get(0.0f);
-
-		if (MinWidthVal == 0.0f && MaxWidthVal == 0.0f)
-		{
-			return SCompoundWidget::ComputeDesiredSize(LayoutScaleMultiplier);
-		}
-		else
-		{
-			FVector2D ChildSize = ChildSlot.GetWidget()->GetDesiredSize();
-
-			float XVal = FMath::Max(MinWidthVal, ChildSize.X);
-			if (MaxWidthVal >= MinWidthVal)
-			{
-				XVal = FMath::Min(MaxWidthVal, XVal);
-			}
-
-			return FVector2D(XVal, ChildSize.Y);
-		}
-	}
-
 	int32 SIndent::OnPaint(const FPaintArgs& Args, const FGeometry& AllottedGeometry,
 							const FSlateRect& MyCullingRect, FSlateWindowElementList& OutDrawElements,
 							int32 LayerId, const FWidgetStyle& InWidgetStyle, bool bParentEnabled) const
 	{
-		constexpr int8 TabSize = 16;
-		
 		TSharedPtr<ITableRow> RowPtr = Row.Pin();
 		if (!RowPtr.IsValid())
 		{
@@ -101,14 +65,14 @@ namespace PropertyInfoViewStyle
 		const FSlateBrush* DropShadowBrush = FAppStyle::Get().GetBrush("DetailsView.ArrayDropShadow");
 
 		int32 IndentLevel = RowPtr->GetIndentLevel();
-		for (int32 i = 0; i < IndentLevel; ++i)
+		for (int32 IndentCount = 0; IndentCount < IndentLevel; ++IndentCount)
 		{
-			FSlateColor BackgroundColor = GetRowBackgroundColor(i);
+			FSlateColor BackgroundColor = GetRowBackgroundColor(IndentCount);
 
 			FSlateDrawElement::MakeBox(
 				OutDrawElements,
 				LayerId,
-				AllottedGeometry.ToPaintGeometry(FVector2D(TabSize * i, 0), FVector2D(TabSize, AllottedGeometry.GetLocalSize().Y)),
+				AllottedGeometry.ToPaintGeometry(FVector2D(TabSize * IndentCount, 0), FVector2D(TabSize, AllottedGeometry.GetLocalSize().Y)),
 				BackgroundBrush,
 				ESlateDrawEffect::None,
 				BackgroundColor.GetColor(InWidgetStyle)
@@ -117,7 +81,7 @@ namespace PropertyInfoViewStyle
 			FSlateDrawElement::MakeBox(
 				OutDrawElements,
 				LayerId + 1,
-				AllottedGeometry.ToPaintGeometry(FVector2D(TabSize * i, 0), FVector2D(TabSize, AllottedGeometry.GetLocalSize().Y)),
+				AllottedGeometry.ToPaintGeometry(FVector2D(TabSize * IndentCount, 0), FVector2D(TabSize, AllottedGeometry.GetLocalSize().Y)),
 				DropShadowBrush
 			);
 		}
@@ -146,7 +110,7 @@ namespace PropertyInfoViewStyle
 		{
 			IndentLevel = RowPtr->GetIndentLevel();
 		}
-		return IndentLevel * 16.0f;
+		return IndentLevel * TabSize;
 	}
 
 	FSlateColor SIndent::GetRowBackgroundColor(int32 IndentLevel) const
diff --git a/Engine/Source/Editor/UnrealEd/Public/PropertyInfoViewStyle.h b/Engine/Source/Editor/UnrealEd/Public/PropertyInfoViewStyle.h
--- a/Engine/Source/Editor/UnrealEd/Public/PropertyInfoViewStyle.h
+++ b/Engine/Source/Editor/UnrealEd/Public/PropertyInfoViewStyle.h
@@ -32,24 +32,6 @@ namespace PropertyInfoViewStyle
 	 */
 	UNREALED_API FSlateColor GetRowBackgroundColor(ITableRow* Row);
 
-	/** Helper class to force a widget to fill in a space. Copied from SDetailSingleItemRow.cpp */
-	class UNREALED_API SConstrainedBox : public SCompoundWidget
-	{
-	public:
-		SLATE_BEGIN_ARGS(SConstrainedBox){}
-			SLATE_DEFAULT_SLOT(FArguments, Content)
-			SLATE_ATTRIBUTE(TOptional<float>, MinWidth)
-			SLATE_ATTRIBUTE(TOptional<float>, MaxWidth)
-		SLATE_END_ARGS()
-
-		void Construct(const FArguments& InArgs);
-
-		virtual FVector2D ComputeDesiredSize(float LayoutScaleMultiplier) const override;
-
-	private:
-		TAttribute< TOptional<float> > MinWidth;
-		TAttribute< TOptional<float> > MaxWidth;
-	};
 
 	/**
 	 * SIndent is a widget used to indent trees in a layered
@@ -80,6 +62,7 @@ namespace PropertyInfoViewStyle
 
 	private:
 		TWeakPtr<ITableRow> Row;
+		float TabSize = 16.0f;
 	};
 
 	/**
