Added DDC2 local executor (out-of-process) implementation. This is as basic as I could make it, and is intended primarily to facilitate debugging and understanding.

diff --git a/Engine/Source/Editor/UnrealEd/Private/DerivedDataBuildLocalExecutor.cpp b/Engine/Source/Editor/UnrealEd/Private/DerivedDataBuildLocalExecutor.cpp
new file mode 100644
--- /dev/null
+++ b/Engine/Source/Editor/UnrealEd/Private/DerivedDataBuildLocalExecutor.cpp
@@ -0,0 +1,291 @@
+// Copyright Epic Games, Inc. All Rights Reserved.
+
+#include "DerivedDataBuildLocalExecutor.h"
+
+#include "DerivedDataBuild.h"
+#include "DerivedDataBuildAction.h"
+#include "DerivedDataBuildInputs.h"
+#include "DerivedDataBuildOutput.h"
+#include "DerivedDataBuildWorker.h"
+#include "DerivedDataRequest.h"
+#include "DerivedDataRequestOwner.h"
+#include "Features/IModularFeatures.h"
+#include "HAL/Event.h"
+#include "HAL/PlatformProcess.h"
+#include "HAL/FileManager.h"
+#include "Misc/Optional.h"
+#include "Misc/Paths.h"
+#include "Misc/PathViews.h"
+#include "Modules/ModuleManager.h"
+#include "Serialization/CompactBinarySerialization.h"
+#include "Serialization/CompactBinaryValidation.h"
+#include "Serialization/CompactBinaryWriter.h"
+
+#include <atomic>
+
+namespace UE::DerivedData
+{
+
+DEFINE_LOG_CATEGORY_STATIC(LogDerivedDataBuildLocalExecutor, Log, All);
+
+/**
+ * This implements a simple local DDC2 executor which spawns all jobs
+ * as local processes. This is intentionally as simple as possible and
+ * everything is synchronous. This is not meant to be used in production, 
+ * and is perhaps primarily useful as a debugging aid
+ * 
+ */
+
+class FLocalBuildWorkerExecutor final : public IBuildWorkerExecutor
+{
+public:
+	FLocalBuildWorkerExecutor()
+	{
+		IModularFeatures::Get().RegisterModularFeature(IBuildWorkerExecutor::GetFeatureName(), this);
+	}
+
+	virtual ~FLocalBuildWorkerExecutor()
+	{
+	}
+
+	void BuildAction(const FBuildAction&			Action,
+					 const FOptionalBuildInputs&	Inputs,
+					 const FBuildWorker&			Worker,
+					 IBuild&						BuildSystem,
+					 EBuildPolicy					Policy,
+					 IRequestOwner&					Owner,
+					 FOnBuildWorkerActionComplete&& OnComplete) final
+	{
+		// Review build action inputs to determine if they need to be materialized/propagated 
+		// (right now, they always will be)
+
+		TArray<FString> MissingInputs;
+		TArray<FStringView> MissingInputViews;
+
+		Action.IterateInputs([&](FStringView Key, const FIoHash& RawHash, uint64 RawSize)
+			{
+				if (Inputs.IsNull() || Inputs.Get().FindInput(Key).IsNull())
+				{
+					MissingInputs.Emplace(Key);
+					MissingInputViews.Add(MissingInputs.Last());
+				}
+			});
+
+		if (!MissingInputViews.IsEmpty())
+		{
+			// Report missing inputs
+			return OnComplete({ Action.GetKey(), {}, MissingInputViews, EStatus::Ok });
+		}
+
+		// This path will execute the build action synchronously in a scratch directory
+		// at this stage, all inputs are available in process
+
+		static std::atomic<int32> SerialNo = 0;
+
+		FString SandboxRoot = FPaths::EngineSavedDir() / TEXT("LocalExec") / TEXT("Scratch");
+		SandboxRoot.AppendInt(++SerialNo);
+
+		// Manifest worker in scratch area
+
+		{
+			TArray<FIoHash> MissingWorkerData;
+
+			TArray<FIoHash> WorkerFileHashes;
+			TArray<TTuple<FStringView, bool>> WorkerFileMeta;
+			Worker.IterateExecutables([&WorkerFileHashes, &WorkerFileMeta](FStringView Path, const FIoHash& RawHash, uint64 RawSize) {
+				WorkerFileHashes.Emplace(RawHash);
+				WorkerFileMeta.Emplace(Path, true);
+			});
+
+			Worker.IterateFiles([&WorkerFileHashes, &WorkerFileMeta](FStringView Path, const FIoHash& RawHash, uint64 RawSize) {
+				WorkerFileHashes.Emplace(RawHash);
+				WorkerFileMeta.Emplace(Path, false);
+			});
+
+			FRequestOwner BlockingOwner(EPriority::Blocking);
+			Worker.FindFileData(WorkerFileHashes, BlockingOwner, [&](FBuildWorkerFileDataCompleteParams&& Params) {
+				uint32 MetaIndex = 0;
+				for (const FCompressedBuffer& Buffer : Params.Files)
+				{
+					const TTuple<FStringView, bool>& Meta = WorkerFileMeta[MetaIndex];
+
+					FString Path { SandboxRoot / FString(Meta.Key) };
+
+					if (TUniquePtr<FArchive> Ar{ IFileManager::Get().CreateFileWriter(*Path, FILEWRITE_Silent) })
+					{
+						FCompositeBuffer DecompressedComposite = Buffer.DecompressToComposite();
+
+						for (auto& Segment : DecompressedComposite.GetSegments())
+						{
+							Ar->Serialize((void*)Segment.GetData(), Segment.GetSize());
+						}
+					}
+
+					++MetaIndex;
+				}
+			});
+			BlockingOwner.Wait();
+
+			// This directory must exist in order for the builder to run correctly
+			FString BinPath{ SandboxRoot / TEXT("Engine/Binaries/Win64") };
+			IFileManager::Get().MakeDirectory(*BinPath);
+		}
+
+		// Manifest inputs in scratch area
+
+		{
+			if (!Inputs.IsNull())
+			{
+				Inputs.Get().IterateInputs([&](FStringView Key, const FCompressedBuffer& Buffer) {
+					// We should be able to accomplish this without decompression
+					FCompositeBuffer DecompressedComposite = Buffer.DecompressToComposite();
+					FIoHash RawHash = FIoHash::HashBuffer(DecompressedComposite);
+
+					TStringBuilder<128> InputPath;
+					InputPath << TEXT("Inputs/") << RawHash;
+
+					FString Path{ SandboxRoot / *InputPath };
+
+					if (TUniquePtr<FArchive> Ar{ IFileManager::Get().CreateFileWriter(*Path, FILEWRITE_Silent) })
+					{
+						FCompressedBuffer Comp = FCompressedBuffer::FromCompressed(Buffer.GetCompressed());
+						*Ar << Comp;
+					}
+				});
+			}
+		}
+
+		// Serialize action specification
+
+		{
+			FString Path{ SandboxRoot / TEXT("Build.action") };
+
+			if (TUniquePtr<FArchive> Ar{ IFileManager::Get().CreateFileWriter(*Path, FILEWRITE_Silent) })
+			{
+				FCbWriter BuildActionWriter;
+				Action.Save(BuildActionWriter);
+				BuildActionWriter.Save(*Ar);
+			}
+		}
+
+		const bool bLaunchDetached = false;
+		const bool bLaunchHidden = false;
+		const bool bLaunchReallyHidden = false;
+		uint32 ProcessID = 0;
+		const int PriorityModifier = 0;
+		const TCHAR* WorkingDirectory = *SandboxRoot;
+		const FString WorkerPath { SandboxRoot / FString{Worker.GetPath()} };
+
+		FProcHandle ProcHandle = FPlatformProcess::CreateProc(
+			*WorkerPath, 
+			TEXT("-Build=build.action"), 
+			bLaunchDetached,
+			bLaunchHidden, 
+			bLaunchReallyHidden, 
+			&ProcessID, 
+			PriorityModifier, 
+			WorkingDirectory, 
+			nullptr, 
+			nullptr);
+
+		FPlatformProcess::WaitForProc(ProcHandle);
+		
+		int32 ExitCode = -1;
+		if (!FPlatformProcess::GetProcReturnCode(ProcHandle, &ExitCode) || ExitCode != 0)
+		{
+			UE_LOG(LogDerivedDataBuildLocalExecutor, Warning, TEXT("Worker process exit code = %d!"), ExitCode);
+			return OnComplete({ Action.GetKey(), {}, {}, EStatus::Error });
+		}
+
+		// Gather results
+
+		FOptionalBuildOutput RemoteBuildOutput;
+
+		{
+			FString BuildOutputPath{ SandboxRoot / TEXT("Build.output") };
+
+			if (TUniquePtr<FArchive> Ar{ IFileManager::Get().CreateFileReader(*BuildOutputPath, FILEREAD_Silent) })
+			{
+				FCbObject BuildOutput = LoadCompactBinary(*Ar).AsObject();
+
+				if (Ar->IsError())
+				{
+					UE_LOG(LogDerivedDataBuildLocalExecutor, Warning, TEXT("Worker error: build output structure not valid!"));
+
+					return OnComplete({ Action.GetKey(), {}, {}, EStatus::Error });
+				}
+
+				RemoteBuildOutput = FBuildOutput::Load(	Action.GetName(), 
+														Action.GetFunction(), 
+														BuildOutput);
+			}
+		}
+
+		if (RemoteBuildOutput.IsNull())
+		{
+			UE_LOG(LogDerivedDataBuildLocalExecutor, Warning, TEXT("Remote execution system error: build output blob missing!"));
+
+			return OnComplete({ Action.GetKey(), {}, {}, EStatus::Error });
+		}
+
+		FBuildOutputBuilder OutputBuilder = BuildSystem.CreateOutput(Action.GetName(), Action.GetFunction());
+
+		for (const FPayload& Payload : RemoteBuildOutput.Get().GetPayloads())
+		{
+			TStringBuilder<128> OutputPath;
+			OutputPath << TEXT("Outputs/") << Payload.GetRawHash();
+
+			FString Path{ SandboxRoot / *OutputPath };
+
+			FCompressedBuffer BufferForPayload;
+
+			if (TUniquePtr<FArchive> Ar{ IFileManager::Get().CreateFileReader(*Path, FILEREAD_Silent) })
+			{
+				BufferForPayload = FCompressedBuffer::FromCompressed(*Ar);
+			}
+
+			if (BufferForPayload.IsNull())
+			{
+				UE_LOG(LogDerivedDataBuildLocalExecutor, Warning, TEXT("Remote execution system error: payload blob missing!"));
+
+				return OnComplete({ Action.GetKey(), {}, {}, EStatus::Error });
+			}
+
+			OutputBuilder.AddPayload(FPayload(Payload.GetId(), BufferForPayload));
+		}
+
+		FBuildOutput BuildOutput = OutputBuilder.Build();
+
+		return OnComplete({ Action.GetKey(), BuildOutput, {}, EStatus::Ok });
+	}
+
+	TConstArrayView<FStringView> GetHostPlatforms() const final
+	{
+		static constexpr FStringView HostPlatforms[]{TEXT("Win64"_SV)};
+		return HostPlatforms;
+	}
+
+	void DumpStats() {}
+};
+
+}  // namespace UE::DerivedData
+
+TOptional<UE::DerivedData::FLocalBuildWorkerExecutor> GLocalBuildWorkerExecutor;
+
+void InitDerivedDataBuildLocalExecutor()
+{
+	if (!GLocalBuildWorkerExecutor.IsSet())
+	{
+		GLocalBuildWorkerExecutor.Emplace();
+	}
+}
+
+void DumpDerivedDataBuildLocalExecutorStats()
+{
+	static bool bHasRun = false;
+	if (GLocalBuildWorkerExecutor.IsSet() && !bHasRun)
+	{
+		bHasRun = true;
+		GLocalBuildWorkerExecutor->DumpStats();
+	}
+}
diff --git a/Engine/Source/Editor/UnrealEd/Private/DerivedDataBuildLocalExecutor.h b/Engine/Source/Editor/UnrealEd/Private/DerivedDataBuildLocalExecutor.h
new file mode 100644
--- /dev/null
+++ b/Engine/Source/Editor/UnrealEd/Private/DerivedDataBuildLocalExecutor.h
@@ -0,0 +1,6 @@
+// Copyright Epic Games, Inc. All Rights Reserved.
+
+#pragma once
+
+void InitDerivedDataBuildLocalExecutor();
+void DumpDerivedDataBuildLocalExecutorStats();
