Cooker: RequestClusters: All transitive dependencies of the collective requests from CookByTheBook or the individual requests from CookOnTheFly are queued when the external request is ingested rather than queued when each request is loaded.
#rnx

diff --git a/Engine/Source/Editor/UnrealEd/Classes/CookOnTheSide/CookOnTheFlyServer.h b/Engine/Source/Editor/UnrealEd/Classes/CookOnTheSide/CookOnTheFlyServer.h
--- a/Engine/Source/Editor/UnrealEd/Classes/CookOnTheSide/CookOnTheFlyServer.h
+++ b/Engine/Source/Editor/UnrealEd/Classes/CookOnTheSide/CookOnTheFlyServer.h
@@ -5,6 +5,7 @@
 #include "CoreMinimal.h"
 
 #include "Containers/ArrayView.h"
+#include "Containers/RingBuffer.h"
 #include "CookOnTheSide/CookLog.h"
 #include "INetworkFileSystemModule.h"
 #include "IPlatformFileSandboxWrapper.h"
@@ -112,6 +113,7 @@ ENUM_CLASS_FLAGS(ECookTickFlags);
 namespace UE::Cook
 {
 	class FExternalRequests;
+	class FRequestCluster;
 	class FSaveCookedPackageContext;
 	class ICookOnTheFlyRequestManager;
 	struct FCookerTimer;
@@ -304,19 +306,22 @@ private:
 	void UpdateDisplay(ECookTickFlags CookFlags, bool bForceDisplay);
 	enum class ECookAction
 	{
-		Done,		// The cook is complete; no requests remain in any non-idle state
-		Request,    // Process the RequestQueue
-		Load,		// Process the LoadQueue
-		LoadLimited,// Process the LoadQueue, stopping when loadqueuelength reaches the desired population level
-		Save,       // Process the SaveQueue
-		SaveLimited,// Process the SaveQueue, stopping when savequeuelength reaches the desired population level
-		YieldTick,  // Progress is blocked by an async result. Temporarily exit TickCookOnTheSide.
-		Cancel,		// Cancel the current CookByTheBook
+		Done,			// The cook is complete; no requests remain in any non-idle state
+		RequestCluster,	// Process the RequestClusters
+		Request,		// Process the RequestQueue
+		Load,			// Process the LoadQueue
+		LoadLimited,	// Process the LoadQueue, stopping when loadqueuelength reaches the desired population level
+		Save,			// Process the SaveQueue
+		SaveLimited,	// Process the SaveQueue, stopping when savequeuelength reaches the desired population level
+		YieldTick,		// Progress is blocked by an async result. Temporarily exit TickCookOnTheSide.
+		Cancel,			// Cancel the current CookByTheBook
 	};
 	/** Inspect all tasks the scheduler could do and return which one it should do. */
 	ECookAction DecideNextCookAction(UE::Cook::FTickStackData& StackData);
-	/** Execute any existing external callbacks and push any existing external cook requests into the RequestQueue. */
+	/** Execute any existing external callbacks and push any existing external cook requests into new RequestClusters. */
 	void PumpExternalRequests(const UE::Cook::FCookerTimer& CookerTimer);
+	/** Pull off each RequestCluster and find all transitive dependencies and push them into the RequestQueue. */
+	void PumpRequestClusters(UE::Cook::FTickStackData& StackData);
 	/** Inspect the next package in the RequestQueue and push it on to its next state. Report the number of PackageDatas that were pushed to load. */
 	void PumpRequests(UE::Cook::FTickStackData& StackData, int32& OutNumPushed);
 	/** Handle the requested PackageData that has been peeled off of the RequestQueue, e.g. by sending it on to the LoadQueue. Report the number of PackageDatas that were pushed to load (includes dependencies, so may be more than 1). */
@@ -1120,6 +1125,11 @@ private:
 	bool bPreloadingEnabled = false;
 	/** If TargetDomain is enabled, we load/save TargetDomainKey hashes and use those to test whether packages have already been cooked in hybrid-iterative builds. */
 	bool bTargetDomainEnabled = true;
+	/**
+	 * Following soft dependencies in request clusters can cause editor-only packages to be cooked due to lack of robust editor-only reference detection.
+	 * Projects can disable it to prevent cooking those packages, at the cost of lower cook performance.
+	 */
+	bool bExploreSoftReferencesOnStart = true;
 
 	/** Timers for tracking how long we have been busy, to manage retries and warnings of deadlock */
 	float SaveBusyTimeLastRetry = 0.f;
@@ -1133,6 +1143,7 @@ private:
 	TUniquePtr<UE::Cook::FPackageTracker> PackageTracker;
 	TUniquePtr<UE::Cook::FPackageDatas> PackageDatas;
 	TUniquePtr<UE::Cook::FExternalRequests> ExternalRequests;
+	TRingBuffer<UE::Cook::FRequestCluster> RequestClusters;
 
 	TArray<FSavePackageContext*> SavePackageContexts;
 	/** Objects that were collected during the single-threaded PreGarbageCollect callback and that should be reported as referenced in CookerAddReferencedObjects. */
@@ -1141,10 +1152,8 @@ private:
 	/** Helper struct for running cooking in diagnostic modes */
 	TUniquePtr<FDiffModeCookServerUtils> DiffModeHelper;
 
-	// temporary -- should eliminate the need for this. Only required right now because FullLoadAndSave 
-	// accesses maps directly
-	friend FPackageNameCache;
 	friend UE::Cook::FPackageData;
 	friend UE::Cook::FPendingCookedPlatformData;
 	friend UE::Cook::FPlatformManager;
+	friend UE::Cook::FRequestCluster;
 };
diff --git a/Engine/Source/Editor/UnrealEd/Private/CookOnTheFlyServer.cpp b/Engine/Source/Editor/UnrealEd/Private/CookOnTheFlyServer.cpp
--- a/Engine/Source/Editor/UnrealEd/Private/CookOnTheFlyServer.cpp
+++ b/Engine/Source/Editor/UnrealEd/Private/CookOnTheFlyServer.cpp
@@ -21,6 +21,7 @@
 #include "Cooker/CookPackageData.h"
 #include "Cooker/CookPlatformManager.h"
 #include "Cooker/CookProfiling.h"
+#include "Cooker/CookRequestCluster.h"
 #include "Cooker/CookRequests.h"
 #include "Cooker/CookTypes.h"
 #include "Cooker/IoStoreCookOnTheFlyRequestManager.h"
@@ -312,55 +313,6 @@ void LogCookerMessage( const FString& MessageText, EMessageSeverity::Type Severi
 	MessageLog.Notify(FText(), EMessageSeverity::Warning, false);
 }
 
-//////////////////////////////////////////////////////////////////////////
-// Cook by the book options
-
-struct UCookOnTheFlyServer::FCookByTheBookOptions
-{
-public:
-	/** Should we generate streaming install manifests (only valid option in cook by the book) */
-	bool							bGenerateStreamingInstallManifests = false;
-
-	/** Should we generate a seperate manifest for map dependencies */
-	bool							bGenerateDependenciesForMaps = false;
-
-	/** Is cook by the book currently running */
-	bool							bRunning = false;
-
-	/** Cancel has been queued will be processed next tick */
-	bool							bCancel = false;
-
-	/** DlcName setup if we are cooking dlc will be used as the directory to save cooked files to */
-	FString							DlcName;
-
-	/** Create a release from this manifest and store it in the releases directory for this cgame */
-	FString							CreateReleaseVersion;
-
-	/** Dependency graph of maps as root objects. */
-	TFastPointerMap<const ITargetPlatform*,TMap<FName,TSet<FName>>> MapDependencyGraphs;
-
-	/** If we are based on a release version of the game this is the set of packages which were cooked in that release. Map from platform name to list of uncooked package filenames */
-	TMap<FName, TArray<FName>>			BasedOnReleaseCookedPackages;
-
-	/** Timing information about cook by the book */
-	double							CookTime = 0.0;
-	double							CookStartTime = 0.0;
-
-	/** error when detecting engine content being used in this cook */
-	bool							bErrorOnEngineContentUse = false;
-	bool							bSkipHardReferences = false;
-	bool							bSkipSoftReferences = false;
-	bool							bFullLoadAndSave = false;
-	bool							bIoStore = false;
-	bool							bZenStore = false;
-	bool							bCookAgainstFixedBase = false;
-	bool							bDlcLoadMainAssetRegistry = false;
-	TArray<FName>					StartupPackages;
-
-	/** Mapping from source packages to their localized variants (based on the culture list in FCookByTheBookStartupOptions) */
-	TMap<FName, TArray<FName>>		SourceToLocalizedPackageVariants;
-};
-
 //////////////////////////////////////////////////////////////////////////
 // Cook on the fly server interface adapter
 
@@ -1425,6 +1377,9 @@ uint32 UCookOnTheFlyServer::TickCookOnTheSide(const float TimeSlice, uint32 &Coo
 			bool bBusy;
 			switch (CookAction)
 			{
+			case ECookAction::RequestCluster:
+				PumpRequestClusters(StackData);
+				break;
 			case ECookAction::Request:
 				PumpRequests(StackData, NumPushed);
 				if (NumPushed > 0)
@@ -1748,6 +1703,11 @@ UCookOnTheFlyServer::ECookAction UCookOnTheFlyServer::DecideNextCookAction(UE::C
 		return ECookAction::YieldTick;
 	}
 
+	if (RequestClusters.Num())
+	{
+		return ECookAction::RequestCluster;
+	}
+
 	UE::Cook::FPackageDataMonitor& Monitor = PackageDatas->GetMonitor();
 	if (Monitor.GetNumUrgent() > 0)
 	{
@@ -1839,12 +1799,12 @@ void UCookOnTheFlyServer::PumpExternalRequests(const UE::Cook::FCookerTimer& Coo
 	}
 	UE_SCOPED_COOKTIMER(PumpExternalRequests);
 
-	UE::Cook::FFilePlatformRequest ToBuild;
+	TArray<UE::Cook::FFilePlatformRequest> BuildRequests;
 	TArray<UE::Cook::FSchedulerCallback> SchedulerCallbacks;
 	UE::Cook::EExternalRequestType RequestType;
 	while (!CookerTimer.IsTimeUp())
 	{
-		RequestType = ExternalRequests->DequeueRequest(SchedulerCallbacks, ToBuild);
+		RequestType = ExternalRequests->DequeueNextCluster(SchedulerCallbacks, BuildRequests);
 		if (RequestType == UE::Cook::EExternalRequestType::None)
 		{
 			// No more requests to process
@@ -1860,34 +1820,40 @@ void UCookOnTheFlyServer::PumpExternalRequests(const UE::Cook::FCookerTimer& Coo
 		}
 		else
 		{
-			check(RequestType == UE::Cook::EExternalRequestType::Cook && ToBuild.IsValid());
-			FName FileName = ToBuild.GetFilename();
+			check(RequestType == UE::Cook::EExternalRequestType::Cook && BuildRequests.Num() > 0);
 #if PROFILE_NETWORK
 			if (NetworkRequestEvent)
 			{
 				NetworkRequestEvent->Trigger();
 			}
 #endif
-#if DEBUG_COOKONTHEFLY
-			UE_LOG(LogCook, Display, TEXT("Processing request for package %s"), *FileName.ToString());
-#endif
+			bool bRequestsAreUrgent = IsCookOnTheFlyMode();
+			UE::Cook::FRequestCluster::AddClusters(*this, MoveTemp(BuildRequests), bRequestsAreUrgent,
+				RequestClusters);
+		}
+	}
+}
 
-			const FName* PackageName = GetPackageNameCache().GetCachedPackageNameFromStandardFileName(FileName, /* bExactMatchRequired */ false, &FileName);
-			if (!PackageName)
+void UCookOnTheFlyServer::PumpRequestClusters(UE::Cook::FTickStackData& StackData)
+{
+	using namespace UE::Cook;
+
+	const FCookerTimer& CookerTimer = StackData.Timer;
+	while (RequestClusters.Num() > 0 && !CookerTimer.IsTimeUp())
+	{
+		FRequestCluster& RequestCluster = RequestClusters.First();
+		bool bComplete = false;
+		RequestCluster.Process(CookerTimer, bComplete);
+		if (bComplete)
+		{
+			for (FRequestCluster::FRequest& Request : RequestCluster.GetRequests())
 			{
-				LogCookerMessage(FString::Printf(TEXT("Could not find package at file %s!"), *ToBuild.GetFilename().ToString()), EMessageSeverity::Error);
-				UE_LOG(LogCook, Error, TEXT("Could not find package at file %s!"), *ToBuild.GetFilename().ToString());
-				UE::Cook::FCompletionCallback CompletionCallback(MoveTemp(ToBuild.GetCompletionCallback()));
-				if (CompletionCallback)
-				{
-					CompletionCallback(nullptr);
-				}
-				continue;
+				FPackageData& PackageData(PackageDatas->FindOrAddPackageData(Request.PackageName, Request.FileName));
+				PackageData.UpdateRequestData(RequestCluster.GetPlatforms(), Request.bIsUrgent,
+					MoveTemp(Request.CompletionCallback));
 			}
 
-			UE::Cook::FPackageData& PackageData(PackageDatas->FindOrAddPackageData(*PackageName, FileName));
-			bool bIsUrgent = IsCookOnTheFlyMode();
-			PackageData.UpdateRequestData(ToBuild.GetPlatforms(), bIsUrgent, MoveTemp(ToBuild.GetCompletionCallback()));
+			RequestClusters.PopFront();
 		}
 	}
 }
@@ -1953,6 +1919,8 @@ void UCookOnTheFlyServer::ProcessRequest(UE::Cook::FPackageData& PackageData, UE
 
 	if (!PackageData.GetIsUrgent() && (!IsCookByTheBookMode() || !CookByTheBookOptions->bSkipHardReferences))
 	{
+		// TODO: We need to AddDependenciesToLoadQueue only for requests that did not already do so
+		// because they were part of a CookRequestCluster
 		int32 NumPushed;
 		AddDependenciesToLoadQueue(PackageData, NumPushed);
 		OutNumPushed += NumPushed;
@@ -4953,6 +4921,7 @@ void UCookOnTheFlyServer::Initialize( ECookMode::Type DesiredCookMode, ECookInit
 		FLinkerLoad::SetPreloadingEnabled(true);
 	}
 	GConfig->GetBool(TEXT("CookSettings"), TEXT("TargetDomainEnabled"), bTargetDomainEnabled, GEditorIni);
+	GConfig->GetBool(TEXT("CookSettings"), TEXT("ExploreSoftReferencesOnStart"), bExploreSoftReferencesOnStart, GEditorIni);
 
 	{
 		const FConfigSection* CacheSettings = GConfig->GetSectionPrivate(TEXT("CookPlatformDataCacheSettings"), false, true, GEditorIni);
@@ -8841,9 +8810,9 @@ uint32 UCookOnTheFlyServer::FullLoadAndSave(uint32& CookedPackageCount)
 		UE_SCOPED_HIERARCHICAL_COOKTIMER(FullLoadAndSave_RequestedLoads);
 		while (ExternalRequests->HasRequests())
 		{
-			UE::Cook::FFilePlatformRequest ToBuild;
+			TArray<UE::Cook::FFilePlatformRequest> BuildRequests;
 			TArray<UE::Cook::FSchedulerCallback> SchedulerCallbacks;
-			UE::Cook::EExternalRequestType RequestType = ExternalRequests->DequeueRequest(SchedulerCallbacks, /* out */ ToBuild);
+			UE::Cook::EExternalRequestType RequestType = ExternalRequests->DequeueNextCluster(SchedulerCallbacks, BuildRequests);
 			if (RequestType == UE::Cook::EExternalRequestType::Callback)
 			{
 				for (UE::Cook::FSchedulerCallback& SchedulerCallback : SchedulerCallbacks)
@@ -8852,21 +8821,23 @@ uint32 UCookOnTheFlyServer::FullLoadAndSave(uint32& CookedPackageCount)
 				}
 				continue;
 			}
-			check(RequestType == UE::Cook::EExternalRequestType::Cook && ToBuild.IsValid());
-
-			const FName BuildFilenameFName = ToBuild.GetFilename();
-			if (!PackageTracker->NeverCookPackageList.Contains(BuildFilenameFName))
+			check(RequestType == UE::Cook::EExternalRequestType::Cook && BuildRequests.Num() > 0);
+			for (UE::Cook::FFilePlatformRequest& ToBuild : BuildRequests)
 			{
-				const FString BuildFilename = BuildFilenameFName.ToString();
-				GIsCookerLoadingPackage = true;
-				UE_SCOPED_HIERARCHICAL_COOKTIMER(LoadPackage);
-				LoadPackage(nullptr, *BuildFilename, LOAD_None);
-				if (GShaderCompilingManager)
+				const FName BuildFilenameFName = ToBuild.GetFilename();
+				if (!PackageTracker->NeverCookPackageList.Contains(BuildFilenameFName))
 				{
-					GShaderCompilingManager->ProcessAsyncResults(true, false);
+					const FString BuildFilename = BuildFilenameFName.ToString();
+					GIsCookerLoadingPackage = true;
+					UE_SCOPED_HIERARCHICAL_COOKTIMER(LoadPackage);
+					LoadPackage(nullptr, *BuildFilename, LOAD_None);
+					if (GShaderCompilingManager)
+					{
+						GShaderCompilingManager->ProcessAsyncResults(true, false);
+					}
+					FAssetCompilingManager::Get().ProcessAsyncTasks(true);
+					GIsCookerLoadingPackage = false;
 				}
-				FAssetCompilingManager::Get().ProcessAsyncTasks(true);
-				GIsCookerLoadingPackage = false;
 			}
 		}
 	}
diff --git a/Engine/Source/Editor/UnrealEd/Private/Cooker/CookRequestCluster.cpp b/Engine/Source/Editor/UnrealEd/Private/Cooker/CookRequestCluster.cpp
new file mode 100644
--- /dev/null
+++ b/Engine/Source/Editor/UnrealEd/Private/Cooker/CookRequestCluster.cpp
@@ -0,0 +1,452 @@
+// Copyright Epic Games, Inc. All Rights Reserved.
+
+#include "Cooker/CookRequestCluster.h"
+
+#include "Algo/Unique.h"
+#include "Algo/Sort.h"
+#include "Cooker/CookPackageData.h"
+#include "Cooker/CookProfiling.h"
+#include "Cooker/CookRequests.h"
+#include "Cooker/PackageNameCache.h"
+#include "Cooker/PackageTracker.h"
+#include "CookOnTheSide/CookOnTheFlyServer.h"
+#include "EditorDomain/EditorDomainUtils.h"
+#include "Misc/StringBuilder.h"
+#include "TargetDomain/TargetDomainUtils.h"
+
+namespace UE::Cook
+{
+
+FRequestCluster::FRequestCluster(UCookOnTheFlyServer& COTFS, TArray<const ITargetPlatform*>&& InPlatforms)
+	: Platforms(MoveTemp(InPlatforms))
+	, PackageDatas(*COTFS.PackageDatas)
+	, AssetRegistry(*IAssetRegistry::Get())
+	, PackageNameCache(COTFS.GetPackageNameCache())
+	, PackageTracker(*COTFS.PackageTracker)
+{
+}
+
+void FRequestCluster::AddClusters(UCookOnTheFlyServer& COTFS, TArray<FFilePlatformRequest>&& InRequests, bool bRequestsAreUrgent,
+	TRingBuffer<FRequestCluster>& OutClusters)
+{
+	if (InRequests.Num() == 0)
+	{
+		return;
+	}
+
+	TArray<FRequestCluster, TInlineAllocator<1>> AddedClusters;
+	auto FindOrAddClusterForPlatforms = [&AddedClusters, &COTFS](TArray<const ITargetPlatform*>&& Platforms)
+	{
+		for (FRequestCluster& Existing : AddedClusters)
+		{
+			if (Existing.GetPlatforms() == Platforms)
+			{
+				Platforms.Reset();
+				return &Existing;
+			}
+		}
+		return &AddedClusters.Emplace_GetRef(COTFS, MoveTemp(Platforms));
+	};
+
+	UE::Cook::FRequestCluster* MRUCluster = FindOrAddClusterForPlatforms(MoveTemp(InRequests[0].GetPlatforms()));
+	// The usual case is all platforms are the same, so reserve the first Cluster's size assuming it will get all requests
+	MRUCluster->Requests.Reserve(InRequests.Num());
+	// Add the first Request to it. Since we've already taken away the Platforms from the first request, we have to handle it specially.
+	MRUCluster->Requests.Add(FRequest(MoveTemp(InRequests[0]), bRequestsAreUrgent));
+
+	for (FFilePlatformRequest& Request : TArrayView<FFilePlatformRequest>(InRequests).Slice(1, InRequests.Num() - 1))
+	{
+		if (Request.GetPlatforms() != MRUCluster->GetPlatforms())
+		{
+			// MRUCluster points to data inside AddedClusters, so we have to recalculate MRUCluster whenever we add
+			MRUCluster = FindOrAddClusterForPlatforms(MoveTemp(Request.GetPlatforms()));
+		}
+		else
+		{
+			Request.GetPlatforms().Reset();
+		}
+
+		MRUCluster->Requests.Add(FRequest(MoveTemp(Request), bRequestsAreUrgent));
+	}
+
+	for (UE::Cook::FRequestCluster& AddedCluster : AddedClusters)
+	{
+		AddedCluster.Initialize(COTFS);
+		OutClusters.Add(MoveTemp(AddedCluster));
+	}
+}
+
+FRequestCluster::FRequest::FRequest(FFilePlatformRequest&& FileRequest, bool bInIsUrgent)
+	: PackageName(NAME_None)
+	, FileName(FileRequest.GetFilename())
+	, CompletionCallback(MoveTemp(FileRequest.GetCompletionCallback()))
+	, bIsUrgent(bInIsUrgent)
+{
+}
+
+FRequestCluster::FRequest::FRequest(FName InPackageName, FName InFileName, FCompletionCallback&& InCompletionCallback,
+	bool bInIsUrgent)
+	: PackageName(InPackageName)
+	, FileName(InFileName)
+	, CompletionCallback(MoveTemp(InCompletionCallback))
+	, bIsUrgent(bInIsUrgent)
+{
+}
+
+void FRequestCluster::Initialize(UCookOnTheFlyServer& COTFS)
+{
+	UCookOnTheFlyServer::FCookByTheBookOptions* Options = COTFS.CookByTheBookOptions;
+	if (Options)
+	{
+		bAllowHardDependencies = !Options->bSkipHardReferences;
+		// We always skip soft dependencies in cluster exploration if the cook commandline is set to skip soft references.
+		// We also need to skip soft dependencies from assetregistry, and make other editor-only robustness fixes, 
+		// if the project has problems with editor-only robustness and has turned bExploreSoftReferencesOnStart off
+		bAllowSoftDependencies = !Options->bSkipSoftReferences && COTFS.bExploreSoftReferencesOnStart;
+		// Build dependency soft references are more robust, and can be enabled whenever the cook mode allows soft references.
+		bAllowSoftBuildDependencies = !Options->bSkipSoftReferences;
+		bErrorOnEngineContentUse = Options->bErrorOnEngineContentUse;
+	}
+	bTargetDomainEnabled = COTFS.bTargetDomainEnabled;
+	if (bErrorOnEngineContentUse)
+	{
+		DLCPath = FPaths::Combine(*COTFS.GetBaseDirectoryForDLC(), TEXT("Content"));
+		FPaths::MakeStandardFilename(DLCPath);
+	}
+}
+
+void FRequestCluster::Process(const FCookerTimer& CookerTimer, bool& bOutComplete)
+{
+	bOutComplete = true;
+
+	FetchPackageNames(CookerTimer, bOutComplete);
+	if (!bOutComplete)
+	{
+		return;
+	}
+	FetchDependencies(CookerTimer, bOutComplete);
+}
+
+void FRequestCluster::FetchPackageNames(const FCookerTimer& CookerTimer, bool& bOutComplete)
+{
+	if (bPackageNamesComplete)
+	{
+		return;
+	}
+
+	constexpr int32 TimerCheckPeriod = 100; // Do not incur the cost of checking the timer on every package
+	int32 RequestsNum = Requests.Num();
+	while (this->NextRequest < RequestsNum)
+	{
+		if (NextRequest % TimerCheckPeriod == 0 && CookerTimer.IsTimeUp())
+		{
+			break;
+		}
+
+		FRequest& Request = Requests[NextRequest];
+		FName OriginalName = Request.FileName;
+#if DEBUG_COOKONTHEFLY
+		UE_LOG(LogCook, Display, TEXT("Processing request for package %s"), *OriginalName.ToString());
+#endif
+		const FName* PackageName = PackageNameCache.GetCachedPackageNameFromStandardFileName(OriginalName,
+			/* bExactMatchRequired */ false, &Request.FileName);
+		if (!PackageName)
+		{
+			LogCookerMessage(FString::Printf(TEXT("Could not find package at file %s!"),
+				*OriginalName.ToString()), EMessageSeverity::Error);
+			UE_LOG(LogCook, Error, TEXT("Could not find package at file %s!"), *OriginalName.ToString());
+			UE::Cook::FCompletionCallback CompletionCallback(MoveTemp(Request.CompletionCallback));
+			if (CompletionCallback)
+			{
+				CompletionCallback(nullptr);
+			}
+			Requests.RemoveAt(NextRequest); // Preserve the order of incoming requests
+			--RequestsNum;
+		}
+		else
+		{
+			Request.PackageName = *PackageName;
+			++this->NextRequest;
+		}
+	}
+	if (NextRequest < RequestsNum)
+	{
+		bOutComplete = false;
+		return;
+	}
+
+	NextRequest = 0;
+	bPackageNamesComplete = true;
+}
+
+void FRequestCluster::FetchDependencies(const FCookerTimer& CookerTimer, bool& bOutComplete)
+{
+	if (bDependenciesComplete)
+	{
+		return;
+	}
+
+	if (!bAllowHardDependencies)
+	{
+		bDependenciesComplete = true;
+		return;
+	}
+
+	// Do a depth first search over each initial Request. Use DFS rather than breadth first so that we can
+	// (a) Cheaply sort the transitive set of requests from root to leaf
+	// (b) Coarsely visit all the references of an Asset, which are likely related, before moving on to a new Asset
+	//
+	// For each new initial request, move all of the new Assets found in its transitive graph to earlier than the previous
+	// requests, so that they will occur earlier than their dependencies from the previous request
+	//
+	// For the output order, we will reverse our visit order (which is in root to leaf order) so that the cooker loads
+	// leaf packages first.
+	// 
+	// When initial request assets are unrelated, we want to load those assets in the original order. Since we are
+	// (a) moving later requests before earlier requests but then (b) reversing the list, the reversals cancel out
+	// and we should iterate the initial requests from front to back.
+	TRingBuffer<FRequestClusterGraphData::FNamePair>& TransitiveRequests(DependencyGraphData.TransitiveRequests);
+	TRingBuffer<FRequestClusterGraphData::FNamePair>& Segment(DependencyGraphData.Segment);
+	TSet<FName>& Visited(DependencyGraphData.Visited);
+	TArray<FRequestClusterGraphData::FStackData>& StackStorage(DependencyGraphData.Stack);
+	int32& StackNum(DependencyGraphData.StackNum);
+	int32 NumRequests = Requests.Num();
+	int32 TimerCounter = 0;
+	auto IsStackEmpty = [&StackNum]()
+	{
+		return StackNum == 0;
+	};
+	auto TopStack = [&StackStorage, &StackNum]() -> FRequestClusterGraphData::FStackData&
+	{
+		return StackStorage[StackNum-1];
+	};
+	auto PushStack = [&StackStorage, &StackNum](FName PackageName) -> FRequestClusterGraphData::FStackData&
+	{
+		if (StackNum == StackStorage.Num())
+		{
+			StackStorage.Emplace();
+		}
+		FRequestClusterGraphData::FStackData& Top = StackStorage[StackNum++];
+		Top.Reset(PackageName);
+		return Top;
+	};
+	auto PopStack = [&StackNum]
+	{
+		--StackNum;
+	};
+
+	auto AddVertex = [this, &Segment, &Visited, &PushStack, &TimerCounter]
+		(FName PackageName, FName FileName, bool bInitial, bool& bOutPushedStack)
+	{
+		bool bAlreadyInSet;
+		Visited.Add(PackageName, &bAlreadyInSet);
+		if (bAlreadyInSet)
+		{
+			bOutPushedStack = false;
+			return;
+		}
+		++TimerCounter;
+
+		// Add the package to the transitive output if its cookable, and always add the initial vertices
+		bool bCookable = IsRequestCookable(PackageName, FileName);
+		if (bInitial || bCookable)
+		{
+			bOutPushedStack = true;
+			Segment.Add(FRequestClusterGraphData::FNamePair{ PackageName, FileName });
+			FRequestClusterGraphData::FStackData& StackData = PushStack(PackageName);
+
+			bool bExploreDependencies = bCookable;
+			if (bCookable && !bAllowSoftDependencies)
+			{
+				// TODO: Editor-only packages may be loaded at editor startup, which makes them skip the request state
+				// and go straight to load and save, where they are culled and returned to idle.
+				// If we add dependencies from these startup packages, in some projects we add transitive dependency
+				// packages that are not otherwise cooked. To prevent breaking projects we are temporarily skipping the
+				// dependency search from these packages.
+				bool bInProgress = false;
+				FPackageData* PackageData = PackageDatas.FindPackageDataByPackageName(PackageName);
+				if (PackageData && (PackageData->IsInProgress() || FindPackage(nullptr, *PackageName.ToString())))
+				{
+					bExploreDependencies = false;
+				}
+			}
+			if (bExploreDependencies)
+			{
+				GetDependencies(PackageName, StackData.Dependencies);
+			}
+		}
+		else
+		{
+			bOutPushedStack = false;
+		}
+	};
+
+	constexpr int32 TimerCheckPeriod = 10; // Do not incur the cost of checking the timer on every package
+	for (;;)
+	{
+		while (!IsStackEmpty())
+		{
+			if (TimerCounter >= TimerCheckPeriod)
+			{
+				if (CookerTimer.IsTimeUp())
+				{
+					bOutComplete = false;
+					return;
+				}
+				TimerCounter = 0;
+			}
+
+			FRequestClusterGraphData::FStackData& CurrentTop(TopStack());
+			bool bPushedStack = false;
+			while (CurrentTop.NextDependency < CurrentTop.Dependencies.Num())
+			{
+				AddVertex(CurrentTop.Dependencies[CurrentTop.NextDependency++], NAME_None, false /* bInitial */, bPushedStack);
+				if (bPushedStack)
+				{
+					// CurrentTop is now invalid
+					break;
+				}
+			}
+			if (bPushedStack)
+			{
+				continue;
+			}
+			check(CurrentTop.NextDependency == CurrentTop.Dependencies.Num());
+			PopStack();
+		}
+
+		TransitiveRequests.Reserve(TransitiveRequests.Num() + Segment.Num());
+		while (!Segment.IsEmpty())
+		{
+			// Pop Segment from back to front so that when we PushFront onto TransitiveRequests they will maintain their same relative order
+			TransitiveRequests.AddFront(Segment.PopValue());
+		}
+
+		while (NextRequest < NumRequests)
+		{
+			FRequest& Request = Requests[NextRequest++];
+			FRequestClusterGraphData::FNamePair NamePair{ Request.PackageName, Request.FileName };
+			bool bPushedStack = false;
+			AddVertex(Request.PackageName, Request.FileName, true /* bInitial */, bPushedStack);
+			if (bPushedStack)
+			{
+				break;
+			}
+		}
+		if (IsStackEmpty())
+		{
+			check(NextRequest == NumRequests);
+			break;
+		}
+	}
+
+	// Map PackageNames to Request for the initial requests, so we can preserve their data (e.g. bIsUrgent) in the output
+	TMap<FName, FRequest*> InitialRequests;
+	for (FRequest& Request : Requests)
+	{
+		InitialRequests.Add(Request.PackageName, &Request);
+	}
+
+	// Convert all NamePairs found to an FRequest and push them all into the output in LeafToRoot order
+	TArray<FRequest> OutRequests;
+	OutRequests.Reserve(TransitiveRequests.Num());
+	for (int32 LeafToRootIndex = TransitiveRequests.Num() - 1; LeafToRootIndex >= 0; --LeafToRootIndex)
+	{
+		FRequestClusterGraphData::FNamePair& NamePair = TransitiveRequests[LeafToRootIndex];
+		FRequest** InitialRequest = InitialRequests.Find(NamePair.PackageName);
+		if (InitialRequest)
+		{
+			OutRequests.Add(MoveTemp(**InitialRequest));
+		}
+		else
+		{
+			OutRequests.Add(FRequest{ NamePair.PackageName, NamePair.FileName, FCompletionCallback(), false /* bIsUrgent */ });
+			COOK_STAT(++DetailedCookStats::NumPreloadedDependencies);
+		}
+	}
+	Swap(OutRequests, Requests);
+
+	TransitiveRequests.Empty();
+	StackStorage.Empty();
+	StackNum = 0;
+	Segment.Empty();
+	Visited.Empty();
+	NextRequest = 0;
+	bDependenciesComplete = true;
+}
+
+void FRequestCluster::GetDependencies(FName PackageName, TArray<FName>& OutDependencies)
+{
+	TArray<FName>& BuildDependencies(OutDependencies);
+	TArray<FName>& RuntimeDependencies(this->RuntimeScratch);
+	BuildDependencies.Reset();
+	RuntimeDependencies.Reset();
+
+	// TODO EditorOnly References: We only fetch Game dependencies, because the cooker explicitly loads all of
+	// the dependencies that we report. And if we explicitly load an EditorOnly dependency, that causes
+	// StaticLoadObjectInternal to SetLoadedByEditorPropertiesOnly(false), which then treats the editor-only package
+	// as needed in-game.
+	UE::AssetRegistry::EDependencyQuery DependencyQuery = UE::AssetRegistry::EDependencyQuery::Game;
+	if (!bAllowSoftDependencies)
+	{
+		DependencyQuery |= UE::AssetRegistry::EDependencyQuery::Hard;
+	}
+	AssetRegistry.GetDependencies(PackageName, BuildDependencies,
+		UE::AssetRegistry::EDependencyCategory::Package, DependencyQuery);
+	if (bTargetDomainEnabled)
+	{
+		for (const ITargetPlatform* TargetPlatform : Platforms)
+		{
+			UE::TargetDomain::TryFetchDependencies(PackageName, TargetPlatform, BuildDependencies,
+				RuntimeDependencies);
+		}
+	}
+	if (bAllowSoftBuildDependencies)
+	{
+		BuildDependencies.Append(RuntimeDependencies);
+	}
+
+	// Sort the list of BuildDependencies to check for uniqueness and make them deterministic
+	Algo::Sort(BuildDependencies, FNameLexicalLess());
+	BuildDependencies.SetNum(Algo::Unique(BuildDependencies));
+}
+
+bool FRequestCluster::IsRequestCookable(FName PackageName, FName& InOutFileName)
+{
+	TStringBuilder<256> NameBuffer;
+	// We need to reject packagenames from adding themselves or their transitive dependencies using all the same rules that
+	// UCookOnTheFlyServer::ProcessRequest uses. Packages that are rejected from cook do not add their dependencies to the cook.
+	PackageName.ToString(NameBuffer);
+	if (FPackageName::IsScriptPackage(NameBuffer))
+	{
+		return false;
+	}
+
+	if (InOutFileName.IsNone())
+	{
+		InOutFileName = PackageNameCache.GetCachedStandardFileName(PackageName);
+		if (InOutFileName.IsNone())
+		{
+			// Package does not exist on disk
+			return false;
+		}
+	}
+
+	if (bErrorOnEngineContentUse && !DLCPath.IsEmpty())
+	{
+		InOutFileName.ToString(NameBuffer);
+		if (!FStringView(NameBuffer).StartsWith(DLCPath))
+		{
+			return false;
+		}
+	}
+
+	if (PackageTracker.NeverCookPackageList.Contains(InOutFileName))
+	{
+		return false;
+	}
+
+	return true;
+}
+
+} // namespace UE::Cook
diff --git a/Engine/Source/Editor/UnrealEd/Private/Cooker/CookRequestCluster.h b/Engine/Source/Editor/UnrealEd/Private/Cooker/CookRequestCluster.h
new file mode 100644
--- /dev/null
+++ b/Engine/Source/Editor/UnrealEd/Private/Cooker/CookRequestCluster.h
@@ -0,0 +1,117 @@
+// Copyright Epic Games, Inc. All Rights Reserved.
+
+#pragma once
+
+#include "Containers/Array.h"
+#include "Containers/ArrayView.h"
+#include "Containers/Map.h"
+#include "Containers/RingBuffer.h"
+#include "Cooker/CookTypes.h"
+#include "UObject/NameTypes.h"
+
+class IAssetRegistry;
+class ITargetPlatform;
+class UCookOnTheFlyServer;
+struct FPackageNameCache;
+namespace UE::Cook { struct FFilePlatformRequest; }
+namespace UE::Cook { struct FPackageTracker; }
+
+namespace UE::Cook
+{
+
+/** Data used during the reentrant dependencies and topological sort operation of FRequestCluster. */
+struct FRequestClusterGraphData
+{
+	struct FNamePair
+	{
+		FName PackageName;
+		FName FileName;
+	};
+	/** All transitive PackageNames found from from all initial requests. Root to leaf order. */
+	TRingBuffer<FNamePair> TransitiveRequests;
+	/** Transitive PackageNames from the current initial request not already in TransitiveRequests. Root to leaf order. */
+	TRingBuffer<FNamePair> Segment;
+	/** Which PackageNames have already been reached (possibly still on the stack). */
+	TSet<FName> Visited;
+	/** Requests that are currently iterating over their dependencies. */
+	struct FStackData
+	{
+		void Reset(FName InPackageName)
+		{
+			PackageName = InPackageName;
+			Dependencies.Reset();
+			NextDependency = 0;
+		}
+		TArray<FName> Dependencies;
+		FName PackageName;
+		int32 NextDependency;
+	};
+	TArray<FStackData> Stack;
+	/** We do not remove FStackData from the stack, because we want to avoid reallocating Dependencies arrays. */
+	int32 StackNum = 0;
+};
+
+/**
+ * A group of external requests sent to CookOnTheFlyServer's tick loop. Transitive dependencies are found and all of the
+ * requested or dependent packagenames are added as requests together to the cooking state machine.
+ */
+class FRequestCluster
+{
+public:
+	/** The cluster's data about a package from the external requests or discovered dependency. */
+	struct FRequest
+	{
+		FName PackageName;
+		FName FileName;
+		FCompletionCallback CompletionCallback;
+		bool bIsUrgent;
+
+		FRequest(FName InPackageName, FName InFileName, FCompletionCallback&& InCompletionCallback, bool bInIsUrgent);
+		FRequest(FFilePlatformRequest&& FileRequest, bool bInIsUrgent);
+	};
+	explicit FRequestCluster(UCookOnTheFlyServer& COTFS, TArray<const ITargetPlatform*>&& Platforms);
+	FRequestCluster(FRequestCluster&&) = default;
+
+	/**
+	 * Calculate the information needed to create a PackageData, and transitive search dependencies for all requests.
+	 * Called repeatedly (due to timeslicing) until bOutComplete is set to true.
+	 */
+	void Process(const FCookerTimer& CookerTimer, bool& bOutComplete);
+
+	TConstArrayView<const ITargetPlatform*> GetPlatforms() const { return Platforms; }
+	TArray<FRequest>& GetRequests() { return Requests; }
+
+	/**
+	 * Create clusters(s) for all the given requests and append them to OutClusters.
+	 * Multiple clusters are necessary if the list of platforms differs between requests.
+	 */
+	static void AddClusters(UCookOnTheFlyServer& COTFS, TArray<FFilePlatformRequest>&& InRequests, bool bRequestsAreUrgent,
+		TRingBuffer<FRequestCluster>& OutClusters);
+
+private:
+	void Initialize(UCookOnTheFlyServer& COTFS);
+	void FetchPackageNames(const FCookerTimer& CookerTimer, bool& bOutComplete);
+	void FetchDependencies(const FCookerTimer& CookerTimer, bool& bOutComplete);
+	void GetDependencies(FName PackageName, TArray<FName>& OutDependencies);
+	bool IsRequestCookable(FName PackageName, FName& InOutFileName);
+
+	TArray<FRequest> Requests;
+	TArray<FName> RuntimeScratch;
+	TArray<const ITargetPlatform*> Platforms;
+	FRequestClusterGraphData DependencyGraphData;
+	FString DLCPath;
+	FPackageDatas& PackageDatas;
+	IAssetRegistry& AssetRegistry;
+	const FPackageNameCache& PackageNameCache;
+	FPackageTracker& PackageTracker;
+	int32 NextRequest = 0;
+	bool bAllowHardDependencies = true;
+	bool bAllowSoftDependencies = true;
+	bool bAllowSoftBuildDependencies = true;
+	bool bTargetDomainEnabled = true;
+	bool bErrorOnEngineContentUse = false;
+	bool bPackageNamesComplete = false;
+	bool bDependenciesComplete = false;
+};
+
+}
diff --git a/Engine/Source/Editor/UnrealEd/Private/Cooker/CookRequests.cpp b/Engine/Source/Editor/UnrealEd/Private/Cooker/CookRequests.cpp
--- a/Engine/Source/Editor/UnrealEd/Private/Cooker/CookRequests.cpp
+++ b/Engine/Source/Editor/UnrealEd/Private/Cooker/CookRequests.cpp
@@ -212,7 +212,7 @@ namespace Cook
 		}
 	}
 
-	EExternalRequestType FExternalRequests::DequeueRequest(TArray<FSchedulerCallback>& OutCallbacks, FFilePlatformRequest& OutToBuild)
+	EExternalRequestType FExternalRequests::DequeueNextCluster(TArray<FSchedulerCallback>& OutCallbacks, TArray<FFilePlatformRequest>& OutBuildRequests)
 	{
 		FScopeLock ScopeLock(&RequestLock);
 
@@ -222,9 +222,13 @@ namespace Cook
 		}
 		else if (Queue.Num())
 		{
-			FName Filename = Queue.PopFrontValue();
-			OutToBuild = RequestMap.FindAndRemoveChecked(Filename);
-			--RequestCount;
+			OutBuildRequests.Reserve(Queue.Num() + OutBuildRequests.Num());
+			while (Queue.Num())
+			{
+				FName Filename = Queue.PopFrontValue();
+				OutBuildRequests.Add(RequestMap.FindAndRemoveChecked(Filename));
+				--RequestCount;
+			}
 			return EExternalRequestType::Cook;
 		}
 		else
diff --git a/Engine/Source/Editor/UnrealEd/Private/Cooker/CookRequests.h b/Engine/Source/Editor/UnrealEd/Private/Cooker/CookRequests.h
--- a/Engine/Source/Editor/UnrealEd/Private/Cooker/CookRequests.h
+++ b/Engine/Source/Editor/UnrealEd/Private/Cooker/CookRequests.h
@@ -97,15 +97,15 @@ namespace Cook
 
 		/**
 		 * If this FExternalRequests has any callbacks, dequeue them all into OutCallbacks and return EExternalRequestType::Callback; Callbacks take priority over cook requests.
-		 * Otherwise, if there are any cook requests, set OutToBuild to the front request and return EExternalRequestType::Cook.
+		 * Otherwise, if there are any cook requests, dequeue them all into OutBuildRequests and return EExternalRequestType::Cook.
 		 * Otherwise, return EExternalRequestType::None.
 		 */
-		EExternalRequestType DequeueRequest(TArray<FSchedulerCallback>& OutCallbacks, FFilePlatformRequest& OutToBuild);
+		EExternalRequestType DequeueNextCluster(TArray<FSchedulerCallback>& OutCallbacks, TArray<FFilePlatformRequest>& OutBuildRequests);
 		/* Move any existing callbacks onto OutCallbacks, and return whether any were added. */
 		bool DequeueCallbacks(TArray<FSchedulerCallback>& OutCallbacks);
 		/** Eliminate all callbacks and cook requests and free memory */
 		void EmptyRequests();
-		/** Move all callbacks into OutCallbacks, and all cook requests into OutCookRequests. This is used when cancelling a cook session. */
+		/** Move all callbacks into OutCallbacks, and all cook requests into OutCookRequests. This is used when canceling a cook session. */
 		void DequeueAll(TArray<FSchedulerCallback>& OutCallbacks, TArray<FFilePlatformRequest>& OutCookRequests);
 
 		/** Remove references to the given platform from all cook requests. */
diff --git a/Engine/Source/Editor/UnrealEd/Private/Cooker/CookTypes.h b/Engine/Source/Editor/UnrealEd/Private/Cooker/CookTypes.h
--- a/Engine/Source/Editor/UnrealEd/Private/Cooker/CookTypes.h
+++ b/Engine/Source/Editor/UnrealEd/Private/Cooker/CookTypes.h
@@ -167,4 +167,53 @@ namespace Cook
 }
 }
 
+//////////////////////////////////////////////////////////////////////////
+// Cook by the book options
+
+struct UCookOnTheFlyServer::FCookByTheBookOptions
+{
+public:
+	/** Should we generate streaming install manifests (only valid option in cook by the book) */
+	bool							bGenerateStreamingInstallManifests = false;
+
+	/** Should we generate a seperate manifest for map dependencies */
+	bool							bGenerateDependenciesForMaps = false;
+
+	/** Is cook by the book currently running */
+	bool							bRunning = false;
+
+	/** Cancel has been queued will be processed next tick */
+	bool							bCancel = false;
+
+	/** DlcName setup if we are cooking dlc will be used as the directory to save cooked files to */
+	FString							DlcName;
+
+	/** Create a release from this manifest and store it in the releases directory for this cgame */
+	FString							CreateReleaseVersion;
+
+	/** Dependency graph of maps as root objects. */
+	TFastPointerMap<const ITargetPlatform*, TMap<FName, TSet<FName>>> MapDependencyGraphs;
+
+	/** If we are based on a release version of the game this is the set of packages which were cooked in that release. Map from platform name to list of uncooked package filenames */
+	TMap<FName, TArray<FName>>			BasedOnReleaseCookedPackages;
+
+	/** Timing information about cook by the book */
+	double							CookTime = 0.0;
+	double							CookStartTime = 0.0;
+
+	/** error when detecting engine content being used in this cook */
+	bool							bErrorOnEngineContentUse = false;
+	bool							bSkipHardReferences = false;
+	bool							bSkipSoftReferences = false;
+	bool							bFullLoadAndSave = false;
+	bool							bIoStore = false;
+	bool							bZenStore = false;
+	bool							bCookAgainstFixedBase = false;
+	bool							bDlcLoadMainAssetRegistry = false;
+	TArray<FName>					StartupPackages;
+
+	/** Mapping from source packages to their localized variants (based on the culture list in FCookByTheBookStartupOptions) */
+	TMap<FName, TArray<FName>>		SourceToLocalizedPackageVariants;
+};
+
 void LogCookerMessage(const FString& MessageText, EMessageSeverity::Type Severity);
