CookRequestCluster: Fix clang compile error, and set ExploreSoftReferencesOnStart=false for now.

diff --git a/Engine/Config/BaseEditor.ini b/Engine/Config/BaseEditor.ini
--- a/Engine/Config/BaseEditor.ini
+++ b/Engine/Config/BaseEditor.ini
@@ -327,6 +327,7 @@ MaxUsedMemory=16384
 [CookSettings]
 BulkDataRegistryEnabled=false
 TargetDomainEnabled=false
+ExploreSoftReferencesOnStart=false
 PackagesPerGC=0
 IdleTimeToGC=20
 MemoryMaxUsedPhysical=16384
diff --git a/Engine/Source/Editor/UnrealEd/Private/Cooker/CookRequestCluster.cpp b/Engine/Source/Editor/UnrealEd/Private/Cooker/CookRequestCluster.cpp
--- a/Engine/Source/Editor/UnrealEd/Private/Cooker/CookRequestCluster.cpp
+++ b/Engine/Source/Editor/UnrealEd/Private/Cooker/CookRequestCluster.cpp
@@ -35,17 +35,17 @@ void FRequestCluster::AddClusters(UCookOnTheFlyServer& COTFS, TArray<FFilePlatfo
 	}
 
 	TArray<FRequestCluster, TInlineAllocator<1>> AddedClusters;
-	auto FindOrAddClusterForPlatforms = [&AddedClusters, &COTFS](TArray<const ITargetPlatform*>&& Platforms)
+	auto FindOrAddClusterForPlatforms = [&AddedClusters, &COTFS](TArray<const ITargetPlatform*>&& InPlatforms)
 	{
 		for (FRequestCluster& Existing : AddedClusters)
 		{
-			if (Existing.GetPlatforms() == Platforms)
+			if (Existing.GetPlatforms() == InPlatforms)
 			{
-				Platforms.Reset();
+				InPlatforms.Reset();
 				return &Existing;
 			}
 		}
-		return &AddedClusters.Emplace_GetRef(COTFS, MoveTemp(Platforms));
+		return &AddedClusters.Emplace_GetRef(COTFS, MoveTemp(InPlatforms));
 	};
 
 	UE::Cook::FRequestCluster* MRUCluster = FindOrAddClusterForPlatforms(MoveTemp(InRequests[0].GetPlatforms()));
