Cooking: Create TargetDomain dependencies and key from the PackageBuildDependencyTracker and the AssetRegistry.
#rnx

diff --git a/Engine/Source/Editor/UnrealEd/Private/Cooker/CookRequestCluster.cpp b/Engine/Source/Editor/UnrealEd/Private/Cooker/CookRequestCluster.cpp
--- a/Engine/Source/Editor/UnrealEd/Private/Cooker/CookRequestCluster.cpp
+++ b/Engine/Source/Editor/UnrealEd/Private/Cooker/CookRequestCluster.cpp
@@ -11,6 +11,7 @@
 #include "Cooker/PackageTracker.h"
 #include "CookOnTheSide/CookOnTheFlyServer.h"
 #include "EditorDomain/EditorDomainUtils.h"
+#include "Interfaces/ITargetPlatform.h"
 #include "Misc/StringBuilder.h"
 #include "TargetDomain/TargetDomainUtils.h"
 
@@ -113,6 +114,12 @@ void FRequestCluster::Initialize(UCookOnTheFlyServer& COTFS)
 		DLCPath = FPaths::Combine(*COTFS.GetBaseDirectoryForDLC(), TEXT("Content"));
 		FPaths::MakeStandardFilename(DLCPath);
 	}
+
+	PackageStores.Reserve(Platforms.Num());
+	for (const ITargetPlatform* TargetPlatform : Platforms)
+	{
+		PackageStores.Add(COTFS.GetPackageStoreWriter(FName(*TargetPlatform->PlatformName())));
+	}
 }
 
 void FRequestCluster::Process(const FCookerTimer& CookerTimer, bool& bOutComplete)
@@ -395,10 +402,16 @@ void FRequestCluster::GetDependencies(FName PackageName, TArray<FName>& OutDepen
 		UE::AssetRegistry::EDependencyCategory::Package, DependencyQuery);
 	if (bTargetDomainEnabled)
 	{
-		for (const ITargetPlatform* TargetPlatform : Platforms)
+		for (int32 PlatIndex = 0; PlatIndex < Platforms.Num(); ++PlatIndex)
 		{
-			UE::TargetDomain::TryFetchDependencies(PackageName, TargetPlatform, BuildDependencies,
-				RuntimeDependencies);
+			const ITargetPlatform* TargetPlatform = Platforms[PlatIndex];
+			IPackageStoreWriter* PackageStore = PackageStores[PlatIndex];
+			if (!PackageStore)
+			{
+				continue;
+			}
+			UE::TargetDomain::TryFetchKeyAndDependencies(PackageStore, PackageName, TargetPlatform, nullptr /* OutHash */, &BuildDependencies,
+				&RuntimeDependencies, nullptr /* OutErrorMessage */);
 		}
 	}
 	if (bAllowSoftBuildDependencies)
diff --git a/Engine/Source/Editor/UnrealEd/Private/Cooker/CookRequestCluster.h b/Engine/Source/Editor/UnrealEd/Private/Cooker/CookRequestCluster.h
--- a/Engine/Source/Editor/UnrealEd/Private/Cooker/CookRequestCluster.h
+++ b/Engine/Source/Editor/UnrealEd/Private/Cooker/CookRequestCluster.h
@@ -98,6 +98,7 @@ private:
 	TArray<FRequest> Requests;
 	TArray<FName> RuntimeScratch;
 	TArray<const ITargetPlatform*> Platforms;
+	TArray<IPackageStoreWriter*> PackageStores;
 	FRequestClusterGraphData DependencyGraphData;
 	FString DLCPath;
 	FPackageDatas& PackageDatas;
diff --git a/Engine/Source/Editor/UnrealEd/Private/TargetDomain/TargetDomainUtils.cpp b/Engine/Source/Editor/UnrealEd/Private/TargetDomain/TargetDomainUtils.cpp
--- a/Engine/Source/Editor/UnrealEd/Private/TargetDomain/TargetDomainUtils.cpp
+++ b/Engine/Source/Editor/UnrealEd/Private/TargetDomain/TargetDomainUtils.cpp
@@ -8,6 +8,7 @@
 #include "Cooker/PackageBuildDependencyTracker.h"
 #include "EditorDomain/EditorDomainUtils.h"
 #include "IO/IoHash.h"
+#include "IO/PackageStoreWriter.h"
 #include "Misc/StringBuilder.h"
 #include "Serialization/CompactBinary.h"
 #include "Serialization/CompactBinaryWriter.h"
@@ -15,12 +16,117 @@
 namespace UE::TargetDomain
 {
 
+bool TryCreateKey(FName PackageName, TArrayView<FName> SortedBuildDependencies, FIoHash* OutHash, FString* OutErrorMessage)
+{
+	IAssetRegistry* AssetRegistry = IAssetRegistry::Get();
+	if (!AssetRegistry)
+	{
+		if (OutErrorMessage) *OutErrorMessage = TEXT("AssetRegistry is unavailable.");
+		return false;
+	}
+	FCbWriter KeyBuilder;
+	UE::EditorDomain::EPackageDigestResult Result;
+	FString ErrorMessage;
+	bool bBlacklisted;
+	Result = UE::EditorDomain::AppendPackageDigest(*AssetRegistry, PackageName, KeyBuilder, bBlacklisted, ErrorMessage);
+	if (Result != UE::EditorDomain::EPackageDigestResult::Success)
+	{
+		if (OutErrorMessage) *OutErrorMessage = MoveTemp(ErrorMessage);
+		return false;
+	}
+
+	for (FName DependencyName : SortedBuildDependencies)
+	{
+		Result = UE::EditorDomain::AppendPackageDigest(*AssetRegistry, DependencyName, KeyBuilder, bBlacklisted, ErrorMessage);
+		if (Result != UE::EditorDomain::EPackageDigestResult::Success)
+		{
+			if (OutErrorMessage)
+			{
+				*OutErrorMessage = FString::Printf(TEXT("Could not create PackageDigest for %s: %s"),
+					*DependencyName.ToString(), *ErrorMessage);
+			}
+			return false;
+		}
+	}
+
+	if (OutHash)
+	{
+		*OutHash = KeyBuilder.Save().GetRangeHash();
+	}
+	return true;
+}
+
 bool TryCollectKeyAndDependencies(UPackage* Package, const ITargetPlatform* TargetPlatform, FIoHash* OutHash, TArray<FName>* OutBuildDependencies,
 	TArray<FName>* OutRuntimeOnlyDependencies, FString* OutErrorMessage)
 {
-	// Implementation is coming in a separate changelist that mostly involves the PackageBuildDependencyTracker
-	if (OutErrorMessage) *OutErrorMessage = TEXT("Not yet implemented.");
-	return false;
+	if (!Package)
+	{
+		if (OutErrorMessage) *OutErrorMessage = TEXT("Invalid null package.");
+		return false;
+	}
+
+	FName PackageName = Package->GetFName();
+	TSet<FName> BuildDependencies;
+	TSet<FName> RuntimeOnlyDependencies;
+
+	IAssetRegistry* AssetRegistry = IAssetRegistry::Get();
+	if (!AssetRegistry)
+	{
+		if (OutErrorMessage) *OutErrorMessage = TEXT("AssetRegistry is unavailable.");
+		return false;
+	}
+
+	TArray<FName> AssetDependencies;
+	AssetRegistry->GetDependencies(PackageName, AssetDependencies, UE::AssetRegistry::EDependencyCategory::Package,
+		UE::AssetRegistry::EDependencyQuery::Game);
+
+	FPackageBuildDependencyTracker& Tracker = FPackageBuildDependencyTracker::Get();
+	TArray<FBuildDependencyAccessData> AccessDatas = Tracker.GetAccessDatas(PackageName);
+
+	BuildDependencies.Reserve(AccessDatas.Num());
+	for (FBuildDependencyAccessData& AccessData : AccessDatas)
+	{
+		if (AccessData.TargetPlatform == TargetPlatform || AccessData.TargetPlatform == nullptr)
+		{
+			BuildDependencies.Add(AccessData.ReferencedPackage);
+		}
+	}
+
+	RuntimeOnlyDependencies.Reserve(AccessDatas.Num());
+	for (FName DependencyName : AssetDependencies)
+	{
+		if (!BuildDependencies.Contains(DependencyName))
+		{
+			RuntimeOnlyDependencies.Add(DependencyName);
+		}
+	}
+
+	TArray<FName> SortedBuild;
+	SortedBuild = BuildDependencies.Array();
+	SortedBuild.Sort(FNameLexicalLess());
+	TArray<FName> SortedRuntimeOnly;
+	SortedRuntimeOnly = RuntimeOnlyDependencies.Array();
+	SortedRuntimeOnly.Sort(FNameLexicalLess());
+
+	if (!TryCreateKey(PackageName, SortedBuild, OutHash, OutErrorMessage))
+	{
+		return false;
+	}
+
+	if (OutBuildDependencies)
+	{
+		*OutBuildDependencies = MoveTemp(SortedBuild);
+	}
+	if (OutRuntimeOnlyDependencies)
+	{
+		*OutRuntimeOnlyDependencies = MoveTemp(SortedRuntimeOnly);
+	}
+	if (OutErrorMessage)
+	{
+		OutErrorMessage->Reset();
+	}
+	
+	return true;
 }
 
 FCbObject CollectDependenciesObject(UPackage* Package, const ITargetPlatform* TargetPlatform, FString* ErrorMessage)
@@ -62,17 +168,65 @@ FCbObject CollectDependenciesObject(UPackage* Package, const ITargetPlatform* Ta
 }
 
 
-bool TryFetchKeyAndDependencies(FName PackageName, const ITargetPlatform* TargetPlatform, FIoHash* OutHash, TArray<FName>* OutDependencies)
+bool TryFetchKeyAndDependencies(IPackageStoreWriter* PackageStore, FName PackageName, const ITargetPlatform* TargetPlatform,
+	FIoHash* OutHash, TArray<FName>* OutBuildDependencies, TArray<FName>* OutRuntimeOnlyDependencies, FString* OutErrorMessage)
 {
-	// TODO: Read from oplog for the packagename, return true if and only if it has a targetdomainkey and the generated key from the collection of dependencies matches
-	return false;
-}
+	FCbObject DependenciesObj = PackageStore->GetTargetDomainDependencies(PackageName);
+	FIoHash StoredKey = DependenciesObj["targetdomainkey"].AsHash();
+	if (StoredKey.IsZero())
+	{
+		if (OutErrorMessage)
+		{
+			*OutErrorMessage = TEXT("Dependencies not in oplog.");
+		}
+		return false;
+	}
 
-bool TryFetchDependencies(FName PackageName, const ITargetPlatform* TargetPlatform,
-	TArray<FName>& OutBuildDependencies, TArray<FName>& OutRuntimeOnlyDependencies)
-{
-	// TODO: Read from oplog for the packagename, return true if and only if it has a targetdomainkey and the generated key from the collection of dependencies matches
-	return false;
+	TArray<FName> BuildDependencies;
+	if (OutBuildDependencies)
+	{
+		OutBuildDependencies->Reset();
+	}
+	else
+	{
+		OutBuildDependencies = &BuildDependencies;
+	}
+
+	for (FCbFieldView DepObj : DependenciesObj["builddependencies"])
+	{
+		if (FString DependencyName(DepObj.AsString()); !DependencyName.IsEmpty())
+		{
+			OutBuildDependencies->Add(FName(*DependencyName));
+		}
+	}
+
+	FIoHash CurrentKey;
+	if (!TryCreateKey(PackageName, *OutBuildDependencies, &CurrentKey, OutErrorMessage))
+	{
+		return false;
+	}
+
+	if (StoredKey != CurrentKey)
+	{
+		if (OutErrorMessage)
+		{
+			*OutErrorMessage = TEXT("Stored key does not match current key.");
+		}
+		return false;
+	}
+
+	if (OutRuntimeOnlyDependencies)
+	{
+		for (FCbFieldView DepObj : DependenciesObj["runtimeonlydependencies"])
+		{
+			if (FString DependencyName(DepObj.AsString()); !DependencyName.IsEmpty())
+			{
+				OutRuntimeOnlyDependencies->Add(FName(*DependencyName));
+			}
+		}
+	}
+
+	return true;
 }
 
 }
diff --git a/Engine/Source/Editor/UnrealEd/Private/TargetDomain/TargetDomainUtils.h b/Engine/Source/Editor/UnrealEd/Private/TargetDomain/TargetDomainUtils.h
--- a/Engine/Source/Editor/UnrealEd/Private/TargetDomain/TargetDomainUtils.h
+++ b/Engine/Source/Editor/UnrealEd/Private/TargetDomain/TargetDomainUtils.h
@@ -5,6 +5,7 @@
 #include "Containers/Array.h"
 #include "UObject/NameTypes.h"
 
+class IPackageStoreWriter;
 class ITargetPlatform;
 class FCbObject;
 class FString;
@@ -14,13 +15,20 @@ struct FIoHash;
 namespace UE::TargetDomain
 {
 
+/** Create the TargetDomainKey based on the EditorDomainKeys of the Package and its dependencies. */
+bool TryCreateKey(FName PackageName, TArrayView<FName> SortedBuildDependencies, FIoHash* OutHash, FString* OutErrorMessage);
+
+/** Collect the Package's dependencies and the key based on them. */
 bool TryCollectKeyAndDependencies(UPackage* Package, const ITargetPlatform* TargetPlatform,
 	FIoHash* OutHash, TArray<FName>* OutBuildDependencies, TArray<FName>* OutRuntimeOnlyDependencies, FString* OutErrorMessage);
+/** Collect the Package's dependencies, and create a FCbObject describing them for storage in the OpLog. */
 FCbObject CollectDependenciesObject(UPackage* Package, const ITargetPlatform* TargetPlatform, FString* ErrorMessage);
 
-bool TryFetchKeyAndDependencies(FName PackageName, const ITargetPlatform* TargetPlatform,
+/**
+ * Read the oplog for the given packagename and fetch the dependencies and key out of it. Use the dependencies to calculate
+ * the current key, and return the dependencies and key only if they key matches.
+ */
+bool TryFetchKeyAndDependencies(IPackageStoreWriter* PackageStore, FName PackageName, const ITargetPlatform* TargetPlatform,
 	FIoHash* OutHash, TArray<FName>* OutBuildDependencies, TArray<FName>* OutRuntimeOnlyDependencies, FString* OutErrorMessage);
-bool TryFetchDependencies(FName PackageName, const ITargetPlatform* TargetPlatform,
-	TArray<FName>& OutBuildDependencies, TArray<FName>& OutRuntimeOnlyDependencies);
 
 }
