Fixup bad merge [REDACTED_HANDLE]. A deletion went missing during conflict resolution. Original change with the deletion was [REDACTED_HANDLE].

diff --git a/Engine/Source/Editor/UnrealEd/Private/AssetThumbnail.cpp b/Engine/Source/Editor/UnrealEd/Private/AssetThumbnail.cpp
--- a/Engine/Source/Editor/UnrealEd/Private/AssetThumbnail.cpp
+++ b/Engine/Source/Editor/UnrealEd/Private/AssetThumbnail.cpp
@@ -1155,8 +1155,6 @@ bool FAssetThumbnailPool::LoadThumbnail(TSharedRef<FThumbnailInfo> ThumbnailInfo
 					);
 				}
 
-				// Since this was rendered, add it to the list of thumbnails that can be rendered in real-time
-				RealTimeThumbnails.AddUnique(ThumbnailInfo);
 				return true;
 			}
 		}
