Fix: Widget used in graph warning when deleting unused widgets

diff --git a/Engine/Source/Editor/UnrealEd/Private/Kismet2/BlueprintEditorUtils.cpp b/Engine/Source/Editor/UnrealEd/Private/Kismet2/BlueprintEditorUtils.cpp
--- a/Engine/Source/Editor/UnrealEd/Private/Kismet2/BlueprintEditorUtils.cpp
+++ b/Engine/Source/Editor/UnrealEd/Private/Kismet2/BlueprintEditorUtils.cpp
@@ -5844,7 +5844,8 @@ bool FBlueprintEditorUtils::IsVariableUsed(const UBlueprint* VariableBlueprint,
 
 				for (const UK2Node_Variable* CurrentNode : GraphNodes)
 				{
-					if (VariableGuid == CurrentNode->VariableReference.GetMemberGuid())
+					if (VariableGuid == CurrentNode->VariableReference.GetMemberGuid() &&
+						VariableName == CurrentNode->GetVarName())
 					{
 						return true;
 					}
