Add TargetDomain.IterativeClassBlacklist and IterativeClassWhitelist.
#rnx

diff --git a/Engine/Config/BaseEditor.ini b/Engine/Config/BaseEditor.ini
--- a/Engine/Config/BaseEditor.ini
+++ b/Engine/Config/BaseEditor.ini
@@ -326,7 +326,6 @@ MaxUsedMemory=16384
 
 [CookSettings]
 BulkDataRegistryEnabled=false
-TargetDomainEnabled=false
 ExploreSoftReferencesOnStart=false
 PackagesPerGC=0
 IdleTimeToGC=20
@@ -382,6 +381,18 @@ ClassBlacklist=/Script/ControlRig.ControlRigBlueprintGeneratedClass
 +PackageBlacklist=/Niagara/Functions/RandomRangeFloat
 +PackageBlacklist=/Niagara/Modules/Update/Color/ScaleColor
 
+[TargetDomain]
+TargetDomainEnabled=false
+; IterativeClassWhitelist/IterativeClassBlacklist:
+;    Used in HybridIterative only.
+;    List of classes that may use iterative cooks. Whitelist superclasses are also whitelisted.
+;    Blacklist subclasses are also blacklisted. Blacklist overrides the whitelist.
+; +IterativeClassWhitelist=/Long/PackageName.ClassName
+; +IterativeClassBlacklist=/Long/PackageName.ClassName
+IterativeClassWhitelist=/Script/Engine.Texture2D
++IterativeClassWhitelist=/Script/CoreUObject.MetaData
++IterativeClassWhitelist=/Script/Engine.AssetImportData
+
 [CookPlatformDataCacheSettings]
 MaterialInstance=100
 Material=100
diff --git a/Engine/Source/Editor/UnrealEd/Private/CookOnTheFlyServer.cpp b/Engine/Source/Editor/UnrealEd/Private/CookOnTheFlyServer.cpp
--- a/Engine/Source/Editor/UnrealEd/Private/CookOnTheFlyServer.cpp
+++ b/Engine/Source/Editor/UnrealEd/Private/CookOnTheFlyServer.cpp
@@ -4920,7 +4920,7 @@ void UCookOnTheFlyServer::Initialize( ECookMode::Type DesiredCookMode, ECookInit
 		bPreloadingEnabled = true;
 		FLinkerLoad::SetPreloadingEnabled(true);
 	}
-	GConfig->GetBool(TEXT("CookSettings"), TEXT("TargetDomainEnabled"), bTargetDomainEnabled, GEditorIni);
+	GConfig->GetBool(TEXT("TargetDomain"), TEXT("TargetDomainEnabled"), bTargetDomainEnabled, GEditorIni);
 	GConfig->GetBool(TEXT("CookSettings"), TEXT("ExploreSoftReferencesOnStart"), bExploreSoftReferencesOnStart, GEditorIni);
 
 	{
diff --git a/Engine/Source/Editor/UnrealEd/Private/EditorDomain/EditorDomain.cpp b/Engine/Source/Editor/UnrealEd/Private/EditorDomain/EditorDomain.cpp
--- a/Engine/Source/Editor/UnrealEd/Private/EditorDomain/EditorDomain.cpp
+++ b/Engine/Source/Editor/UnrealEd/Private/EditorDomain/EditorDomain.cpp
@@ -40,14 +40,20 @@ public:
 
 	static IPackageResourceManager* SetPackageResourceManager()
 	{
-		bool bEditorDomainEnabled = false;
 		if (GIsEditor && (!IsRunningCommandlet() || IsRunningCookCommandlet()))
 		{
+			bool bEditorDomainEnabled = true;
 			GConfig->GetBool(TEXT("EditorDomain"), TEXT("EditorDomainEnabled"), bEditorDomainEnabled, GEditorIni);
 			if (GConfig->GetBool(TEXT("CookSettings"), TEXT("EditorDomainEnabled"), bEditorDomainEnabled, GEditorIni))
 			{
 				UE_LOG(LogEditorDomain, Error, TEXT("Editor.ini:[CookSettings]:EditorDomainEnabled is deprecated, use Editor.ini:[EditorDomain]:EditorDomainEnabled instead."));
 			}
+			bool bTargetDomainEnabled = true;
+			GConfig->GetBool(TEXT("TargetDomain"), TEXT("TargetDomainEnabled"), bTargetDomainEnabled, GEditorIni);
+			if (bEditorDomainEnabled || bTargetDomainEnabled)
+			{
+				UE::EditorDomain::UtilsInitialize();
+			}
 			if (bEditorDomainEnabled)
 			{
 				// Set values for config settings EditorDomain depends on
@@ -158,15 +164,15 @@ bool FEditorDomain::TryFindOrAddPackageSource(const FPackagePath& PackagePath, T
 
 	FString ErrorMessage;
 	FPackageDigest PackageDigest;
-	bool bIsBlacklisted;
-	EPackageDigestResult Result = GetPackageDigest(*AssetRegistry, PackageName, PackageDigest, bIsBlacklisted, ErrorMessage);
+	bool bEditorDomainEnabledForPackage;
+	EPackageDigestResult Result = GetPackageDigest(*AssetRegistry, PackageName, PackageDigest, bEditorDomainEnabledForPackage, ErrorMessage);
 	switch (Result)
 	{
 	case EPackageDigestResult::Success:
 		PackageSource = new FPackageSource();
 		PackageSource->Digest = PackageDigest;
 		OutSource = PackageSource;
-		if (!bEditorDomainReadEnabled || bIsBlacklisted)
+		if (!bEditorDomainReadEnabled || !bEditorDomainEnabledForPackage)
 		{
 			PackageSource->Source = EPackageSource::Workspace;
 		}
diff --git a/Engine/Source/Editor/UnrealEd/Private/EditorDomain/EditorDomainUtils.cpp b/Engine/Source/Editor/UnrealEd/Private/EditorDomain/EditorDomainUtils.cpp
--- a/Engine/Source/Editor/UnrealEd/Private/EditorDomain/EditorDomainUtils.cpp
+++ b/Engine/Source/Editor/UnrealEd/Private/EditorDomain/EditorDomainUtils.cpp
@@ -15,8 +15,10 @@
 #include "HAL/FileManager.h"
 #include "IO/PackageStoreWriter.h"
 #include "Memory/SharedBuffer.h"
+#include "Misc/CoreDelegates.h"
 #include "Misc/FileHelper.h"
 #include "Misc/PackagePath.h"
+#include "Misc/ScopeRWLock.h"
 #include "Misc/StringBuilder.h"
 #include "Serialization/CompactBinaryWriter.h"
 #include "UObject/CoreRedirects.h"
@@ -31,12 +33,23 @@ namespace EditorDomain
 {
 
 FClassDigestMap GClassDigests;
-
 FClassDigestMap& GetClassDigests()
 {
 	return GClassDigests;
 }
 
+TSet<FName> GEditorClassBlacklist;
+const TSet<FName>& GetEditorClassBlacklist()
+{
+	return GEditorClassBlacklist;
+}
+
+TSet<FName> GEditorPackageBlacklist;
+const TSet<FName>& GetEditorPackageBlacklist()
+{
+	return GEditorPackageBlacklist;
+}
+
 // Change to a new guid when EditorDomain needs to be invalidated
 const TCHAR* EditorDomainVersion = TEXT("D1718C34CA7C47AEB87A1607568E25B0");
 // Identifier of the CacheBuckets for EditorDomain tables
@@ -60,10 +73,10 @@ static bool GetEditorDomainSaveUnversioned()
 	return bEditorDomainSaveUnversioned;
 }
 
-EPackageDigestResult AppendPackageDigest(FCbWriter& Writer, bool& bOutIsBlacklisted, FString& OutErrorMessage,
+EPackageDigestResult AppendPackageDigest(FCbWriter& Writer, bool& bOutEditorDomainEnabled, FString& OutErrorMessage,
 	const FAssetPackageData& PackageData, FName PackageName)
 {
-	bOutIsBlacklisted = false;
+	bOutEditorDomainEnabled = true;
 
 	int32 CurrentFileVersionUE = GPackageFileUEVersion;
 	int32 CurrentFileVersionLicenseeUE = GPackageFileLicenseeUEVersion;
@@ -108,7 +121,7 @@ EPackageDigestResult AppendPackageDigest(FCbWriter& Writer, bool& bOutIsBlacklis
 			RemainingClasses = RemainingClasses.Slice(NextClass, PackageData.ImportedClasses.Num() - NextClass);
 			PrecacheClassDigests(RemainingClasses);
 		}
-		FScopeLock ClassDigestsScopeLock(&ClassDigests.Lock);
+		FReadScopeLock ClassDigestsScopeLock(ClassDigests.Lock);
 		for (; NextClass < PackageData.ImportedClasses.Num(); ++NextClass)
 		{
 			FName ClassName = PackageData.ImportedClasses[NextClass];
@@ -121,7 +134,7 @@ EPackageDigestResult AppendPackageDigest(FCbWriter& Writer, bool& bOutIsBlacklis
 			{
 				Writer << ExistingData->SchemaHash;
 			}
-			bOutIsBlacklisted |= ExistingData->bBlacklisted;
+			bOutEditorDomainEnabled &= ExistingData->bEditorDomainEnabled;
 		}
 	}
 	return EPackageDigestResult::Success;
@@ -134,7 +147,7 @@ void PrecacheClassDigests(TConstArrayView<FName> ClassNames, TMap<FName, FClassD
 	TArray<FName, TInlineAllocator<10>> ClassesToAdd;
 	ClassesToAdd.Reserve(ClassNames.Num());
 	{
-		FScopeLock ClassDigestsScopeLock(&ClassDigests.Lock);
+		FWriteScopeLock ClassDigestsScopeLock(ClassDigests.Lock);
 		for (FName ClassName : ClassNames)
 		{
 			FClassDigestData* DigestData = ClassDigests.Map.Find(ClassName);
@@ -193,10 +206,10 @@ void PrecacheClassDigests(TConstArrayView<FName> ClassNames, TMap<FName, FClassD
 		
 		FClassData& ClassData = ClassDatas.Emplace_GetRef();
 		ClassData.Name = ClassName;
-		ClassData.DigestData.bBlacklisted = GetClassBlacklist().Contains(ClassName);
+		ClassData.DigestData.bEditorDomainEnabled = !GetEditorClassBlacklist().Contains(ClassName);
 		if (LookupName != ClassName)
 		{
-			ClassData.DigestData.bBlacklisted |= GetClassBlacklist().Contains(LookupName);
+			ClassData.DigestData.bEditorDomainEnabled &= !GetEditorClassBlacklist().Contains(LookupName);
 		}
 
 		if (Struct)
@@ -246,9 +259,9 @@ void PrecacheClassDigests(TConstArrayView<FName> ClassNames, TMap<FName, FClassD
 
 	TMap<FName, FClassData> RemainingBatch;
 	{
-		FScopeLock ClassDigestsScopeLock(&ClassDigests.Lock);
+		FWriteScopeLock ClassDigestsScopeLock(ClassDigests.Lock);
 
-		// Look up the data for the parent of each class, so we can propagate bBlacklisted from the parent,
+		// Look up the data for the parent of each class, so we can propagate bEditorDomainEnabled from the parent,
 		// once parent data is propagated, add it to the ClassDigests map.
 		// For any parents missing data, keep the class for a second pass that adds the parent
 		for (FClassData& ClassData : ClassDatas)
@@ -259,7 +272,7 @@ void PrecacheClassDigests(TConstArrayView<FName> ClassNames, TMap<FName, FClassD
 				FClassDigestData* ParentDigest = ClassDigests.Map.Find(ClassData.ParentName);
 				if (ParentDigest)
 				{
-					ClassData.DigestData.bBlacklisted |= ParentDigest->bBlacklisted;
+					ClassData.DigestData.bEditorDomainEnabled &= ParentDigest->bEditorDomainEnabled;
 				}
 				else
 				{
@@ -342,7 +355,7 @@ void PrecacheClassDigests(TConstArrayView<FName> ClassNames, TMap<FName, FClassD
 		// If the superclass was not found, due to a bad redirect or a missing blueprint assetregistry entry, then give up and treat the class as having no parent
 		if (ParentDigest)
 		{
-			ClassData.DigestData.bBlacklisted |= ParentDigest->bBlacklisted;
+			ClassData.DigestData.bEditorDomainEnabled &= ParentDigest->bEditorDomainEnabled;
 		}
 	};
 	for (TPair<FName, FClassData>& RemainingPair : RemainingBatch)
@@ -352,7 +365,7 @@ void PrecacheClassDigests(TConstArrayView<FName> ClassNames, TMap<FName, FClassD
 
 	// Add the now-complete RemainingBatch digests to ClassDigests
 	{
-		FScopeLock ClassDigestsScopeLock(&ClassDigests.Lock);
+		FWriteScopeLock ClassDigestsScopeLock(ClassDigests.Lock);
 		for (TPair<FName, FClassData>& RemainingPair : RemainingBatch)
 		{
 			if (OutDatas)
@@ -376,12 +389,6 @@ TSet<FName> ConstructClassBlacklist()
 	return Result;
 }
 
-const TSet<FName>& GetClassBlacklist()
-{
-	static TSet<FName> ClassBlacklist = ConstructClassBlacklist();
-	return ClassBlacklist;
-}
-
 TSet<FName> ConstructPackageNameBlacklist()
 {
 	TSet<FName> Result;
@@ -402,17 +409,128 @@ TSet<FName> ConstructPackageNameBlacklist()
 	return Result;
 }
 
-const TSet<FName>& GetPackageNameBlacklist()
+TSet<FName> ConstructTargetIterativeClassBlacklist()
+{
+	TSet<FName> Result;
+	TArray<FString> BlacklistArray;
+	GConfig->GetArray(TEXT("TargetDomain"), TEXT("IterativeClassBlacklist"), BlacklistArray, GEditorIni);
+	for (const FString& ClassPathName : BlacklistArray)
+	{
+		Result.Add(FName(*ClassPathName));
+	}
+	return Result;
+}
+
+void ConstructTargetIterativeClassWhitelist()
+{
+	// We're using a whitelist with a blacklist override, so the blacklist is only needed when creating the whitelist
+	TSet<FName> BlacklistFNames = ConstructTargetIterativeClassBlacklist();
+
+	// Whitelist elements implicitly whitelist all parent classes, so instead of consulting a list and propagating
+	// from parent classes every time we read a new class, we have to iterate the list for all classes up front and
+	// propagate _TO_ parent classes. Note that we only support whitelisting native classes, otherwise we would have
+	// to wait for the AssetRegistry to finish loading to be sure we could find every specified whitelist class.
+
+	// Declare a recursive Visit function. Every class we visit is whitelisted, and we visit its superclasses.
+	// To decide whether a visited class is enabled, we also have to get IsBlacklisted recursively from the parent.
+	TSet<FName> EnabledFNames;
+	TStringBuilder<256> NameStringBuffer;
+	TMap<FName, TOptional<bool>> Visited;
+	auto EnableClassIfNotBlacklisted = [&Visited, &EnabledFNames, &BlacklistFNames, &NameStringBuffer]
+		(FName PathName, UStruct* Struct, bool& bOutIsBlacklisted, auto& EnableClassIfNotBlacklistedRef)
+	{
+		int32 KeyHash = GetTypeHash(PathName);
+		TOptional<bool>& BlacklistValue = Visited.FindOrAddByHash(KeyHash, PathName);
+		if (BlacklistValue.IsSet())
+		{
+			bOutIsBlacklisted = *BlacklistValue;
+			return;
+		}
+		BlacklistValue = false; // If there is a cycle in the class graph, we will encounter PathName again, so initialize to false
+
+		bool bParentBlacklisted = false;
+		UStruct* ParentStruct = Struct->GetSuperStruct();
+		if (ParentStruct)
+		{
+			NameStringBuffer.Reset();
+			ParentStruct->GetPathName(nullptr, NameStringBuffer);
+			EnableClassIfNotBlacklistedRef(FName(*NameStringBuffer), ParentStruct,
+				bParentBlacklisted, EnableClassIfNotBlacklistedRef);
+		}
+
+		bOutIsBlacklisted = bParentBlacklisted || BlacklistFNames.Contains(PathName);
+		if (bOutIsBlacklisted)
+		{
+			// Call FindOrAdd again, since the recursive calls may have altered the map and invalidated the BlacklistValue reference
+			Visited.FindOrAddByHash(KeyHash, PathName) = bOutIsBlacklisted;
+		}
+		else
+		{
+			EnabledFNames.Add(PathName);
+		}
+	};
+
+	TArray<FString> WhitelistLeafNames;
+	GConfig->GetArray(TEXT("TargetDomain"), TEXT("IterativeClassWhitelist"), WhitelistLeafNames, GEditorIni);
+	for (const FString& ClassPathName : WhitelistLeafNames)
+	{
+		if (!FPackageName::IsScriptPackage(ClassPathName))
+		{
+			continue;
+		}
+		UStruct* Struct = FindObject<UStruct>(nullptr, *ClassPathName);
+		if (!Struct)
+		{
+			continue;
+		}
+		bool bUnusedIsBlacklisted;
+		EnableClassIfNotBlacklisted(FName(*ClassPathName), Struct, bUnusedIsBlacklisted, EnableClassIfNotBlacklisted);
+	}
+
+	TArray<FName> EnabledFNamesArray = EnabledFNames.Array();
+	PrecacheClassDigests(EnabledFNamesArray, nullptr /* OutDatas */);
+	FClassDigestMap& ClassDigests = GetClassDigests();
+	{
+		FWriteScopeLock ClassDigestsScopeLock(ClassDigests.Lock);
+		for (FName ClassPathName : EnabledFNamesArray)
+		{
+			FClassDigestData* DigestData = ClassDigests.Map.Find(ClassPathName);
+			if (DigestData)
+			{
+				DigestData->bTargetIterativeEnabled = true;
+			}
+		}
+	}
+}
+
+FDelegateHandle GUtilsPostInitDelegate;
+void UtilsPostEngineInit();
+
+void UtilsInitialize()
 {
-	static TSet<FName> PackageNameBlacklist = ConstructPackageNameBlacklist();
-	return PackageNameBlacklist;
+	GEditorClassBlacklist = ConstructClassBlacklist();
+	GEditorPackageBlacklist = ConstructPackageNameBlacklist();
+
+	// Constructing whitelists requires use of UStructs, and the early SetPackageResourceManager
+	// where UtilsInitialize is called is too early; trying to call UStruct->GetSchemaHash at that
+	// time will break the UClass. Defer the construction of whitelist-based data until OnPostEngineInit
+	GUtilsPostInitDelegate = FCoreDelegates::OnPostEngineInit.AddLambda([]() { UtilsPostEngineInit(); });
+}
+
+void UtilsPostEngineInit()
+{
+	FCoreDelegates::OnPostEngineInit.Remove(GUtilsPostInitDelegate);
+	GUtilsPostInitDelegate.Reset();
+
+	// Note that constructing whitelists depends on all blacklists having been parsed already
+	ConstructTargetIterativeClassWhitelist();
 }
 
 EPackageDigestResult GetPackageDigest(IAssetRegistry& AssetRegistry, FName PackageName,
-	FPackageDigest& OutPackageDigest, bool& bOutIsBlacklisted, FString& OutErrorMessage)
+	FPackageDigest& OutPackageDigest, bool& bOutEditorDomainEnabled, FString& OutErrorMessage)
 {
 	FCbWriter Builder;
-	EPackageDigestResult Result = AppendPackageDigest(AssetRegistry, PackageName, Builder, bOutIsBlacklisted, OutErrorMessage);
+	EPackageDigestResult Result = AppendPackageDigest(AssetRegistry, PackageName, Builder, bOutEditorDomainEnabled, OutErrorMessage);
 	if (Result == EPackageDigestResult::Success)
 	{
 		OutPackageDigest = Builder.Save().GetRangeHash();
@@ -421,7 +539,7 @@ EPackageDigestResult GetPackageDigest(IAssetRegistry& AssetRegistry, FName Packa
 }
 
 EPackageDigestResult AppendPackageDigest(IAssetRegistry& AssetRegistry, FName PackageName,
-	FCbWriter& Builder, bool& bOutIsBlacklisted, FString& OutErrorMessage)
+	FCbWriter& Builder, bool& bOutEditorDomainEnabled, FString& OutErrorMessage)
 {
 	AssetRegistry.WaitForPackage(PackageName.ToString());
 	TOptional<FAssetPackageData> PackageData = AssetRegistry.GetAssetPackageDataCopy(PackageName);
@@ -429,11 +547,11 @@ EPackageDigestResult AppendPackageDigest(IAssetRegistry& AssetRegistry, FName Pa
 	{
 		OutErrorMessage = FString::Printf(TEXT("Package %s does not exist in the AssetRegistry"),
 			*PackageName.ToString());
-		bOutIsBlacklisted = false;
+		bOutEditorDomainEnabled = true;
 		return EPackageDigestResult::FileDoesNotExist;
 	}
-	EPackageDigestResult Result = AppendPackageDigest(Builder, bOutIsBlacklisted, OutErrorMessage, *PackageData, PackageName);
-	bOutIsBlacklisted |= GetPackageNameBlacklist().Contains(PackageName);
+	EPackageDigestResult Result = AppendPackageDigest(Builder, bOutEditorDomainEnabled, OutErrorMessage, *PackageData, PackageName);
+	bOutEditorDomainEnabled &= !GetEditorPackageBlacklist().Contains(PackageName);
 	return Result;
 }
 
@@ -624,9 +742,9 @@ bool TrySavePackage(UPackage* Package)
 
 	FString ErrorMessage;
 	FPackageDigest PackageDigest;
-	bool bIsBlacklisted;
+	bool bEditorDomainEnabled;
 	EPackageDigestResult FindHashResult = GetPackageDigest(*IAssetRegistry::Get(), Package->GetFName(), PackageDigest,
-		bIsBlacklisted, ErrorMessage);
+		bEditorDomainEnabled, ErrorMessage);
 	switch (FindHashResult)
 	{
 	case EPackageDigestResult::Success:
@@ -635,7 +753,7 @@ bool TrySavePackage(UPackage* Package)
 		UE_LOG(LogEditorDomain, Warning, TEXT("Could not save package to EditorDomain: %s."), *ErrorMessage)
 		return false;
 	}
-	if (bIsBlacklisted)
+	if (!bEditorDomainEnabled)
 	{
 		UE_LOG(LogEditorDomain, Verbose, TEXT("Skipping save of blacklisted package to EditorDomain: %s."), *Package->GetName())
 		return false;
@@ -726,9 +844,9 @@ void GetBulkDataList(FName PackageName, UE::DerivedData::IRequestOwner& Owner, T
 {
 	FString ErrorMessage;
 	FPackageDigest PackageDigest;
-	bool bIsBlacklisted;
+	bool bEditorDomainEnabled;
 	EPackageDigestResult FindHashResult = GetPackageDigest(*IAssetRegistry::Get(), PackageName, PackageDigest,
-		bIsBlacklisted, ErrorMessage);
+		bEditorDomainEnabled, ErrorMessage);
 	switch (FindHashResult)
 	{
 	case EPackageDigestResult::Success:
@@ -739,7 +857,7 @@ void GetBulkDataList(FName PackageName, UE::DerivedData::IRequestOwner& Owner, T
 		return;
 	}
 	}
-	if (bIsBlacklisted)
+	if (!bEditorDomainEnabled)
 	{
 		Callback(FSharedBuffer());
 		return;
@@ -760,9 +878,9 @@ void PutBulkDataList(FName PackageName, FSharedBuffer Buffer)
 {
 	FString ErrorMessage;
 	FPackageDigest PackageDigest;
-	bool bIsBlacklisted;
+	bool bEditorDomainEnabled;
 	EPackageDigestResult FindHashResult = GetPackageDigest(*IAssetRegistry::Get(), PackageName, PackageDigest,
-		bIsBlacklisted, ErrorMessage);
+		bEditorDomainEnabled, ErrorMessage);
 	switch (FindHashResult)
 	{
 	case EPackageDigestResult::Success:
@@ -772,7 +890,7 @@ void PutBulkDataList(FName PackageName, FSharedBuffer Buffer)
 		return;
 	}
 	}
-	if (bIsBlacklisted)
+	if (!bEditorDomainEnabled)
 	{
 		return;
 	}
@@ -797,8 +915,8 @@ void GetBulkDataPayloadId(FName PackageName, const FGuid& BulkDataId, UE::Derive
 {
 	FString ErrorMessage;
 	FCbWriter Builder;
-	bool bIsBlacklisted;
-	EPackageDigestResult FindHashResult = AppendPackageDigest(*IAssetRegistry::Get(), PackageName, Builder, bIsBlacklisted, ErrorMessage);
+	bool bEditorDomainEnabled;
+	EPackageDigestResult FindHashResult = AppendPackageDigest(*IAssetRegistry::Get(), PackageName, Builder, bEditorDomainEnabled, ErrorMessage);
 	switch (FindHashResult)
 	{
 	case EPackageDigestResult::Success:
@@ -809,7 +927,7 @@ void GetBulkDataPayloadId(FName PackageName, const FGuid& BulkDataId, UE::Derive
 		return;
 	}
 	}
-	if (bIsBlacklisted)
+	if (!bEditorDomainEnabled)
 	{
 		Callback(FSharedBuffer());
 		return;
@@ -831,8 +949,8 @@ void PutBulkDataPayloadId(FName PackageName, const FGuid& BulkDataId, FSharedBuf
 {
 	FString ErrorMessage;
 	FCbWriter Builder;
-	bool bIsBlacklisted;
-	EPackageDigestResult FindHashResult = AppendPackageDigest(*IAssetRegistry::Get(), PackageName, Builder, bIsBlacklisted, ErrorMessage);
+	bool bEditorDomainEnabled;
+	EPackageDigestResult FindHashResult = AppendPackageDigest(*IAssetRegistry::Get(), PackageName, Builder, bEditorDomainEnabled, ErrorMessage);
 	switch (FindHashResult)
 	{
 	case EPackageDigestResult::Success:
@@ -842,7 +960,7 @@ void PutBulkDataPayloadId(FName PackageName, const FGuid& BulkDataId, FSharedBuf
 		return;
 	}
 	}
-	if (bIsBlacklisted)
+	if (!bEditorDomainEnabled)
 	{
 		return;
 	}
diff --git a/Engine/Source/Editor/UnrealEd/Private/EditorDomain/EditorDomainUtils.h b/Engine/Source/Editor/UnrealEd/Private/EditorDomain/EditorDomainUtils.h
--- a/Engine/Source/Editor/UnrealEd/Private/EditorDomain/EditorDomainUtils.h
+++ b/Engine/Source/Editor/UnrealEd/Private/EditorDomain/EditorDomainUtils.h
@@ -8,6 +8,7 @@
 #include "Containers/UnrealString.h"
 #include "DerivedDataCache.h"
 #include "EditorDomain/EditorDomain.h"
+#include "HAL/CriticalSection.h"
 #include "Misc/ConfigCacheIni.h"
 #include "Templates/Function.h"
 #include "UObject/NameTypes.h"
@@ -39,14 +40,17 @@ struct FClassDigestData
 {
 	FBlake3Hash SchemaHash;
 	bool bNative = false;
-	bool bBlacklisted = false;
+	/** EditorDomainEnabled whitelists everything and uses only a blacklist, so defaults to true. */
+	bool bEditorDomainEnabled = true;
+	/** bTargetIterativeEnabled uses a whitelist (with a blacklist override), so defaults to false. */
+	bool bTargetIterativeEnabled = false;
 };
 
 /** Threadsafe cache of ClassName -> Digest data for calculating EditorDomain Digests */
 struct FClassDigestMap
 {
 	TMap<FName, FClassDigestData> Map;
-	FCriticalSection Lock;
+	FRWLock Lock;
 };
 
 /**
@@ -54,18 +58,14 @@ struct FClassDigestMap
  * Reads information from the AssetRegistry to compute the digest.
  */
 EPackageDigestResult GetPackageDigest(IAssetRegistry& AssetRegistry, FName PackageName,
-	FPackageDigest& OutPackageDigest, bool& bOutIsBlacklisted, FString& OutErrorMessage);
+	FPackageDigest& OutPackageDigest, bool& bOutEditorDomainEnabled, FString& OutErrorMessage);
 /** Appends the fields to calculate the packagedigest; call Builder.Save().GetRangeHash() to get digest. */
 EPackageDigestResult AppendPackageDigest(IAssetRegistry& AssetRegistry, FName PackageName,
-	FCbWriter& Builder, bool& bOutIsBlacklisted, FString& OutErrorMessage);
+	FCbWriter& Builder, bool& bOutEditorDomainEnabled, FString& OutErrorMessage);
 
 /** For any ClassNames not already in ClassDigests, look up their UStruct and add them. */
 void PrecacheClassDigests(TConstArrayView<FName> ClassNames, TMap<FName, FClassDigestData>* OutDatas = nullptr);
 
-/** Return the TSet of class fullpath names that was parsed from the Editor.ini:[EditorDomain]:ClassBlacklist array. */
-const TSet<FName>& GetClassBlacklist();
-const TSet<FName>& GetPackageNameBlacklist();
-
 /** Get the CacheRequest for the given package from the EditorDomain cache bucket. */
 void RequestEditorDomainPackage(const FPackagePath& PackagePath,
 	const FPackageDigest& PackageDigest, UE::DerivedData::ECachePolicy SkipFlags,
@@ -90,4 +90,7 @@ void PutBulkDataPayloadId(FName PackageName, const FGuid& BulkDataId, FSharedBuf
 /** Accessor for the global ClassDigest map shared by systems needing to calculate PackageDigests. */
 FClassDigestMap& GetClassDigests();
 
+/** Initializes some global config-driven values used by the EditorDomain and TargetDomain. */
+void UtilsInitialize();
+
 }
diff --git a/Engine/Source/Editor/UnrealEd/Private/TargetDomain/TargetDomainUtils.cpp b/Engine/Source/Editor/UnrealEd/Private/TargetDomain/TargetDomainUtils.cpp
--- a/Engine/Source/Editor/UnrealEd/Private/TargetDomain/TargetDomainUtils.cpp
+++ b/Engine/Source/Editor/UnrealEd/Private/TargetDomain/TargetDomainUtils.cpp
@@ -9,6 +9,7 @@
 #include "EditorDomain/EditorDomainUtils.h"
 #include "IO/IoHash.h"
 #include "IO/PackageStoreWriter.h"
+#include "Misc/ScopeRWLock.h"
 #include "Misc/StringBuilder.h"
 #include "Serialization/CompactBinary.h"
 #include "Serialization/CompactBinaryWriter.h"
@@ -229,4 +230,38 @@ bool TryFetchKeyAndDependencies(IPackageStoreWriter* PackageStore, FName Package
 	return true;
 }
 
+bool IsIterativeEnabled(FName PackageName)
+{
+	IAssetRegistry* AssetRegistry = IAssetRegistry::Get();
+	if (!AssetRegistry)
+	{
+		return false;
+	}
+	TOptional<FAssetPackageData> PackageDataOpt = AssetRegistry->GetAssetPackageDataCopy(PackageName);
+	if (!PackageDataOpt)
+	{
+		return false;
+	}
+	FAssetPackageData& PackageData = *PackageDataOpt;
+
+	UE::EditorDomain::FClassDigestMap& ClassDigests = UE::EditorDomain::GetClassDigests();
+	FReadScopeLock ClassDigestsScopeLock(ClassDigests.Lock);
+	for (FName ClassName : PackageData.ImportedClasses)
+	{
+		UE::EditorDomain::FClassDigestData* ExistingData = ClassDigests.Map.Find(ClassName);
+		if (!ExistingData)
+		{
+			// All whitelisted classes are added to ClassDigests at startup, so if the class is not in ClassDigests,
+			// it is not whitelisted
+			return false;
+		}
+		if (!ExistingData->bTargetIterativeEnabled)
+		{
+			return false;
+		}
+	}
+	return true;
 }
+
+} // namespace UE::TargetDomain
+
diff --git a/Engine/Source/Editor/UnrealEd/Private/TargetDomain/TargetDomainUtils.h b/Engine/Source/Editor/UnrealEd/Private/TargetDomain/TargetDomainUtils.h
--- a/Engine/Source/Editor/UnrealEd/Private/TargetDomain/TargetDomainUtils.h
+++ b/Engine/Source/Editor/UnrealEd/Private/TargetDomain/TargetDomainUtils.h
@@ -31,4 +31,7 @@ FCbObject CollectDependenciesObject(UPackage* Package, const ITargetPlatform* Ta
 bool TryFetchKeyAndDependencies(IPackageStoreWriter* PackageStore, FName PackageName, const ITargetPlatform* TargetPlatform,
 	FIoHash* OutHash, TArray<FName>* OutBuildDependencies, TArray<FName>* OutRuntimeOnlyDependencies, FString* OutErrorMessage);
 
+/** Return whether iterative cook is enabled for the given packagename, based on used-class whitelist/blacklist. */
+bool IsIterativeEnabled(FName PackageName);
+
 }
