Moving Watch pins to live with breakpoints in their new home: FPerBlueprintSettings

diff --git a/Engine/Plugins/Animation/ControlRig/Source/ControlRigDeveloper/Private/ControlRigBlueprint.cpp b/Engine/Plugins/Animation/ControlRig/Source/ControlRigDeveloper/Private/ControlRigBlueprint.cpp
--- a/Engine/Plugins/Animation/ControlRig/Source/ControlRigDeveloper/Private/ControlRigBlueprint.cpp
+++ b/Engine/Plugins/Animation/ControlRig/Source/ControlRigDeveloper/Private/ControlRigBlueprint.cpp
@@ -29,6 +29,7 @@
 
 #if WITH_EDITOR
 #include "IControlRigEditorModule.h"
+#include "Kismet2/KismetDebugUtilities.h"
 #include "Kismet2/BlueprintEditorUtils.h"
 #include "ControlRigBlueprintUtils.h"
 #include "Settings/ControlRigSettings.h"
@@ -2891,11 +2892,14 @@ void UControlRigBlueprint::HandleModifiedEvent(ERigVMGraphNotifType InNotifType,
 								{
 									if (ModelPin->RequiresWatch())
 									{
-										WatchedPins.AddUnique(EdPin);
+										if (!FKismetDebugUtilities::IsPinBeingWatched(this, EdPin))
+										{
+											FKismetDebugUtilities::AddPinWatch(this, EdPin);
+										}
 									}
 									else
 									{
-										WatchedPins.Remove(EdPin);
+										FKismetDebugUtilities::RemovePinWatch(this, EdPin);
 									}
 
 									if(InNotifType == ERigVMGraphNotifType::PinWatchedChanged)
diff --git a/Engine/Plugins/Animation/ControlRig/Source/ControlRigEditor/Private/Editor/ControlRigEditor.cpp b/Engine/Plugins/Animation/ControlRig/Source/ControlRigEditor/Private/Editor/ControlRigEditor.cpp
--- a/Engine/Plugins/Animation/ControlRig/Source/ControlRigEditor/Private/Editor/ControlRigEditor.cpp
+++ b/Engine/Plugins/Animation/ControlRig/Source/ControlRigEditor/Private/Editor/ControlRigEditor.cpp
@@ -92,6 +92,7 @@
 #include "ControlRigContextMenuContext.h"
 #include "Types/ISlateMetaData.h"
 #include "ControlRigGraphDetails.h"
+#include "Kismet2/KismetDebugUtilities.h"
 
 #define LOCTEXT_NAMESPACE "ControlRigEditor"
 
@@ -3950,39 +3951,43 @@ void FControlRigEditor::UpdateStaleWatchedPins()
 	}
 
 	TSet<UEdGraphPin*> AllPins;
+	uint16 WatchCount;
 
 	// Find all unique pins being watched
-	for (const FEdGraphPinReference& PinRef : ControlRigBP->WatchedPins)
-	{
-		UEdGraphPin* Pin = PinRef.Get();
-		if (Pin == nullptr)
+	FKismetDebugUtilities::ForeachPinWatch(
+		ControlRigBP,
+		[&AllPins, &WatchCount](UEdGraphPin* Pin)
 		{
-			continue;
-		}
+			++WatchCount;
+			if (Pin == nullptr)
+			{
+				return; // ~continue
+			}
 
-		UEdGraphNode* OwningNode = Pin->GetOwningNode();
-		// during node reconstruction, dead pins get moved to the transient 
-		// package (so just in case this blueprint got saved with dead pin watches)
-		if (OwningNode == nullptr)
-		{
-			continue;
-		}
+			UEdGraphNode* OwningNode = Pin->GetOwningNode();
+			// during node reconstruction, dead pins get moved to the transient 
+			// package (so just in case this blueprint got saved with dead pin watches)
+			if (OwningNode == nullptr)
+			{
+				return; // ~continue
+			}
 
-		if (!OwningNode->Pins.Contains(Pin))
-		{
-			continue;
-		}
+			if (!OwningNode->Pins.Contains(Pin))
+			{
+				return; // ~continue
+			}
 
-		AllPins.Add(Pin);
-	}
+			AllPins.Add(Pin);
+		}
+	);
 
 	// Refresh watched pins with unique pins (throw away null or duplicate watches)
-	if (ControlRigBP->WatchedPins.Num() != AllPins.Num())
+	if (WatchCount != AllPins.Num())
 	{
 		ControlRigBP->Status = BS_Dirty;
 	}
 
-	ControlRigBP->WatchedPins.Empty();
+	FKismetDebugUtilities::ClearPinWatches(ControlRigBP);
 
 	TArray<URigVMGraph*> Models = ControlRigBP->GetAllModels();
 	for (URigVMGraph* Model : Models)
@@ -4001,7 +4006,7 @@ void FControlRigEditor::UpdateStaleWatchedPins()
 	}
 	for (UEdGraphPin* Pin : AllPins)
 	{
-		ControlRigBP->WatchedPins.Add(Pin);
+		FKismetDebugUtilities::AddPinWatch(ControlRigBP, Pin);
 		UEdGraph* EdGraph = Pin->GetOwningNode()->GetGraph();
 		ControlRigBP->GetController(EdGraph)->SetPinIsWatched(Pin->GetName(), true, false);
 	}
diff --git a/Engine/Source/Editor/BlueprintGraph/Private/EdGraphSchema_K2.cpp b/Engine/Source/Editor/BlueprintGraph/Private/EdGraphSchema_K2.cpp
--- a/Engine/Source/Editor/BlueprintGraph/Private/EdGraphSchema_K2.cpp
+++ b/Engine/Source/Editor/BlueprintGraph/Private/EdGraphSchema_K2.cpp
@@ -7115,7 +7115,7 @@ void UEdGraphSchema_K2::RecombinePin(UEdGraphPin* Pin) const
 		}
 
 		GraphNode->Pins.Remove(SubPin);
-		Blueprint->WatchedPins.Remove(SubPin);
+		FKismetDebugUtilities::RemovePinWatch(Blueprint, SubPin);
 	}
 
 	if (Pin->Direction == EGPD_Input)
diff --git a/Engine/Source/Editor/BlueprintGraph/Private/K2Node.cpp b/Engine/Source/Editor/BlueprintGraph/Private/K2Node.cpp
--- a/Engine/Source/Editor/BlueprintGraph/Private/K2Node.cpp
+++ b/Engine/Source/Editor/BlueprintGraph/Private/K2Node.cpp
@@ -100,7 +100,7 @@ void UK2Node::PostLoad()
 			{
 				if (UEdGraphPin* NewPin = UEdGraphPin::FindPinCreatedFromDeprecatedPin(WatchedPin))
 				{
-					BP->WatchedPins.Add(NewPin);
+					FKismetDebugUtilities::AddPinWatch(BP, NewPin);
 				}
 
 				BP->DeprecatedPinWatches.RemoveAt(WatchIdx);
@@ -1019,15 +1019,10 @@ void UK2Node::ReconstructSinglePin(UEdGraphPin* NewPin, UEdGraphPin* OldPin, ERe
 		}
 	}
 
-	// Update the blueprints watched pins as the old pin will be going the way of the dodo
-	for (int32 WatchIndex = 0; WatchIndex < Blueprint->WatchedPins.Num(); ++WatchIndex)
+	// if OldPin has a watch, swap the watch for NewPin
+	if(FKismetDebugUtilities::RemovePinWatch(Blueprint, OldPin))
 	{
-		UEdGraphPin* WatchedPin = Blueprint->WatchedPins[WatchIndex].Get();
-		if( WatchedPin == OldPin )
-		{
-			WatchedPin = NewPin;
-			break;
-		}
+		FKismetDebugUtilities::AddPinWatch(Blueprint, NewPin);
 	}
 }
 
diff --git a/Engine/Source/Editor/BlueprintGraph/Public/BlueprintEditorSettings.h b/Engine/Source/Editor/BlueprintGraph/Public/BlueprintEditorSettings.h
--- a/Engine/Source/Editor/BlueprintGraph/Public/BlueprintEditorSettings.h
+++ b/Engine/Source/Editor/BlueprintGraph/Public/BlueprintEditorSettings.h
@@ -21,6 +21,7 @@ enum ESaveOnCompile
 
 /** Blueprint Editor settings that are different for each
 *	blueprint.
+*	See FKismetDebugUtilities for helper functions
 */
 USTRUCT()
 struct BLUEPRINTGRAPH_API FPerBlueprintSettings 
@@ -29,6 +30,22 @@ struct BLUEPRINTGRAPH_API FPerBlueprintSettings
 	
 	UPROPERTY()
 	TArray<FBlueprintBreakpoint> Breakpoints;
+
+	UPROPERTY()
+	TArray<FEdGraphPinReference> WatchedPins;
+
+	bool operator==(const FPerBlueprintSettings& Other) const
+	{
+		return Breakpoints == Other.Breakpoints && WatchedPins == Other.WatchedPins;
+	}
+};
+
+template<> struct TStructOpsTypeTraits<FPerBlueprintSettings> : public TStructOpsTypeTraitsBase2<FPerBlueprintSettings>
+{
+	enum
+	{
+		WithIdenticalViaEquality = true
+	};
 };
 
 UCLASS(config=EditorPerProjectUserSettings)
diff --git a/Engine/Source/Editor/GraphEditor/Private/KismetNodes/KismetNodeInfoContext.cpp b/Engine/Source/Editor/GraphEditor/Private/KismetNodes/KismetNodeInfoContext.cpp
--- a/Engine/Source/Editor/GraphEditor/Private/KismetNodes/KismetNodeInfoContext.cpp
+++ b/Engine/Source/Editor/GraphEditor/Private/KismetNodes/KismetNodeInfoContext.cpp
@@ -7,6 +7,7 @@
 #include "EdGraphSchema_K2.h"
 #include "K2Node_CallFunction.h"
 #include "Kismet2/BlueprintEditorUtils.h"
+#include "Kismet2/KismetDebugUtilities.h"
 
 //////////////////////////////////////////////////////////////////////////
 // FKismetNodeInfoContext
@@ -86,23 +87,26 @@ FKismetNodeInfoContext::FKismetNodeInfoContext(UEdGraph* SourceGraph)
 		}
 
 		// Covert the watched pin array into a set
-		for (auto WatchedPinIt = SourceBlueprint->WatchedPins.CreateConstIterator(); WatchedPinIt; ++WatchedPinIt)
-		{
-			UEdGraphPin* WatchedPin = WatchedPinIt->Get();
-			if (!ensure(WatchedPin))
-			{
-				continue;
-			}
+		FKismetDebugUtilities::ForeachPinWatch(
+			SourceBlueprint,
+			[&WatchedPinSet = WatchedPinSet, &WatchedNodeSet = WatchedNodeSet]
+				(UEdGraphPin* WatchedPin)
+				{
+					if (!ensure(WatchedPin))
+					{
+						return; // ~continue
+					}
 
-			UEdGraphNode* OwningNode = WatchedPin->GetOuter();
-			if (!ensure(OwningNode != NULL)) // shouldn't happen, but just in case a dead pin was added to the WatchedPins array
-			{
-				continue;
-			}
-			check(OwningNode == WatchedPin->GetOwningNode());
+					UEdGraphNode* OwningNode = WatchedPin->GetOuter();
+					if (!ensure(OwningNode != NULL)) // shouldn't happen, but just in case a dead pin was added to the WatchedPins array
+					{
+						return; // ~continue
+					}
+					check(OwningNode == WatchedPin->GetOwningNode());
 
-			WatchedPinSet.Add(WatchedPin);
-			WatchedNodeSet.Add(OwningNode);
-		}
+					WatchedPinSet.Add(WatchedPin);
+					WatchedNodeSet.Add(OwningNode);
+				}
+		);
 	}
 }
diff --git a/Engine/Source/Editor/Kismet/Private/BlueprintEditor.cpp b/Engine/Source/Editor/Kismet/Private/BlueprintEditor.cpp
--- a/Engine/Source/Editor/Kismet/Private/BlueprintEditor.cpp
+++ b/Engine/Source/Editor/Kismet/Private/BlueprintEditor.cpp
@@ -3991,7 +3991,8 @@ bool FBlueprintEditor::HasAnyDisabledBreakpoints() const
 
 bool FBlueprintEditor::HasAnyWatches() const
 {
-	return GetBlueprintObj() && GetBlueprintObj()->WatchedPins.Num() > 0;
+	const UBlueprint* Blueprint = GetBlueprintObj();;
+	return Blueprint && FKismetDebugUtilities::BlueprintHasPinWatches(Blueprint);
 }
 
 // Jumps to a hyperlinked node, pin, or graph, if it belongs to this blueprint
diff --git a/Engine/Source/Editor/Kismet/Private/Debugging/SKismetDebuggingView.cpp b/Engine/Source/Editor/Kismet/Private/Debugging/SKismetDebuggingView.cpp
--- a/Engine/Source/Editor/Kismet/Private/Debugging/SKismetDebuggingView.cpp
+++ b/Engine/Source/Editor/Kismet/Private/Debugging/SKismetDebuggingView.cpp
@@ -951,10 +951,13 @@ public:
 				// Create children for each watch
 				if(IsDebugLineTypeActive(DLT_Watch))
 				{
-					for (FEdGraphPinReference WatchedPin : ParentBP->WatchedPins)
-					{
-						EnsureChildIsAdded(OutChildren, FWatchLineItem(WatchedPin.Get(), ParentObject));
-					}
+					FKismetDebugUtilities::ForeachPinWatch(
+						ParentBP,
+						[this, &OutChildren, ParentObject](UEdGraphPin* WatchedPin)
+						{
+							EnsureChildIsAdded(OutChildren, FWatchLineItem(WatchedPin, ParentObject));
+						}
+					);
 				}
 
 				// It could also have active latent behaviors
@@ -1058,7 +1061,7 @@ protected:
 	{
 		if (UBlueprint* BP = Cast<UBlueprint>(ObjectRef.Get()))
 		{
-			if (BP->WatchedPins.Num() > 0)
+			if (FKismetDebugUtilities::BlueprintHasPinWatches(BP))
 			{
 				FUIAction ClearAllWatches(
 					FExecuteAction::CreateStatic( &FDebuggingActionCallbacks::ClearWatches, BP )
diff --git a/Engine/Source/Editor/Kismet/Private/WatchPointViewer.cpp b/Engine/Source/Editor/Kismet/Private/WatchPointViewer.cpp
--- a/Engine/Source/Editor/Kismet/Private/WatchPointViewer.cpp
+++ b/Engine/Source/Editor/Kismet/Private/WatchPointViewer.cpp
@@ -183,9 +183,9 @@ namespace
 			}
 			FText BlueprintName = FText::FromString(BlueprintObj->GetName());
 
-			for (const FEdGraphPinReference& PinRef : BlueprintObj->WatchedPins)
-			{
-				if (UEdGraphPin* Pin = PinRef.Get())
+			FKismetDebugUtilities::ForeachPinWatch(
+				BlueprintObj.Get(),
+				[BlueprintObj, BlueprintName](UEdGraphPin* Pin)
 				{
 					FText GraphName = FText::FromString(Pin->GetOwningNode()->GetGraph()->GetName());
 					FText NodeName = Pin->GetOwningNode()->GetNodeTitle(ENodeTitleType::ListView);
@@ -210,7 +210,7 @@ namespace
 						)
 					);
 				}
-			}
+			);
 		}
 	}
 
@@ -842,55 +842,56 @@ void WatchViewer::UpdateInstancedWatchDisplay()
 			}
 
 			// We have a valid instance, iterate over all the watched pins and create rows for them
-			for (const FEdGraphPinReference& PinRef : BlueprintObj->WatchedPins)
-			{
-				UEdGraphPin* Pin = PinRef.Get();
-
-				FText GraphName = FText::FromString(Pin->GetOwningNode()->GetGraph()->GetName());
-				FText NodeName = Pin->GetOwningNode()->GetNodeTitle(ENodeTitleType::ListView);
-
-				TSharedPtr<FPropertyInstanceInfo> DebugInfo;
-				const FKismetDebugUtilities::EWatchTextResult WatchStatus = FKismetDebugUtilities::GetDebugInfo(DebugInfo, BlueprintObj, BlueprintInstance, Pin);
-
-				if (WatchStatus != FKismetDebugUtilities::EWTR_Valid)
+			FKismetDebugUtilities::ForeachPinWatch(
+				BlueprintObj,
+				[BlueprintObj, BlueprintInstance, BlueprintName](UEdGraphPin* Pin)
 				{
-					const UEdGraphSchema* Schema = Pin->GetOwningNode()->GetSchema();
-					DebugInfo->DisplayName = Schema->GetPinDisplayName(Pin);
-					DebugInfo->Type = UEdGraphSchema_K2::TypeToText(Pin->PinType);
-
-					switch (WatchStatus)
-					{
-					case FKismetDebugUtilities::EWTR_NotInScope:
-						DebugInfo->Value = LOCTEXT("NotInScope", "(not in scope)");
-						break;
-
-					case FKismetDebugUtilities::EWTR_NoProperty:
-						DebugInfo->Value = LOCTEXT("NoDebugData", "(no debug data)");
-						break;
+					FText GraphName = FText::FromString(Pin->GetOwningNode()->GetGraph()->GetName());
+					FText NodeName = Pin->GetOwningNode()->GetNodeTitle(ENodeTitleType::ListView);
 
-					case FKismetDebugUtilities::EWTR_NoDebugObject:
-						DebugInfo->Value = LOCTEXT("NoDebugObject", "(no debug object)");
-						break;
+					TSharedPtr<FPropertyInstanceInfo> DebugInfo;
+					const FKismetDebugUtilities::EWatchTextResult WatchStatus = FKismetDebugUtilities::GetDebugInfo(DebugInfo, BlueprintObj, BlueprintInstance, Pin);
 
-					default:
-						// do nothing
-						break;
+					if (WatchStatus != FKismetDebugUtilities::EWTR_Valid)
+					{
+						const UEdGraphSchema* Schema = Pin->GetOwningNode()->GetSchema();
+						DebugInfo->DisplayName = Schema->GetPinDisplayName(Pin);
+						DebugInfo->Type = UEdGraphSchema_K2::TypeToText(Pin->PinType);
+
+						switch (WatchStatus)
+						{
+						case FKismetDebugUtilities::EWTR_NotInScope:
+							DebugInfo->Value = LOCTEXT("NotInScope", "(not in scope)");
+							break;
+
+						case FKismetDebugUtilities::EWTR_NoProperty:
+							DebugInfo->Value = LOCTEXT("NoDebugData", "(no debug data)");
+							break;
+
+						case FKismetDebugUtilities::EWTR_NoDebugObject:
+							DebugInfo->Value = LOCTEXT("NoDebugObject", "(no debug object)");
+							break;
+
+						default:
+							// do nothing
+							break;
+						}
 					}
-				}
 
-				Private_InstanceWatchSource.Add(
-					MakeShared<FWatchRow>(
-						BlueprintObj,
-						Pin->GetOwningNode(),
-						Pin,
-						BlueprintInstance,
-						BlueprintName,
-						GraphName,
-						NodeName,
-						*DebugInfo
-						)
-				);
-			}
+					Private_InstanceWatchSource.Add(
+						MakeShared<FWatchRow>(
+							BlueprintObj,
+							Pin->GetOwningNode(),
+							Pin,
+							BlueprintInstance,
+							BlueprintName,
+							GraphName,
+							NodeName,
+							*DebugInfo
+							)
+					);
+				}
+			);
 		}
 
 		// Notify subscribers:
diff --git a/Engine/Source/Editor/UnrealEd/Private/Kismet2/BlueprintEditorUtils.cpp b/Engine/Source/Editor/UnrealEd/Private/Kismet2/BlueprintEditorUtils.cpp
--- a/Engine/Source/Editor/UnrealEd/Private/Kismet2/BlueprintEditorUtils.cpp
+++ b/Engine/Source/Editor/UnrealEd/Private/Kismet2/BlueprintEditorUtils.cpp
@@ -7548,39 +7548,43 @@ void FBlueprintEditorUtils::UpdateTransactionalFlags(UBlueprint* Blueprint)
 void FBlueprintEditorUtils::UpdateStalePinWatches( UBlueprint* Blueprint )
 {
 	TSet<UEdGraphPin*> AllPins;
-
+	uint16 WatchCount;
+	
 	// Find all unique pins being watched
-	for (const FEdGraphPinReference& PinRef : Blueprint->WatchedPins)
-	{
-		UEdGraphPin* Pin = PinRef.Get();
-		if (Pin == nullptr)
+	FKismetDebugUtilities::ForeachPinWatch(
+		Blueprint,
+		[&AllPins, &WatchCount](UEdGraphPin* Pin)
 		{
-			continue;
-		}
+			++WatchCount;
+			if (Pin == nullptr)
+			{
+				return; // ~continue
+			}
 
-		UEdGraphNode* OwningNode = Pin->GetOwningNode();
-		// during node reconstruction, dead pins get moved to the transient 
-		// package (so just in case this blueprint got saved with dead pin watches)
-		if (OwningNode == nullptr)
-		{
-			continue;
-		}
+			UEdGraphNode* OwningNode = Pin->GetOwningNode();
+			// during node reconstruction, dead pins get moved to the transient 
+			// package (so just in case this blueprint got saved with dead pin watches)
+			if (OwningNode == nullptr)
+			{
+				return; // ~continue
+			}
 
-		if (!OwningNode->Pins.Contains(Pin))
-		{
-			continue;
+			if (!OwningNode->Pins.Contains(Pin))
+			{
+				return; // ~continue
+			}
+
+			AllPins.Add(Pin);
 		}
-		
-		AllPins.Add(Pin);
-	}
+	);
 
 	// Refresh watched pins with unique pins (throw away null or duplicate watches)
-	if (Blueprint->WatchedPins.Num() != AllPins.Num())
+	if (WatchCount != AllPins.Num())
 	{
-		Blueprint->WatchedPins.Empty();
+		FKismetDebugUtilities::ClearPinWatches(Blueprint);
 		for (UEdGraphPin* Pin : AllPins)
 		{
-			Blueprint->WatchedPins.Add(Pin);
+			FKismetDebugUtilities::AddPinWatch(Blueprint, Pin);
 		}
 
 		Blueprint->Status = BS_Dirty;
diff --git a/Engine/Source/Editor/UnrealEd/Private/Kismet2/KismetDebugUtilities.cpp b/Engine/Source/Editor/UnrealEd/Private/Kismet2/KismetDebugUtilities.cpp
--- a/Engine/Source/Editor/UnrealEd/Private/Kismet2/KismetDebugUtilities.cpp
+++ b/Engine/Source/Editor/UnrealEd/Private/Kismet2/KismetDebugUtilities.cpp
@@ -784,21 +784,34 @@ bool FKismetDebugUtilities::IsSingleStepping()
 		|| Data.TargetGraphStackDepth != INDEX_NONE; 
 }
 
-TArray<FBlueprintBreakpoint>* FKismetDebugUtilities::GetBreakpoints(const UBlueprint* Blueprint)
+FPerBlueprintSettings* FKismetDebugUtilities::GetPerBlueprintSettings(const UBlueprint* Blueprint)
 {
 	UBlueprintEditorSettings* Settings = GetMutableDefault<UBlueprintEditorSettings>();
 	check(Settings);
-	
-	FPerBlueprintSettings* Found = Settings->PerBlueprintSettings.Find(Blueprint->GetPathName());
-	if(Found)
-	{
-		check(!Found->Breakpoints.IsEmpty());
-		return &Found->Breakpoints;
-	}
-	return nullptr;
+	return Settings->PerBlueprintSettings.Find(Blueprint->GetPathName());
+}
+
+TArray<FBlueprintBreakpoint>* FKismetDebugUtilities::GetBreakpoints(const UBlueprint* Blueprint)
+{
+	FPerBlueprintSettings* Settings = GetPerBlueprintSettings(Blueprint);
+
+	// return nullptr if there's no breakpoints associated w/ this blueprint
+	return (!Settings || Settings->Breakpoints.IsEmpty()) ?
+		nullptr :
+		&Settings->Breakpoints;
+}
+
+TArray<FEdGraphPinReference>* FKismetDebugUtilities::GetWatchedPins(const UBlueprint* Blueprint)
+{
+	FPerBlueprintSettings* Settings = GetPerBlueprintSettings(Blueprint);
+
+	// return nullptr if there's no breakpoints associated w/ this blueprint
+	return (!Settings || Settings->WatchedPins.IsEmpty()) ?
+		nullptr :
+		&Settings->WatchedPins;
 }
 
-void FKismetDebugUtilities::SaveBreakpoints()
+void FKismetDebugUtilities::SaveBlueprintEditorSettings()
 {
 	UBlueprintEditorSettings* Settings = GetMutableDefault<UBlueprintEditorSettings>();
 	check(Settings);
@@ -821,6 +834,48 @@ void FKismetDebugUtilities::CleanupBreakpoints(const UBlueprint* Blueprint)
 	);
 }
 
+void FKismetDebugUtilities::CleanupWatches(const UBlueprint* Blueprint)
+{
+	RemovePinWatchesByPredicate(
+		Blueprint,
+		[Blueprint](const UEdGraphPin* Pin)->bool
+		{
+			if(Pin)
+			{
+				if(UEdGraphNode* Node = Pin->GetOwningNode())
+				{
+					if(UEdGraph* Graph = Node->GetGraph())
+					{
+						TArray<UEdGraph*> BPGraphs;
+						Blueprint->GetAllGraphs(BPGraphs);
+						if(BPGraphs.Contains(Graph))
+						{
+							return false;
+						}
+					}
+				}
+			}
+			return true;
+		}
+	);
+}
+
+void FKismetDebugUtilities::RemoveEmptySettings(const FString& BlueprintPath)
+{
+	UBlueprintEditorSettings* Settings = GetMutableDefault<UBlueprintEditorSettings>();
+	check(Settings);
+	if(FPerBlueprintSettings *PerBlueprintSettings =
+		Settings->PerBlueprintSettings.Find(BlueprintPath))
+	{
+		// if all settings data is default, we can remove it from the map
+		if(*PerBlueprintSettings == FPerBlueprintSettings())
+		{
+			Settings->PerBlueprintSettings.Remove(BlueprintPath);
+		}
+		SaveBlueprintEditorSettings();
+	}
+}
+
 //////////////////////////////////////////////////////////////////////////
 // Breakpoint
 
@@ -869,7 +924,7 @@ void FKismetDebugUtilities::SetBreakpointEnabled(FBlueprintBreakpoint& Breakpoin
 
 	Breakpoint.bEnabled = bIsEnabled;
 	SetBreakpointInternal(Breakpoint, Breakpoint.bEnabled);
-	SaveBreakpoints();
+	SaveBlueprintEditorSettings();
 }
 
 // Set or clear the enabled flag for the breakpoint
@@ -1033,7 +1088,7 @@ void FKismetDebugUtilities::CreateBreakpoint(const UBlueprint* Blueprint, UEdGra
 	BlueprintSettings.Breakpoints.Emplace();
 	SetBreakpointEnabled(BlueprintSettings.Breakpoints.Top(), bIsEnabled);
 	SetBreakpointLocation(BlueprintSettings.Breakpoints.Top(), Node);
-	SaveBreakpoints();
+	SaveBlueprintEditorSettings();
 }
 
 void FKismetDebugUtilities::ForeachBreakpoint(const UBlueprint* Blueprint,
@@ -1070,7 +1125,7 @@ void FKismetDebugUtilities::RemoveBreakpointsByPredicate(const UBlueprint* Bluep
 				// keeps the ini file clean by removing empty arrays
 				ClearBreakpoints(Blueprint);
 			}
-			SaveBreakpoints();
+			SaveBlueprintEditorSettings();
 		}
 	}
 }
@@ -1182,19 +1237,22 @@ void FKismetDebugUtilities::ClearBreakpoints(const UBlueprint* Blueprint)
 
 
 void FKismetDebugUtilities::ClearBreakpointsForPath(const FString& BlueprintPath)
-{
+{	
 	UBlueprintEditorSettings* Settings = GetMutableDefault<UBlueprintEditorSettings>();
 	check(Settings);
-
-	FPerBlueprintSettings Collection;
-	if (Settings->PerBlueprintSettings.RemoveAndCopyValue(BlueprintPath, Collection))
+	if(FPerBlueprintSettings *PerBlueprintSettings =
+		Settings->PerBlueprintSettings.Find(BlueprintPath))
 	{
-		for(FBlueprintBreakpoint& Breakpoint : Collection.Breakpoints)
+		for(FBlueprintBreakpoint& Breakpoint : PerBlueprintSettings->Breakpoints)
 		{
 			// notify debugger that the breakpont has been removed
 			SetBreakpointLocation(Breakpoint, nullptr);	
 		}
-		SaveBreakpoints();
+		PerBlueprintSettings->Breakpoints.Empty();
+		
+		RemoveEmptySettings(BlueprintPath);
+        
+		SaveBlueprintEditorSettings();
 	}
 }
 
@@ -1249,43 +1307,139 @@ bool FKismetDebugUtilities::CanWatchPin(const UBlueprint* Blueprint, const UEdGr
 
 bool FKismetDebugUtilities::IsPinBeingWatched(const UBlueprint* Blueprint, const UEdGraphPin* Pin)
 {
-	return Blueprint->WatchedPins.Contains(const_cast<UEdGraphPin*>(Pin));
+	if(TArray<FEdGraphPinReference>* WatchedPins = GetWatchedPins(Blueprint))
+	{
+		return WatchedPins->Contains(const_cast<UEdGraphPin*>(Pin));
+	}
+	return false;
 }
 
-void FKismetDebugUtilities::RemovePinWatch(UBlueprint* Blueprint, const UEdGraphPin* Pin)
+bool FKismetDebugUtilities::RemovePinWatch(const UBlueprint* Blueprint, const UEdGraphPin* Pin)
 {
-	UEdGraphPin* NonConstPin = const_cast<UEdGraphPin*>(Pin);
-	Blueprint->WatchedPins.Remove(NonConstPin);
-	Blueprint->MarkPackageDirty();
-	Blueprint->PostEditChange();
-	WatchedPinsListChangedEvent.Broadcast(Blueprint);
+	return RemovePinWatchesByPredicate(
+		Blueprint,
+		[Pin](const UEdGraphPin* Other)
+		{
+			return Pin->PinId == Other->PinId;
+		}
+	);
 }
 
-void FKismetDebugUtilities::TogglePinWatch(UBlueprint* Blueprint, const UEdGraphPin* Pin)
+void FKismetDebugUtilities::AddPinWatch(const UBlueprint* Blueprint, const UEdGraphPin* Pin)
 {
-	int32 ExistingWatchIndex = Blueprint->WatchedPins.Find(const_cast<UEdGraphPin*>(Pin));
+	UBlueprintEditorSettings* Settings = GetMutableDefault<UBlueprintEditorSettings>();
+	check(Settings);
+	FPerBlueprintSettings &BlueprintSettings = Settings->PerBlueprintSettings.FindOrAdd(Blueprint->GetPathName());
+
+	// ensure that this node doesn't already contain a breakpoint
+	checkSlow(!BlueprintSettings.WatchedPins.ContainsByPredicate(
+		[Pin](const FEdGraphPinReference& Other)
+		{
+			return Other == Pin;
+		}
+	));
+	
+	BlueprintSettings.WatchedPins.Add(Pin);
+	SaveBlueprintEditorSettings();
+	WatchedPinsListChangedEvent.Broadcast(const_cast<UBlueprint*>(Blueprint));
+}
+
+
+void FKismetDebugUtilities::TogglePinWatch(const UBlueprint* Blueprint, const UEdGraphPin* Pin)
+{
+	int32 ExistingWatchIndex = INDEX_NONE;
+	
+	if(TArray<FEdGraphPinReference>* WatchedPins = GetWatchedPins(Blueprint))
+	{
+		ExistingWatchIndex = WatchedPins->Find(const_cast<UEdGraphPin*>(Pin));
+	}
 
 	if (ExistingWatchIndex != INDEX_NONE)
 	{
-		FKismetDebugUtilities::RemovePinWatch(Blueprint, Pin);
+		RemovePinWatch(Blueprint, Pin);
 	}
 	else
 	{
-		Blueprint->WatchedPins.Add(const_cast<UEdGraphPin*>(Pin));
-		Blueprint->MarkPackageDirty();
-		Blueprint->PostEditChange();
+		AddPinWatch(Blueprint, Pin);
 	}
+}
 
-	WatchedPinsListChangedEvent.Broadcast(Blueprint);
+void FKismetDebugUtilities::ClearPinWatches(const UBlueprint* Blueprint)
+{
+	UBlueprintEditorSettings* Settings = GetMutableDefault<UBlueprintEditorSettings>();
+	check(Settings);
+	if(FPerBlueprintSettings *PerBlueprintSettings =
+		Settings->PerBlueprintSettings.Find(Blueprint->GetPathName()))
+	{
+		PerBlueprintSettings->WatchedPins.Empty();
+		
+		RemoveEmptySettings(Blueprint->GetPathName());
+        		
+        SaveBlueprintEditorSettings();
+        WatchedPinsListChangedEvent.Broadcast(const_cast<UBlueprint*>(Blueprint));
+	}
 }
 
-void FKismetDebugUtilities::ClearPinWatches(UBlueprint* Blueprint)
+bool FKismetDebugUtilities::BlueprintHasPinWatches(const UBlueprint* Blueprint)
 {
-	Blueprint->WatchedPins.Empty();
-	Blueprint->MarkPackageDirty();
-	Blueprint->PostEditChange();
+	return GetWatchedPins(Blueprint) != nullptr;
+}
 
-	WatchedPinsListChangedEvent.Broadcast(Blueprint);
+void FKismetDebugUtilities::ForeachPinWatch(const UBlueprint* Blueprint, TFunctionRef<void(UEdGraphPin*)> Task)
+{
+	if(TArray<FEdGraphPinReference>* Pins = GetWatchedPins(Blueprint))
+	{
+		for(FEdGraphPinReference& PinRef : *Pins)
+		{
+			if(UEdGraphPin* Pin = PinRef.Get())
+			{
+				Task(Pin);
+			}
+		}
+	}
+}
+
+bool FKismetDebugUtilities::RemovePinWatchesByPredicate(const UBlueprint* Blueprint,
+	const TFunctionRef<bool(const UEdGraphPin*)> Predicate)
+{
+	auto ModifiedPedicate = [Predicate](FEdGraphPinReference& PinRef)
+	{
+		UEdGraphPin* Pin = PinRef.Get();
+		return Pin && Predicate(Pin);
+	};
+	
+	if(TArray<FEdGraphPinReference>* Pins = GetWatchedPins(Blueprint))
+	{
+		if(Pins->RemoveAllSwap(ModifiedPedicate, false))
+		{
+			if(Pins->IsEmpty())
+			{
+				// keeps the ini file clean by removing empty arrays
+				ClearPinWatches(Blueprint);
+			}
+			SaveBlueprintEditorSettings();
+			WatchedPinsListChangedEvent.Broadcast(const_cast<UBlueprint*>(Blueprint));
+			return true;
+		}
+	}
+	return false;
+}
+
+UEdGraphPin* FKismetDebugUtilities::FindPinWatchByPredicate(const UBlueprint* Blueprint,
+	const TFunctionRef<bool(const UEdGraphPin*)> Predicate)
+{
+	if(TArray<FEdGraphPinReference>* Pins = GetWatchedPins(Blueprint))
+	{
+		for(FEdGraphPinReference& PinRef : *Pins)
+		{
+			UEdGraphPin* Pin = PinRef.Get();
+			if(Pin && Predicate(Pin))
+			{
+				return Pin;
+			}
+		}
+	}
+	return nullptr;
 }
 
 // Gets the watched tooltip for a specified site
diff --git a/Engine/Source/Editor/UnrealEd/Public/Kismet2/KismetDebugUtilities.h b/Engine/Source/Editor/UnrealEd/Public/Kismet2/KismetDebugUtilities.h
--- a/Engine/Source/Editor/UnrealEd/Public/Kismet2/KismetDebugUtilities.h
+++ b/Engine/Source/Editor/UnrealEd/Public/Kismet2/KismetDebugUtilities.h
@@ -288,9 +288,32 @@ public:
 
 	static bool CanWatchPin(const UBlueprint* Blueprint, const UEdGraphPin* Pin);
 	static bool IsPinBeingWatched(const UBlueprint* Blueprint, const UEdGraphPin* Pin);
-	static void TogglePinWatch(UBlueprint* Blueprint, const UEdGraphPin* Pin);
-	static void RemovePinWatch(UBlueprint* Blueprint, const UEdGraphPin* Pin);
-	static void ClearPinWatches(UBlueprint* Blueprint);
+	static void TogglePinWatch(const UBlueprint* Blueprint, const UEdGraphPin* Pin);
+	static bool RemovePinWatch(const UBlueprint* Blueprint, const UEdGraphPin* Pin);
+	static void AddPinWatch(const UBlueprint* Blueprint, const UEdGraphPin* Pin);
+	static void ClearPinWatches(const UBlueprint* Blueprint);
+	static bool BlueprintHasPinWatches(const UBlueprint* Blueprint);
+
+	/**
+	* Performs a task on every breakpoint in the provided blueprint
+	* @param Blueprint The owning blueprint of the breakpoints to iterate
+	* @param Task function to be called on every element
+	*/
+	static void ForeachPinWatch(const UBlueprint* Blueprint, TFunctionRef<void(UEdGraphPin*)> Task);
+
+	/**
+	* Removes any breakpoint that matches the provided predicate
+	* @param Blueprint The owning blueprint of the breakpoints to iterate
+	* @param Predicate function that returns true if a breakpoint should be removed
+	*/
+	static bool RemovePinWatchesByPredicate(const UBlueprint* Blueprint, const TFunctionRef<bool(const UEdGraphPin*)> Predicate);
+
+	/**
+	* Returns the first breakpoint that matches the provided predicate or nullptr if nothing matched
+	* @param Blueprint The owning blueprint of the breakpoints to iterate
+	* @param Predicate function that returns true for the found breakpoint
+	*/
+	static UEdGraphPin* FindPinWatchByPredicate(const UBlueprint* Blueprint, const TFunctionRef<bool(const UEdGraphPin*)> Predicate);
 
 	enum EWatchTextResult
 	{
@@ -341,16 +364,25 @@ protected:
 	* @return	EWTR_Valid if the debug data could be found, otherwise an appropriate error code
 	*/
 	static EWatchTextResult FindDebuggingData(UBlueprint* Blueprint, UObject* ActiveObject, const UEdGraphPin* WatchPin, FProperty*& OutProperty, void*& OutData, void*& OutDelta, UObject*& OutParent, TArray<UObject*>& SeenObjects, bool* bOutIsDirectPtr = nullptr);
-	
+
+	/**	Retrieve the user settings associated with a blueprint.
+	*	returns null if the blueprint has default settings (no breakpoints and no watches) */
+	static struct FPerBlueprintSettings* GetPerBlueprintSettings(const UBlueprint* Blueprint);
 
 	/**	Retrieve the Array of breakpoints associated with a blueprint.
 	*	returns null if there are no breakpoints associated with this blueprint */
 	static TArray<FBlueprintBreakpoint>* GetBreakpoints(const UBlueprint* Blueprint);
+
+	/**	Retrieve the Array of breakpoints associated with a blueprint.
+	*	returns null if there are no breakpoints associated with this blueprint */
+	static TArray<FEdGraphPinReference>* GetWatchedPins(const UBlueprint* Blueprint);
 	
 	/** Save any modifications made to breakpoints */
-	static void SaveBreakpoints();
+	static void SaveBlueprintEditorSettings();
 
 	static void CleanupBreakpoints(const UBlueprint* Blueprint);
+	static void CleanupWatches(const UBlueprint* Blueprint);
+	static void RemoveEmptySettings(const FString& BlueprintPath);
 private:
 	FKismetDebugUtilities() {}
 };
diff --git a/Engine/Source/Runtime/Engine/Classes/EdGraph/EdGraphPin.h b/Engine/Source/Runtime/Engine/Classes/EdGraph/EdGraphPin.h
--- a/Engine/Source/Runtime/Engine/Classes/EdGraph/EdGraphPin.h
+++ b/Engine/Source/Runtime/Engine/Classes/EdGraph/EdGraphPin.h
@@ -278,7 +278,7 @@ struct ENGINE_API FEdGraphPinReference
 private:
 	/** The node that owns the pin referred to by this struct. Updated at Set and Save time. */
 	UPROPERTY()
-	TWeakObjectPtr<UEdGraphNode> OwningNode;
+	TSoftObjectPtr<UEdGraphNode> OwningNode;
 
 	/** The pin's unique ID. Updated at Set and Save time. */
 	UPROPERTY()
diff --git a/Engine/Source/Runtime/Engine/Classes/Engine/Blueprint.h b/Engine/Source/Runtime/Engine/Classes/Engine/Blueprint.h
--- a/Engine/Source/Runtime/Engine/Classes/Engine/Blueprint.h
+++ b/Engine/Source/Runtime/Engine/Classes/Engine/Blueprint.h
@@ -638,11 +638,13 @@ PRAGMA_ENABLE_DEPRECATION_WARNINGS
 	UPROPERTY()
 	TArray<FBPEditorBookmarkNode> BookmarkNodes;
 
+	// moved to FPerBlueprintSettings
 	UPROPERTY()
 	TArray<TObjectPtr<class UDEPRECATED_Breakpoint>> Breakpoints_DEPRECATED;
 
+	// moved to FPerBlueprintSettings
 	UPROPERTY()
-	TArray<FEdGraphPinReference> WatchedPins;
+	TArray<FEdGraphPinReference> WatchedPins_DEPRECATED;
 
 	UPROPERTY()
 	TArray<TObjectPtr<class UEdGraphPin_Deprecated>> DeprecatedPinWatches;
