Eliminated unnecessary decompression step in local executor - we can just use the hash provided by FCompressedBuffer instead

Also added some comments

diff --git a/Engine/Source/Editor/UnrealEd/Private/DerivedDataBuildLocalExecutor.cpp b/Engine/Source/Editor/UnrealEd/Private/DerivedDataBuildLocalExecutor.cpp
--- a/Engine/Source/Editor/UnrealEd/Private/DerivedDataBuildLocalExecutor.cpp
+++ b/Engine/Source/Editor/UnrealEd/Private/DerivedDataBuildLocalExecutor.cpp
@@ -78,7 +78,12 @@ public:
 		}
 
 		// This path will execute the build action synchronously in a scratch directory
-		// at this stage, all inputs are available in process
+		// 
+		// At this stage, all inputs are available in process
+		//
+		// Currently no cleanup whatsoever is performed, so inputs/outputs can be inspected
+		// this could be problematic for large runs so we should probably add support for
+		// configurable cleanup policies
 
 		static std::atomic<int32> SerialNo = 0;
 
@@ -133,26 +138,20 @@ public:
 
 		// Manifest inputs in scratch area
 
+		if (!Inputs.IsNull())
 		{
-			if (!Inputs.IsNull())
-			{
-				Inputs.Get().IterateInputs([&](FStringView Key, const FCompressedBuffer& Buffer) {
-					// We should be able to accomplish this without decompression
-					FCompositeBuffer DecompressedComposite = Buffer.DecompressToComposite();
-					FIoHash RawHash = FIoHash::HashBuffer(DecompressedComposite);
-
-					TStringBuilder<128> InputPath;
-					InputPath << TEXT("Inputs/") << RawHash;
+			Inputs.Get().IterateInputs([&](FStringView Key, const FCompressedBuffer& Buffer) {
+				TStringBuilder<128> InputPath;
+				InputPath << TEXT("Inputs/") << Buffer.GetRawHash();
 
-					FString Path{ SandboxRoot / *InputPath };
+				FString Path{ SandboxRoot / *InputPath };
 
-					if (TUniquePtr<FArchive> Ar{ IFileManager::Get().CreateFileWriter(*Path, FILEWRITE_Silent) })
-					{
-						FCompressedBuffer Comp = FCompressedBuffer::FromCompressed(Buffer.GetCompressed());
-						*Ar << Comp;
-					}
-				});
-			}
+				if (TUniquePtr<FArchive> Ar{ IFileManager::Get().CreateFileWriter(*Path, FILEWRITE_Silent) })
+				{
+					FCompressedBuffer Comp = FCompressedBuffer::FromCompressed(Buffer.GetCompressed());
+					*Ar << Comp;
+				}
+			});
 		}
 
 		// Serialize action specification
