Ensure that the status of the FindMissingBlobs/BatchUploadBlobs/BatchReadBlobs operations is explicitly checked before allowing progress to the subsequent remote execution steps.

diff --git a/Engine/Source/Editor/UnrealEd/Private/DerivedDataBuildRemoteExecutor.cpp b/Engine/Source/Editor/UnrealEd/Private/DerivedDataBuildRemoteExecutor.cpp
--- a/Engine/Source/Editor/UnrealEd/Private/DerivedDataBuildRemoteExecutor.cpp
+++ b/Engine/Source/Editor/UnrealEd/Private/DerivedDataBuildRemoteExecutor.cpp
@@ -117,6 +117,7 @@ private:
 
 	struct FMerkleTreeFileBuilder
 	{
+		FStringView Path;
 		FFileNode File;
 		EFileType Type;
 		FCompositeBuffer ContentBytes;
@@ -125,6 +126,7 @@ private:
 	struct FMerkleTreeDirectoryBuilder
 	{
 		FStringView Name;
+		FStringView Path;
 		FDirectory Directory;
 		TOptional<FDigest> Digest;
 		TArray<uint8> ContentBytes;
@@ -193,6 +195,7 @@ private:
 	bool ValidateUploadSuccess(const FBatchUpdateBlobsResponse& BatchUpdateBlobsResponse);
 	FOptionalBuildOutput ComposeBuildOutput(FBatchReadBlobsResponse& BatchReadBlobsResponse, EStatus& OutStatus);
 	bool ProcessCancellation();
+	bool IsStatusOk(FStatus Status, const TCHAR* OperationDesc);
 
 	// Async steps
 	TFuture<TPair<FStatus, FFindMissingBlobsResponse>> DetermineMissingBlobsAsync();
@@ -202,10 +205,10 @@ private:
 	TFuture<TPair<FStatus, FBatchReadBlobsResponse>> DownloadResultsAsync();
 
 	// Post-step flow
-	void OnMissingBlobsDetermined(FFindMissingBlobsResponse&& Result);
-	void OnMissingBlobsUploaded(const FBatchUpdateBlobsResponse& Result);
+	void OnMissingBlobsDetermined(FStatus Status, FFindMissingBlobsResponse&& Result);
+	void OnMissingBlobsUploaded(FStatus Status, const FBatchUpdateBlobsResponse& Result);
 	void OnExecutionCompleted(FExecuteResponse&& Result);
-	void OnOutputBlobsDownloaded(FBatchReadBlobsResponse&& Result);
+	void OnOutputBlobsDownloaded(FStatus Status, FBatchReadBlobsResponse&& Result);
 };
 
 class FRemoteBuildWorkerExecutor final: public IBuildWorkerExecutor
@@ -484,7 +487,7 @@ FRemoteBuildExecutionRequest::FRemoteBuildExecutionRequest(
 {
 	Owner.Begin(this);
 	DetermineMissingBlobsAsync()
-		.Next([this] (TPair<FStatus, FFindMissingBlobsResponse>&& Result) { OnMissingBlobsDetermined(MoveTemp(Result.Value)); });
+		.Next([this] (TPair<FStatus, FFindMissingBlobsResponse>&& Result) { OnMissingBlobsDetermined(Result.Key, MoveTemp(Result.Value)); });
 }
 
 FRemoteBuildExecutionRequest::~FRemoteBuildExecutionRequest()
@@ -502,6 +505,7 @@ FRemoteBuildExecutionRequest::FMerkleTreeDirectoryBuilder& FRemoteBuildExecution
 	{
 		DirectoryBuilderIndex = State.Directories.Num();
 		FMerkleTreeDirectoryBuilder& NewNode = State.Directories.AddDefaulted_GetRef();
+		NewNode.Path = Path;
 		NewNode.Name = Path.IsEmpty() ? Path : FPathViews::GetCleanFilename(Path);
 	}
 
@@ -534,6 +538,7 @@ int32 FRemoteBuildExecutionRequest::AddMerkleTreeFile(FStringView Path, const FI
 	int32 NewFileIndex = State.Files.Num();
 	State.DigestFilesystemIndex.Add(NewNode.Digest, FVariantIndex(ENodeType::File, NewFileIndex));
 	FMerkleTreeFileBuilder& FileBuilder = State.Files.AddDefaulted_GetRef();
+	FileBuilder.Path = Path;
 	FileBuilder.File = NewNode; // Duplicates the node in the state's file array
 	FileBuilder.Type = FileType;
 	if (ContentBytes)
@@ -817,16 +822,48 @@ bool FRemoteBuildExecutionRequest::ProcessCancellation()
 	return false;
 }
 
+bool FRemoteBuildExecutionRequest::IsStatusOk(FStatus Status, const TCHAR* OperationDesc)
+{
+	if (!Status.Ok())
+	{
+		UE_LOG(LogDerivedDataBuildRemoteExecutor, Warning, TEXT("Remote execution system error: operation '%s' produced an error result (%d)!"), OperationDesc, (int)Status.Code);
+
+		State.Owner.End(this, [this]
+		{
+			CompletionCallback({State.BuildAction.GetKey(), {}, {}, EStatus::Error});
+			CompletionEvent->Trigger();
+		});
+		return false;
+	}
+	return true;
+}
+
 TFuture<TPair<FStatus, FFindMissingBlobsResponse>> FRemoteBuildExecutionRequest::DetermineMissingBlobsAsync()
 {
 	BuildMerkleTreeNodes();
 
 	State.FindMissingBlobsRequest.InstanceName = Executor.InstanceName;
 	State.FindMissingBlobsRequest.BlobDigests.Add(State.ActionDigest);
+	UE_LOG(LogDerivedDataBuildRemoteExecutor, Verbose, TEXT("Checking CAS presence of action (hash: %s) of size %d."), *::LexToString(State.ActionDigest.Hash), State.ActionContentBytes.Num());
+
 	State.FindMissingBlobsRequest.BlobDigests.Add(State.Action.CommandDigest);
+	UE_LOG(LogDerivedDataBuildRemoteExecutor, Verbose, TEXT("Checking CAS presence of command (hash: %s) of size %d."), *::LexToString(State.Action.CommandDigest.Hash), State.CommandContentBytes.Num());
+
 	for (const TPair<FDigest, FVariantIndex>& FilesystemItem : State.DigestFilesystemIndex)
 	{
 		State.FindMissingBlobsRequest.BlobDigests.AddUnique(FilesystemItem.Key);
+		switch (FilesystemItem.Value.NodeType)
+		{
+		case ENodeType::Directory:
+			UE_LOG(LogDerivedDataBuildRemoteExecutor, Verbose, TEXT("Checking CAS presence of directory '%s' (hash: %s) of size %d."), *FString(State.Directories[FilesystemItem.Value.Index].Path), *::LexToString(FilesystemItem.Key.Hash), State.Directories[FilesystemItem.Value.Index].ContentBytes.Num());
+			break;
+		case ENodeType::File:
+			UE_LOG(LogDerivedDataBuildRemoteExecutor, Verbose, TEXT("Checking CAS presence of file '%s' (hash: %s, type: %s) of size %d."), *FString(State.Files[FilesystemItem.Value.Index].Path), *::LexToString(FilesystemItem.Key.Hash), LexToString(State.Files[FilesystemItem.Value.Index].Type), FilesystemItem.Key.SizeBytes);
+			break;
+		default:
+			checkNoEntry();
+			break;
+		}
 	}
 	return Executor.ContentAddressableStorage->FindMissingBlobsAsync(State.FindMissingBlobsRequest);
 }
@@ -906,7 +943,7 @@ TFuture<TPair<FStatus, FBatchUpdateBlobsResponse>> FRemoteBuildExecutionRequest:
 				Executor.Stats.TotalDirectoryBlobsUploaded.AddBlob(State.Directories[VariantIndex.Index].ContentBytes.Num());
 				NewRequest.Digest = State.Directories[VariantIndex.Index].Digest.GetValue();
 				NewRequest.Data = MoveTemp(State.Directories[VariantIndex.Index].ContentBytes);
-				UE_LOG(LogDerivedDataBuildRemoteExecutor, Verbose, TEXT("Uploading directory '%s' (hash: %s) of upload size %d."), *FString(State.Directories[VariantIndex.Index].Name), *::LexToString(NewRequest.Digest.Hash), NewRequest.Data.Num());
+				UE_LOG(LogDerivedDataBuildRemoteExecutor, Verbose, TEXT("Uploading directory '%s' (hash: %s) of upload size %d."), *FString(State.Directories[VariantIndex.Index].Path), *::LexToString(NewRequest.Digest.Hash), NewRequest.Data.Num());
 				break;
 			case ENodeType::File:
 				{
@@ -921,7 +958,7 @@ TFuture<TPair<FStatus, FBatchUpdateBlobsResponse>> FRemoteBuildExecutionRequest:
 						NewRequest.Data.Append((const uint8 *)Segment.GetData(), Segment.GetSize());
 					}
 					FileBuffer.Reset();
-					UE_LOG(LogDerivedDataBuildRemoteExecutor, Verbose, TEXT("Uploading file '%s' (hash: %s, type: %s) of upload size %d."), *State.Files[VariantIndex.Index].File.Name, *::LexToString(NewRequest.Digest.Hash), LexToString(State.Files[VariantIndex.Index].Type), NewRequest.Data.Num());
+					UE_LOG(LogDerivedDataBuildRemoteExecutor, Verbose, TEXT("Uploading file '%s' (hash: %s, type: %s) of upload size %d."), *FString(State.Files[VariantIndex.Index].Path), *::LexToString(NewRequest.Digest.Hash), LexToString(State.Files[VariantIndex.Index].Type), NewRequest.Data.Num());
 				}
 				break;
 			default:
@@ -956,9 +993,9 @@ TFuture<TPair<FStatus, FBatchReadBlobsResponse>> FRemoteBuildExecutionRequest::D
 	return Executor.ContentAddressableStorage->BatchReadBlobsAsync(State.BatchReadBlobsRequest);
 }
 
-void FRemoteBuildExecutionRequest::OnMissingBlobsDetermined(FFindMissingBlobsResponse&& Result)
+void FRemoteBuildExecutionRequest::OnMissingBlobsDetermined(FStatus Status, FFindMissingBlobsResponse&& Result)
 {
-	if (ProcessCancellation())
+	if (ProcessCancellation() || !IsStatusOk(Status, TEXT("FindMissingBlobs")))
 	{
 		return;
 	}
@@ -1000,7 +1037,7 @@ void FRemoteBuildExecutionRequest::OnMissingBlobsDetermined(FFindMissingBlobsRes
 	if (!State.FindMissingBlobsResponse.MissingBlobDigests.IsEmpty())
 	{
 		UploadMissingBlobsAsync()
-			.Next([this] (const TPair<FStatus, FBatchUpdateBlobsResponse>& InnerResult) { OnMissingBlobsUploaded(InnerResult.Value); });
+			.Next([this] (const TPair<FStatus, FBatchUpdateBlobsResponse>& InnerResult) { OnMissingBlobsUploaded(InnerResult.Key, InnerResult.Value); });
 	}
 	else
 	{
@@ -1009,9 +1046,9 @@ void FRemoteBuildExecutionRequest::OnMissingBlobsDetermined(FFindMissingBlobsRes
 	}
 }
 
-void FRemoteBuildExecutionRequest::OnMissingBlobsUploaded(const FBatchUpdateBlobsResponse& Result)
+void FRemoteBuildExecutionRequest::OnMissingBlobsUploaded(FStatus Status, const FBatchUpdateBlobsResponse& Result)
 {
-	if (ProcessCancellation())
+	if (ProcessCancellation() || !IsStatusOk(Status, TEXT("BatchUploadBlobs")))
 	{
 		return;
 	}
@@ -1045,7 +1082,7 @@ void FRemoteBuildExecutionRequest::OnExecutionCompleted(FExecuteResponse&& Resul
 		DownloadResultsAsync()
 			.Next([this] (TPair<FStatus, FBatchReadBlobsResponse>&& InnerResult)
 				{
-					OnOutputBlobsDownloaded(MoveTemp(InnerResult.Value));
+					OnOutputBlobsDownloaded(InnerResult.Key, MoveTemp(InnerResult.Value));
 				});
 	}
 	else
@@ -1060,9 +1097,9 @@ void FRemoteBuildExecutionRequest::OnExecutionCompleted(FExecuteResponse&& Resul
 
 }
 
-void FRemoteBuildExecutionRequest::OnOutputBlobsDownloaded(FBatchReadBlobsResponse&& Result)
+void FRemoteBuildExecutionRequest::OnOutputBlobsDownloaded(FStatus Status, FBatchReadBlobsResponse&& Result)
 {
-	if (ProcessCancellation())
+	if (ProcessCancellation() || !IsStatusOk(Status, TEXT("BatchReadBlobs")))
 	{
 		return;
 	}
