Add histogram and alphabetic list to the shader stats commandlet.

- Allows an easier comparison between two builds looking for changes in what was added/removed.

#[review] [REDACTED_HANDLE]

diff --git a/Engine/Source/Editor/UnrealEd/Private/Commandlets/GatherShaderStatsFromMaterialsCommandlet.cpp b/Engine/Source/Editor/UnrealEd/Private/Commandlets/GatherShaderStatsFromMaterialsCommandlet.cpp
--- a/Engine/Source/Editor/UnrealEd/Private/Commandlets/GatherShaderStatsFromMaterialsCommandlet.cpp
+++ b/Engine/Source/Editor/UnrealEd/Private/Commandlets/GatherShaderStatsFromMaterialsCommandlet.cpp
@@ -15,20 +15,99 @@
 
 DEFINE_LOG_CATEGORY_STATIC(LogGatherShaderStatsFromMaterialsCommandlet, Log, All);
 
-class FLogger
+class FShaderStatsGatheringContext
 {
 public:
-	FLogger() = delete;
-	FLogger(const FString& FileName)
+	FShaderStatsGatheringContext() = delete;
+	FShaderStatsGatheringContext(const FString& FileName)
 	{
 		DebugWriter = IFileManager::Get().CreateFileWriter(*FileName);
 	}
-	~FLogger()
+	~FShaderStatsGatheringContext()
 	{
 		DebugWriter->Close();
 		delete DebugWriter;
 	}
 
+	void AddToHistogram(const TCHAR* VertexFactoryName, const TCHAR* ShaderPipelineName, const TCHAR* ShaderTypeName)
+	{
+		FString ShaderType(ShaderTypeName);
+		if (int32* Existing = ShaderTypeHistogram.Find(ShaderType))
+		{
+			++(*Existing);
+		}
+		else
+		{
+			ShaderTypeHistogram.FindOrAdd(ShaderType, 1);
+		}
+
+		FString AbsoluteShaderName = (ShaderPipelineName != nullptr) ? FString::Printf(TEXT("%s.%s.%s"), VertexFactoryName, ShaderPipelineName, ShaderTypeName) : FString::Printf(TEXT("%s.%s"), VertexFactoryName, ShaderTypeName);
+		if (int32* Existing = FullShaderTypeHistogram.Find(AbsoluteShaderName))
+		{
+			++(*Existing);
+		}
+		else
+		{
+			FullShaderTypeHistogram.FindOrAdd(AbsoluteShaderName, 1);
+		}
+	}
+
+	void PrintHistogram()
+	{
+		if (ShaderTypeHistogram.Num() > 0)
+		{
+			ShaderTypeHistogram.ValueSort(TGreater<int32>());
+			const char ShaderTypeHeader[] = "\nShaderType, Count\n";
+			DebugWriter->Serialize(const_cast<char*>(ShaderTypeHeader), sizeof(ShaderTypeHeader) - 1);
+			for (TPair<FString, int32> ShaderUsage : ShaderTypeHistogram)
+			{
+				FString OuputLine = FString::Printf(TEXT("%s, %d\n"), *ShaderUsage.Key, ShaderUsage.Value);
+				DebugWriter->Serialize(const_cast<ANSICHAR*>(StringCast<ANSICHAR>(*OuputLine).Get()), OuputLine.Len());
+			}
+		}
+
+		if (FullShaderTypeHistogram.Num() > 0)
+		{
+			FullShaderTypeHistogram.ValueSort(TGreater<int32>());
+			const char FullShaderTypeHeader[] = "\nFullShaderType, Count\n";
+			DebugWriter->Serialize(const_cast<char*>(FullShaderTypeHeader), sizeof(FullShaderTypeHeader) - 1);
+			for (TPair<FString, int32> ShaderUsage : FullShaderTypeHistogram)
+			{
+				FString OuputLine = FString::Printf(TEXT("%s, %d\n"), *ShaderUsage.Key, ShaderUsage.Value);
+				DebugWriter->Serialize(const_cast<ANSICHAR*>(StringCast<ANSICHAR>(*OuputLine).Get()), OuputLine.Len());
+			}
+		}
+	}
+
+	void PrintAlphabeticList()
+	{
+		if (ShaderTypeHistogram.Num() > 0)
+		{
+			ShaderTypeHistogram.KeySort(TGreater<FString>());
+			const char ShaderTypeAlphabeticHeader[] = "\nShaderType only\n";
+			DebugWriter->Serialize(const_cast<char*>(ShaderTypeAlphabeticHeader), sizeof(ShaderTypeAlphabeticHeader) - 1);
+			for (TPair<FString, int32> ShaderUsage : ShaderTypeHistogram)
+			{
+				// do not print numbers here as it complicates the diff
+				FString OuputLine = FString::Printf(TEXT("%s\n"), *ShaderUsage.Key);
+				DebugWriter->Serialize(const_cast<ANSICHAR*>(StringCast<ANSICHAR>(*OuputLine).Get()), OuputLine.Len());
+			}
+		}
+
+		if (FullShaderTypeHistogram.Num() > 0)
+		{
+			FullShaderTypeHistogram.KeySort(TGreater<FString>());
+			const char FullShaderTypeAlphabeticHeader[] = "\nFullShaderType only\n";
+			DebugWriter->Serialize(const_cast<char*>(FullShaderTypeAlphabeticHeader), sizeof(FullShaderTypeAlphabeticHeader) - 1);
+			for (TPair<FString, int32> ShaderUsage : FullShaderTypeHistogram)
+			{
+				// do not print numbers here as it complicates the diff
+				FString OuputLine = FString::Printf(TEXT("%s\n"), *ShaderUsage.Key);
+				DebugWriter->Serialize(const_cast<ANSICHAR*>(StringCast<ANSICHAR>(*OuputLine).Get()), OuputLine.Len());
+			}
+		}
+	}
+
 	void Log(const FString& OutString)
 	{
 		//UE_LOG(LogGatherShaderStatsFromMaterialsCommandlet, Log, TEXT("%s"), *OutString);
@@ -38,6 +117,12 @@ public:
 
 private:
 	FArchive* DebugWriter = nullptr;
+
+	/** Map of shader type names (no matter the vertex factory) to their counts. */
+	TMap<FString, int32>	ShaderTypeHistogram;
+
+	/** Map of full shader display names to their counts. */
+	TMap<FString, int32>	FullShaderTypeHistogram;
 };
 
 UGatherShaderStatsFromMaterialsCommandlet::UGatherShaderStatsFromMaterialsCommandlet(const FObjectInitializer& ObjectInitializer)
@@ -60,7 +145,7 @@ int GetTotalShaders(const TArray<FDebugShaderTypeInfo>& OutShaderInfo)
 	return TotalShadersForMaterial;
 }
 
-void PrintDebugShaderInfo(FLogger& Output, const TArray<FDebugShaderTypeInfo>& OutShaderInfo)
+void PrintDebugShaderInfo(FShaderStatsGatheringContext& Output, const TArray<FDebugShaderTypeInfo>& OutShaderInfo)
 {
 	for (const FDebugShaderTypeInfo& ShaderInfo : OutShaderInfo)
 	{
@@ -78,6 +163,7 @@ void PrintDebugShaderInfo(FLogger& Output, const TArray<FDebugShaderTypeInfo>& O
 		for (FShaderType* ShaderType : ShaderInfo.ShaderTypes)
 		{
 			Output.Log(FString::Printf(TEXT("\t\t%s"), ShaderType->GetName()));
+			Output.AddToHistogram(ShaderInfo.VFType->GetName(), nullptr, ShaderType->GetName());
 		}
 
 		for (const FDebugShaderPipelineInfo& PipelineInfo : ShaderInfo.Pipelines)
@@ -87,6 +173,7 @@ void PrintDebugShaderInfo(FLogger& Output, const TArray<FDebugShaderTypeInfo>& O
 			for (FShaderType* ShaderType : PipelineInfo.ShaderTypes)
 			{
 				Output.Log(FString::Printf(TEXT("\t\t\t%s"), ShaderType->GetName()));
+				Output.AddToHistogram(ShaderInfo.VFType->GetName(), PipelineInfo.Pipeline->GetName(), ShaderType->GetName());
 			}
 		}
 
@@ -94,7 +181,7 @@ void PrintDebugShaderInfo(FLogger& Output, const TArray<FDebugShaderTypeInfo>& O
 	}
 }
 
-int ProcessMaterials(const EShaderPlatform ShaderPlatform, FLogger& Output, TArray<FAssetData>& MaterialList)
+int ProcessMaterials(const EShaderPlatform ShaderPlatform, FShaderStatsGatheringContext& Output, TArray<FAssetData>& MaterialList)
 {
 	int TotalShaders = 0;
 
@@ -123,7 +210,7 @@ int ProcessMaterials(const EShaderPlatform ShaderPlatform, FLogger& Output, TArr
 	return TotalShaders;
 }
 
-int ProcessMaterialInstances(const EShaderPlatform ShaderPlatform, FLogger& Output, TArray<FAssetData>& MaterialInstanceList)
+int ProcessMaterialInstances(const EShaderPlatform ShaderPlatform, FShaderStatsGatheringContext& Output, TArray<FAssetData>& MaterialInstanceList)
 {
 	int TotalShaders = 0;
 
@@ -266,7 +353,7 @@ int32 UGatherShaderStatsFromMaterialsCommandlet::Main(const FString& Params)
 	const FString TimeNow = FDateTime::Now().ToString();
 	FString FileName = FPaths::Combine(*FPaths::ProjectSavedDir(), FString::Printf(TEXT("MaterialStats/ShaderStatsFromMaterials-%s.txt"), *TimeNow));
 
-	FLogger Output(FileName);
+	FShaderStatsGatheringContext Output(FileName);
 
 	int TotalShaders = 0;
 	int TotalAssets = 0;
@@ -287,10 +374,14 @@ int32 UGatherShaderStatsFromMaterialsCommandlet::Main(const FString& Params)
 	Output.Log(TEXT("Summary"));
 	Output.Log(FString::Printf(TEXT("Total Assets: %d"), TotalAssets));
 	Output.Log(FString::Printf(TEXT("Total Shaders: %d"), TotalShaders));
+	Output.Log(FString::Printf(TEXT("Histogram:")));
+	Output.PrintHistogram();
+	Output.Log(FString::Printf(TEXT("Alphabetic:")));
+	Output.PrintAlphabeticList();
 
 	const double EndTime = FPlatformTime::Seconds() - StartTime;
 	Output.Log(TEXT(""));
 	Output.Log(FString::Printf(TEXT("Commandlet Took: %lf"), EndTime));
 
-	return true;
+	return 0;
 }
