Deprecate IAssetEditorInstance::GetToolbarTabId and implementations - the toolbar tab is no longer registered by asset editors

diff --git a/Engine/Source/Editor/UnrealEd/Private/MiniCurveEditor.cpp b/Engine/Source/Editor/UnrealEd/Private/MiniCurveEditor.cpp
--- a/Engine/Source/Editor/UnrealEd/Private/MiniCurveEditor.cpp
+++ b/Engine/Source/Editor/UnrealEd/Private/MiniCurveEditor.cpp
@@ -79,12 +79,6 @@ bool SMiniCurveEditor::CloseWindow()
 	return true;
 }
 
-FName SMiniCurveEditor::GetToolbarTabId() const
-{
-	//@TODO: This editor should probably derive from FAssetEditorToolkit instead!
-	return NAME_None;
-}
-
 TSharedPtr<class FTabManager> SMiniCurveEditor::GetAssociatedTabManager()
 {
 	//@TODO: This editor should probably derive from FAssetEditorToolkit instead!
diff --git a/Engine/Source/Editor/UnrealEd/Private/Toolkits/AssetEditorToolkit.cpp b/Engine/Source/Editor/UnrealEd/Private/Toolkits/AssetEditorToolkit.cpp
--- a/Engine/Source/Editor/UnrealEd/Private/Toolkits/AssetEditorToolkit.cpp
+++ b/Engine/Source/Editor/UnrealEd/Private/Toolkits/AssetEditorToolkit.cpp
@@ -42,8 +42,6 @@ TWeakPtr< IToolkitHost > FAssetEditorToolkit::PreviousWorldCentricToolkitHostFor
 TSharedPtr<FExtensibilityManager> FAssetEditorToolkit::SharedMenuExtensibilityManager;
 TSharedPtr<FExtensibilityManager> FAssetEditorToolkit::SharedToolBarExtensibilityManager;
 
-const FName FAssetEditorToolkit::ToolbarTabId( TEXT( "AssetEditorToolkit_Toolbar" ) );
-
 const FName FAssetEditorToolkit::DefaultAssetEditorToolBarName("AssetEditor.DefaultToolBar");
 
 FAssetEditorToolkit::FAssetEditorToolkit()
@@ -327,7 +325,6 @@ void FAssetEditorToolkit::RegisterTabSpawners(const TSharedRef<class FTabManager
 
 void FAssetEditorToolkit::UnregisterTabSpawners(const TSharedRef<class FTabManager>& InTabManager)
 {
-	InTabManager->UnregisterTabSpawner( ToolbarTabId );
 	InTabManager->ClearLocalWorkspaceMenuCategories();
 }
 
@@ -918,29 +915,6 @@ bool FAssetEditorToolkit::ShouldPromptForNewFilesOnReload(const UObject& Editing
 	return true;
 }
 
-TSharedRef<SDockTab> FAssetEditorToolkit::SpawnTab_Toolbar( const FSpawnTabArgs& Args )
-{
-	check( Args.GetTabId() == ToolbarTabId );
-
-	TSharedRef<SDockTab> DockTab = SNew(SDockTab)
-		.Label( NSLOCTEXT("AssetEditorToolkit", "Toolbar_TabTitle", "Toolbar") )
-		.ShouldAutosize(true)
-		[
-			SAssignNew(ToolbarWidgetContent, SBorder)
-			.Padding(0)
-			.BorderImage(FEditorStyle::GetBrush("NoBorder"))
-		];
-
-	if (Toolbar.IsValid())
-	{
-		ToolbarWidgetContent->SetContent(Toolbar.ToSharedRef());
-	}
-
-	return DockTab;
-}
-
-
-
 void FAssetEditorToolkit::FillDefaultFileMenuCommands(FToolMenuSection& InSection)
 {
 	const FToolMenuInsert InsertPosition(NAME_None, EToolMenuInsertType::First);
diff --git a/Engine/Source/Editor/UnrealEd/Private/Tools/UAssetEditor.cpp b/Engine/Source/Editor/UnrealEd/Private/Tools/UAssetEditor.cpp
--- a/Engine/Source/Editor/UnrealEd/Private/Tools/UAssetEditor.cpp
+++ b/Engine/Source/Editor/UnrealEd/Private/Tools/UAssetEditor.cpp
@@ -38,11 +38,6 @@ bool UAssetEditor::CloseWindow()
 	return true;
 }
 
-FName UAssetEditor::GetToolbarTabId() const
-{
-	return ToolkitInstance->GetToolbarTabId();
-}
-
 TSharedPtr<class FTabManager> UAssetEditor::GetAssociatedTabManager()
 {
 	return ToolkitInstance->GetAssociatedTabManager();
@@ -83,4 +78,4 @@ void UAssetEditor::OnToolkitClosed()
 	AssetEditorSubsystem->UnregisterUAssetEditor(this);
 }
 
-#undef LOCTEXT_NAMESPACE
\ No newline at end of file
+#undef LOCTEXT_NAMESPACE
diff --git a/Engine/Source/Editor/UnrealEd/Public/MiniCurveEditor.h b/Engine/Source/Editor/UnrealEd/Public/MiniCurveEditor.h
--- a/Engine/Source/Editor/UnrealEd/Public/MiniCurveEditor.h
+++ b/Engine/Source/Editor/UnrealEd/Public/MiniCurveEditor.h
@@ -34,7 +34,6 @@ public:
 	virtual bool CloseWindow() override;
 	virtual bool IsPrimaryEditor() const override { return true; };
 	virtual void InvokeTab(const struct FTabId& TabId) override {}
-	virtual FName GetToolbarTabId() const override;
 	virtual TSharedPtr<class FTabManager> GetAssociatedTabManager() override;
 	virtual double GetLastActivationTime() override;
 	virtual void RemoveEditingAsset(UObject* Asset) override;
diff --git a/Engine/Source/Editor/UnrealEd/Public/Subsystems/AssetEditorSubsystem.h b/Engine/Source/Editor/UnrealEd/Public/Subsystems/AssetEditorSubsystem.h
--- a/Engine/Source/Editor/UnrealEd/Public/Subsystems/AssetEditorSubsystem.h
+++ b/Engine/Source/Editor/UnrealEd/Public/Subsystems/AssetEditorSubsystem.h
@@ -32,7 +32,8 @@ public:
 	virtual bool CloseWindow() = 0;
 	virtual bool IsPrimaryEditor() const = 0;
 	virtual void InvokeTab(const struct FTabId& TabId) = 0;
-	virtual FName GetToolbarTabId() const = 0;
+	UE_DEPRECATED(5.0, "Toolbar tab no longer exists and tab ID will return None; do not add it to layouts")
+	virtual FName GetToolbarTabId() const { return NAME_None; }
 	virtual TSharedPtr<class FTabManager> GetAssociatedTabManager() = 0;
 	virtual double GetLastActivationTime() = 0;
 	virtual void RemoveEditingAsset(UObject* Asset) = 0;
diff --git a/Engine/Source/Editor/UnrealEd/Public/Toolkits/AssetEditorToolkit.h b/Engine/Source/Editor/UnrealEd/Public/Toolkits/AssetEditorToolkit.h
--- a/Engine/Source/Editor/UnrealEd/Public/Toolkits/AssetEditorToolkit.h
+++ b/Engine/Source/Editor/UnrealEd/Public/Toolkits/AssetEditorToolkit.h
@@ -142,7 +142,6 @@ public:
 	virtual bool CloseWindow() override;
 	virtual bool IsPrimaryEditor() const override { return true; };
 	virtual void InvokeTab(const FTabId& TabId) override;
-	virtual FName GetToolbarTabId() const override { return ToolbarTabId; }
 	virtual TSharedPtr<FTabManager> GetAssociatedTabManager() override;
 	virtual double GetLastActivationTime() override;
 	virtual void RemoveEditingAsset(UObject* Asset) override;
@@ -357,9 +356,6 @@ protected:
 	virtual void CreateEditorModeManager() override;
 
 private:
-	/** Spawns the toolbar tab */
-	TSharedRef<SDockTab> SpawnTab_Toolbar(const FSpawnTabArgs& Args);
-
 	// Callback for persisting the Asset Editor's layout.
 	void HandleTabManagerPersistLayout( const TSharedRef<FTabManager::FLayout>& LayoutToSave )
 	{
@@ -445,7 +441,4 @@ private:
 
 	/** Whether the buttons on the default toolbar use small icons */
 	bool bIsToolbarUsingSmallIcons;
-
-	/**	The tab ids for all the tabs used */
-	static const FName ToolbarTabId;
 };
diff --git a/Engine/Source/Editor/UnrealEd/Public/Tools/UAssetEditor.h b/Engine/Source/Editor/UnrealEd/Public/Tools/UAssetEditor.h
--- a/Engine/Source/Editor/UnrealEd/Public/Tools/UAssetEditor.h
+++ b/Engine/Source/Editor/UnrealEd/Public/Tools/UAssetEditor.h
@@ -34,7 +34,6 @@ public:
 		return true;
 	}
 	virtual void InvokeTab(const FTabId& TabId) override {}
-	virtual FName GetToolbarTabId() const override;
 	virtual TSharedPtr<FTabManager> GetAssociatedTabManager() override;
 	virtual double GetLastActivationTime() override
 	{
