Add Unreal Insights bookmark when PIE is clicked in the editor to easily see when a in-editor PIE session starts.

diff --git a/Engine/Source/Editor/UnrealEd/Private/Kismet2/DebuggerCommands.cpp b/Engine/Source/Editor/UnrealEd/Private/Kismet2/DebuggerCommands.cpp
--- a/Engine/Source/Editor/UnrealEd/Private/Kismet2/DebuggerCommands.cpp
+++ b/Engine/Source/Editor/UnrealEd/Private/Kismet2/DebuggerCommands.cpp
@@ -1315,6 +1315,8 @@ FSlateIcon FInternalPlayWorldCommandCallbacks::GetRepeatLastPlayIcon()
 
 void FInternalPlayWorldCommandCallbacks::PlayInViewport_Clicked()
 {
+	TRACE_BOOKMARK(TEXT("PIE"));
+
 	FLevelEditorModule& LevelEditorModule = FModuleManager::GetModuleChecked<FLevelEditorModule>(TEXT("LevelEditor"));
 
 	/** Set PlayInViewPort as the last executed play command */
