DerivedDataBuildLocalExecutor:

- Implemented scratch directory cleanup at startup
- Optimized handling of Inputs - no longer decompresses payload to get hash

diff --git a/Engine/Source/Editor/UnrealEd/Private/DerivedDataBuildLocalExecutor.cpp b/Engine/Source/Editor/UnrealEd/Private/DerivedDataBuildLocalExecutor.cpp
--- a/Engine/Source/Editor/UnrealEd/Private/DerivedDataBuildLocalExecutor.cpp
+++ b/Engine/Source/Editor/UnrealEd/Private/DerivedDataBuildLocalExecutor.cpp
@@ -38,9 +38,19 @@ DEFINE_LOG_CATEGORY_STATIC(LogDerivedDataBuildLocalExecutor, Log, All);
 
 class FLocalBuildWorkerExecutor final : public IBuildWorkerExecutor
 {
+	FString SandboxRootDir = FPaths::EngineSavedDir() / TEXT("LocalExec");
+
 public:
 	FLocalBuildWorkerExecutor()
 	{
+		// Clean out any leftovers from a previous run
+
+		UE_LOG(LogDerivedDataBuildLocalExecutor, Warning, TEXT("Deleting existing local execution state from '%s'"), *SandboxRootDir);
+
+		const bool RequireExists = false;
+		const bool Tree = true;
+		IFileManager::Get().DeleteDirectory(*SandboxRootDir, RequireExists, Tree);
+
 		IModularFeatures::Get().RegisterModularFeature(IBuildWorkerExecutor::GetFeatureName(), this);
 	}
 
@@ -87,7 +97,7 @@ public:
 
 		static std::atomic<int32> SerialNo = 0;
 
-		FString SandboxRoot = FPaths::EngineSavedDir() / TEXT("LocalExec") / TEXT("Scratch");
+		FString SandboxRoot = SandboxRootDir / TEXT("Scratch");
 		SandboxRoot.AppendInt(++SerialNo);
 
 		// Manifest worker in scratch area
@@ -142,12 +152,13 @@ public:
 		{
 			Inputs.Get().IterateInputs([&](FStringView Key, const FCompressedBuffer& Buffer) {
 				TStringBuilder<128> InputPath;
-				InputPath << TEXT("Inputs/") << Buffer.GetRawHash();
+				InputPath << TEXT("Inputs/") << FIoHash(Buffer.GetRawHash());
 
 				FString Path{ SandboxRoot / *InputPath };
 
 				if (TUniquePtr<FArchive> Ar{ IFileManager::Get().CreateFileWriter(*Path, FILEWRITE_Silent) })
 				{
+					// Silly workaround for FArchive::operator<< not accepting const objects
 					FCompressedBuffer Comp = FCompressedBuffer::FromCompressed(Buffer.GetCompressed());
 					*Ar << Comp;
 				}
