EditorDomain: Change whitelist/blacklist to new terms allowlist/blocklist.
#rnx

diff --git a/Engine/Config/BaseEditor.ini b/Engine/Config/BaseEditor.ini
--- a/Engine/Config/BaseEditor.ini
+++ b/Engine/Config/BaseEditor.ini
@@ -373,25 +373,25 @@ bUseFullGCAssetClassNames=true
 [EditorDomain]
 EditorDomainEnabled=false
 SaveUnversioned=false
-; +ClassBlacklist=/Long/PackageName.ClassName
-ClassBlacklist=/Script/ControlRig.ControlRigBlueprintGeneratedClass
-; +PackageBlacklist=/Long/PackageName OR ../../../Standard/Filename/NoExtension OR D:\Absolute\Filename\NoExtension
+; +ClassBlockList=/Long/PackageName.ClassName
+ClassBlockList=/Script/ControlRig.ControlRigBlueprintGeneratedClass
+; +PackageBlockList=/Long/PackageName OR ../../../Standard/Filename/NoExtension OR D:\Absolute\Filename\NoExtension
 
-; Know bad files that are problematic if all of them come from EditorDomain. Blacklisting any one of these is sufficient to prevent the bug
-+PackageBlacklist=/Niagara/Functions/RandomRangeFloat
-+PackageBlacklist=/Niagara/Modules/Update/Color/ScaleColor
+; Know bad files that are problematic if all of them come from EditorDomain. Blocking any one of these is sufficient to prevent the bug
++PackageBlockList=/Niagara/Functions/RandomRangeFloat
++PackageBlockList=/Niagara/Modules/Update/Color/ScaleColor
 
 [TargetDomain]
 TargetDomainEnabled=false
-; IterativeClassWhitelist/IterativeClassBlacklist:
+; IterativeClassAllowList/IterativeClassBlockList:
 ;    Used in HybridIterative only.
-;    List of classes that may use iterative cooks. Whitelist superclasses are also whitelisted.
-;    Blacklist subclasses are also blacklisted. Blacklist overrides the whitelist.
-; +IterativeClassWhitelist=/Long/PackageName.ClassName
-; +IterativeClassBlacklist=/Long/PackageName.ClassName
-IterativeClassWhitelist=/Script/Engine.Texture2D
-+IterativeClassWhitelist=/Script/CoreUObject.MetaData
-+IterativeClassWhitelist=/Script/Engine.AssetImportData
+;    List of classes that may use iterative cooks. Allowed superclasses are also allowed.
+;    Blocked subclasses are also blocked. BlockList overrides the AllowList.
+; +IterativeClassAllowList=/Long/PackageName.ClassName
+; +IterativeClassBlockList=/Long/PackageName.ClassName
+IterativeClassAllowList=/Script/Engine.Texture2D
++IterativeClassAllowList=/Script/CoreUObject.MetaData
++IterativeClassAllowList=/Script/Engine.AssetImportData
 
 [CookPlatformDataCacheSettings]
 MaterialInstance=100
diff --git a/Engine/Source/Editor/UnrealEd/Private/EditorDomain/EditorDomainUtils.cpp b/Engine/Source/Editor/UnrealEd/Private/EditorDomain/EditorDomainUtils.cpp
--- a/Engine/Source/Editor/UnrealEd/Private/EditorDomain/EditorDomainUtils.cpp
+++ b/Engine/Source/Editor/UnrealEd/Private/EditorDomain/EditorDomainUtils.cpp
@@ -38,16 +38,16 @@ FClassDigestMap& GetClassDigests()
 	return GClassDigests;
 }
 
-TSet<FName> GEditorClassBlacklist;
-const TSet<FName>& GetEditorClassBlacklist()
+TSet<FName> GEditorClassBlockList;
+const TSet<FName>& GetEditorClassBlockList()
 {
-	return GEditorClassBlacklist;
+	return GEditorClassBlockList;
 }
 
-TSet<FName> GEditorPackageBlacklist;
-const TSet<FName>& GetEditorPackageBlacklist()
+TSet<FName> GEditorPackageBlockList;
+const TSet<FName>& GetEditorPackageBlockList()
 {
-	return GEditorPackageBlacklist;
+	return GEditorPackageBlockList;
 }
 
 // Change to a new guid when EditorDomain needs to be invalidated
@@ -206,10 +206,10 @@ void PrecacheClassDigests(TConstArrayView<FName> ClassNames, TMap<FName, FClassD
 		
 		FClassData& ClassData = ClassDatas.Emplace_GetRef();
 		ClassData.Name = ClassName;
-		ClassData.DigestData.bEditorDomainEnabled = !GetEditorClassBlacklist().Contains(ClassName);
+		ClassData.DigestData.bEditorDomainEnabled = !GetEditorClassBlockList().Contains(ClassName);
 		if (LookupName != ClassName)
 		{
-			ClassData.DigestData.bEditorDomainEnabled &= !GetEditorClassBlacklist().Contains(LookupName);
+			ClassData.DigestData.bEditorDomainEnabled &= !GetEditorClassBlockList().Contains(LookupName);
 		}
 
 		if (Struct)
@@ -377,30 +377,30 @@ void PrecacheClassDigests(TConstArrayView<FName> ClassNames, TMap<FName, FClassD
 	}
 }
 
-TSet<FName> ConstructClassBlacklist()
+TSet<FName> ConstructClassBlockList()
 {
 	TSet<FName> Result;
-	TArray<FString> BlacklistArray;
-	GConfig->GetArray(TEXT("EditorDomain"), TEXT("ClassBlacklist"), BlacklistArray, GEditorIni);
-	for (const FString& ClassPathName : BlacklistArray)
+	TArray<FString> BlockListArray;
+	GConfig->GetArray(TEXT("EditorDomain"), TEXT("ClassBlockList"), BlockListArray, GEditorIni);
+	for (const FString& ClassPathName : BlockListArray)
 	{
 		Result.Add(FName(*ClassPathName));
 	}
 	return Result;
 }
 
-TSet<FName> ConstructPackageNameBlacklist()
+TSet<FName> ConstructPackageNameBlockList()
 {
 	TSet<FName> Result;
-	TArray<FString> BlacklistArray;
-	GConfig->GetArray(TEXT("EditorDomain"), TEXT("PackageBlacklist"), BlacklistArray, GEditorIni);
+	TArray<FString> BlockListArray;
+	GConfig->GetArray(TEXT("EditorDomain"), TEXT("PackageBlockList"), BlockListArray, GEditorIni);
 	FString PackageName;
 	FString ErrorReason;
-	for (const FString& PackageNameOrFilename : BlacklistArray)
+	for (const FString& PackageNameOrFilename : BlockListArray)
 	{
 		if (!FPackageName::TryConvertFilenameToLongPackageName(PackageNameOrFilename, PackageName, &ErrorReason))
 		{
-			UE_LOG(LogEditorDomain, Warning, TEXT("Editor.ini:[EditorDomain]:PackageBlacklist: Could not convert %s to a LongPackageName: %s"),
+			UE_LOG(LogEditorDomain, Warning, TEXT("Editor.ini:[EditorDomain]:PackageBlocklist: Could not convert %s to a LongPackageName: %s"),
 				*PackageNameOrFilename, *ErrorReason);
 			continue;
 		}
@@ -409,60 +409,60 @@ TSet<FName> ConstructPackageNameBlacklist()
 	return Result;
 }
 
-TSet<FName> ConstructTargetIterativeClassBlacklist()
+TSet<FName> ConstructTargetIterativeClassBlockList()
 {
 	TSet<FName> Result;
-	TArray<FString> BlacklistArray;
-	GConfig->GetArray(TEXT("TargetDomain"), TEXT("IterativeClassBlacklist"), BlacklistArray, GEditorIni);
-	for (const FString& ClassPathName : BlacklistArray)
+	TArray<FString> BlockListArray;
+	GConfig->GetArray(TEXT("TargetDomain"), TEXT("IterativeClassBlockList"), BlockListArray, GEditorIni);
+	for (const FString& ClassPathName : BlockListArray)
 	{
 		Result.Add(FName(*ClassPathName));
 	}
 	return Result;
 }
 
-void ConstructTargetIterativeClassWhitelist()
+void ConstructTargetIterativeClassAllowList()
 {
-	// We're using a whitelist with a blacklist override, so the blacklist is only needed when creating the whitelist
-	TSet<FName> BlacklistFNames = ConstructTargetIterativeClassBlacklist();
+	// We're using a allowlist with a blocklist override, so the blocklist is only needed when creating the allowlist
+	TSet<FName> BlockListFNames = ConstructTargetIterativeClassBlockList();
 
-	// Whitelist elements implicitly whitelist all parent classes, so instead of consulting a list and propagating
+	// AllowList elements implicitly allow all parent classes, so instead of consulting a list and propagating
 	// from parent classes every time we read a new class, we have to iterate the list for all classes up front and
-	// propagate _TO_ parent classes. Note that we only support whitelisting native classes, otherwise we would have
-	// to wait for the AssetRegistry to finish loading to be sure we could find every specified whitelist class.
+	// propagate _TO_ parent classes. Note that we only support allowlisting native classes, otherwise we would have
+	// to wait for the AssetRegistry to finish loading to be sure we could find every specified allowed class.
 
-	// Declare a recursive Visit function. Every class we visit is whitelisted, and we visit its superclasses.
-	// To decide whether a visited class is enabled, we also have to get IsBlacklisted recursively from the parent.
+	// Declare a recursive Visit function. Every class we visit is allowlisted, and we visit its superclasses.
+	// To decide whether a visited class is enabled, we also have to get IsBlockListed recursively from the parent.
 	TSet<FName> EnabledFNames;
 	TStringBuilder<256> NameStringBuffer;
 	TMap<FName, TOptional<bool>> Visited;
-	auto EnableClassIfNotBlacklisted = [&Visited, &EnabledFNames, &BlacklistFNames, &NameStringBuffer]
-		(FName PathName, UStruct* Struct, bool& bOutIsBlacklisted, auto& EnableClassIfNotBlacklistedRef)
+	auto EnableClassIfNotBlocked = [&Visited, &EnabledFNames, &BlockListFNames, &NameStringBuffer]
+		(FName PathName, UStruct* Struct, bool& bOutIsBlocked, auto& EnableClassIfNotBlockedRef)
 	{
 		int32 KeyHash = GetTypeHash(PathName);
-		TOptional<bool>& BlacklistValue = Visited.FindOrAddByHash(KeyHash, PathName);
-		if (BlacklistValue.IsSet())
+		TOptional<bool>& BlockedValue = Visited.FindOrAddByHash(KeyHash, PathName);
+		if (BlockedValue.IsSet())
 		{
-			bOutIsBlacklisted = *BlacklistValue;
+			bOutIsBlocked = *BlockedValue;
 			return;
 		}
-		BlacklistValue = false; // If there is a cycle in the class graph, we will encounter PathName again, so initialize to false
+		BlockedValue = false; // If there is a cycle in the class graph, we will encounter PathName again, so initialize to false
 
-		bool bParentBlacklisted = false;
+		bool bParentBlocked = false;
 		UStruct* ParentStruct = Struct->GetSuperStruct();
 		if (ParentStruct)
 		{
 			NameStringBuffer.Reset();
 			ParentStruct->GetPathName(nullptr, NameStringBuffer);
-			EnableClassIfNotBlacklistedRef(FName(*NameStringBuffer), ParentStruct,
-				bParentBlacklisted, EnableClassIfNotBlacklistedRef);
+			EnableClassIfNotBlockedRef(FName(*NameStringBuffer), ParentStruct,
+				bParentBlocked, EnableClassIfNotBlockedRef);
 		}
 
-		bOutIsBlacklisted = bParentBlacklisted || BlacklistFNames.Contains(PathName);
-		if (bOutIsBlacklisted)
+		bOutIsBlocked = bParentBlocked || BlockListFNames.Contains(PathName);
+		if (bOutIsBlocked)
 		{
-			// Call FindOrAdd again, since the recursive calls may have altered the map and invalidated the BlacklistValue reference
-			Visited.FindOrAddByHash(KeyHash, PathName) = bOutIsBlacklisted;
+			// Call FindOrAdd again, since the recursive calls may have altered the map and invalidated the BlockedValue reference
+			Visited.FindOrAddByHash(KeyHash, PathName) = bOutIsBlocked;
 		}
 		else
 		{
@@ -470,9 +470,9 @@ void ConstructTargetIterativeClassWhitelist()
 		}
 	};
 
-	TArray<FString> WhitelistLeafNames;
-	GConfig->GetArray(TEXT("TargetDomain"), TEXT("IterativeClassWhitelist"), WhitelistLeafNames, GEditorIni);
-	for (const FString& ClassPathName : WhitelistLeafNames)
+	TArray<FString> AllowListLeafNames;
+	GConfig->GetArray(TEXT("TargetDomain"), TEXT("IterativeClassAllowList"), AllowListLeafNames, GEditorIni);
+	for (const FString& ClassPathName : AllowListLeafNames)
 	{
 		if (!FPackageName::IsScriptPackage(ClassPathName))
 		{
@@ -483,8 +483,8 @@ void ConstructTargetIterativeClassWhitelist()
 		{
 			continue;
 		}
-		bool bUnusedIsBlacklisted;
-		EnableClassIfNotBlacklisted(FName(*ClassPathName), Struct, bUnusedIsBlacklisted, EnableClassIfNotBlacklisted);
+		bool bUnusedIsBlocked;
+		EnableClassIfNotBlocked(FName(*ClassPathName), Struct, bUnusedIsBlocked, EnableClassIfNotBlocked);
 	}
 
 	TArray<FName> EnabledFNamesArray = EnabledFNames.Array();
@@ -508,12 +508,12 @@ void UtilsPostEngineInit();
 
 void UtilsInitialize()
 {
-	GEditorClassBlacklist = ConstructClassBlacklist();
-	GEditorPackageBlacklist = ConstructPackageNameBlacklist();
+	GEditorClassBlockList = ConstructClassBlockList();
+	GEditorPackageBlockList = ConstructPackageNameBlockList();
 
-	// Constructing whitelists requires use of UStructs, and the early SetPackageResourceManager
+	// Constructing allowlists requires use of UStructs, and the early SetPackageResourceManager
 	// where UtilsInitialize is called is too early; trying to call UStruct->GetSchemaHash at that
-	// time will break the UClass. Defer the construction of whitelist-based data until OnPostEngineInit
+	// time will break the UClass. Defer the construction of allowlist-based data until OnPostEngineInit
 	GUtilsPostInitDelegate = FCoreDelegates::OnPostEngineInit.AddLambda([]() { UtilsPostEngineInit(); });
 }
 
@@ -522,8 +522,8 @@ void UtilsPostEngineInit()
 	FCoreDelegates::OnPostEngineInit.Remove(GUtilsPostInitDelegate);
 	GUtilsPostInitDelegate.Reset();
 
-	// Note that constructing whitelists depends on all blacklists having been parsed already
-	ConstructTargetIterativeClassWhitelist();
+	// Note that constructing AllowLists depends on all BlockLists having been parsed already
+	ConstructTargetIterativeClassAllowList();
 }
 
 EPackageDigestResult GetPackageDigest(IAssetRegistry& AssetRegistry, FName PackageName,
@@ -551,7 +551,7 @@ EPackageDigestResult AppendPackageDigest(IAssetRegistry& AssetRegistry, FName Pa
 		return EPackageDigestResult::FileDoesNotExist;
 	}
 	EPackageDigestResult Result = AppendPackageDigest(Builder, bOutEditorDomainEnabled, OutErrorMessage, *PackageData, PackageName);
-	bOutEditorDomainEnabled &= !GetEditorPackageBlacklist().Contains(PackageName);
+	bOutEditorDomainEnabled &= !GetEditorPackageBlockList().Contains(PackageName);
 	return Result;
 }
 
@@ -755,7 +755,7 @@ bool TrySavePackage(UPackage* Package)
 	}
 	if (!bEditorDomainEnabled)
 	{
-		UE_LOG(LogEditorDomain, Verbose, TEXT("Skipping save of blacklisted package to EditorDomain: %s."), *Package->GetName())
+		UE_LOG(LogEditorDomain, Verbose, TEXT("Skipping save of blocked package to EditorDomain: %s."), *Package->GetName())
 		return false;
 	}
 
diff --git a/Engine/Source/Editor/UnrealEd/Private/EditorDomain/EditorDomainUtils.h b/Engine/Source/Editor/UnrealEd/Private/EditorDomain/EditorDomainUtils.h
--- a/Engine/Source/Editor/UnrealEd/Private/EditorDomain/EditorDomainUtils.h
+++ b/Engine/Source/Editor/UnrealEd/Private/EditorDomain/EditorDomainUtils.h
@@ -40,9 +40,9 @@ struct FClassDigestData
 {
 	FBlake3Hash SchemaHash;
 	bool bNative = false;
-	/** EditorDomainEnabled whitelists everything and uses only a blacklist, so defaults to true. */
+	/** EditorDomainEnabled allows everything and uses only a blocklist, so defaults to true. */
 	bool bEditorDomainEnabled = true;
-	/** bTargetIterativeEnabled uses a whitelist (with a blacklist override), so defaults to false. */
+	/** bTargetIterativeEnabled uses an allowlist (with a blocklist override), so defaults to false. */
 	bool bTargetIterativeEnabled = false;
 };
 
diff --git a/Engine/Source/Editor/UnrealEd/Private/TargetDomain/TargetDomainUtils.cpp b/Engine/Source/Editor/UnrealEd/Private/TargetDomain/TargetDomainUtils.cpp
--- a/Engine/Source/Editor/UnrealEd/Private/TargetDomain/TargetDomainUtils.cpp
+++ b/Engine/Source/Editor/UnrealEd/Private/TargetDomain/TargetDomainUtils.cpp
@@ -251,8 +251,8 @@ bool IsIterativeEnabled(FName PackageName)
 		UE::EditorDomain::FClassDigestData* ExistingData = ClassDigests.Map.Find(ClassName);
 		if (!ExistingData)
 		{
-			// All whitelisted classes are added to ClassDigests at startup, so if the class is not in ClassDigests,
-			// it is not whitelisted
+			// All allowlisted classes are added to ClassDigests at startup, so if the class is not in ClassDigests,
+			// it is not allowlisted
 			return false;
 		}
 		if (!ExistingData->bTargetIterativeEnabled)
diff --git a/Engine/Source/Editor/UnrealEd/Private/TargetDomain/TargetDomainUtils.h b/Engine/Source/Editor/UnrealEd/Private/TargetDomain/TargetDomainUtils.h
--- a/Engine/Source/Editor/UnrealEd/Private/TargetDomain/TargetDomainUtils.h
+++ b/Engine/Source/Editor/UnrealEd/Private/TargetDomain/TargetDomainUtils.h
@@ -31,7 +31,7 @@ FCbObject CollectDependenciesObject(UPackage* Package, const ITargetPlatform* Ta
 bool TryFetchKeyAndDependencies(IPackageStoreWriter* PackageStore, FName PackageName, const ITargetPlatform* TargetPlatform,
 	FIoHash* OutHash, TArray<FName>* OutBuildDependencies, TArray<FName>* OutRuntimeOnlyDependencies, FString* OutErrorMessage);
 
-/** Return whether iterative cook is enabled for the given packagename, based on used-class whitelist/blacklist. */
+/** Return whether iterative cook is enabled for the given packagename, based on used-class allowlist/blocklist. */
 bool IsIterativeEnabled(FName PackageName);
 
 }
