PackageStoreWriter: Split IPackageStoreWriter into IPackageWriter, IPackageCookWrite, IPackageStoreWriter.
We plan to implement an IPackageWriter for workspace domain editor files to reduce duplication in SavePackage, and implementations for that purpose do not need to handle cook functionality.
IOStore-specific functionality is not needed by the cooker, and is pushed into IPackageStoreWriter.
#rnx

diff --git a/Engine/Source/Developer/IoStoreUtilities/Private/IoStoreUtilities.cpp b/Engine/Source/Developer/IoStoreUtilities/Private/IoStoreUtilities.cpp
--- a/Engine/Source/Developer/IoStoreUtilities/Private/IoStoreUtilities.cpp
+++ b/Engine/Source/Developer/IoStoreUtilities/Private/IoStoreUtilities.cpp
@@ -3914,7 +3914,7 @@ int32 Staged2Zen(const FString& BuildPath, const FKeyChain& KeyChain, const ITar
 	struct FBulkDataInfo
 	{
 		FString FileName;
-		IPackageStoreWriter::FBulkDataInfo::EType BulkDataType;
+		IPackageWriter::FBulkDataInfo::EType BulkDataType;
 		TTuple<FIoStoreReader*, FIoChunkId> Chunk;
 	};
 
@@ -3999,15 +3999,15 @@ int32 Staged2Zen(const FString& BuildPath, const FKeyChain& KeyChain, const ITar
 							BulkDataInfo.Chunk = MakeTuple(Reader.Get(), ChunkInfo.Id);
 							if (ChunkType == EIoChunkType::OptionalBulkData)
 							{
-								BulkDataInfo.BulkDataType = IPackageStoreWriter::FBulkDataInfo::Optional;
+								BulkDataInfo.BulkDataType = IPackageWriter::FBulkDataInfo::Optional;
 							}
 							else if (ChunkType == EIoChunkType::MemoryMappedBulkData)
 							{
-								BulkDataInfo.BulkDataType = IPackageStoreWriter::FBulkDataInfo::Mmap;
+								BulkDataInfo.BulkDataType = IPackageWriter::FBulkDataInfo::Mmap;
 							}
 							else
 							{
-								BulkDataInfo.BulkDataType = IPackageStoreWriter::FBulkDataInfo::Standard;
+								BulkDataInfo.BulkDataType = IPackageWriter::FBulkDataInfo::Standard;
 							}
 						}
 					}
@@ -4059,7 +4059,7 @@ int32 Staged2Zen(const FString& BuildPath, const FKeyChain& KeyChain, const ITar
 	FString MetaDataOutputPath = FPaths::Combine(CookedOutputPath, FApp::GetProjectName(), TEXT("Metadata"));
 	TUniquePtr<FZenStoreWriter> ZenStoreWriter = MakeUnique<FZenStoreWriter>(CookedOutputPath, MetaDataOutputPath, TargetPlatform, true);
 	
-	IPackageStoreWriter::FCookInfo CookInfo;
+	ICookedPackageWriter::FCookInfo CookInfo;
 	ZenStoreWriter->BeginCook(CookInfo);
 	int32 LocalPackageIndex = 0;
 	TArray<FPackageInfo> PackagesArray;
@@ -4069,12 +4069,12 @@ int32 Staged2Zen(const FString& BuildPath, const FKeyChain& KeyChain, const ITar
 	{
 		const FPackageInfo& PackageInfo = PackagesArray[Index];
 
-		IPackageStoreWriter::FBeginPackageInfo BeginPackageInfo;
+		IPackageWriter::FBeginPackageInfo BeginPackageInfo;
 		BeginPackageInfo.PackageName = PackageInfo.PackageName;
 
 		ZenStoreWriter->BeginPackage(BeginPackageInfo);
 
-		IPackageStoreWriter::FPackageInfo PackageStorePackageInfo;
+		IPackageWriter::FPackageInfo PackageStorePackageInfo;
 		PackageStorePackageInfo.PackageName = PackageInfo.PackageName;
 		PackageStorePackageInfo.LooseFilePath = PackageInfo.FileName;
 		PackageStorePackageInfo.ChunkId = PackageInfo.Chunk.Value;
@@ -4084,7 +4084,7 @@ int32 Staged2Zen(const FString& BuildPath, const FKeyChain& KeyChain, const ITar
 
 		for (const FBulkDataInfo& BulkDataInfo : PackageInfo.BulkData)
 		{
-			IPackageStoreWriter::FBulkDataInfo PackageStoreBulkDataInfo;
+			IPackageWriter::FBulkDataInfo PackageStoreBulkDataInfo;
 			PackageStoreBulkDataInfo.PackageName = PackageInfo.PackageName;
 			PackageStoreBulkDataInfo.LooseFilePath = BulkDataInfo.FileName;
 			PackageStoreBulkDataInfo.ChunkId = BulkDataInfo.Chunk.Value;
@@ -4093,7 +4093,7 @@ int32 Staged2Zen(const FString& BuildPath, const FKeyChain& KeyChain, const ITar
 			ZenStoreWriter->WriteBulkdata(PackageStoreBulkDataInfo, BulkDataBuffer, TArray<FFileRegion>());
 		}
 		
-		IPackageStoreWriter::FCommitPackageInfo CommitInfo;
+		IPackageWriter::FCommitPackageInfo CommitInfo;
 		CommitInfo.PackageName = PackageInfo.PackageName;
 		ZenStoreWriter->CommitPackage(CommitInfo);
 
diff --git a/Engine/Source/Developer/IoStoreUtilities/Private/ZenStoreWriter.cpp b/Engine/Source/Developer/IoStoreUtilities/Private/ZenStoreWriter.cpp
--- a/Engine/Source/Developer/IoStoreUtilities/Private/ZenStoreWriter.cpp
+++ b/Engine/Source/Developer/IoStoreUtilities/Private/ZenStoreWriter.cpp
@@ -233,7 +233,7 @@ FZenStoreWriter::FZenStoreWriter(
 	, OutputPath(InOutputPath)
 	, MetadataDirectoryPath(InMetadataDirectoryPath)
 	, PackageStoreOptimizer(new FPackageStoreOptimizer())
-	, CookMode(IPackageStoreWriter::FCookInfo::CookByTheBookMode)
+	, CookMode(ICookedPackageWriter::FCookInfo::CookByTheBookMode)
 {
 	FString HostName = TEXT("localhost");
 	uint16 Port = 1337;
@@ -488,7 +488,7 @@ bool FZenStoreWriter::WriteAdditionalFile(const FAdditionalFileInfo& Info, const
 
 void FZenStoreWriter::WriteLinkerAdditionalData(const FLinkerAdditionalDataInfo& Info, const FIoBuffer& Data, const TArray<FFileRegion>& FileRegions)
 {
-	// Should not be called because IsLinkerAdditionalDataInSeparateArchive returned false
+	// Should not be called because bsLinkerAdditionalDataInSeparateArchive is false
 	checkNoEntry();
 }
 
@@ -496,7 +496,7 @@ void FZenStoreWriter::BeginCook(const FCookInfo& Info)
 {
 	CookMode = Info.CookMode;
 
-	if (Info.CookMode == IPackageStoreWriter::FCookInfo::CookOnTheFlyMode)
+	if (Info.CookMode == ICookedPackageWriter::FCookInfo::CookOnTheFlyMode)
 	{
 		FCbPackage Pkg;
 		FCbWriter PackageObj;
@@ -730,7 +730,7 @@ FCbObject FZenStoreWriter::GetTargetDomainDependencies(FName PackageName)
 		return FCbObject();
 	}
 
-	const IPackageStoreWriter::FCookedPackageInfo& Info = CookedPackagesInfo[*Idx];
+	const ICookedPackageWriter::FCookedPackageInfo& Info = CookedPackagesInfo[*Idx];
 	TIoStatusOr<FIoBuffer> BufferResult = HttpClient->ReadOpLogAttachment(WriteToString<48>(Info.TargetDomainDependencies));
 	if (!BufferResult.IsOk())
 	{
diff --git a/Engine/Source/Developer/IoStoreUtilities/Public/FilePackageStoreWriter.h b/Engine/Source/Developer/IoStoreUtilities/Public/FilePackageStoreWriter.h
--- a/Engine/Source/Developer/IoStoreUtilities/Public/FilePackageStoreWriter.h
+++ b/Engine/Source/Developer/IoStoreUtilities/Public/FilePackageStoreWriter.h
@@ -4,7 +4,7 @@
 
 #include "Serialization/AsyncLoading2.h"
 #include "IO/PackageStore.h"
-#include "IO/PackageStoreWriter.h"
+#include "PackageStoreWriter.h"
 
 class FIoStoreWriterContext;
 class FIoStoreWriter;
@@ -59,6 +59,9 @@ private:
 	TUniquePtr<FZenServerInfo> ZenServerInfo;
 };
 
+/**
+ * A PackageWriter that saves cooked packages for use by IoStore, and stores them in loose files.
+ */
 class FFilePackageStoreWriter
 	: public IPackageStoreWriter
 {
@@ -71,7 +74,7 @@ public:
 	IOSTOREUTILITIES_API virtual void WriteBulkdata(const FBulkDataInfo& Info, const FIoBuffer& BulkData, const TArray<FFileRegion>& FileRegions) override;
 	IOSTOREUTILITIES_API virtual void WriteLinkerAdditionalData(const FLinkerAdditionalDataInfo& Info, const FIoBuffer& Data, const TArray<FFileRegion>& FileRegions) override
 	{
-		// Should not be called because IsLinkerAdditionalDataInSeparateArchive returned false
+		// Should not be called because bLinkerAdditionalDataInSeparateArchive is false
 		checkNoEntry();
 	}
 
diff --git a/Engine/Source/Developer/IoStoreUtilities/Public/PackageStoreWriter.h b/Engine/Source/Developer/IoStoreUtilities/Public/PackageStoreWriter.h
new file mode 100644
--- /dev/null
+++ b/Engine/Source/Developer/IoStoreUtilities/Public/PackageStoreWriter.h
@@ -0,0 +1,39 @@
+// Copyright Epic Games, Inc. All Rights Reserved.
+
+#pragma once
+
+#include "IO/PackageStore.h"
+#include "Serialization/PackageWriter.h"
+
+class IPackageStoreWriter : public ICookedPackageWriter
+{
+public:
+	/** Identify as a member of this interface from the ICookedPackageWriter api. */
+	virtual IPackageStoreWriter* AsPackageStoreWriter() override
+	{
+		return this;
+	}
+
+	/**
+	 * Returns all cooked package store entries.
+	 */
+	virtual void GetEntries(TFunction<void(TArrayView<const FPackageStoreEntryResource>)>&&) = 0;
+
+	/**
+	 * Package commit event arguments
+	 */
+	struct FCommitEventArgs
+	{
+		FName PlatformName;
+		FName PackageName;
+		int32 EntryIndex = INDEX_NONE;
+		TArrayView<const FPackageStoreEntryResource> Entries;
+		TArray<FAdditionalFileInfo> AdditionalFiles;
+	};
+
+	/**
+	 * Broadcasted after a package has been committed, i.e cooked.
+	 */
+	DECLARE_EVENT_OneParam(IPackageStoreWriter, FCommitEvent, const FCommitEventArgs&);
+	virtual FCommitEvent& OnCommit() = 0;
+};
\ No newline at end of file
diff --git a/Engine/Source/Developer/IoStoreUtilities/Public/ZenStoreWriter.h b/Engine/Source/Developer/IoStoreUtilities/Public/ZenStoreWriter.h
--- a/Engine/Source/Developer/IoStoreUtilities/Public/ZenStoreWriter.h
+++ b/Engine/Source/Developer/IoStoreUtilities/Public/ZenStoreWriter.h
@@ -6,8 +6,8 @@
 #include "IO/IoDispatcher.h"
 #include "Containers/Map.h"
 #include "Serialization/CompactBinary.h"
+#include "Serialization/PackageWriter.h"
 #include "IO/PackageStore.h"
-#include "IO/PackageStoreWriter.h"
 #include "FilePackageStoreWriter.h"
 
 namespace UE {
@@ -20,9 +20,8 @@ class FCbPackage;
 class FCbWriter;
 
 /** 
- * Client for interfacing with Zen storage service
+ * A PackageStoreWriter that saves cooked packages for use by IoStore, and stores them in the Zen storage service.
  */
-
 class FZenStoreWriter
 	: public IPackageStoreWriter
 {
@@ -126,7 +125,7 @@ private:
 	class FZenStoreHttpQueue;
 	TUniquePtr<FZenStoreHttpQueue>		HttpQueue;
 	
-	IPackageStoreWriter::FCookInfo::ECookMode CookMode;
+	ICookedPackageWriter::FCookInfo::ECookMode CookMode;
 
 	FZenStats							ZenStats;
 };
diff --git a/Engine/Source/Editor/UnrealEd/Classes/CookOnTheSide/CookOnTheFlyServer.h b/Engine/Source/Editor/UnrealEd/Classes/CookOnTheSide/CookOnTheFlyServer.h
--- a/Engine/Source/Editor/UnrealEd/Classes/CookOnTheSide/CookOnTheFlyServer.h
+++ b/Engine/Source/Editor/UnrealEd/Classes/CookOnTheSide/CookOnTheFlyServer.h
@@ -26,7 +26,7 @@ class FDiffModeCookServerUtils;
 class FReferenceCollector;
 class FSavePackageContext;
 class IAssetRegistry;
-class IPackageStoreWriter;
+class ICookedPackageWriter;
 class IPlugin;
 class ITargetPlatform;
 struct FPackageNameCache;
@@ -117,6 +117,7 @@ namespace UE::Cook
 	class FSaveCookedPackageContext;
 	class ICookOnTheFlyRequestManager;
 	struct FCookerTimer;
+	struct FCookSavePackageContext;
 	struct FGeneratorPackage;
 	struct FPackageData;
 	struct FPackageDatas;
@@ -1109,8 +1110,8 @@ private:
 
 	uint32 FullLoadAndSave(uint32& CookedPackageCount);
 
-	IPackageStoreWriter* GetPackageStoreWriter(const FName& PlatformName) const;
-	FSavePackageContext* GetSavePackageContext(const ITargetPlatform* TargetPlatform) const;
+	ICookedPackageWriter* GetPackageWriter(const FName& PlatformName) const;
+	UE::Cook::FCookSavePackageContext* GetSavePackageContext(const ITargetPlatform* TargetPlatform) const;
 
 	uint32		StatLoadedPackageCount = 0;
 	uint32		StatSavedPackageCount = 0;
@@ -1145,7 +1146,7 @@ private:
 	TUniquePtr<UE::Cook::FExternalRequests> ExternalRequests;
 	TRingBuffer<UE::Cook::FRequestCluster> RequestClusters;
 
-	TArray<FSavePackageContext*> SavePackageContexts;
+	TArray<UE::Cook::FCookSavePackageContext*> SavePackageContexts;
 	/** Objects that were collected during the single-threaded PreGarbageCollect callback and that should be reported as referenced in CookerAddReferencedObjects. */
 	TArray<UObject*> GCKeepObjects;
 	UE::Cook::FPackageData* SavingPackageData = nullptr;
diff --git a/Engine/Source/Editor/UnrealEd/Private/CookOnTheFlyServer.cpp b/Engine/Source/Editor/UnrealEd/Private/CookOnTheFlyServer.cpp
--- a/Engine/Source/Editor/UnrealEd/Private/CookOnTheFlyServer.cpp
+++ b/Engine/Source/Editor/UnrealEd/Private/CookOnTheFlyServer.cpp
@@ -466,15 +466,15 @@ public:
 		return true;
 	}
 
-	virtual IPackageStoreWriter* GetPackageStoreWriter(const FName& PlatformName) override
+	virtual ICookedPackageWriter* GetPackageWriter(const FName& PlatformName) override
 	{
 		if (Cooker.CookOnTheFlyOptions.bZenStore)
 		{
-			for (FSavePackageContext* Context : Cooker.SavePackageContexts)
+			for (UE::Cook::FCookSavePackageContext* Context : Cooker.SavePackageContexts)
 			{
-				if (FName(*Context->TargetPlatform->PlatformName()) == PlatformName)
+				if (FName(*Context->SaveContext.TargetPlatform->PlatformName()) == PlatformName)
 				{
-					return Context->PackageStoreWriter;
+					return Context->PackageWriter;
 				}
 			}
 		}
@@ -506,9 +506,9 @@ public:
 
 		if (Cooker.CookOnTheFlyOptions.bZenStore)
 		{
-			for (FSavePackageContext* Context : Cooker.SavePackageContexts)
+			for (UE::Cook::FCookSavePackageContext* Context : Cooker.SavePackageContexts)
 			{
-				Context->PackageStoreWriter->Flush();
+				Context->PackageWriter->Flush();
 			}
 		}
 
@@ -3201,7 +3201,7 @@ private: // Used only by UCookOnTheFlyServer, which has private access
 	// Per-platform data, only valid in PerPlatform callbacks
 	const ITargetPlatform* TargetPlatform = nullptr;
 	FSavePackageContext* SavePackageContext = nullptr;
-	IPackageStoreWriter* PackageStoreWriter = nullptr;
+	ICookedPackageWriter* PackageWriter = nullptr;
 	FString PlatFilename;
 	FSavePackageResultStruct SavePackageResult;
 	bool bPlatformSetupSuccessful = false;
@@ -4188,7 +4188,7 @@ private:
 public:
 
 
-	void Initialize(TConstArrayView<FSavePackageContext*> SavePackageContexts)
+	void Initialize(TConstArrayView<UE::Cook::FCookSavePackageContext*> SavePackageContexts)
 	{
 		if (bInitialized)
 		{
@@ -4215,14 +4215,14 @@ public:
 		bInitialized = true;
 	}
 
-	void InitializePlatformSettings(TConstArrayView<FSavePackageContext*> SavePackageContexts)
+	void InitializePlatformSettings(TConstArrayView<UE::Cook::FCookSavePackageContext*> SavePackageContexts)
 	{
 		if (IsRunningCookDiff() || IsRunningCookLinkerDiff())
 		{
 			bool bCompatible = true;
-			for (FSavePackageContext* SavePackageContext : SavePackageContexts)
+			for (UE::Cook::FCookSavePackageContext* SavePackageContext : SavePackageContexts)
 			{
-				if (SavePackageContext && SavePackageContext->PackageStoreWriter)
+				if (SavePackageContext && SavePackageContext->PackageWriter)
 				{
 					bCompatible = false;
 					break;
@@ -4231,7 +4231,7 @@ public:
 			if (!bCompatible)
 			{
 				const TCHAR* CommandLineArg = IsRunningCookDiff() ? TEXT("-DIFFONLY") : TEXT("-LINKERDIFF");
-				UE_LOG(LogCook, Fatal, TEXT("%s was enabled, but PackageStoreWriter is also enabled and PackageStoreWriter does not support %s."),
+				UE_LOG(LogCook, Fatal, TEXT("%s was enabled, but PackageWriter is also enabled and PackageWriter does not support %s."),
 					CommandLineArg, CommandLineArg);
 				bDiffEnabled = false;
 				LinkerDiffMode = LDM_None;
@@ -4457,7 +4457,7 @@ void UCookOnTheFlyServer::SaveCookedPackage(UE::Cook::FSaveCookedPackageContext&
 			{
 				if (DiffModeHelper->IsRunningCookDiff())
 				{
-					check(!Context.PackageStoreWriter); // IsRunningCookDiff should have been disabled by DiffModeHelper if we have any PackageStoreWriter
+					check(!Context.PackageWriter); // IsRunningCookDiff should have been disabled by DiffModeHelper if we have any PackageWriter
 					DiffModeHelper->ProcessPackage(Package);
 
 					// When looking for deterministic cook issues, first serialize the package to memory and do a simple diff with the existing package
@@ -4477,7 +4477,7 @@ void UCookOnTheFlyServer::SaveCookedPackage(UE::Cook::FSaveCookedPackageContext&
 				}
 				else if (DiffModeHelper->IsRunningCookLinkerDiff())
 				{
-					check(!Context.PackageStoreWriter); // IsRunningCookLinkerDiff should have been disabled by DiffModeHelper if we have any PackageStoreWriter
+					check(!Context.PackageWriter); // IsRunningCookLinkerDiff should have been disabled by DiffModeHelper if we have any PackageWriter
 
 					Context.SavePackageResult = GEditor->Save(Package, Context.World, Context.FlagsToCook, *Context.PlatFilename,
 						GError, nullptr, Context.bEndianSwap, false, Context.SaveFlags, TargetPlatform,
@@ -4646,14 +4646,23 @@ void FSaveCookedPackageContext::SetupPlatform(const ITargetPlatform* InTargetPla
 		World->PersistentLevel->HandleLegacyMapBuildData();
 	}
 
-	SavePackageContext = COTFS.GetSavePackageContext(TargetPlatform);
-	PackageStoreWriter = SavePackageContext ? SavePackageContext->PackageStoreWriter : nullptr;
-	if (PackageStoreWriter)
+	FCookSavePackageContext* CookContext = COTFS.GetSavePackageContext(TargetPlatform);
+	if (CookContext)
 	{
-		IPackageStoreWriter::FBeginPackageInfo Info;
+		SavePackageContext = &CookContext->SaveContext;
+		PackageWriter = CookContext->PackageWriter;
+	}
+	else
+	{
+		SavePackageContext = nullptr;
+		PackageWriter = nullptr;
+	}
+	if (PackageWriter)
+	{
+		ICookedPackageWriter::FBeginPackageInfo Info;
 		Info.PackageName = Package->GetFName();
 
-		PackageStoreWriter->BeginPackage(Info);
+		PackageWriter->BeginPackage(Info);
 	}
 
 	// Indicate Setup was successful
@@ -4666,7 +4675,7 @@ void FSaveCookedPackageContext::FinishPlatform()
 	bool bSuccessful = SavePackageResult.IsSuccessful();
 	bool bLocalReferencedOnlyByEditorOnlyData = SavePackageResult == ESavePackageResult::ReferencedOnlyByEditorOnlyData;
 
-	if (bPlatformSetupSuccessful && PackageStoreWriter)
+	if (bPlatformSetupSuccessful && PackageWriter)
 	{
 		FAssetRegistryGenerator* Generator = COTFS.PlatformManager->GetPlatformData(TargetPlatform)->RegistryGenerator.Get();
 		FAssetPackageData* AssetPackageData = Generator->GetAssetPackageData(Package->GetFName());
@@ -4679,7 +4688,7 @@ void FSaveCookedPackageContext::FinishPlatform()
 			TargetDomainDependencies = UE::TargetDomain::CollectDependenciesObject(Package, TargetPlatform, nullptr /* ErrorMessage */);
 		}
 
-		IPackageStoreWriter::FCommitPackageInfo Info;
+		ICookedPackageWriter::FCommitPackageInfo Info;
 		Info.bSucceeded = bSuccessful;
 		Info.PackageName = Package->GetFName();
 		PRAGMA_DISABLE_DEPRECATION_WARNINGS
@@ -4687,7 +4696,7 @@ void FSaveCookedPackageContext::FinishPlatform()
 		PRAGMA_ENABLE_DEPRECATION_WARNINGS
 		Info.TargetDomainDependencies = TargetDomainDependencies;
 
-		PackageStoreWriter->CommitPackage(Info);
+		PackageWriter->CommitPackage(Info);
 	}
 
 	// Update asset registry
@@ -6228,14 +6237,14 @@ void UCookOnTheFlyServer::PopulateCookedPackagesFromPackageStore(const TArrayVie
 
 	struct FPreviousAssetPackageData
 	{
-		FPreviousAssetPackageData(const TArray<IPackageStoreWriter::FCookedPackageInfo>& InCookedPackages)
+		FPreviousAssetPackageData(const TArray<ICookedPackageWriter::FCookedPackageInfo>& InCookedPackages)
 		{
 			AssetPackageDatas.SetNum(InCookedPackages.Num());
 			AssetPackageDataMap.Reserve(InCookedPackages.Num());
 
 			for (int32 Idx = 0, Count = InCookedPackages.Num(); Idx < Count; ++Idx)
 			{
-				const IPackageStoreWriter::FCookedPackageInfo& CookedPackageInfo = InCookedPackages[Idx];
+				const ICookedPackageWriter::FCookedPackageInfo& CookedPackageInfo = InCookedPackages[Idx];
 				FAssetPackageData& AssetPackageData = AssetPackageDatas[Idx];
 				
 				AssetPackageData.CookedHash = CookedPackageInfo.Hash;
@@ -6258,14 +6267,14 @@ void UCookOnTheFlyServer::PopulateCookedPackagesFromPackageStore(const TArrayVie
 		
 		UE::Cook::FPlatformData* PlatformData = PlatformManager->GetPlatformData(Platform);
 		FAssetRegistryGenerator* PlatformAssetRegistry = PlatformData->RegistryGenerator.Get();
-		IPackageStoreWriter* PackageStoreWriter = GetPackageStoreWriter(FName(*Platform->PlatformName()));
+		ICookedPackageWriter* PackageWriter = GetPackageWriter(FName(*Platform->PlatformName()));
 		
 		check(PlatformData);
 		check(PlatformAssetRegistry);
-		check(PackageStoreWriter);
+		check(PackageWriter);
 
-		TArray<IPackageStoreWriter::FCookedPackageInfo> CookedPackages;
-		PackageStoreWriter->GetCookedPackages(CookedPackages);
+		TArray<ICookedPackageWriter::FCookedPackageInfo> CookedPackages;
+		PackageWriter->GetCookedPackages(CookedPackages);
 		
 		UE_LOG(LogCook, Display, TEXT("Found '%d' cooked package(s) in package store"), CookedPackages.Num());
 
@@ -6293,7 +6302,7 @@ void UCookOnTheFlyServer::PopulateCookedPackagesFromPackageStore(const TArrayVie
 		
 		UE_LOG(LogCook, Display, TEXT("Keeping '%d' and removing '%d' cooked package(s)"), Difference.IdenticalCookedPackages.Num(), PackagesToRemove.Num());
 		
-		PackageStoreWriter->RemoveCookedPackages(PackagesToRemove);
+		PackageWriter->RemoveCookedPackages(PackagesToRemove);
 		
 		TArray<FName> PackagesToKeep;
 		for (const FName& IdenticalPackageName : Difference.IdenticalCookedPackages)
@@ -7888,7 +7897,7 @@ void UCookOnTheFlyServer::InitializePackageStore(const TArrayView<const ITargetP
 		const FString ResolvedRootPath = RootPathSandbox.Replace(TEXT("[Platform]"), *PlatformString);
 		const FString ResolvedProjectPath = ProjectPathSandbox.Replace(TEXT("[Platform]"), *PlatformString);
 
-		IPackageStoreWriter* PackageStoreWriter = nullptr;
+		ICookedPackageWriter* PackageWriter = nullptr;
 		if (IsUsingIoStore())
 		{
 			FString MetadataDirectoryPath = ProjectPathSandbox / TEXT("Metadata");
@@ -7896,14 +7905,14 @@ void UCookOnTheFlyServer::InitializePackageStore(const TArrayView<const ITargetP
 
 			if (IsUsingZenStore())
 			{
-				PackageStoreWriter = new FZenStoreWriter(	ResolvedRootPath, 
+				PackageWriter = new FZenStoreWriter(	ResolvedRootPath, 
 															MetadataDirectoryPath, 
 															TargetPlatform, 
 															bIsCleanBuild);
 			}
 			else
 			{
-				PackageStoreWriter = new FFilePackageStoreWriter(ResolvedRootPath, MetadataDirectoryPath, TargetPlatform);
+				PackageWriter = new FFilePackageStoreWriter(ResolvedRootPath, MetadataDirectoryPath, TargetPlatform);
 			}
 		}
 
@@ -7913,7 +7922,8 @@ void UCookOnTheFlyServer::InitializePackageStore(const TArrayView<const ITargetP
 		bool bLegacyBulkDataOffsets = false;
 		PlatformEngineIni.GetBool(TEXT("Core.System"), TEXT("LegacyBulkDataOffsets"), bLegacyBulkDataOffsets);
 
-		FSavePackageContext* SavePackageContext = new FSavePackageContext(TargetPlatform, PackageStoreWriter, bLegacyBulkDataOffsets);
+		UE::Cook::FCookSavePackageContext* SavePackageContext 
+			= new UE::Cook::FCookSavePackageContext(TargetPlatform, PackageWriter, bLegacyBulkDataOffsets);
 		SavePackageContexts.Add(SavePackageContext);
 	}
 
@@ -7924,11 +7934,11 @@ void UCookOnTheFlyServer::InitializePackageStore(const TArrayView<const ITargetP
 			PopulateCookedPackagesFromPackageStore(TargetPlatforms);
 		}
 
-		for (const FSavePackageContext* Context : SavePackageContexts)
+		for (const UE::Cook::FCookSavePackageContext* Context : SavePackageContexts)
 		{
-			IPackageStoreWriter::FCookInfo CookInfo;
-			CookInfo.CookMode = IsCookOnTheFlyMode() ? IPackageStoreWriter::FCookInfo::CookOnTheFlyMode : IPackageStoreWriter::FCookInfo::CookByTheBookMode;
-			Context->PackageStoreWriter->BeginCook(CookInfo);
+			ICookedPackageWriter::FCookInfo CookInfo;
+			CookInfo.CookMode = IsCookOnTheFlyMode() ? ICookedPackageWriter::FCookInfo::CookOnTheFlyMode : ICookedPackageWriter::FCookInfo::CookByTheBookMode;
+			Context->PackageWriter->BeginCook(CookInfo);
 		}
 	}
 }
@@ -7938,13 +7948,13 @@ void UCookOnTheFlyServer::FinalizePackageStore()
 	UE_SCOPED_HIERARCHICAL_COOKTIMER(FinalizePackageStore);
 
 	UE_LOG(LogCook, Display, TEXT("Finalize package store(s)..."));
-	for (FSavePackageContext* PackageContext : SavePackageContexts)
+	for (UE::Cook::FCookSavePackageContext* PackageContext : SavePackageContexts)
 	{
 		if (PackageContext != nullptr)
 		{
-			if (PackageContext->PackageStoreWriter != nullptr)
+			if (PackageContext->PackageWriter != nullptr)
 			{
-				PackageContext->PackageStoreWriter->EndCook();
+				PackageContext->PackageWriter->EndCook();
 			}
 		}
 	}
@@ -7955,7 +7965,7 @@ void UCookOnTheFlyServer::FinalizePackageStore()
 
 void UCookOnTheFlyServer::ClearPackageStoreContexts()
 {
-	for (FSavePackageContext* Context : SavePackageContexts)
+	for (UE::Cook::FCookSavePackageContext* Context : SavePackageContexts)
 	{
 		delete Context;
 	}
@@ -9234,7 +9244,7 @@ uint32 UCookOnTheFlyServer::FullLoadAndSave(uint32& CookedPackageCount)
 						}
 								
 						GIsCookerLoadingPackage = true;
-						FSavePackageContext* const SavePackageContext = (IsCookByTheBookMode() && SavePackageContexts.Num() > 0) ? SavePackageContexts[PlatformIndex] : nullptr;
+						FSavePackageContext* const SavePackageContext = (IsCookByTheBookMode() && SavePackageContexts.Num() > 0) ? &SavePackageContexts[PlatformIndex]->SaveContext : nullptr;
 						FSavePackageResultStruct SaveResult = GEditor->Save(Package, World, FlagsToCook, *PlatFilename, GError, NULL, bSwap, false, SaveFlags, Target, FDateTime::MinValue(), 
 																			false, /*DiffMap*/ nullptr,	SavePackageContext);
 						GIsCookerLoadingPackage = false;
@@ -9318,24 +9328,24 @@ uint32 UCookOnTheFlyServer::FullLoadAndSave(uint32& CookedPackageCount)
 	return Result;
 }
 
-IPackageStoreWriter* UCookOnTheFlyServer::GetPackageStoreWriter(const FName& PlatformName) const
+ICookedPackageWriter* UCookOnTheFlyServer::GetPackageWriter(const FName& PlatformName) const
 {
-	for (const FSavePackageContext* Context : SavePackageContexts)
+	for (const UE::Cook::FCookSavePackageContext* Context : SavePackageContexts)
 	{
-		if (FName(*Context->TargetPlatform->PlatformName()) == PlatformName)
+		if (FName(*Context->SaveContext.TargetPlatform->PlatformName()) == PlatformName)
 		{
-			return Context->PackageStoreWriter;
+			return Context->PackageWriter;
 		}
 	}
 
 	return nullptr;
 }
 
-FSavePackageContext* UCookOnTheFlyServer::GetSavePackageContext(const ITargetPlatform* TargetPlatform) const
+UE::Cook::FCookSavePackageContext* UCookOnTheFlyServer::GetSavePackageContext(const ITargetPlatform* TargetPlatform) const
 {
-	for (FSavePackageContext* Context : SavePackageContexts)
+	for (UE::Cook::FCookSavePackageContext* Context : SavePackageContexts)
 	{
-		if (Context->TargetPlatform == TargetPlatform)
+		if (Context->SaveContext.TargetPlatform == TargetPlatform)
 		{
 			return Context;
 		}
diff --git a/Engine/Source/Editor/UnrealEd/Private/Cooker/CookOnTheFlyServerInterface.h b/Engine/Source/Editor/UnrealEd/Private/Cooker/CookOnTheFlyServerInterface.h
--- a/Engine/Source/Editor/UnrealEd/Private/Cooker/CookOnTheFlyServerInterface.h
+++ b/Engine/Source/Editor/UnrealEd/Private/Cooker/CookOnTheFlyServerInterface.h
@@ -6,7 +6,7 @@
 #include "CookTypes.h"
 
 struct FODSCRequestPayload;
-class IPackageStoreWriter;
+class ICookedPackageWriter;
 
 namespace UE { namespace Cook
 {
@@ -119,7 +119,7 @@ public:
 	 *
 	 * @param PlatformName The platform name.
 	 */
-	virtual IPackageStoreWriter* GetPackageStoreWriter(const FName& PlatformName) = 0;
+	virtual ICookedPackageWriter* GetPackageWriter(const FName& PlatformName) = 0;
 
 	/**
 	 * A somewhat temporary event triggered when all pending packages has been cooked.
diff --git a/Engine/Source/Editor/UnrealEd/Private/Cooker/CookRequestCluster.cpp b/Engine/Source/Editor/UnrealEd/Private/Cooker/CookRequestCluster.cpp
--- a/Engine/Source/Editor/UnrealEd/Private/Cooker/CookRequestCluster.cpp
+++ b/Engine/Source/Editor/UnrealEd/Private/Cooker/CookRequestCluster.cpp
@@ -115,10 +115,10 @@ void FRequestCluster::Initialize(UCookOnTheFlyServer& COTFS)
 		FPaths::MakeStandardFilename(DLCPath);
 	}
 
-	PackageStores.Reserve(Platforms.Num());
+	PackageWriters.Reserve(Platforms.Num());
 	for (const ITargetPlatform* TargetPlatform : Platforms)
 	{
-		PackageStores.Add(COTFS.GetPackageStoreWriter(FName(*TargetPlatform->PlatformName())));
+		PackageWriters.Add(COTFS.GetPackageWriter(FName(*TargetPlatform->PlatformName())));
 	}
 }
 
@@ -405,12 +405,12 @@ void FRequestCluster::GetDependencies(FName PackageName, TArray<FName>& OutDepen
 		for (int32 PlatIndex = 0; PlatIndex < Platforms.Num(); ++PlatIndex)
 		{
 			const ITargetPlatform* TargetPlatform = Platforms[PlatIndex];
-			IPackageStoreWriter* PackageStore = PackageStores[PlatIndex];
-			if (!PackageStore)
+			ICookedPackageWriter* PackageWriter = PackageWriters[PlatIndex];
+			if (!PackageWriter)
 			{
 				continue;
 			}
-			UE::TargetDomain::TryFetchKeyAndDependencies(PackageStore, PackageName, TargetPlatform, nullptr /* OutHash */, &BuildDependencies,
+			UE::TargetDomain::TryFetchKeyAndDependencies(PackageWriter, PackageName, TargetPlatform, nullptr /* OutHash */, &BuildDependencies,
 				&RuntimeDependencies, nullptr /* OutErrorMessage */);
 		}
 	}
diff --git a/Engine/Source/Editor/UnrealEd/Private/Cooker/CookRequestCluster.h b/Engine/Source/Editor/UnrealEd/Private/Cooker/CookRequestCluster.h
--- a/Engine/Source/Editor/UnrealEd/Private/Cooker/CookRequestCluster.h
+++ b/Engine/Source/Editor/UnrealEd/Private/Cooker/CookRequestCluster.h
@@ -98,7 +98,7 @@ private:
 	TArray<FRequest> Requests;
 	TArray<FName> RuntimeScratch;
 	TArray<const ITargetPlatform*> Platforms;
-	TArray<IPackageStoreWriter*> PackageStores;
+	TArray<ICookedPackageWriter*> PackageWriters;
 	FRequestClusterGraphData DependencyGraphData;
 	FString DLCPath;
 	FPackageDatas& PackageDatas;
diff --git a/Engine/Source/Editor/UnrealEd/Private/Cooker/CookTypes.h b/Engine/Source/Editor/UnrealEd/Private/Cooker/CookTypes.h
--- a/Engine/Source/Editor/UnrealEd/Private/Cooker/CookTypes.h
+++ b/Engine/Source/Editor/UnrealEd/Private/Cooker/CookTypes.h
@@ -7,7 +7,9 @@
 #include "CookOnTheSide/CookOnTheFlyServer.h" // ECookTickFlags
 #include "HAL/Platform.h"
 #include "Logging/TokenizedMessage.h"
+#include "Serialization/PackageWriter.h"
 #include "Templates/Function.h"
+#include "UObject/SavePackage.h"
 
 class FString;
 
@@ -53,9 +55,7 @@ template<typename KeyType, typename SetAllocator = FDefaultSetAllocator>
 class TFastPointerSet : public TSet<KeyType, TFastPointerSetKeyFuncs<KeyType>, SetAllocator>
 {};
 
-namespace UE
-{
-namespace Cook
+namespace UE::Cook
 {
 	struct FPackageData;
 	/** A function that is called when a requested package finishes cooking (when successful, failed, or skipped) */
@@ -160,12 +160,25 @@ namespace Cook
 		}
 	};
 
+	/** Context data passed into SavePackage for a given TargetPlatform. */
+	struct FCookSavePackageContext
+	{
+		FCookSavePackageContext(const ITargetPlatform* InTargetPlatform,
+			ICookedPackageWriter* InPackageWriter, bool bInForceLegacyOffsets)
+			: SaveContext(InTargetPlatform, InPackageWriter, bInForceLegacyOffsets)
+			, PackageWriter(InPackageWriter)
+		{
+		}
+		// Note the SaveContext takes responsibility for deleting PackageWriter
+		FSavePackageContext SaveContext;
+		ICookedPackageWriter* PackageWriter;
+	};
+
 	/* Thread Local Storage access to identify which thread is the SchedulerThread for cooking. */
 	void InitializeTls();
 	bool IsSchedulerThread();
 	void SetIsSchedulerThread(bool bValue);
 }
-}
 
 //////////////////////////////////////////////////////////////////////////
 // Cook by the book options
diff --git a/Engine/Source/Editor/UnrealEd/Private/Cooker/IoStoreCookOnTheFlyRequestManager.cpp b/Engine/Source/Editor/UnrealEd/Private/Cooker/IoStoreCookOnTheFlyRequestManager.cpp
--- a/Engine/Source/Editor/UnrealEd/Private/Cooker/IoStoreCookOnTheFlyRequestManager.cpp
+++ b/Engine/Source/Editor/UnrealEd/Private/Cooker/IoStoreCookOnTheFlyRequestManager.cpp
@@ -5,7 +5,7 @@
 #include "CookOnTheFly.h"
 #include "Modules/ModuleManager.h"
 #include "Serialization/CookOnTheFlyPackageStore.h"
-#include "IO/PackageStoreWriter.h"
+#include "PackageStoreWriter.h"
 #include "ShaderCompiler.h"
 #include "HAL/PlatformTime.h"
 #include "Misc/PackageName.h"
@@ -396,7 +396,7 @@ private:
 		FCriticalSection CriticalSection;
 		FName PlatformName;
 		FCookOnTheFlyPackageTracker PackageTracker;
-		IPackageStoreWriter* PackageStoreWriter = nullptr;
+		IPackageStoreWriter* PackageWriter = nullptr;
 	};
 
 	virtual bool Initialize() override
@@ -445,15 +445,16 @@ private:
 					TUniquePtr<FPlatformContext>& Context = PlatformContexts.Add(Client.PlatformName, MakeUnique<FPlatformContext>());
 					Context->PlatformName = Client.PlatformName;
 					
-					IPackageStoreWriter* PackageStoreWriter = CookOnTheFlyServer.GetPackageStoreWriter(Client.PlatformName);
-					Context->PackageStoreWriter = PackageStoreWriter;
+					IPackageStoreWriter* PackageWriter = CookOnTheFlyServer.GetPackageWriter(Client.PlatformName)->AsPackageStoreWriter();
+					check(PackageWriter); // This class should not be used except when COTFS is using an IPackageStoreWriter
+					Context->PackageWriter = PackageWriter;
 					
-					PackageStoreWriter->GetEntries([&Context](TArrayView<const FPackageStoreEntryResource> Entries)
+					PackageWriter->GetEntries([&Context](TArrayView<const FPackageStoreEntryResource> Entries)
 					{
 						Context->PackageTracker.AddCompletedPackages(Entries);
 					});
 
-					PackageStoreWriter->OnCommit().AddRaw(this, &FIoStoreCookOnTheFlyRequestManager::OnPackageCooked);
+					PackageWriter->OnCommit().AddRaw(this, &FIoStoreCookOnTheFlyRequestManager::OnPackageCooked);
 				}
 
 				return true;
@@ -520,7 +521,7 @@ private:
 		PackageStoreData.TotalFailedPackages	= FailedPackages.Num();
 		//PackageStoreData.FailedPackages			= FailedPackages;
 
-		Context.PackageStoreWriter->GetEntries([&CookedEntryIndices, &PackageStoreData](TArrayView<const FPackageStoreEntryResource> Entries)
+		Context.PackageWriter->GetEntries([&CookedEntryIndices, &PackageStoreData](TArrayView<const FPackageStoreEntryResource> Entries)
 		{
 			for (int32 EntryIndex : CookedEntryIndices)
 			{
@@ -597,7 +598,7 @@ private:
 			TArray<FString> Filenames;
 			TArray<FIoChunkId> ChunkIds;
 
-			for (const IPackageStoreWriter::FAdditionalFileInfo& FileInfo : EventArgs.AdditionalFiles)
+			for (const ICookedPackageWriter::FAdditionalFileInfo& FileInfo : EventArgs.AdditionalFiles)
 			{
 				UE_LOG(LogCookOnTheFly, Verbose, TEXT("Sending additional cooked file '%s'"), *FileInfo.Filename);
 				Filenames.Add(FileInfo.Filename);
@@ -731,7 +732,7 @@ private:
 				PackageStoreData.TotalFailedPackages = Context.PackageTracker.GetFailedPackages().Num();
 				PackageStoreData.FailedPackages = MoveTemp(CompletedPackages.FailedPackages);
 
-				Context.PackageStoreWriter->GetEntries([&CompletedPackages, &PackageStoreData](TArrayView<const FPackageStoreEntryResource> Entries)
+				Context.PackageWriter->GetEntries([&CompletedPackages, &PackageStoreData](TArrayView<const FPackageStoreEntryResource> Entries)
 				{
 					for (int32 EntryIndex : CompletedPackages.CookedEntryIndices)
 					{
diff --git a/Engine/Source/Editor/UnrealEd/Private/EditorDomain/EditorDomainUtils.cpp b/Engine/Source/Editor/UnrealEd/Private/EditorDomain/EditorDomainUtils.cpp
--- a/Engine/Source/Editor/UnrealEd/Private/EditorDomain/EditorDomainUtils.cpp
+++ b/Engine/Source/Editor/UnrealEd/Private/EditorDomain/EditorDomainUtils.cpp
@@ -13,7 +13,6 @@
 #include "DerivedDataRequestOwner.h"
 #include "Editor.h"
 #include "HAL/FileManager.h"
-#include "IO/PackageStoreWriter.h"
 #include "Memory/SharedBuffer.h"
 #include "Misc/CoreDelegates.h"
 #include "Misc/FileHelper.h"
@@ -21,6 +20,7 @@
 #include "Misc/ScopeRWLock.h"
 #include "Misc/StringBuilder.h"
 #include "Serialization/CompactBinaryWriter.h"
+#include "Serialization/PackageWriter.h"
 #include "UObject/CoreRedirects.h"
 #include "UObject/ObjectVersion.h"
 #include "UObject/Package.h"
@@ -586,68 +586,38 @@ void RequestEditorDomainPackage(const FPackagePath& PackagePath,
 }
 
 /** Stores data from SavePackage in accessible fields */
-class FMemoryPackageStoreWriter : public IPackageStoreWriter
+class FMemoryPackageWriter : public IPackageWriter
 {
 public:
-	// IPackageStoreWriter
-	virtual bool IsAdditionalFilesNeedLinkerSize() const override { return true; }
-	virtual bool IsLinkerAdditionalDataInSeparateArchive() const override { return true; }
+	// IPackageWriter
+	virtual FCapabilities GetCapabilities() const override
+	{
+		FCapabilities Result;
+		Result.bAdditionalFilesNeedLinkerSize = true;
+		Result.bLinkerAdditionalDataInSeparateArchive = true;
+		return Result;
+	}
 
 	virtual void WritePackageData(const FPackageInfo& Info, const FIoBuffer& PackageData, const TArray<FFileRegion>& FileRegions) override;
 	virtual void WriteBulkdata(const FBulkDataInfo& Info, const FIoBuffer& BulkData, const TArray<FFileRegion>& FileRegions) override;
 	virtual void WriteLinkerAdditionalData(const FLinkerAdditionalDataInfo& Info, const FIoBuffer& Data, const TArray<FFileRegion>& FileRegions) override;
-	virtual void Flush() override
-	{
-	}
 	virtual void BeginPackage(const FBeginPackageInfo& Info) override;
 	virtual void CommitPackage(const FCommitPackageInfo& Info) override
 	{
 	}
-
-	// IPackageStoreWriter cooking and accessor interface: not implemented in this writer
 	virtual bool WriteAdditionalFile(const FAdditionalFileInfo& Info, const FIoBuffer& FileData) override
 	{
+		// WriteAdditionalFile is only used when saving cooked packages, so we do not need to implement it
 		checkNoEntry();
 		return false;
 	}
-	virtual void BeginCook(const FCookInfo& Info) override
-	{
-		checkNoEntry();
-	}
-	virtual void EndCook() override
-	{
-		checkNoEntry();
-	}
-	virtual void GetEntries(TFunction<void(TArrayView<const FPackageStoreEntryResource>)>&&) override
-	{
-		checkNoEntry();
-	}
-	DECLARE_DERIVED_EVENT(FMemoryPackageStoreWriter, IPackageStoreWriter::FCommitEvent, FCommitEvent);
-	virtual FCommitEvent& OnCommit() override
-	{
-		checkNoEntry();
-		return CommitEvent;
-	}
-	virtual void GetCookedPackages(TArray<FCookedPackageInfo>& OutCookedPackages) override
-	{
-		checkNoEntry();
-	}
-	virtual void RemoveCookedPackages(TArrayView<const FName> PackageNamesToRemove) override
-	{
-		checkNoEntry();
-	}
-	virtual FCbObject GetTargetDomainDependencies(FName InPackageName) override
-	{
-		checkNoEntry();
-		return FCbObject();
-	}
 
-	// FMemoryPackageStoreWriter interface
+	// FMemoryPackageWriter interface
 	uint64 GetHeaderSize() const { return HeaderSize; }
 	FSharedBuffer& GetHeaderAndExports() { return HeaderAndExports; }
 	/**
-	 * BulkDatas in this array are views into the FMemoryPackageStoreWriter.
-	 * Callers should call MakeOwned if they make copies that need to outlive the FMemoryPackageStoreWriter.
+	 * BulkDatas in this array are views into the FMemoryPackageWriter.
+	 * Callers should call MakeOwned if they make copies that need to outlive the FMemoryPackageWriter.
 	 */
 	TArray<FSharedBuffer>& GetBulkDatas() { return BulkDataRegions; }
 
@@ -656,12 +626,11 @@ private:
 
 	FSharedBuffer HeaderAndExports;
 	TArray<FSharedBuffer> BulkDataRegions;
-	FCommitEvent CommitEvent;
 	FName PackageName;
 	uint64 HeaderSize = 0;
 };
 
-void FMemoryPackageStoreWriter::SetPackageName(FName InPackageName)
+void FMemoryPackageWriter::SetPackageName(FName InPackageName)
 {
 	if (PackageName.IsNone())
 	{
@@ -669,11 +638,11 @@ void FMemoryPackageStoreWriter::SetPackageName(FName InPackageName)
 	}
 	else
 	{
-		checkf(PackageName == InPackageName, TEXT("FMemoryPackageStoreWriter received different PackageNames in WritePackageData and WriteBulkdata."));
+		checkf(PackageName == InPackageName, TEXT("FMemoryPackageWriter received different PackageNames in WritePackageData and WriteBulkdata."));
 	}
 }
 
-void FMemoryPackageStoreWriter::BeginPackage(const FBeginPackageInfo& Info)
+void FMemoryPackageWriter::BeginPackage(const FBeginPackageInfo& Info)
 {
 	SetPackageName(Info.PackageName);
 }
@@ -687,11 +656,11 @@ FSharedBuffer IoBufferToSharedBuffer(const FIoBuffer& InBuffer)
 	return FSharedBuffer{ FSharedBuffer::TakeOwnership(DataPtr, DataSize, FMemory::Free) };
 };
 
-void FMemoryPackageStoreWriter::WritePackageData(const FPackageInfo& Info, const FIoBuffer& PackageData, const TArray<FFileRegion>& FileRegions)
+void FMemoryPackageWriter::WritePackageData(const FPackageInfo& Info, const FIoBuffer& PackageData, const TArray<FFileRegion>& FileRegions)
 {
 	for (const FFileRegion& FileRegion : FileRegions)
 	{
-		checkf(FileRegion.Type == EFileRegionType::None, TEXT("FMemoryPackageStoreWriter does not currently support FileRegion types other than None."));
+		checkf(FileRegion.Type == EFileRegionType::None, TEXT("FMemoryPackageWriter does not currently support FileRegion types other than None."));
 	}
 	SetPackageName(Info.PackageName);
 	// Info.LooseFilePath is ignored
@@ -700,26 +669,26 @@ void FMemoryPackageStoreWriter::WritePackageData(const FPackageInfo& Info, const
 	HeaderAndExports = IoBufferToSharedBuffer(PackageData);
 }
 
-void FMemoryPackageStoreWriter::WriteBulkdata(const FBulkDataInfo& Info, const FIoBuffer& BulkData, const TArray<FFileRegion>& FileRegions)
+void FMemoryPackageWriter::WriteBulkdata(const FBulkDataInfo& Info, const FIoBuffer& BulkData, const TArray<FFileRegion>& FileRegions)
 {
 	SetPackageName(Info.PackageName);
 	// Info.LooseFilePath is ignored
 	// Info.ChunkId is ignored
-	checkf(Info.BulkdataType == IPackageStoreWriter::FBulkDataInfo::Standard, TEXT("MemoryPackageStoreWriter does not currently support BulkData types other than Standard."));
+	checkf(Info.BulkdataType == IPackageWriter::FBulkDataInfo::Standard, TEXT("MemoryPackageWriter does not currently support BulkData types other than Standard."));
 
 	FSharedBuffer BulkDataOwner = IoBufferToSharedBuffer(BulkData);
 	const uint8* BulkDataStart = reinterpret_cast<const uint8*>(BulkDataOwner.GetData());
 	uint64 BulkDataLen = BulkDataOwner.GetSize();
 	for (const FFileRegion& FileRegion : FileRegions)
 	{
-		checkf(FileRegion.Type == EFileRegionType::None, TEXT("FMemoryPackageStoreWriter does not currently support FileRegion types other than None."));
+		checkf(FileRegion.Type == EFileRegionType::None, TEXT("FMemoryPackageWriter does not currently support FileRegion types other than None."));
 		checkf(FileRegion.Offset + FileRegion.Length <= BulkDataLen, TEXT("FileRegions in WriteBulkdata were outside of the range of the BulkData's size."));
 		check(FileRegion.Length > 0); // SavePackage is not allowed to call WriteBulkData with empty bulkdatas
 		BulkDataRegions.Add(FSharedBuffer::MakeView(BulkDataStart + FileRegion.Offset, FileRegion.Length, BulkDataOwner));
 	}
 }
 
-void FMemoryPackageStoreWriter::WriteLinkerAdditionalData(const FLinkerAdditionalDataInfo& Info, const FIoBuffer& Data, const TArray<FFileRegion>& FileRegions)
+void FMemoryPackageWriter::WriteLinkerAdditionalData(const FLinkerAdditionalDataInfo& Info, const FIoBuffer& Data, const TArray<FFileRegion>& FileRegions)
 {
 	SetPackageName(Info.PackageName);
 
@@ -728,7 +697,7 @@ void FMemoryPackageStoreWriter::WriteLinkerAdditionalData(const FLinkerAdditiona
 	uint64 DataLen = DataOwner.GetSize();
 	for (const FFileRegion& FileRegion : FileRegions)
 	{
-		checkf(FileRegion.Type == EFileRegionType::None, TEXT("FMemoryPackageStoreWriter does not currently support FileRegion types other than None."));
+		checkf(FileRegion.Type == EFileRegionType::None, TEXT("FMemoryPackageWriter does not currently support FileRegion types other than None."));
 		checkf(FileRegion.Offset + FileRegion.Length <= DataLen, TEXT("FileRegions in WriteLinkerAdditionalData were outside of the range of the Data's size."));
 		check(FileRegion.Length > 0); // SavePackage is not allowed to call WriteLinkerAdditionalData with empty regions
 		BulkDataRegions.Add(FSharedBuffer::MakeView(DataStart + FileRegion.Offset, FileRegion.Length, DataOwner));
@@ -761,7 +730,7 @@ bool TrySavePackage(UPackage* Package)
 
 	uint32 SaveFlags = SAVE_NoError // Do not crash the SaveServer on an error
 		| SAVE_BulkDataByReference	// EditorDomain saves reference bulkdata from the WorkspaceDomain rather than duplicating it
-		| SAVE_Async				// SavePackage support for PackageStoreWriter is only implemented with SAVE_Async
+		| SAVE_Async				// SavePackage support for PackageWriter is only implemented with SAVE_Async
 		// EDITOR_DOMAIN_TODO: Add a a save flag that specifies the creation of a deterministic guid
 		// | SAVE_KeepGUID;			// Prevent indeterminism by keeping the Guid
 		;
@@ -788,9 +757,9 @@ bool TrySavePackage(UPackage* Package)
 		SaveFlags |= bSaveUnversioned ? SAVE_Unversioned : 0;
 	}
 
-	FMemoryPackageStoreWriter* PackageStoreWriter = new FMemoryPackageStoreWriter();
-	FSavePackageContext SavePackageContext(nullptr /* TargetPlatform */, PackageStoreWriter, false /* bInForceLegacyOffsets */);
-	FSavePackageResultStruct Result = GEditor->Save(Package, nullptr, RF_Standalone, TEXT("EditorDomainPackageStoreWriter"),
+	FMemoryPackageWriter* PackageWriter = new FMemoryPackageWriter();
+	FSavePackageContext SavePackageContext(nullptr /* TargetPlatform */, PackageWriter, false /* bInForceLegacyOffsets */);
+	FSavePackageResultStruct Result = GEditor->Save(Package, nullptr, RF_Standalone, TEXT("EditorDomainPackageWriter"),
 		GError, nullptr /* Conform */, false /* bForceByteSwapping */, true /* bWarnOfLongFilename */,
 		SaveFlags, nullptr /* TargetPlatform */, FDateTime::MinValue(), false /* bSlowTask */,
 		nullptr /* DiffMap */, &SavePackageContext);
@@ -816,14 +785,14 @@ bool TrySavePackage(UPackage* Package)
 	};
 
 	int32 AttachmentIndex = 1; // 0 is not a valid value for IntToPayloadId
-	const FSharedBuffer& ExportsBuffer = PackageStoreWriter->GetHeaderAndExports();
+	const FSharedBuffer& ExportsBuffer = PackageWriter->GetHeaderAndExports();
 	check(ExportsBuffer.GetSize() > 0); // Header+Exports segment is non-zero in length
 	RecordBuilder.AddAttachment(ExportsBuffer, IntToPayloadId(AttachmentIndex++));
 	uint64 FileSize = ExportsBuffer.GetSize();
-	for (const FSharedBuffer& BulkBuffer : PackageStoreWriter->GetBulkDatas())
+	for (const FSharedBuffer& BulkBuffer : PackageWriter->GetBulkDatas())
 	{
 		int64 BulkSize = BulkBuffer.GetSize();
-		check(BulkSize > 0); // We checked this before adding the Region to the PackageStoreWriter
+		check(BulkSize > 0); // We checked this before adding the Region to the PackageWriter
 		RecordBuilder.AddAttachment(BulkBuffer, IntToPayloadId(AttachmentIndex++));
 		FileSize += BulkSize;
 	}
diff --git a/Engine/Source/Editor/UnrealEd/Private/TargetDomain/TargetDomainUtils.cpp b/Engine/Source/Editor/UnrealEd/Private/TargetDomain/TargetDomainUtils.cpp
--- a/Engine/Source/Editor/UnrealEd/Private/TargetDomain/TargetDomainUtils.cpp
+++ b/Engine/Source/Editor/UnrealEd/Private/TargetDomain/TargetDomainUtils.cpp
@@ -8,11 +8,11 @@
 #include "Cooker/PackageBuildDependencyTracker.h"
 #include "EditorDomain/EditorDomainUtils.h"
 #include "IO/IoHash.h"
-#include "IO/PackageStoreWriter.h"
 #include "Misc/ScopeRWLock.h"
 #include "Misc/StringBuilder.h"
 #include "Serialization/CompactBinary.h"
 #include "Serialization/CompactBinaryWriter.h"
+#include "Serialization/PackageWriter.h"
 
 namespace UE::TargetDomain
 {
@@ -169,10 +169,10 @@ FCbObject CollectDependenciesObject(UPackage* Package, const ITargetPlatform* Ta
 }
 
 
-bool TryFetchKeyAndDependencies(IPackageStoreWriter* PackageStore, FName PackageName, const ITargetPlatform* TargetPlatform,
+bool TryFetchKeyAndDependencies(ICookedPackageWriter* PackageWriter, FName PackageName, const ITargetPlatform* TargetPlatform,
 	FIoHash* OutHash, TArray<FName>* OutBuildDependencies, TArray<FName>* OutRuntimeOnlyDependencies, FString* OutErrorMessage)
 {
-	FCbObject DependenciesObj = PackageStore->GetTargetDomainDependencies(PackageName);
+	FCbObject DependenciesObj = PackageWriter->GetTargetDomainDependencies(PackageName);
 	FIoHash StoredKey = DependenciesObj["targetdomainkey"].AsHash();
 	if (StoredKey.IsZero())
 	{
diff --git a/Engine/Source/Editor/UnrealEd/Private/TargetDomain/TargetDomainUtils.h b/Engine/Source/Editor/UnrealEd/Private/TargetDomain/TargetDomainUtils.h
--- a/Engine/Source/Editor/UnrealEd/Private/TargetDomain/TargetDomainUtils.h
+++ b/Engine/Source/Editor/UnrealEd/Private/TargetDomain/TargetDomainUtils.h
@@ -5,7 +5,7 @@
 #include "Containers/Array.h"
 #include "UObject/NameTypes.h"
 
-class IPackageStoreWriter;
+class ICookedPackageWriter;
 class ITargetPlatform;
 class FCbObject;
 class FString;
@@ -28,7 +28,7 @@ FCbObject CollectDependenciesObject(UPackage* Package, const ITargetPlatform* Ta
  * Read the oplog for the given packagename and fetch the dependencies and key out of it. Use the dependencies to calculate
  * the current key, and return the dependencies and key only if they key matches.
  */
-bool TryFetchKeyAndDependencies(IPackageStoreWriter* PackageStore, FName PackageName, const ITargetPlatform* TargetPlatform,
+bool TryFetchKeyAndDependencies(ICookedPackageWriter* PackageStore, FName PackageName, const ITargetPlatform* TargetPlatform,
 	FIoHash* OutHash, TArray<FName>* OutBuildDependencies, TArray<FName>* OutRuntimeOnlyDependencies, FString* OutErrorMessage);
 
 /** Return whether iterative cook is enabled for the given packagename, based on used-class allowlist/blocklist. */
diff --git a/Engine/Source/Runtime/Core/Public/IO/PackageStoreWriter.h b/Engine/Source/Runtime/Core/Public/IO/PackageStoreWriter.h
deleted file mode 100644
diff --git a/Engine/Source/Runtime/Core/Public/Serialization/PackageWriter.h b/Engine/Source/Runtime/Core/Public/Serialization/PackageWriter.h
new file mode 100644
--- /dev/null
+++ b/Engine/Source/Runtime/Core/Public/Serialization/PackageWriter.h
@@ -0,0 +1,200 @@
+// Copyright Epic Games, Inc. All Rights Reserved.
+
+#pragma once
+
+#include "Containers/StringView.h"
+#include "IO/IoDispatcher.h"
+#include "Misc/SecureHash.h"
+#include "Serialization/CompactBinary.h"
+
+class IPackageStoreWriter;
+struct FPackageStoreEntryResource;
+
+/** Interface for SavePackage to write packages to storage. */
+class IPackageWriter
+{
+public:
+	virtual ~IPackageWriter() = default;
+
+
+	struct FCapabilities
+	{
+		/** Whether this writer needs BulkDatas to be written after the Linker's archive has finalized its size.
+
+			Some PackageWriters need that behavior because they put the BulkDatas in a segment following
+			the exports in a Composite archive.
+		 */
+		bool bAdditionalFilesNeedLinkerSize = false;
+		/** Whether data stored in Linker.AdditionalDataToAppend should be serialized to a separate archive.
+
+			If false, the data will be serialized to the end of the LinkerSave archive instead.
+		*/
+		bool bLinkerAdditionalDataInSeparateArchive = false;
+	};
+
+	/** Return capabilities/settings this PackageWriter has/requires 
+	  */
+	virtual FCapabilities GetCapabilities() const
+	{
+		return FCapabilities();
+	}
+
+	// Events the PackageWriter receives
+	struct FBeginPackageInfo
+	{
+		FName	PackageName;
+	};
+
+	/** Mark the beginning of a package store transaction for the specified package
+
+		This must be called before any data is produced for a given package
+	  */
+	virtual void BeginPackage(const FBeginPackageInfo& Info) = 0;
+
+	struct FCommitPackageInfo
+	{
+		FName PackageName;
+		FGuid PackageGuid;
+		FCbObject TargetDomainDependencies;
+		bool bSucceeded = false;
+	};
+
+	/** Finalize a package started with BeginPackage()
+	  */
+	virtual void CommitPackage(const FCommitPackageInfo& Info) = 0;
+
+	struct FPackageInfo
+	{
+		FName		PackageName;
+		FString		LooseFilePath;
+		uint64		HeaderSize;
+		FIoChunkId	ChunkId = FIoChunkId::InvalidChunkId;
+	};
+
+	/** Write package data (exports and serialized header)
+
+		This may only be called after a BeginPackage() call has been made
+		to signal the start of a package store transaction
+	  */
+	virtual void WritePackageData(const FPackageInfo& Info, const FIoBuffer& PackageData, const TArray<FFileRegion>& FileRegions) = 0;
+
+	struct FBulkDataInfo
+	{
+		enum EType
+		{
+			Standard,
+			Mmap,
+			Optional
+		};
+
+		FName		PackageName;
+		EType		BulkdataType = Standard;
+		FString		LooseFilePath;
+		FIoChunkId	ChunkId = FIoChunkId::InvalidChunkId;
+	};
+
+	/** Write bulk data for the current package
+	  */
+	virtual void WriteBulkdata(const FBulkDataInfo& Info, const FIoBuffer& BulkData, const TArray<FFileRegion>& FileRegions) = 0;
+
+	struct FAdditionalFileInfo
+	{
+		FName	PackageName;
+		FString	Filename;
+		FIoChunkId ChunkId;
+	};
+
+	/** Write separate files written by UObjects during cooking via UObject::CookAdditionalFiles. */
+	virtual bool WriteAdditionalFile(const FAdditionalFileInfo& Info, const FIoBuffer& FileData) = 0;
+
+	struct FLinkerAdditionalDataInfo
+	{
+		FName PackageName;
+	};
+	/** Write separate data written by UObjects via FLinkerSave::AdditionalDataToAppend.
+
+		This function will not be called unless GetCapabilities().bLinkerAdditionalDataInSeparateArchive is true.
+		If that function is false, the data was inlined into the PackageData instead.
+	*/
+	virtual void WriteLinkerAdditionalData(const FLinkerAdditionalDataInfo& Info, const FIoBuffer& Data, const TArray<FFileRegion>& FileRegions) = 0;
+};
+
+/** Interface for cooking that writes cooked packages to storage usable by the runtime game. */
+class ICookedPackageWriter : public IPackageWriter
+{
+public:
+	virtual ~ICookedPackageWriter() = default;
+
+	struct FCookInfo
+	{
+		enum ECookMode
+		{
+			CookByTheBookMode,
+			CookOnTheFlyMode
+		};
+
+		ECookMode CookMode = CookByTheBookMode;
+	};
+
+	/** Signal the start of a cooking pass
+
+		Package data may only be produced after BeginCook() has been called and
+		before EndCook() is called
+	  */
+	virtual void BeginCook(const FCookInfo& Info) = 0;
+
+	/** Signal the end of a cooking pass.
+	  */
+	virtual void EndCook() = 0;
+
+	/**
+	 * Flush any outstanding writes.
+	 */
+	virtual void Flush() = 0;
+
+	struct FCookedPackageInfo
+	{
+		FName PackageName;
+		FMD5Hash Hash;
+		FGuid PackageGuid;
+		int64 DiskSize = -1;
+		FIoHash TargetDomainDependencies;
+	};
+
+	/**
+	 * Returns a list of cooked package(s).
+	 */
+	virtual void GetCookedPackages(TArray<FCookedPackageInfo>& OutCookedPackages) = 0;
+
+	/**
+	 * Returns the TargetDomainDependencies that were previously commited for the given PackageName.
+	 * Returns an empty object if not found.
+	 */
+	virtual FCbObject GetTargetDomainDependencies(FName PackageName) = 0;
+
+	/**
+	 * Remove cooked package(s) that has been modified since the last cook.
+	 */
+	virtual void RemoveCookedPackages(TArrayView<const FName> PackageNamesToRemove) = 0;
+
+	/** Downcast function for ICookedPackageWriters that implement the IPackageStoreWriter inherited interface. */
+	virtual IPackageStoreWriter* AsPackageStoreWriter()
+	{
+		return nullptr;
+	}
+};
+
+static inline const ANSICHAR* LexToString(IPackageWriter::FBulkDataInfo::EType Value)
+{
+	switch (Value)
+	{
+	case IPackageWriter::FBulkDataInfo::Standard:
+		return "Standard";
+	case IPackageWriter::FBulkDataInfo::Mmap:
+		return "Mmap";
+	case IPackageWriter::FBulkDataInfo::Optional:
+		return "Optional";
+	}
+
+	return "Unknown";
+}
diff --git a/Engine/Source/Runtime/CoreUObject/Private/UObject/SavePackage.cpp b/Engine/Source/Runtime/CoreUObject/Private/UObject/SavePackage.cpp
--- a/Engine/Source/Runtime/CoreUObject/Private/UObject/SavePackage.cpp
+++ b/Engine/Source/Runtime/CoreUObject/Private/UObject/SavePackage.cpp
@@ -63,12 +63,12 @@
 #include "Serialization/Formatters/BinaryArchiveFormatter.h"
 #include "Serialization/Formatters/JsonArchiveOutputFormatter.h"
 #include "Serialization/ArchiveUObjectFromStructuredArchive.h"
+#include "Serialization/PackageWriter.h"
 #include "Serialization/UnversionedPropertySerialization.h"
 #include "UObject/AsyncWorkSequence.h"
 #include "Misc/ScopeExit.h"
 #include "Misc/PackageAccessTracking.h"
 #include "Misc/PackageAccessTrackingOps.h"
-#include "IO/PackageStoreWriter.h"
 
 #if ENABLE_COOK_STATS
 #include "ProfilingDebugging/ScopedTimers.h"
@@ -4162,7 +4162,7 @@ FSavePackageResultStruct UPackage::Save(UPackage* InOuter, UObject* Base, EObjec
 
 				SlowTask.EnterProgressFrame(1, NSLOCTEXT("Core", "SerializingBulkData", "Serializing bulk data"));
 
-				IPackageStoreWriter* PackageStoreWriter = SavePackageContext ? SavePackageContext->PackageStoreWriter : nullptr;
+				IPackageWriter* PackageWriter = SavePackageContext ? SavePackageContext->PackageWriter : nullptr;
 				FSavePackageOutputFileArray AdditionalOutputFiles;
 				auto WriteAdditionalFiles = [&](int64 LinkerSize) -> ESavePackageResult
 				{
@@ -4175,7 +4175,7 @@ FSavePackageResultStruct UPackage::Save(UPackage* InOuter, UObject* Base, EObjec
 						return SaveResult;
 					}
 
-					ESavePackageResult AppendAdditionalDataResult = SavePackageUtilities::AppendAdditionalData(*Linker.Get(), DataStartOffset, PackageStoreWriter);
+					ESavePackageResult AppendAdditionalDataResult = SavePackageUtilities::AppendAdditionalData(*Linker.Get(), DataStartOffset, SavePackageContext);
 					if (AppendAdditionalDataResult != ESavePackageResult::Success)
 					{
 						return AppendAdditionalDataResult;
@@ -4218,7 +4218,7 @@ FSavePackageResultStruct UPackage::Save(UPackage* InOuter, UObject* Base, EObjec
 					return ESavePackageResult::Success;
 				};
 			
-				bool bAdditionalFilesNeedLinkerSize = PackageStoreWriter ? PackageStoreWriter->IsAdditionalFilesNeedLinkerSize() : false;
+				bool bAdditionalFilesNeedLinkerSize = PackageWriter ? SavePackageContext->PackageWriterCapabilities.bAdditionalFilesNeedLinkerSize : false;
 				if (!bAdditionalFilesNeedLinkerSize)
 				{
 					ESavePackageResult AdditionalFilesResult = WriteAdditionalFiles(-1);
@@ -4246,7 +4246,7 @@ FSavePackageResultStruct UPackage::Save(UPackage* InOuter, UObject* Base, EObjec
 					{
 						return AdditionalFilesResult;
 					}
-					checkf(Linker->Tell() == PackageSize, TEXT("The writing of additional files is not allowed to append to the LinkerSave when IsAdditionalFilesNeedLinkerSize is true."));
+					checkf(Linker->Tell() == PackageSize, TEXT("The writing of additional files is not allowed to append to the LinkerSave when bAdditionalFilesNeedLinkerSize is true."));
 				}
 
 				// Save the import map.
@@ -4512,7 +4512,7 @@ FSavePackageResultStruct UPackage::Save(UPackage* InOuter, UObject* Base, EObjec
 							FLargeMemoryWriter* Writer = static_cast<FLargeMemoryWriter*>(Linker->Saver);
 							const int64 DataSize = Writer->TotalSize();
 
-							if (PackageStoreWriter)
+							if (PackageWriter)
 							{
 								checkf(AdditionalOutputFiles.IsEmpty(), TEXT("Saving additional output files during cooking is not supported!"));
 
@@ -4534,7 +4534,7 @@ FSavePackageResultStruct UPackage::Save(UPackage* InOuter, UObject* Base, EObjec
 
 								const int32 HeaderSize = Linker->Summary.TotalHeaderSize;
 
-								IPackageStoreWriter::FPackageInfo PackageInfo;
+								IPackageWriter::FPackageInfo PackageInfo;
 								PackageInfo.PackageName = InOuter->GetFName();
 								PackageInfo.LooseFilePath = Filename;
 								PackageInfo.HeaderSize = HeaderSize;
@@ -4542,7 +4542,7 @@ FSavePackageResultStruct UPackage::Save(UPackage* InOuter, UObject* Base, EObjec
 								FPackageId PackageId = FPackageId::FromName(PackageInfo.PackageName);
 								PackageInfo.ChunkId = CreateIoChunkId(PackageId.Value(), 0, EIoChunkType::ExportBundleData);
 
-								PackageStoreWriter->WritePackageData(PackageInfo, IoBuffer, Linker->FileRegions);
+								PackageWriter->WritePackageData(PackageInfo, IoBuffer, Linker->FileRegions);
 							}
 							else if (IsEventDrivenLoaderEnabledInCookedBuilds() && Linker->IsCooking())
 							{
@@ -4584,7 +4584,7 @@ FSavePackageResultStruct UPackage::Save(UPackage* InOuter, UObject* Base, EObjec
 #if WITH_EDITOR
 						else
 						{
-							checkf(!PackageStoreWriter, TEXT("PackageStoreWriter is not currently supported when diffing."));
+							checkf(!PackageWriter, TEXT("PackageWriter is not currently supported when diffing."));
 						}
 #endif
 						Linker->CloseAndDestroySaver();
@@ -4596,7 +4596,7 @@ FSavePackageResultStruct UPackage::Save(UPackage* InOuter, UObject* Base, EObjec
 					// Move the temporary file.
 					else
 					{
-						checkf(!PackageStoreWriter, TEXT("PackageStoreWriter is not currently supported with synchronous writes."));
+						checkf(!PackageWriter, TEXT("PackageWriter is not currently supported with synchronous writes."));
 						checkf(TempFilename.IsSet(), TEXT("The package should've been saved to a tmp file first!"));
 
 						// When saving in text format we will have two temp files, so we need to manually delete the non-textbased one
@@ -4804,7 +4804,7 @@ bool UPackage::SavePackage(UPackage* InOuter, UObject* Base, EObjectFlags TopLev
 
 FSavePackageContext::~FSavePackageContext()
 {
-	delete PackageStoreWriter;
+	delete PackageWriter;
 }
 
 #endif	// UE_WITH_SAVEPACKAGE
diff --git a/Engine/Source/Runtime/CoreUObject/Private/UObject/SavePackage/SaveContext.cpp b/Engine/Source/Runtime/CoreUObject/Private/UObject/SavePackage/SaveContext.cpp
--- a/Engine/Source/Runtime/CoreUObject/Private/UObject/SavePackage/SaveContext.cpp
+++ b/Engine/Source/Runtime/CoreUObject/Private/UObject/SavePackage/SaveContext.cpp
@@ -2,7 +2,7 @@
 
 #include "SaveContext.h"
 
-#include "IO/PackageStoreWriter.h"
+#include "Serialization/PackageWriter.h"
 
 void FSaveContext::MarkUnsaveable(UObject* InObject)
 {
@@ -59,6 +59,12 @@ bool FSaveContext::IsUnsaveable(UObject* InObject) const
 
 bool FSaveContext::IsAdditionalFilesNeedLinkerSize() const
 {
-	IPackageStoreWriter* PackageStoreWriter = GetPackageStoreWriter();
-	return PackageStoreWriter ? PackageStoreWriter->IsAdditionalFilesNeedLinkerSize() : false;
+	if (SaveArgs.SavePackageContext && SaveArgs.SavePackageContext->PackageWriter)
+	{
+		return SaveArgs.SavePackageContext->PackageWriterCapabilities.bAdditionalFilesNeedLinkerSize;
+	}
+	else
+	{
+		return false;
+	}
 }
diff --git a/Engine/Source/Runtime/CoreUObject/Private/UObject/SavePackage/SaveContext.h b/Engine/Source/Runtime/CoreUObject/Private/UObject/SavePackage/SaveContext.h
--- a/Engine/Source/Runtime/CoreUObject/Private/UObject/SavePackage/SaveContext.h
+++ b/Engine/Source/Runtime/CoreUObject/Private/UObject/SavePackage/SaveContext.h
@@ -548,9 +548,9 @@ public:
 		return ObjectSaveContext;
 	}
 
-	IPackageStoreWriter* GetPackageStoreWriter() const
+	IPackageWriter* GetPackageWriter() const
 	{
-		return SaveArgs.SavePackageContext ? SaveArgs.SavePackageContext->PackageStoreWriter : nullptr;
+		return SaveArgs.SavePackageContext ? SaveArgs.SavePackageContext->PackageWriter : nullptr;
 	}
 
 	bool IsAdditionalFilesNeedLinkerSize() const;
diff --git a/Engine/Source/Runtime/CoreUObject/Private/UObject/SavePackage/SavePackageUtilities.cpp b/Engine/Source/Runtime/CoreUObject/Private/UObject/SavePackage/SavePackageUtilities.cpp
--- a/Engine/Source/Runtime/CoreUObject/Private/UObject/SavePackage/SavePackageUtilities.cpp
+++ b/Engine/Source/Runtime/CoreUObject/Private/UObject/SavePackage/SavePackageUtilities.cpp
@@ -16,6 +16,7 @@
 #include "SaveContext.h"
 #include "Serialization/BulkData.h"
 #include "Serialization/LargeMemoryWriter.h"
+#include "Serialization/PackageWriter.h"
 #include "UObject/AsyncWorkSequence.h"
 #include "UObject/Class.h"
 #include "UObject/GCScopeLock.h"
@@ -31,7 +32,6 @@
 #include "UObject/UObjectHash.h"
 #include "UObject/UObjectThreadContext.h"
 #include "Virtualization/VirtualizedBulkData.h"
-#include "IO/PackageStoreWriter.h"
 
 DEFINE_LOG_CATEGORY(LogSavePackage);
 UE_TRACE_CHANNEL_DEFINE(SaveTimeChannel);
@@ -2129,14 +2129,15 @@ public:
 	TArray<FFileRegion> FileRegions;
 };
 
-ESavePackageResult AppendAdditionalData(FLinkerSave& Linker, int64& InOutDataStartOffset, IPackageStoreWriter* PackageStoreWriter)
+ESavePackageResult AppendAdditionalData(FLinkerSave& Linker, int64& InOutDataStartOffset, FSavePackageContext* SavePackageContext)
 {
 	if (Linker.AdditionalDataToAppend.Num() == 0)
 	{
 		return ESavePackageResult::Success;
 	}
 
-	if (PackageStoreWriter && PackageStoreWriter->IsLinkerAdditionalDataInSeparateArchive())
+	IPackageWriter* PackageWriter = SavePackageContext ? SavePackageContext->PackageWriter : nullptr;
+	if (PackageWriter && SavePackageContext->PackageWriterCapabilities.bLinkerAdditionalDataInSeparateArchive)
 	{
 		FLargeMemoryWriterWithRegions DataArchive;
 		for (FLinkerSave::AdditionalDataCallback& Callback : Linker.AdditionalDataToAppend)
@@ -2149,10 +2150,10 @@ ESavePackageResult AppendAdditionalData(FLinkerSave& Linker, int64& InOutDataSta
 				DataArchive.FileRegions.Add(FFileRegion(RegionStart, RegionEnd - RegionStart, EFileRegionType::None));
 			}
 		}
-		IPackageStoreWriter::FLinkerAdditionalDataInfo DataInfo{ Linker.LinkerRoot->GetFName() };
+		IPackageWriter::FLinkerAdditionalDataInfo DataInfo{ Linker.LinkerRoot->GetFName() };
 		int64 DataSize = DataArchive.TotalSize();
 		FIoBuffer DataBuffer = FIoBuffer(FIoBuffer::AssumeOwnership, DataArchive.ReleaseOwnership(), DataSize);
-		PackageStoreWriter->WriteLinkerAdditionalData(DataInfo, DataBuffer, DataArchive.FileRegions);
+		PackageWriter->WriteLinkerAdditionalData(DataInfo, DataBuffer, DataArchive.FileRegions);
 		InOutDataStartOffset += DataSize;
 	}
 	else
@@ -2261,7 +2262,7 @@ ESavePackageResult SaveBulkData(FLinkerSave* Linker, int64& InOutStartOffset, co
 {
 	// Now we write all the bulkdata that is supposed to be at the end of the package
 	// and fix up the offset
-	IPackageStoreWriter* PackageStoreWriter = SavePackageContext != nullptr ? SavePackageContext->PackageStoreWriter : nullptr;
+	IPackageWriter* PackageWriter = SavePackageContext != nullptr ? SavePackageContext->PackageWriter : nullptr;
 	Linker->Summary.BulkDataStartOffset = InOutStartOffset;
 
 	if (Linker->BulkDataToAppend.Num() == 0)
@@ -2296,7 +2297,7 @@ ESavePackageResult SaveBulkData(FLinkerSave* Linker, int64& InOutStartOffset, co
 	const bool bSeparateSegmentsEnabled = CookerUseSeparateSegments.bEnable && Linker->IsCooking();
 
 	int64 LinkerStart = 0;
-	if (PackageStoreWriter || bSeparateSegmentsEnabled)
+	if (PackageWriter || bSeparateSegmentsEnabled)
 	{
 		BulkArchive.Reset(new FLargeMemoryWriterWithRegions);
 		if (bSeparateSegmentsEnabled)
@@ -2322,7 +2323,7 @@ ESavePackageResult SaveBulkData(FLinkerSave* Linker, int64& InOutStartOffset, co
 		bUseFileRegions = TargetPlatform->SupportsFeature(ETargetPlatformFeatures::CookFileRegionMetadata);
 		BulkDataAlignment = TargetPlatform->GetMemoryMappingAlignment();
 	}
-	else if (PackageStoreWriter)
+	else if (PackageWriter)
 	{
 		bUseFileRegions = true;
 		bOneBulkPerRegion = true;
@@ -2405,7 +2406,7 @@ ESavePackageResult SaveBulkData(FLinkerSave* Linker, int64& InOutStartOffset, co
 		{
 			SaveLocation = ESaveLocation::SeparateSegment;
 		}
-		else if (PackageStoreWriter)
+		else if (PackageWriter)
 		{
 			SaveLocation = ESaveLocation::SeparateArchiveAtEndOfPackage;
 		}
@@ -2550,7 +2551,7 @@ ESavePackageResult SaveBulkData(FLinkerSave* Linker, int64& InOutStartOffset, co
 	{
 		const bool bWriteBulkToDisk = !bDiffing;
 
-		if (PackageStoreWriter && bWriteBulkToDisk)
+		if (PackageWriter && bWriteBulkToDisk)
 		{
 			auto AddSizeAndConvertToIoBuffer = [&TotalPackageSizeUncompressed](FLargeMemoryWriter* Writer)
 			{
@@ -2559,7 +2560,7 @@ ESavePackageResult SaveBulkData(FLinkerSave* Linker, int64& InOutStartOffset, co
 				return FIoBuffer(FIoBuffer::AssumeOwnership, Writer->ReleaseOwnership(), TotalSize);
 			};
 
-			IPackageStoreWriter::FBulkDataInfo BulkInfo;
+			IPackageWriter::FBulkDataInfo BulkInfo;
 			BulkInfo.PackageName = InOuter->GetFName();
 
 			FPackageId PackageId = FPackageId::FromName(BulkInfo.PackageName);
@@ -2567,23 +2568,23 @@ ESavePackageResult SaveBulkData(FLinkerSave* Linker, int64& InOutStartOffset, co
 			if (BulkArchive->TotalSize())
 			{
 				BulkInfo.ChunkId = CreateIoChunkId(PackageId.Value(), 0, EIoChunkType::BulkData);
-				BulkInfo.BulkdataType = IPackageStoreWriter::FBulkDataInfo::Standard;
+				BulkInfo.BulkdataType = IPackageWriter::FBulkDataInfo::Standard;
 				BulkInfo.LooseFilePath = FPaths::ChangeExtension(Filename, LexToString(EPackageExtension::BulkDataDefault));
-				PackageStoreWriter->WriteBulkdata(BulkInfo, AddSizeAndConvertToIoBuffer(BulkArchive.Get()), BulkArchive->FileRegions);
+				PackageWriter->WriteBulkdata(BulkInfo, AddSizeAndConvertToIoBuffer(BulkArchive.Get()), BulkArchive->FileRegions);
 			}
 			if (OptionalBulkArchive && OptionalBulkArchive->TotalSize())
 			{
 				BulkInfo.ChunkId = CreateIoChunkId(PackageId.Value(), 0, EIoChunkType::OptionalBulkData);
-				BulkInfo.BulkdataType = IPackageStoreWriter::FBulkDataInfo::Optional;
+				BulkInfo.BulkdataType = IPackageWriter::FBulkDataInfo::Optional;
 				BulkInfo.LooseFilePath = FPaths::ChangeExtension(Filename, LexToString(EPackageExtension::BulkDataOptional));
-				PackageStoreWriter->WriteBulkdata(BulkInfo, AddSizeAndConvertToIoBuffer(OptionalBulkArchive.Get()), OptionalBulkArchive->FileRegions);
+				PackageWriter->WriteBulkdata(BulkInfo, AddSizeAndConvertToIoBuffer(OptionalBulkArchive.Get()), OptionalBulkArchive->FileRegions);
 			}
 			if (MappedBulkArchive && MappedBulkArchive->TotalSize())
 			{
 				BulkInfo.ChunkId = CreateIoChunkId(PackageId.Value(), 0, EIoChunkType::MemoryMappedBulkData);
-				BulkInfo.BulkdataType = IPackageStoreWriter::FBulkDataInfo::Mmap;
+				BulkInfo.BulkdataType = IPackageWriter::FBulkDataInfo::Mmap;
 				BulkInfo.LooseFilePath = FPaths::ChangeExtension(Filename, LexToString(EPackageExtension::BulkDataMemoryMapped));
-				PackageStoreWriter->WriteBulkdata(BulkInfo, AddSizeAndConvertToIoBuffer(MappedBulkArchive.Get()), MappedBulkArchive->FileRegions);
+				PackageWriter->WriteBulkdata(BulkInfo, AddSizeAndConvertToIoBuffer(MappedBulkArchive.Get()), MappedBulkArchive->FileRegions);
 			}
 		}
 		else
@@ -2622,7 +2623,7 @@ ESavePackageResult SaveBulkData(FLinkerSave* Linker, int64& InOutStartOffset, co
 
 	if (!bSeparateSegmentsEnabled)
 	{
-		if (PackageStoreWriter)
+		if (PackageWriter)
 		{
 			check(BulkArchive.IsValid() && LinkerStart == 0);
 			InOutStartOffset += BulkArchive->TotalSize();
diff --git a/Engine/Source/Runtime/CoreUObject/Private/UObject/SavePackage/SavePackageUtilities.h b/Engine/Source/Runtime/CoreUObject/Private/UObject/SavePackage/SavePackageUtilities.h
--- a/Engine/Source/Runtime/CoreUObject/Private/UObject/SavePackage/SavePackageUtilities.h
+++ b/Engine/Source/Runtime/CoreUObject/Private/UObject/SavePackage/SavePackageUtilities.h
@@ -21,7 +21,7 @@ class FMD5;
 class FPackagePath;
 class FSaveContext;
 class FSavePackageContext;
-class IPackageStoreWriter;
+class IPackageWriter;
 template<typename StateType> class TAsyncWorkSequence;
 
 enum class ESavePackageResult;
@@ -399,7 +399,7 @@ namespace SavePackageUtilities
 	 *        Bytes written to separate files do not contribute to this offset.
 	 * @param InOuter The package being saved
 	 * @param Filename The filename the package is being saved to. Used to decide the filename of separate archives, if any are used.
-	 * @param SavePackageContext Provides The PackageStoreWriter and other parameters of how the bulkdata should be saved.
+	 * @param SavePackageContext Provides The PackageWriter and other parameters of how the bulkdata should be saved.
 	 * @param SaveFlags The flags passed into SavePackage; affects how the bulkdata should be saved.
 	 * @param bTextFormat True if package is saving to text. Bulkdata has special handling for text output.
 	 * @param bDiffing True if the package is only diffing, so bulk data should be recorded but not saved to disk.
@@ -413,16 +413,16 @@ namespace SavePackageUtilities
 	
 	/**
 	 * Used to append additional data to the end of the package file by invoking callbacks stored in the linker.
-	 * They may be saved to the end of the file, or to a separate archive passed into the PackageStoreWriter.
+	 * They may be saved to the end of the file, or to a separate archive passed into the PackageWriter.
 	 
 	 * @param Linker The linker containing the exports. Provides the list of AdditionalData, and the data may write to it as their target archive.
 	 * @param InOutStartOffset In value is the offset in the Linker's archive where the datas will be put. If SavePackageContext settings direct
 	 *        the datas to write in a separate archive that will be combined after the linker, the value is the offset after the Linker archive's 
 	 *        totalsize and after any previous post-Linker archive data such as BulkDatas.
 	 *        Output value is incremented by the number of bytes written the Linker or the separate archive at the end of the linker.
-	 * @param PackageStoreWriter If non-null and configured to require it, data is passed to this writer rather than appended to the Linker archive.
+	 * @param SavePackageContext If non-null and configured to require it, data is passed to this PackageWriter on this context rather than appended to the Linker archive.
 	  */
-	ESavePackageResult AppendAdditionalData(FLinkerSave& Linker, int64& InOutDataStartOffset, IPackageStoreWriter* PackageStoreWriter);
+	ESavePackageResult AppendAdditionalData(FLinkerSave& Linker, int64& InOutDataStartOffset, FSavePackageContext* SavePackageContext);
 	
 	/** Used to create the sidecar file (.upayload) from payloads that have been added to the linker */
 	ESavePackageResult CreatePayloadSidecarFile(FLinkerSave& Linker, const FPackagePath& PackagePath, const bool bSaveAsync, const bool bWriteToDisk, FSavePackageOutputFileArray& AdditionalPackageFiles);
diff --git a/Engine/Source/Runtime/CoreUObject/Private/UObject/SavePackage2.cpp b/Engine/Source/Runtime/CoreUObject/Private/UObject/SavePackage2.cpp
--- a/Engine/Source/Runtime/CoreUObject/Private/UObject/SavePackage2.cpp
+++ b/Engine/Source/Runtime/CoreUObject/Private/UObject/SavePackage2.cpp
@@ -23,6 +23,7 @@
 #include "Serialization/ArchiveStackTrace.h"
 #include "Serialization/Formatters/JsonArchiveOutputFormatter.h"
 #include "Serialization/LargeMemoryWriter.h"
+#include "Serialization/PackageWriter.h"
 #include "Serialization/PropertyLocalizationDataGathering.h"
 #include "Serialization/UnversionedPropertySerialization.h"
 #include "UObject/AsyncWorkSequence.h"
@@ -39,7 +40,6 @@
 #include "UObject/SavePackage/SavePackageUtilities.h"
 #include "UObject/UObjectGlobals.h"
 #include "UObject/UObjectHash.h"
-#include "IO/PackageStoreWriter.h"
 
 #if ENABLE_COOK_STATS
 #include "ProfilingDebugging/ScopedTimers.h"
@@ -1628,7 +1628,7 @@ ESavePackageResult WriteExports(FStructuredArchive::FRecord& StructuredArchiveRo
 ESavePackageResult WriteAdditionalExportFiles(FSaveContext& SaveContext)
 {
 	FSavePackageContext* SavePackageContext = SaveContext.GetSavePackageContext();
-	IPackageStoreWriter* PackageStoreWriter = SavePackageContext ? SavePackageContext->PackageStoreWriter : nullptr;
+	IPackageWriter* PackageWriter = SavePackageContext ? SavePackageContext->PackageWriter : nullptr;
 
 	if (SaveContext.IsCooking() && SaveContext.AdditionalFilesFromExports.Num() > 0)
 	{
@@ -1639,15 +1639,15 @@ ESavePackageResult WriteAdditionalExportFiles(FSaveContext& SaveContext)
 			const int64 Size = Writer.TotalSize();
 			SaveContext.TotalPackageSizeUncompressed += Size;
 
-			if (PackageStoreWriter)
+			if (PackageWriter)
 			{
-				IPackageStoreWriter::FAdditionalFileInfo FileInfo;
+				IPackageWriter::FAdditionalFileInfo FileInfo;
 				FileInfo.PackageName = SaveContext.GetPackage()->GetFName();
 				FileInfo.Filename = *Writer.GetArchiveName();
 
 				FIoBuffer FileData(FIoBuffer::Wrap, Writer.GetData(), Size);
 
-				const bool bWrittenToPackageStore = PackageStoreWriter->WriteAdditionalFile(FileInfo, FileData);
+				const bool bWrittenToPackageStore = PackageWriter->WriteAdditionalFile(FileInfo, FileData);
 				bWriteFileToDisk &= !bWrittenToPackageStore;
 			}
 
@@ -1774,7 +1774,7 @@ ESavePackageResult FinalizeFile(FStructuredArchive::FRecord& StructuredArchiveRo
 	}
 
 	FSavePackageContext* SavePackageContext = SaveContext.GetSavePackageContext();
-	IPackageStoreWriter* PackageStoreWriter = SavePackageContext ? SavePackageContext->PackageStoreWriter : nullptr;
+	IPackageWriter* PackageWriter = SavePackageContext ? SavePackageContext->PackageWriter : nullptr;
 	if (SaveContext.IsSaveAsync())
 	{
 		FString PathToSave = SaveContext.GetFilename();
@@ -1825,7 +1825,7 @@ ESavePackageResult FinalizeFile(FStructuredArchive::FRecord& StructuredArchiveRo
 
 			checkf(SaveContext.IsCooking() == false || SaveContext.AdditionalPackageFiles.IsEmpty(), TEXT("Saving additional output files during cooking is not currently supported! (%s)"), *PathToSave);
 
-			if (PackageStoreWriter)
+			if (PackageWriter)
 			{
 				FIoBuffer IoBuffer(FIoBuffer::AssumeOwnership, Writer->ReleaseOwnership(), DataSize);
 
@@ -1842,7 +1842,7 @@ ESavePackageResult FinalizeFile(FStructuredArchive::FRecord& StructuredArchiveRo
 
 				const int32 HeaderSize = Linker->Summary.TotalHeaderSize;
 
-				IPackageStoreWriter::FPackageInfo PackageInfo;
+				IPackageWriter::FPackageInfo PackageInfo;
 				PackageInfo.PackageName = SaveContext.GetPackage()->GetFName();
 				PackageInfo.LooseFilePath = SaveContext.GetFilename();
 				PackageInfo.HeaderSize = HeaderSize;
@@ -1850,7 +1850,7 @@ ESavePackageResult FinalizeFile(FStructuredArchive::FRecord& StructuredArchiveRo
 				FPackageId PackageId = FPackageId::FromName(PackageInfo.PackageName);
 				PackageInfo.ChunkId = CreateIoChunkId(PackageId.Value(), 0, EIoChunkType::ExportBundleData);
 
-				PackageStoreWriter->WritePackageData(PackageInfo, IoBuffer, Linker->FileRegions);
+				PackageWriter->WritePackageData(PackageInfo, IoBuffer, Linker->FileRegions);
 			}
 			else
 			{
@@ -1879,12 +1879,12 @@ ESavePackageResult FinalizeFile(FStructuredArchive::FRecord& StructuredArchiveRo
 		}
 		else
 		{
-			checkf(!PackageStoreWriter, TEXT("PackageStoreWriter is not currently supported when diffing."));
+			checkf(!PackageWriter, TEXT("PackageWriter is not currently supported when diffing."));
 		}
 	}
 	else
 	{
-		checkf(!PackageStoreWriter, TEXT("PackageStoreWriter is not currently supported with synchronous writes."));
+		checkf(!PackageWriter, TEXT("PackageWriter is not currently supported with synchronous writes."));
 
 		// Destroy archives used for saving, closing file handle.
 		if (SaveContext.CloseLinkerArchives() == false)
@@ -2109,7 +2109,8 @@ ESavePackageResult InnerSave(FSaveContext& SaveContext)
 		}
 	}
 
-	if (!SaveContext.IsAdditionalFilesNeedLinkerSize())
+	bool bAdditionalFilesNeedLinkerSize = SaveContext.IsAdditionalFilesNeedLinkerSize();
+	if (!bAdditionalFilesNeedLinkerSize)
 	{
 		SaveContext.Result = WriteAdditionalFiles(SaveContext, SlowTask, -1);
 		if (SaveContext.Result != ESavePackageResult::Success)
@@ -2129,14 +2130,14 @@ ESavePackageResult InnerSave(FSaveContext& SaveContext)
 	int32 PackageSize = SaveContext.Linker->Tell();
 	SaveContext.TotalPackageSizeUncompressed += PackageSize;
 
-	if (SaveContext.IsAdditionalFilesNeedLinkerSize())
+	if (bAdditionalFilesNeedLinkerSize)
 	{
 		SaveContext.Result = WriteAdditionalFiles(SaveContext, SlowTask, PackageSize);
 		if (SaveContext.Result != ESavePackageResult::Success)
 		{
 			return SaveContext.Result;
 		}
-		checkf(SaveContext.Linker->Tell() == PackageSize, TEXT("The writing of additional files is not allowed to append to the LinkerSave when IsAdditionalFilesNeedLinkerSize is true."));
+		checkf(SaveContext.Linker->Tell() == PackageSize, TEXT("The writing of additional files is not allowed to append to the LinkerSave when bAdditionalFilesNeedLinkerSize is true."));
 	}
 
 	for (const FSavePackageOutputFile& File : SaveContext.AdditionalPackageFiles)
@@ -2214,7 +2215,7 @@ ESavePackageResult WriteAdditionalFiles(FSaveContext& SaveContext, FScopedSlowTa
 		SaveContext.IsComputeHash(), SaveContext.AsyncWriteAndHashSequence, SaveContext.TotalPackageSizeUncompressed);
 
 	// Add any pending data blobs to the end of the file by invoking the callbacks
-	ESavePackageResult Result = SavePackageUtilities::AppendAdditionalData(*SaveContext.Linker.Get(), DataStartOffset, SaveContext.GetPackageStoreWriter());
+	ESavePackageResult Result = SavePackageUtilities::AppendAdditionalData(*SaveContext.Linker.Get(), DataStartOffset, SaveContext.GetSavePackageContext());
 	if (Result != ESavePackageResult::Success)
 	{
 		return Result;
diff --git a/Engine/Source/Runtime/CoreUObject/Public/UObject/LinkerSave.h b/Engine/Source/Runtime/CoreUObject/Public/UObject/LinkerSave.h
--- a/Engine/Source/Runtime/CoreUObject/Public/UObject/LinkerSave.h
+++ b/Engine/Source/Runtime/CoreUObject/Public/UObject/LinkerSave.h
@@ -75,7 +75,7 @@ public:
 
 	/**
 	 * Callback for arbitrary serializers to append data to the end of the ExportsArchive.
-	 * Some PackageStoreWriters used by SavePackage will write this data to a separate archive.
+	 * Some PackageWriters used by SavePackage will write this data to a separate archive.
 	 * 
 	 * @param ExportsArchive The archive containing the UObjects and structs, this is always this LinkerSave.
 	 * @param DataArchive The archive to which the data should be written. Might be this LinkerSave, or might be a separate archive.
diff --git a/Engine/Source/Runtime/CoreUObject/Public/UObject/SavePackage.h b/Engine/Source/Runtime/CoreUObject/Public/UObject/SavePackage.h
--- a/Engine/Source/Runtime/CoreUObject/Public/UObject/SavePackage.h
+++ b/Engine/Source/Runtime/CoreUObject/Public/UObject/SavePackage.h
@@ -7,6 +7,7 @@
 #include "Containers/Map.h"
 #include "UObject/NameTypes.h"
 #include "Serialization/FileRegions.h"
+#include "Serialization/PackageWriter.h"
 #include "Misc/DateTime.h"
 #include "ObjectMacros.h"
 
@@ -21,7 +22,7 @@ class FPackagePath;
 class FSavePackageContext;
 class FArchiveDiffMap;
 class FOutputDevice;
-class IPackageStoreWriter;
+class IPackageWriter;
 
 /**
  * Struct to encapsulate arguments specific to saving one package
@@ -54,17 +55,22 @@ struct FSavePackageArgs
 class FSavePackageContext
 {
 public:
-	FSavePackageContext(const ITargetPlatform* InTargetPlatform, IPackageStoreWriter* InPackageStoreWriter, bool InbForceLegacyOffsets)
+	FSavePackageContext(const ITargetPlatform* InTargetPlatform, IPackageWriter* InPackageWriter, bool InbForceLegacyOffsets)
 	: TargetPlatform(InTargetPlatform)
-	, PackageStoreWriter(InPackageStoreWriter) 
+	, PackageWriter(InPackageWriter) 
 	, bForceLegacyOffsets(InbForceLegacyOffsets)
 	{
+		if (PackageWriter)
+		{
+			PackageWriterCapabilities = PackageWriter->GetCapabilities();
+		}
 	}
 
 	COREUOBJECT_API ~FSavePackageContext();
 
 	const ITargetPlatform* const TargetPlatform;
-	IPackageStoreWriter* const PackageStoreWriter;
+	IPackageWriter* const PackageWriter;
+	IPackageWriter::FCapabilities PackageWriterCapabilities;
 	const bool bForceLegacyOffsets;
 };
 
