CookOnTheFlyServer: Create FLooseCookedPackageWriter and refactor UCookOnTheFlyServer to use ICookedPackageWriter in all cases rather than having a separate path for writing cooked files to disk.
#rnx

diff --git a/Engine/Source/Developer/IoStoreUtilities/Private/IoStoreUtilities.cpp b/Engine/Source/Developer/IoStoreUtilities/Private/IoStoreUtilities.cpp
--- a/Engine/Source/Developer/IoStoreUtilities/Private/IoStoreUtilities.cpp
+++ b/Engine/Source/Developer/IoStoreUtilities/Private/IoStoreUtilities.cpp
@@ -4057,9 +4057,10 @@ int32 Staged2Zen(const FString& BuildPath, const FKeyChain& KeyChain, const ITar
 	}
 
 	FString MetaDataOutputPath = FPaths::Combine(CookedOutputPath, FApp::GetProjectName(), TEXT("Metadata"));
-	TUniquePtr<FZenStoreWriter> ZenStoreWriter = MakeUnique<FZenStoreWriter>(CookedOutputPath, MetaDataOutputPath, TargetPlatform, true);
+	TUniquePtr<FZenStoreWriter> ZenStoreWriter = MakeUnique<FZenStoreWriter>(CookedOutputPath, MetaDataOutputPath, TargetPlatform);
 	
 	ICookedPackageWriter::FCookInfo CookInfo;
+	CookInfo.bCleanBuild = true;
 	ZenStoreWriter->BeginCook(CookInfo);
 	int32 LocalPackageIndex = 0;
 	TArray<FPackageInfo> PackagesArray;
diff --git a/Engine/Source/Developer/IoStoreUtilities/Private/ZenStoreHttpClient.cpp b/Engine/Source/Developer/IoStoreUtilities/Private/ZenStoreHttpClient.cpp
--- a/Engine/Source/Developer/IoStoreUtilities/Private/ZenStoreHttpClient.cpp
+++ b/Engine/Source/Developer/IoStoreUtilities/Private/ZenStoreHttpClient.cpp
@@ -50,8 +50,7 @@ FZenStoreHttpClient::Initialize(FStringView InProjectId,
 	FStringView InOplogId, 
 	FStringView ServerRoot,
 	FStringView EngineRoot,
-	FStringView ProjectRoot,
-	bool		IsCleanBuild)
+	FStringView ProjectRoot)
 {
 	TRACE_CPUPROFILER_EVENT_SCOPE(ZenStoreHttp_Initialize);
 
@@ -105,81 +104,70 @@ FZenStoreHttpClient::Initialize(FStringView InProjectId,
 		}
 	}
 
-	// Establish oplog
+	OplogPath = WriteToString<128>("/prj/", InProjectId, "/oplog/", InOplogId);
+	OplogNewEntryPath = WriteToString<128>("/prj/", InProjectId, "/oplog/", InOplogId, "/new");
+	OplogPrepNewEntryPath = WriteToString<128>("/prj/", InProjectId, "/oplog/", InOplogId, "/prep");
+
+	bAllowRead = true;
+	bAllowEdit = true;
+}
+
+void FZenStoreHttpClient::EstablishWritableOpLog(FStringView InProjectId, FStringView InOplogId, bool bCleanBuild)
+{
+	UE::Zen::FZenScopedRequestPtr Request(RequestPool.Get());
 
+	if (bCleanBuild)
 	{
-		UE::Zen::FZenScopedRequestPtr Request(RequestPool.Get());
+		UE_LOG(LogZenStore, Display, TEXT("Deleting oplog '%s'/'%s' if it exists"), *FString(InProjectId), *FString(InOplogId));
+		Request->PerformBlockingDelete(OplogPath);
+		Request->Reset();
+	}
 
-		TStringBuilder<128> OplogUri;
-		OplogUri << "/prj/" << InProjectId << "/oplog/" << InOplogId;
+	TArray64<uint8> GetBuffer;
+	UE::Zen::FZenHttpRequest::Result Res = Request->PerformBlockingDownload(OplogPath, &GetBuffer);
+	FCbObjectView OplogInfo;
 
-		OplogPath = OplogUri;
+	if (Res == Zen::FZenHttpRequest::Result::Success && Request->GetResponseCode() == 200)
+	{
+		UE_LOG(LogZenStore, Display, TEXT("Zen oplog '%s'/'%s' already exists"), *FString(InProjectId), *FString(InOplogId));
 
-		if (IsCleanBuild)
-		{
-			UE_LOG(LogZenStore, Display, TEXT("Deleting oplog '%s'/'%s' if it exists"), *FString(InProjectId), *FString(InOplogId));
-			Request->PerformBlockingDelete(OplogUri);
-			Request->Reset();
-		}
+		OplogInfo = FCbObjectView(GetBuffer.GetData());
+	}
+	else
+	{
+		FMemoryView Payload;
 
-		TArray64<uint8> GetBuffer;
-		UE::Zen::FZenHttpRequest::Result Res = Request->PerformBlockingDownload(OplogUri, &GetBuffer);
-		FCbObjectView OplogInfo;
+		Request->Reset();
+		Res = Request->PerformBlockingPost(OplogPath, Payload);
 
-		if (Res == Zen::FZenHttpRequest::Result::Success && Request->GetResponseCode() == 200)
+		if (Res != Zen::FZenHttpRequest::Result::Success)
 		{
-			UE_LOG(LogZenStore, Display, TEXT("Zen oplog '%s'/'%s' already exists"), *FString(InProjectId), *FString(InOplogId));
+			UE_LOG(LogZenStore, Error, TEXT("Zen oplog '%s'/'%s' creation FAILED"), *FString(InProjectId), *FString(InOplogId));
 
-			OplogInfo = FCbObjectView(GetBuffer.GetData());
+			// TODO: how to recover / handle this?
+		}
+		else if (Request->GetResponseCode() == 201)
+		{
+			UE_LOG(LogZenStore, Display, TEXT("Zen oplog '%s'/'%s' created"), *FString(InProjectId), *FString(InOplogId));
 		}
 		else
 		{
-			FMemoryView Payload;
-
-			Request->Reset();
-			Res = Request->PerformBlockingPost(OplogUri, Payload);
-
-			if (Res != Zen::FZenHttpRequest::Result::Success)
-			{
-				UE_LOG(LogZenStore, Error, TEXT("Zen oplog '%s'/'%s' creation FAILED"), *FString(InProjectId), *FString(InOplogId));
-
-				// TODO: how to recover / handle this?
-			}
-			else if (Request->GetResponseCode() == 201)
-			{
-				UE_LOG(LogZenStore, Display, TEXT("Zen oplog '%s'/'%s' created"), *FString(InProjectId), *FString(InOplogId));
-			}
-			else
-			{
-				UE_LOG(LogZenStore, Warning, TEXT("Zen oplog '%s'/'%s' creation returned success but not HTTP 201"), *FString(InProjectId), *FString(InOplogId));
-			}
-
-			// Issue another GET to retrieve information
-
-			GetBuffer.Reset();
-			Request->Reset();
-			Res = Request->PerformBlockingDownload(OplogUri, &GetBuffer);
-
-			if (Res == Zen::FZenHttpRequest::Result::Success && Request->GetResponseCode() == 200)
-			{
-				OplogInfo = FCbObjectView(GetBuffer.GetData());
-			}
+			UE_LOG(LogZenStore, Warning, TEXT("Zen oplog '%s'/'%s' creation returned success but not HTTP 201"), *FString(InProjectId), *FString(InOplogId));
 		}
 
-		TempDirPath = FUTF8ToTCHAR(OplogInfo["tempdir"].AsString());
-	}
+		// Issue another GET to retrieve information
 
-	{
-		TStringBuilder<128> OplogUri;
-		OplogUri << "/prj/" << InProjectId << "/oplog/" << InOplogId << "/new";
+		GetBuffer.Reset();
+		Request->Reset();
+		Res = Request->PerformBlockingDownload(OplogPath, &GetBuffer);
 
-		OplogNewEntryPath = OplogUri;
+		if (Res == Zen::FZenHttpRequest::Result::Success && Request->GetResponseCode() == 200)
+		{
+			OplogInfo = FCbObjectView(GetBuffer.GetData());
+		}
 	}
 
-	OplogPrepNewEntryPath = TStringBuilder<128>().AppendAnsi("/prj/").Append(InProjectId).AppendAnsi("/oplog/").Append(InOplogId).AppendAnsi("/prep");
-
-	bAllowRead = true;
-	bAllowEdit = true;
+	TempDirPath = FUTF8ToTCHAR(OplogInfo["tempdir"].AsString());
 }
 
 void FZenStoreHttpClient::InitializeReadOnly(FStringView InProjectId, FStringView InOplogId)
@@ -193,10 +181,8 @@ void FZenStoreHttpClient::InitializeReadOnly(FStringView InProjectId, FStringVie
 	{
 		UE::Zen::FZenScopedRequestPtr Request(RequestPool.Get());
 
-		TStringBuilder<128> ProjectUri;
-		ProjectUri << "/prj/" << InProjectId;
 		TArray64<uint8> GetBuffer;
-		UE::Zen::FZenHttpRequest::Result Res = Request->PerformBlockingDownload(ProjectUri, &GetBuffer);
+		UE::Zen::FZenHttpRequest::Result Res = Request->PerformBlockingDownload(WriteToString<128>("/prj/", InProjectId), &GetBuffer);
 
 		// TODO: how to handle failure here? This is probably the most likely point of failure
 		// if the service is not up or not responding
@@ -212,13 +198,10 @@ void FZenStoreHttpClient::InitializeReadOnly(FStringView InProjectId, FStringVie
 	{
 		UE::Zen::FZenScopedRequestPtr Request(RequestPool.Get());
 
-		TStringBuilder<128> OplogUri;
-		OplogUri << "/prj/" << InProjectId << "/oplog/" << InOplogId;
-
-		OplogPath = OplogUri;
+		OplogPath = WriteToString<128>("/prj/", InProjectId, "/oplog/", InOplogId);
 
 		TArray64<uint8> GetBuffer;
-		UE::Zen::FZenHttpRequest::Result Res = Request->PerformBlockingDownload(OplogUri, &GetBuffer);
+		UE::Zen::FZenHttpRequest::Result Res = Request->PerformBlockingDownload(OplogPath, &GetBuffer);
 
 		if (Res != Zen::FZenHttpRequest::Result::Success || Request->GetResponseCode() != 200)
 		{
@@ -575,8 +558,11 @@ void FZenStoreHttpClient::Initialize(
 	FStringView InOplogId,
 	FStringView ServerRoot,
 	FStringView EngineRoot,
-	FStringView ProjectRoot,
-	bool		IsCleanBuild)
+	FStringView ProjectRoot)
+{
+}
+
+void FZenStoreHttpClient::EstablishWritableOpLog(FStringView InProjectId, FStringView InOplogId, bool bCleanBuild)
 {
 }
 
diff --git a/Engine/Source/Developer/IoStoreUtilities/Private/ZenStoreHttpClient.h b/Engine/Source/Developer/IoStoreUtilities/Private/ZenStoreHttpClient.h
--- a/Engine/Source/Developer/IoStoreUtilities/Private/ZenStoreHttpClient.h
+++ b/Engine/Source/Developer/IoStoreUtilities/Private/ZenStoreHttpClient.h
@@ -28,8 +28,8 @@ public:
 					FStringView InOplogId, 
 					FStringView ServerRoot, 
 					FStringView EngineRoot, 
-					FStringView ProjectRoot,
-					bool		IsCleanBuild);
+					FStringView ProjectRoot);
+	void EstablishWritableOpLog(FStringView InProjectId, FStringView InOplogId, bool bCleanBuild);
 
 	void InitializeReadOnly(FStringView InProjectId, FStringView InOplogId);
 
diff --git a/Engine/Source/Developer/IoStoreUtilities/Private/ZenStoreWriter.cpp b/Engine/Source/Developer/IoStoreUtilities/Private/ZenStoreWriter.cpp
--- a/Engine/Source/Developer/IoStoreUtilities/Private/ZenStoreWriter.cpp
+++ b/Engine/Source/Developer/IoStoreUtilities/Private/ZenStoreWriter.cpp
@@ -227,13 +227,14 @@ private:
 FZenStoreWriter::FZenStoreWriter(
 	const FString& InOutputPath, 
 	const FString& InMetadataDirectoryPath, 
-	const ITargetPlatform* InTargetPlatform,
-	bool IsCleanBuild)
+	const ITargetPlatform* InTargetPlatform
+)
 	: TargetPlatform(*InTargetPlatform)
 	, OutputPath(InOutputPath)
 	, MetadataDirectoryPath(InMetadataDirectoryPath)
 	, PackageStoreOptimizer(new FPackageStoreOptimizer())
 	, CookMode(ICookedPackageWriter::FCookInfo::CookByTheBookMode)
+	, bInitialized(false)
 {
 	FString HostName = TEXT("localhost");
 	uint16 Port = 1337;
@@ -257,8 +258,7 @@ FZenStoreWriter::FZenStoreWriter(
 							OplogId, 
 							AbsServerRoot, 
 							AbsEngineDir, 
-							AbsProjectDir, 
-							IsCleanBuild);
+							AbsProjectDir );
 
 	PackageStoreOptimizer->Initialize(InTargetPlatform);
 
@@ -271,83 +271,6 @@ FZenStoreWriter::FZenStoreWriter(
 	ZenFileSystemManifest = MakeUnique<FZenFileSystemManifest>(TargetPlatform, OutputPath);
 	
 	HttpQueue = MakeUnique<FZenStoreHttpQueue>(*HttpClient);
-
-	if (!IsCleanBuild)
-	{
-		UE_LOG(LogZenStoreWriter, Display, TEXT("Fetching oplog..."), *ProjectId, *OplogId);
-
-		TFuture<FIoStatus> FutureOplogStatus = HttpClient->GetOplog().Next([this](TIoStatusOr<FCbObject> OplogStatus)
-		{
-			if (!OplogStatus.IsOk())
-			{
-				return OplogStatus.Status();
-			}
-			
-			FCbObject Oplog = OplogStatus.ConsumeValueOrDie();
-			
-			if (Oplog["entries"])
-			{
-				for (FCbField& OplogEntry : Oplog["entries"].AsArray())
-				{
-					FCbObject OplogObj						= OplogEntry.AsObject();
-
-					if (OplogObj["package"])
-					{
-						FCbObject PackageObj				= OplogObj["package"].AsObject();
-
-						const FGuid PkgGuid					= PackageObj["guid"].AsUuid();
-						const FIoHash PkgHash				= PackageObj["data"].AsHash();
-						const int64	PkgDiskSize				= PackageObj["disksize"].AsUInt64();
-						FPackageStoreEntryResource Entry	= FPackageStoreEntryResource::FromCbObject(OplogObj["packagestoreentry"].AsObject());
-						const FName PackageName				= Entry.PackageName;
-						FIoHash DependenciesHash			= OplogObj["dependencies"].AsHash();
-						const int32 Index					= PackageStoreEntries.Num();
-
-						PackageStoreEntries.Add(MoveTemp(Entry));
-						CookedPackagesInfo.Add(FCookedPackageInfo { PackageName, IoHashToMD5(PkgHash), PkgGuid, PkgDiskSize, DependenciesHash });
-						PackageNameToIndex.Add(PackageName, Index);
-					}
-				}
-			}
-
-			return FIoStatus::Ok;
-		});
-
-		UE_LOG(LogZenStoreWriter, Display, TEXT("Fetching file manifest..."), *ProjectId, *OplogId);
-
-		TIoStatusOr<FCbObject> FileStatus = HttpClient->GetFiles().Get();
-		if (FileStatus .IsOk())
-		{
-			FCbObject FilesObj = FileStatus.ConsumeValueOrDie();
-			for (FCbField& FileEntry : FilesObj["files"])
-			{
-				FCbObject FileObj	= FileEntry.AsObject();
-				FCbObjectId FileId	= FileObj["id"].AsObjectId();
-				FString ServerPath	= FString(FileObj["serverpath"].AsString());
-				FString ClientPath	= FString(FileObj["clientpath"].AsString());
-
-				FIoChunkId FileChunkId;
-				FileChunkId.Set(FileId.GetView());
-					
-				ZenFileSystemManifest->AddManifestEntry(FileChunkId, MoveTemp(ServerPath), MoveTemp(ClientPath));
-			}
-
-			UE_LOG(LogZenStoreWriter, Display, TEXT("Fetched '%d' files(s) from oplog '%s/%s'"), ZenFileSystemManifest->NumEntries(), *ProjectId, *OplogId);
-		}
-		else
-		{
-			UE_LOG(LogZenStoreWriter, Warning, TEXT("Failed to fetch file(s) from oplog '%s/%s'"), *ProjectId, *OplogId);
-		}
-
-		if (FutureOplogStatus.Get().IsOk())
-		{
-			UE_LOG(LogZenStoreWriter, Display, TEXT("Fetched '%d' packges(s) from oplog '%s/%s'"), PackageStoreEntries.Num(), *ProjectId, *OplogId);
-		}
-		else
-		{
-			UE_LOG(LogZenStoreWriter, Warning, TEXT("Failed to fetch oplog '%s/%s'"), *ProjectId, *OplogId);
-		}
-	}
 }
 
 FZenStoreWriter::~FZenStoreWriter()
@@ -496,6 +419,98 @@ void FZenStoreWriter::BeginCook(const FCookInfo& Info)
 {
 	CookMode = Info.CookMode;
 
+	if (!bInitialized)
+	{
+		FString ProjectId = FApp::GetProjectName();
+		FString OplogId = TargetPlatform.PlatformName();
+		HttpClient->EstablishWritableOpLog(ProjectId, OplogId, Info.bCleanBuild);
+
+		if (!Info.bCleanBuild)
+		{
+			UE_LOG(LogZenStoreWriter, Display, TEXT("Fetching oplog..."), *ProjectId, *OplogId);
+
+			TFuture<FIoStatus> FutureOplogStatus = HttpClient->GetOplog().Next([this](TIoStatusOr<FCbObject> OplogStatus)
+				{
+					if (!OplogStatus.IsOk())
+					{
+						return OplogStatus.Status();
+					}
+
+					FCbObject Oplog = OplogStatus.ConsumeValueOrDie();
+
+					if (Oplog["entries"])
+					{
+						for (FCbField& OplogEntry : Oplog["entries"].AsArray())
+						{
+							FCbObject OplogObj = OplogEntry.AsObject();
+
+							if (OplogObj["package"])
+							{
+								FCbObject PackageObj = OplogObj["package"].AsObject();
+
+								const FGuid PkgGuid = PackageObj["guid"].AsUuid();
+								const FIoHash PkgHash = PackageObj["data"].AsHash();
+								const int64	PkgDiskSize = PackageObj["disksize"].AsUInt64();
+								FPackageStoreEntryResource Entry = FPackageStoreEntryResource::FromCbObject(OplogObj["packagestoreentry"].AsObject());
+								const FName PackageName = Entry.PackageName;
+								FIoHash DependenciesHash = OplogObj["dependencies"].AsHash();
+								const int32 Index = PackageStoreEntries.Num();
+
+								PackageStoreEntries.Add(MoveTemp(Entry));
+								CookedPackagesInfo.Add(FCookedPackageInfo{ PackageName, IoHashToMD5(PkgHash), PkgGuid, PkgDiskSize, DependenciesHash });
+								PackageNameToIndex.Add(PackageName, Index);
+							}
+						}
+					}
+
+					return FIoStatus::Ok;
+				});
+
+			UE_LOG(LogZenStoreWriter, Display, TEXT("Fetching file manifest..."), *ProjectId, *OplogId);
+
+			TIoStatusOr<FCbObject> FileStatus = HttpClient->GetFiles().Get();
+			if (FileStatus.IsOk())
+			{
+				FCbObject FilesObj = FileStatus.ConsumeValueOrDie();
+				for (FCbField& FileEntry : FilesObj["files"])
+				{
+					FCbObject FileObj = FileEntry.AsObject();
+					FCbObjectId FileId = FileObj["id"].AsObjectId();
+					FString ServerPath = FString(FileObj["serverpath"].AsString());
+					FString ClientPath = FString(FileObj["clientpath"].AsString());
+
+					FIoChunkId FileChunkId;
+					FileChunkId.Set(FileId.GetView());
+
+					ZenFileSystemManifest->AddManifestEntry(FileChunkId, MoveTemp(ServerPath), MoveTemp(ClientPath));
+				}
+
+				UE_LOG(LogZenStoreWriter, Display, TEXT("Fetched '%d' files(s) from oplog '%s/%s'"), ZenFileSystemManifest->NumEntries(), *ProjectId, *OplogId);
+			}
+			else
+			{
+				UE_LOG(LogZenStoreWriter, Warning, TEXT("Failed to fetch file(s) from oplog '%s/%s'"), *ProjectId, *OplogId);
+			}
+
+			if (FutureOplogStatus.Get().IsOk())
+			{
+				UE_LOG(LogZenStoreWriter, Display, TEXT("Fetched '%d' packges(s) from oplog '%s/%s'"), PackageStoreEntries.Num(), *ProjectId, *OplogId);
+			}
+			else
+			{
+				UE_LOG(LogZenStoreWriter, Warning, TEXT("Failed to fetch oplog '%s/%s'"), *ProjectId, *OplogId);
+			}
+		}
+		bInitialized = true;
+	}
+	else
+	{
+		if (Info.bCleanBuild)
+		{
+			RemoveCookedPackages();
+		}
+	}
+
 	if (Info.CookMode == ICookedPackageWriter::FCookInfo::CookOnTheFlyMode)
 	{
 		FCbPackage Pkg;
@@ -766,7 +781,7 @@ void FZenStoreWriter::RemoveCookedPackages(TArrayView<const FName> PackageNamesT
 	
 	TArray<FPackageStoreEntryResource> PreviousPackageStoreEntries = MoveTemp(PackageStoreEntries);
 	TArray<FCookedPackageInfo> PreviousCookedPackageInfo = MoveTemp(CookedPackagesInfo);
-	PackageNameToIndex.Empty(false);
+	PackageNameToIndex.Empty();
 
 	if (NumPackagesToKeep > 0)
 	{
@@ -786,6 +801,13 @@ void FZenStoreWriter::RemoveCookedPackages(TArrayView<const FName> PackageNamesT
 	}
 }
 
+void FZenStoreWriter::RemoveCookedPackages()
+{
+	PackageStoreEntries.Empty();
+	CookedPackagesInfo.Empty();
+	PackageNameToIndex.Empty();
+}
+
 void FZenStoreWriter::CreateProjectMetaData(FCbPackage& Pkg, FCbWriter& PackageObj, bool bGenerateContainerHeader)
 {
 	// File Manifest
diff --git a/Engine/Source/Developer/IoStoreUtilities/Public/FilePackageStoreWriter.h b/Engine/Source/Developer/IoStoreUtilities/Public/FilePackageStoreWriter.h
--- a/Engine/Source/Developer/IoStoreUtilities/Public/FilePackageStoreWriter.h
+++ b/Engine/Source/Developer/IoStoreUtilities/Public/FilePackageStoreWriter.h
@@ -80,8 +80,8 @@ public:
 
 	IOSTOREUTILITIES_API virtual bool WriteAdditionalFile(const FAdditionalFileInfo& Info, const FIoBuffer& FileData) override { return false; }
 
-	IOSTOREUTILITIES_API virtual void BeginCook(const FCookInfo& Info);
-	IOSTOREUTILITIES_API virtual void EndCook();
+	IOSTOREUTILITIES_API virtual void BeginCook(const FCookInfo& Info) override;
+	IOSTOREUTILITIES_API virtual void EndCook() override;
 
 	virtual void GetEntries(TFunction<void(TArrayView<const FPackageStoreEntryResource>)>&& Callback) override
 	{
@@ -98,13 +98,16 @@ public:
 	virtual void Flush() override
 	{ }
 
-	virtual void GetCookedPackages(TArray<FCookedPackageInfo>& OutCookedPackages)
+	virtual void GetCookedPackages(TArray<FCookedPackageInfo>& OutCookedPackages) override
 	{ }
 
 	virtual FCbObject GetTargetDomainDependencies(FName PackageName) override
 	{ return FCbObject(); }
 
-	virtual void RemoveCookedPackages(TArrayView<const FName> PackageNamesToRemove)
+	virtual void RemoveCookedPackages(TArrayView<const FName> PackageNamesToRemove) override
+	{ }
+
+	virtual void RemoveCookedPackages() override
 	{ }
 
 private:
diff --git a/Engine/Source/Developer/IoStoreUtilities/Public/ZenStoreWriter.h b/Engine/Source/Developer/IoStoreUtilities/Public/ZenStoreWriter.h
--- a/Engine/Source/Developer/IoStoreUtilities/Public/ZenStoreWriter.h
+++ b/Engine/Source/Developer/IoStoreUtilities/Public/ZenStoreWriter.h
@@ -28,8 +28,7 @@ class FZenStoreWriter
 public:
 	IOSTOREUTILITIES_API FZenStoreWriter(	const FString& OutputPath, 
 											const FString& MetadataDirectoryPath, 
-											const ITargetPlatform* TargetPlatform, 
-											bool IsCleanBuild);
+											const ITargetPlatform* TargetPlatform);
 
 	IOSTOREUTILITIES_API ~FZenStoreWriter();
 
@@ -58,6 +57,7 @@ public:
 	IOSTOREUTILITIES_API virtual void GetCookedPackages(TArray<FCookedPackageInfo>& OutCookedPackages) override;
 	IOSTOREUTILITIES_API virtual FCbObject GetTargetDomainDependencies(FName PackageName) override;
 	IOSTOREUTILITIES_API virtual void RemoveCookedPackages(TArrayView<const FName> PackageNamesToRemove) override;
+	IOSTOREUTILITIES_API virtual void RemoveCookedPackages() override;
 
 private:
 	void CreateProjectMetaData(FCbPackage& Pkg, FCbWriter& PackageObj, bool bGenerateContainerHeader);
@@ -128,4 +128,6 @@ private:
 	ICookedPackageWriter::FCookInfo::ECookMode CookMode;
 
 	FZenStats							ZenStats;
+
+	bool								bInitialized;
 };
diff --git a/Engine/Source/Editor/UnrealEd/Classes/CookOnTheSide/CookOnTheFlyServer.h b/Engine/Source/Editor/UnrealEd/Classes/CookOnTheSide/CookOnTheFlyServer.h
--- a/Engine/Source/Editor/UnrealEd/Classes/CookOnTheSide/CookOnTheFlyServer.h
+++ b/Engine/Source/Editor/UnrealEd/Classes/CookOnTheSide/CookOnTheFlyServer.h
@@ -214,7 +214,6 @@ private:
 	TUniquePtr<class FSandboxPlatformFile> SandboxFile;
 	TUniquePtr<FAsyncIODelete> AsyncIODelete; // Helper for deleting the old cook directory asynchronously
 	FString DefaultAsyncIODeletePlatformName; // Default platform name to be used by GetAsyncIODelete()
-	bool bIsInitializingSandbox = false; // stop recursion into callbacks when we are initializing sandbox
 	bool bIsSavingPackage = false; // used to stop recursive mark package dirty functions
 	bool bSaveAsyncAllowed = false; // True if and only if command line options and all other restrictions allow the use of SAVE_Async
 	/** Set to true during CookOnTheFly if a plugin is calling RequestPackage and we should therefore not make assumptions about when platforms are done cooking */
@@ -247,11 +246,6 @@ private:
 	double LastProgressDisplayTime = 0;
 	double LastDiagnosticsDisplayTime = 0;
 
-	FName ConvertCookedPathToUncookedPath(
-		const FString& SandboxRootDir, const FString& RelativeRootDir,
-		const FString& SandboxProjectDir, const FString& RelativeProjectDir,
-		const FString& CookedPath, FString& OutUncookedPath) const;
-
 	/** Get dependencies for package */
 	const TArray<FName>& GetFullPackageDependencies(const FName& PackageName) const;
 	mutable TMap<FName, TArray<FName>> CachedFullPackageDependencies;
@@ -518,6 +512,8 @@ public:
 	UE_DEPRECATED(4.25, "Use version that takes const ITargetPlatform* instead")
 	void ClearPlatformCookedData(const FString& PlatformName);
 
+	/** Clear all in-process flags that indicate packages have been cooked for the given platforms. */
+	void ResetCookResults(TConstArrayView<const ITargetPlatform*> TargetPlatforms);
 
 	/**
 	* Recompile any global shader changes 
@@ -799,14 +795,9 @@ private:
 		UE::Cook::FPackageData* ReportingPackageData = nullptr);
 
 	/**
-	* Initialize the sandbox for @param TargetPlatforms
-	*/
-	void InitializeSandbox(const TArrayView<const ITargetPlatform* const>& TargetPlatforms);
-
-	/**
-	* Initialize the package store
+	* Initialize the sandbox for a new cook session with @param TargetPlatforms
 	*/
-	void InitializePackageStore(const TArrayView<const ITargetPlatform* const>& TargetPlatforms);
+	void BeginCookSandbox(TConstArrayView<const ITargetPlatform*> TargetPlatforms);
 
 	/**
 	* Finalize the package store
@@ -953,21 +944,12 @@ private:
 	*/
 	FString GetSandboxDirectory( const FString& PlatformName ) const;
 
-	/* Delete the sandbox directory (asynchronously) for the given platform in preparation for a clean cook */
-	void DeleteSandboxDirectory(const FString& PlatformName);
-
 	/** If not already done, set the default platform name to be used by GetAsyncIODelete(). */
 	void TrySetDefaultAsyncIODeletePlatform(const FString& PlatformName);
 
 	/* Create the delete-old-cooked-directory helper.*/
 	FAsyncIODelete& GetAsyncIODelete();
 
-	/* Create the delete-old-cooked-directory helper. PlatformName tells it where to create its temp directory. AsyncDeleteDirectory is DeleteSandboxDirectory internal use only and should otherwise be set to nullptr. */
-	FAsyncIODelete& GetAsyncIODelete(const FString& PlatformName, const FString* AsyncDeleteDirectory = nullptr);
-
-	/** Get the directory that should be used for AsyncDeletes based on the given platform. SandboxDirectory is DeleteSandboxDirectory internal use only and should otherwise be set to nullptr. */
-	FString GetAsyncDeleteDirectory(const FString& PlatformName, const FString* SandboxDirectory = nullptr) const;
-
 	bool IsCookingDLC() const;
 
 	/**
@@ -1044,26 +1026,10 @@ private:
 	FString GetOutputDirectoryOverride() const;
 
 	/**
-	* Populate the cooked packages list from the on disk content using time stamps and dependencies to figure out if they are ok
-	* delete any local content which is out of date
-	* 
-	* @param Platforms to process
-	*/
-	void PopulateCookedPackagesFromDisk(const TArrayView<const ITargetPlatform* const>& Platforms);
-	
-	/**
-	 * Populate cooked packages from meta data available in the package store.
+	 * Populate cooked packages from the PackageWriter's previous manifest and assetregistry of cooked output.
+	 * Delete any out of date packages from the PackageWriter's manifest.
 	 */
-	void PopulateCookedPackagesFromPackageStore(const TArrayView<const ITargetPlatform* const>& Platforms);
-
-	/**
-	* Searches the disk for all the cooked files in the sandbox path provided
-	* Returns a map of the uncooked file path matches to the cooked file path for each package which exists
-	*
-	* @param UncookedpathToCookedPath out Map of the uncooked path matched with the cooked package which exists
-	* @param SandboxRootDir root dir to search for cooked packages in
-	*/
-	void GetAllCookedFiles(TMap<FName, FName>& UncookedPathToCookedPath, const FString& SandboxRootDir);
+	void PopulateCookedPackages(const TConstArrayView<const ITargetPlatform*> TargetPlatforms);
 
 	/** Loads the platform-independent asset registry for use by the cooker */
 	void GenerateAssetRegistry();
@@ -1110,8 +1076,11 @@ private:
 
 	uint32 FullLoadAndSave(uint32& CookedPackageCount);
 
-	ICookedPackageWriter* GetPackageWriter(const FName& PlatformName) const;
-	UE::Cook::FCookSavePackageContext* GetSavePackageContext(const ITargetPlatform* TargetPlatform) const;
+	ICookedPackageWriter& FindOrCreatePackageWriter(const ITargetPlatform* TargetPlatform);
+	void FindOrCreateSaveContexts(TConstArrayView<const ITargetPlatform*> TargetPlatforms);
+	UE::Cook::FCookSavePackageContext& FindOrCreateSaveContext(const ITargetPlatform* TargetPlatform);
+	/** Allocate a new FCookSavePackageContext and ICookedPackageWriter for the given platform. */
+	UE::Cook::FCookSavePackageContext* CreateSaveContext(const ITargetPlatform* TargetPlatform) const;
 
 	uint32		StatLoadedPackageCount = 0;
 	uint32		StatSavedPackageCount = 0;
diff --git a/Engine/Source/Editor/UnrealEd/Private/Commandlets/AssetRegistryGenerator.cpp b/Engine/Source/Editor/UnrealEd/Private/Commandlets/AssetRegistryGenerator.cpp
--- a/Engine/Source/Editor/UnrealEd/Private/Commandlets/AssetRegistryGenerator.cpp
+++ b/Engine/Source/Editor/UnrealEd/Private/Commandlets/AssetRegistryGenerator.cpp
@@ -795,17 +795,16 @@ void FAssetRegistryGenerator::CleanManifestDirectories()
 	CleanTempPackagingDirectory(TargetPlatform->PlatformName());
 }
 
-bool FAssetRegistryGenerator::LoadPreviousAssetRegistry(const FString& Filename)
+void FAssetRegistryGenerator::SetPreviousAssetRegistry(TUniquePtr<FAssetRegistryState>&& InPreviousState)
 {
-	// First try development asset registry
-	FArrayReader SerializedAssetData;
-
-	if (IFileManager::Get().FileExists(*Filename) && FFileHelper::LoadFileToArray(SerializedAssetData, *Filename))
+	if (InPreviousState)
 	{
-		return PreviousState.Load(SerializedAssetData);
+		PreviousState = MoveTemp(*InPreviousState);
+	}
+	else
+	{
+		PreviousState.Reset();
 	}
-
-	return false;
 }
 
 void FAssetRegistryGenerator::InjectEncryptionData(FAssetRegistryState& TargetState)
diff --git a/Engine/Source/Editor/UnrealEd/Private/CookOnTheFlyServer.cpp b/Engine/Source/Editor/UnrealEd/Private/CookOnTheFlyServer.cpp
--- a/Engine/Source/Editor/UnrealEd/Private/CookOnTheFlyServer.cpp
+++ b/Engine/Source/Editor/UnrealEd/Private/CookOnTheFlyServer.cpp
@@ -25,6 +25,7 @@
 #include "Cooker/CookRequests.h"
 #include "Cooker/CookTypes.h"
 #include "Cooker/IoStoreCookOnTheFlyRequestManager.h"
+#include "Cooker/LooseCookedPackageWriter.h"
 #include "Cooker/NetworkFileCookOnTheFlyRequestManager.h"
 #include "Cooker/PackageNameCache.h"
 #include "Cooker/PackageTracker.h"
@@ -119,7 +120,6 @@
 #include "ZenStoreWriter.h"
 
 #define LOCTEXT_NAMESPACE "Cooker"
-#define REMAPPED_PLUGINS TEXT("RemappedPlugins")
 
 DEFINE_LOG_CATEGORY(LogCook);
 
@@ -187,59 +187,6 @@ const TCHAR* GeneratedPackageSubPath = TEXT("_Generated_");
 /* helper structs functions
  *****************************************************************************/
 
-class FPackageSearchVisitor : public IPlatformFile::FDirectoryVisitor
-{
-	TArray<FString>& FoundFiles;
-public:
-	FPackageSearchVisitor(TArray<FString>& InFoundFiles)
-		: FoundFiles(InFoundFiles)
-	{}
-	virtual bool Visit(const TCHAR* FilenameOrDirectory, bool bIsDirectory)
-	{
-		if (bIsDirectory == false)
-		{
-			FString Filename(FilenameOrDirectory);
-			if (Filename.EndsWith(TEXT(".uasset")) || Filename.EndsWith(TEXT(".umap")))
-			{
-				FoundFiles.Add(Filename);
-			}
-		}
-		return true;
-	}
-};
-
-class FAdditionalPackageSearchVisitor: public IPlatformFile::FDirectoryVisitor
-{
-	TSet<FString>& FoundMapFilesNoExt;
-	TArray<FString>& FoundOtherFiles;
-public:
-	FAdditionalPackageSearchVisitor(TSet<FString>& InFoundMapFiles, TArray<FString>& InFoundOtherFiles)
-		: FoundMapFilesNoExt(InFoundMapFiles), FoundOtherFiles(InFoundOtherFiles)
-	{}
-	virtual bool Visit(const TCHAR* FilenameOrDirectory, bool bIsDirectory)
-	{
-		if (bIsDirectory == false)
-		{
-			FString Filename(FilenameOrDirectory);
-			if (Filename.EndsWith(TEXT(".uasset")) || Filename.EndsWith(TEXT(".umap")))
-			{
-				FoundMapFilesNoExt.Add(FPaths::SetExtension(Filename, ""));
-			}
-			else if ( Filename.EndsWith(TEXT(".uexp")) || Filename.EndsWith(TEXT(".ubulk")) )
-			{
-				FoundOtherFiles.Add(Filename);
-			}
-		}
-		return true;
-	}
-};
-
-const FString& GetAssetRegistryPath()
-{
-	static const FString AssetRegistryPath = FPaths::ProjectDir();
-	return AssetRegistryPath;
-}
-
 /**
  * Return the release asset registry filename for the release version supplied
  */
@@ -337,24 +284,24 @@ public:
 		return Cooker.SandboxFile->GetSandboxDirectory();
 	}
 
-	virtual bool AddPlatform(const FName& PlatformName) override
+	virtual const ITargetPlatform* AddPlatform(const FName& PlatformName) override
 	{
+		const ITargetPlatform* Result;
 		{
 			UE::Cook::FPlatformManager::FReadScopeLock PlatformScopeLock(Cooker.PlatformManager->ReadLockPlatforms());
-			if (!Cooker.AddCookOnTheFlyPlatform(PlatformName))
+			Result = Cooker.AddCookOnTheFlyPlatform(PlatformName);
+			if (!Result)
 			{
 				UE_LOG(LogCook, Warning, TEXT("Trying to add invalid platform '%s' on the fly"), *PlatformName.ToString());
-				return false;
+				return nullptr;
 			}
 
 			Cooker.PlatformManager->AddRefCookOnTheFlyPlatform(PlatformName, Cooker);
 		}
 
-		Cooker.TrySetDefaultAsyncIODeletePlatform(PlatformName.ToString());
-
-		return true;
+		return Result;
 	}
-	
+
 	virtual void RemovePlatform(const FName& PlatformName) override
 	{
 		UE::Cook::FPlatformManager::FReadScopeLock PlatformScopeLock(Cooker.PlatformManager->ReadLockPlatforms());
@@ -466,20 +413,9 @@ public:
 		return true;
 	}
 
-	virtual ICookedPackageWriter* GetPackageWriter(const FName& PlatformName) override
+	virtual ICookedPackageWriter& GetPackageWriter(const ITargetPlatform* TargetPlatform) override
 	{
-		if (Cooker.CookOnTheFlyOptions.bZenStore)
-		{
-			for (UE::Cook::FCookSavePackageContext* Context : Cooker.SavePackageContexts)
-			{
-				if (FName(*Context->SaveContext.TargetPlatform->PlatformName()) == PlatformName)
-				{
-					return Context->PackageWriter;
-				}
-			}
-		}
-
-		return nullptr;
+		return Cooker.FindOrCreatePackageWriter(TargetPlatform);
 	}
 
 	DECLARE_DERIVED_EVENT(UCookOnTheFlyServer::FCookOnTheFlyServerInterface, UE::Cook::ICookOnTheFlyServer::FFlushEvent, FFlushEvent);
@@ -504,12 +440,9 @@ public:
 
 		UE_LOG(LogCook, Log, TEXT("Flushing..."));
 
-		if (Cooker.CookOnTheFlyOptions.bZenStore)
+		for (UE::Cook::FCookSavePackageContext* Context : Cooker.SavePackageContexts)
 		{
-			for (UE::Cook::FCookSavePackageContext* Context : Cooker.SavePackageContexts)
-			{
-				Context->PackageWriter->Flush();
-			}
+			Context->PackageWriter->Flush();
 		}
 
 		if (FlushEvent.IsBound())
@@ -572,7 +505,6 @@ UCookOnTheFlyServer::UCookOnTheFlyServer(const FObjectInitializer& ObjectInitial
 	CurrentCookMode(ECookMode::CookOnTheFly),
 	CookByTheBookOptions(nullptr),
 	CookFlags(ECookInitializationFlags::None),
-	bIsInitializingSandbox(false),
 	bIsSavingPackage(false),
 	AssetRegistry(nullptr),
 	PackageDatas(MakeUnique<UE::Cook::FPackageDatas>(*this))
@@ -686,7 +618,7 @@ bool UCookOnTheFlyServer::StartCookOnTheFly(FCookOnTheFlyOptions InCookOnTheFlyO
 		for (ITargetPlatform* TargetPlatform : CookOnTheFlyOptions.TargetPlatforms)
 		{
 			AddCookOnTheFlyPlatform(FName(*TargetPlatform->PlatformName()));
-	}
+		}
 	}
 
 	UE_LOG(LogCook, Display, TEXT("Starting '%s' cook-on-the-fly server"),
@@ -724,6 +656,7 @@ const ITargetPlatform* UCookOnTheFlyServer::AddCookOnTheFlyPlatform(const FName&
 		// Platform has already been added by this function or by StartCookByTheBook
 		return PlatformData->TargetPlatform;
 	}
+	TrySetDefaultAsyncIODeletePlatform(PlatformData->TargetPlatform->PlatformName());
 
 	if (IsInGameThread())
 	{
@@ -791,7 +724,8 @@ void UCookOnTheFlyServer::AddCookOnTheFlyPlatformFromGameThread(ITargetPlatform*
 	TArrayView<ITargetPlatform* const> NewTargetPlatforms(&TargetPlatform,1);
 
 	RefreshPlatformAssetRegistries(NewTargetPlatforms);
-	InitializeSandbox(NewTargetPlatforms);
+	FindOrCreateSaveContexts(NewTargetPlatforms);
+	BeginCookSandbox(NewTargetPlatforms);
 	InitializeTargetPlatforms(NewTargetPlatforms);
 	InitializeShadersForCookOnTheFly(NewTargetPlatforms);
 
@@ -802,8 +736,7 @@ void UCookOnTheFlyServer::AddCookOnTheFlyPlatformFromGameThread(ITargetPlatform*
 	{
 		Generator->SaveAssetRegistry(GetSandboxAssetRegistryFilename(), true);
 	}
-	InitializePackageStore(NewTargetPlatforms);
-	check(PlatformData->bIsSandboxInitialized); // This should have been set by InitializeSandbox, and it is what we use to determine whether a platform has been initialized
+	check(PlatformData->bIsSandboxInitialized); // This should have been set by BeginCookSandbox, and it is what we use to determine whether a platform has been initialized
 }
 
 void UCookOnTheFlyServer::OnTargetPlatformsInvalidated()
@@ -3200,6 +3133,7 @@ private: // Used only by UCookOnTheFlyServer, which has private access
 
 	// Per-platform data, only valid in PerPlatform callbacks
 	const ITargetPlatform* TargetPlatform = nullptr;
+	FCookSavePackageContext* CookContext = nullptr;
 	FSavePackageContext* SavePackageContext = nullptr;
 	ICookedPackageWriter* PackageWriter = nullptr;
 	FString PlatFilename;
@@ -4222,7 +4156,7 @@ public:
 			bool bCompatible = true;
 			for (UE::Cook::FCookSavePackageContext* SavePackageContext : SavePackageContexts)
 			{
-				if (SavePackageContext && SavePackageContext->PackageWriter)
+				if (!SavePackageContext->PackageWriterCapabilities.bDiffModeSupported)
 				{
 					bCompatible = false;
 					break;
@@ -4231,7 +4165,7 @@ public:
 			if (!bCompatible)
 			{
 				const TCHAR* CommandLineArg = IsRunningCookDiff() ? TEXT("-DIFFONLY") : TEXT("-LINKERDIFF");
-				UE_LOG(LogCook, Fatal, TEXT("%s was enabled, but PackageWriter is also enabled and PackageWriter does not support %s."),
+				UE_LOG(LogCook, Fatal, TEXT("%s was enabled, but -iostore is also enabled and iostore PackageWriters do not support %s."),
 					CommandLineArg, CommandLineArg);
 				bDiffEnabled = false;
 				LinkerDiffMode = LDM_None;
@@ -4457,7 +4391,7 @@ void UCookOnTheFlyServer::SaveCookedPackage(UE::Cook::FSaveCookedPackageContext&
 			{
 				if (DiffModeHelper->IsRunningCookDiff())
 				{
-					check(!Context.PackageWriter); // IsRunningCookDiff should have been disabled by DiffModeHelper if we have any PackageWriter
+					check(Context.CookContext->PackageWriterCapabilities.bDiffModeSupported); // IsRunningCookDiff should have been disabled by DiffModeHelper if we are using a PackageWriter
 					DiffModeHelper->ProcessPackage(Package);
 
 					// When looking for deterministic cook issues, first serialize the package to memory and do a simple diff with the existing package
@@ -4477,7 +4411,7 @@ void UCookOnTheFlyServer::SaveCookedPackage(UE::Cook::FSaveCookedPackageContext&
 				}
 				else if (DiffModeHelper->IsRunningCookLinkerDiff())
 				{
-					check(!Context.PackageWriter); // IsRunningCookLinkerDiff should have been disabled by DiffModeHelper if we have any PackageWriter
+					check(Context.CookContext->PackageWriterCapabilities.bDiffModeSupported); // IsRunningCookDiff should have been disabled by DiffModeHelper if we are using a PackageWriter
 
 					Context.SavePackageResult = GEditor->Save(Package, Context.World, Context.FlagsToCook, *Context.PlatFilename,
 						GError, nullptr, Context.bEndianSwap, false, Context.SaveFlags, TargetPlatform,
@@ -4646,24 +4580,13 @@ void FSaveCookedPackageContext::SetupPlatform(const ITargetPlatform* InTargetPla
 		World->PersistentLevel->HandleLegacyMapBuildData();
 	}
 
-	FCookSavePackageContext* CookContext = COTFS.GetSavePackageContext(TargetPlatform);
-	if (CookContext)
-	{
-		SavePackageContext = &CookContext->SaveContext;
-		PackageWriter = CookContext->PackageWriter;
-	}
-	else
-	{
-		SavePackageContext = nullptr;
-		PackageWriter = nullptr;
-	}
-	if (PackageWriter)
-	{
-		ICookedPackageWriter::FBeginPackageInfo Info;
-		Info.PackageName = Package->GetFName();
+	CookContext = &COTFS.FindOrCreateSaveContext(TargetPlatform);
+	SavePackageContext = &CookContext->SaveContext;
+	PackageWriter = CookContext->PackageWriter;
+	ICookedPackageWriter::FBeginPackageInfo Info;
+	Info.PackageName = Package->GetFName();
 
-		PackageWriter->BeginPackage(Info);
-	}
+	PackageWriter->BeginPackage(Info);
 
 	// Indicate Setup was successful
 	bPlatformSetupSuccessful = true;
@@ -4675,7 +4598,7 @@ void FSaveCookedPackageContext::FinishPlatform()
 	bool bSuccessful = SavePackageResult.IsSuccessful();
 	bool bLocalReferencedOnlyByEditorOnlyData = SavePackageResult == ESavePackageResult::ReferencedOnlyByEditorOnlyData;
 
-	if (bPlatformSetupSuccessful && PackageWriter)
+	if (bPlatformSetupSuccessful)
 	{
 		FAssetRegistryGenerator* Generator = COTFS.PlatformManager->GetPlatformData(TargetPlatform)->RegistryGenerator.Get();
 		FAssetPackageData* AssetPackageData = Generator->GetAssetPackageData(Package->GetFName());
@@ -5834,118 +5757,6 @@ bool UCookOnTheFlyServer::SaveCurrentIniSettings(const ITargetPlatform* TargetPl
 
 }
 
-FName UCookOnTheFlyServer::ConvertCookedPathToUncookedPath(
-	const FString& SandboxRootDir, const FString& RelativeRootDir,
-	const FString& SandboxProjectDir, const FString& RelativeProjectDir,
-	const FString& CookedPath, FString& OutUncookedPath) const
-{
-	OutUncookedPath.Reset();
-
-	// Check for remapped plugins' cooked content
-	if (PluginsToRemap.Num() > 0 && CookedPath.Contains(REMAPPED_PLUGINS))
-	{
-		int32 RemappedIndex = CookedPath.Find(REMAPPED_PLUGINS);
-		check(RemappedIndex >= 0);
-		static uint32 RemappedPluginStrLen = FCString::Strlen(REMAPPED_PLUGINS);
-		// Snip everything up through the RemappedPlugins/ off so we can find the plugin it corresponds to
-		FString PluginPath = CookedPath.RightChop(RemappedIndex + RemappedPluginStrLen + 1);
-		// Find the plugin that owns this content
-		for (TSharedRef<IPlugin> Plugin : PluginsToRemap)
-		{
-			if (PluginPath.StartsWith(Plugin->GetName()))
-			{
-				OutUncookedPath = Plugin->GetContentDir();
-				static uint32 ContentStrLen = FCString::Strlen(TEXT("Content/"));
-				// Chop off the pluginName/Content since it's part of the full path
-				OutUncookedPath /= PluginPath.RightChop(Plugin->GetName().Len() + ContentStrLen);
-				break;
-			}
-		}
-
-		if (OutUncookedPath.Len() > 0)
-		{
-			return FName(*OutUncookedPath);
-		}
-		// Otherwise fall through to sandbox handling
-	}
-
-	auto BuildUncookedPath = 
-		[&OutUncookedPath](const FString& CookedPath, const FString& CookedRoot, const FString& UncookedRoot)
-	{
-		OutUncookedPath.AppendChars(*UncookedRoot, UncookedRoot.Len());
-		OutUncookedPath.AppendChars(*CookedPath + CookedRoot.Len(), CookedPath.Len() - CookedRoot.Len());
-	};
-
-	if (CookedPath.StartsWith(SandboxRootDir))
-	{
-		// Optimized CookedPath.StartsWith(SandboxProjectDir) that does not compare all of SandboxRootDir again
-		if (CookedPath.Len() >= SandboxProjectDir.Len() && 
-			0 == FCString::Strnicmp(
-				*CookedPath + SandboxRootDir.Len(),
-				*SandboxProjectDir + SandboxRootDir.Len(),
-				SandboxProjectDir.Len() - SandboxRootDir.Len()))
-		{
-			BuildUncookedPath(CookedPath, SandboxProjectDir, RelativeProjectDir);
-		}
-		else
-		{
-			BuildUncookedPath(CookedPath, SandboxRootDir, RelativeRootDir);
-		}
-	}
-	else
-	{
-		FString FullCookedFilename = FPaths::ConvertRelativePathToFull(CookedPath);
-		BuildUncookedPath(FullCookedFilename, SandboxRootDir, RelativeRootDir);
-	}
-
-	// Convert to a standard filename as required by FPackageNameCache where this path is used.
-	FPaths::MakeStandardFilename(OutUncookedPath);
-
-	return FName(*OutUncookedPath);
-}
-
-void UCookOnTheFlyServer::GetAllCookedFiles(TMap<FName, FName>& UncookedPathToCookedPath, const FString& SandboxRootDir)
-{
-	TRACE_CPUPROFILER_EVENT_SCOPE(UCookOnTheFlyServer::GetAllCookedFiles);
-
-	TArray<FString> CookedFiles;
-	{
-		IPlatformFile& PlatformFile = FPlatformFileManager::Get().GetPlatformFile();
-		FPackageSearchVisitor PackageSearch(CookedFiles);
-		PlatformFile.IterateDirectoryRecursively(*SandboxRootDir, PackageSearch);
-	}
-
-	const FString SandboxProjectDir = FPaths::Combine(*SandboxRootDir, FApp::GetProjectName()) + TEXT("/");
-	const FString RelativeRootDir = FPaths::GetRelativePathToRoot();
-	const FString RelativeProjectDir = FPaths::ProjectDir();
-	FString UncookedFilename;
-	UncookedFilename.Reserve(1024);
-
-	for (const FString& CookedFile : CookedFiles)
-	{
-		const FName CookedFName(*CookedFile);
-		const FName UncookedFName = ConvertCookedPathToUncookedPath(
-			SandboxRootDir, RelativeRootDir,
-			SandboxProjectDir, RelativeProjectDir,
-			CookedFile, UncookedFilename);
-
-		UncookedPathToCookedPath.Add(UncookedFName, CookedFName);
-	}
-}
-
-void UCookOnTheFlyServer::DeleteSandboxDirectory(const FString& PlatformName)
-{
-	FString SandboxDirectory = GetSandboxDirectory(PlatformName);
-	FPaths::NormalizeDirectoryName(SandboxDirectory);
-	FString AsyncDeleteDirectory = GetAsyncDeleteDirectory(PlatformName, &SandboxDirectory);
-
-	FAsyncIODelete& LocalAsyncIODelete = GetAsyncIODelete(PlatformName, &AsyncDeleteDirectory);
-	LocalAsyncIODelete.DeleteDirectory(SandboxDirectory);
-
-	// UE_DEPRECATED(4.25, "Delete the old location for AsyncDeleteDirectory until all users have cooked at least once")
-	LocalAsyncIODelete.DeleteDirectory(SandboxDirectory + TEXT("AsyncDelete"));
-}
-
 void UCookOnTheFlyServer::TrySetDefaultAsyncIODeletePlatform(const FString& PlatformName)
 {
 	if (DefaultAsyncIODeletePlatformName.IsEmpty())
@@ -5956,284 +5767,27 @@ void UCookOnTheFlyServer::TrySetDefaultAsyncIODeletePlatform(const FString& Plat
 
 FAsyncIODelete& UCookOnTheFlyServer::GetAsyncIODelete()
 {
-	check(!DefaultAsyncIODeletePlatformName.IsEmpty());
-	return GetAsyncIODelete(DefaultAsyncIODeletePlatformName);
-}
-
-FAsyncIODelete& UCookOnTheFlyServer::GetAsyncIODelete(const FString& PlatformName, const FString* AsyncDeleteDirectory)
-{
-	FAsyncIODelete* AsyncIODeletePtr = AsyncIODelete.Get();
-	if (!AsyncIODeletePtr)
+	if (AsyncIODelete)
 	{
-		FString Buffer;
-		if (!AsyncDeleteDirectory)
-		{
-			Buffer = GetAsyncDeleteDirectory(PlatformName);
-			AsyncDeleteDirectory = &Buffer;
-		}
-		AsyncIODelete = MakeUnique<FAsyncIODelete>(*AsyncDeleteDirectory);
-		AsyncIODeletePtr = AsyncIODelete.Get();
+		return *AsyncIODelete;
 	}
-	// If we have already created the AsyncIODelete, we ignore the input PlatformName and use the existing AsyncIODelete initialized from whatever platform we used before
-	// The PlatformName is used only to construct a directory that we can be sure no other process is using (because a sandbox can only be cooked by one process at a time)
-	return *AsyncIODeletePtr;
-}
 
-FString UCookOnTheFlyServer::GetAsyncDeleteDirectory(const FString& PlatformName, const FString* SandboxDirectory) const
-{
+	// The PlatformName is used construct a directory that we can be sure no other process is using (because a sandbox can only be cooked by one process at a time)
 	// The TempRoot we will delete into is a sibling of the the Platform-specific sandbox directory, with name [PlatformDir]AsyncDelete
 	// Note that two UnrealEd-Cmd processes cooking to the same sandbox at the same time will therefore cause an error since FAsyncIODelete doesn't handle multiple processes sharing TempRoots.
 	// That simultaneous-cook behavior is also not supported in other assumptions throughout the cooker.
-	FString Buffer;
-	if (!SandboxDirectory)
-	{
-		// Avoid recalculating SandboxDirectory if caller supplied it, but if they didn't, calculate it here
-		Buffer = GetSandboxDirectory(PlatformName);
-		FPaths::NormalizeDirectoryName(Buffer);
-		SandboxDirectory = &Buffer;
-	}
-	return (*SandboxDirectory) + TEXT("_Del");
-}
-
-void UCookOnTheFlyServer::PopulateCookedPackagesFromDisk(const TArrayView<const ITargetPlatform* const>& Platforms)
-{
-	TRACE_CPUPROFILER_EVENT_SCOPE(UCookOnTheFlyServer::PopulateCookedPackagesFromDisk);
-
-	// See what files are out of date in the sandbox folder
-	for (int32 Index = 0; Index < Platforms.Num(); Index++)
-	{
-		TArray<FString> CookedPackagesToDelete;
-
-		const ITargetPlatform* Target = Platforms[Index];
-		UE::Cook::FPlatformData* PlatformData = PlatformManager->GetPlatformData(Target);
-		FString SandboxPath = GetSandboxDirectory(Target->PlatformName());
-
-		FString EngineSandboxPath = SandboxFile->ConvertToSandboxPath(*FPaths::EngineDir()) + TEXT("/");
-		EngineSandboxPath.ReplaceInline(TEXT("[Platform]"), *Target->PlatformName());
-
-		FString GameSandboxPath = SandboxFile->ConvertToSandboxPath(*(FPaths::ProjectDir() + TEXT("a.txt")));
-		GameSandboxPath.ReplaceInline(TEXT("a.txt"), TEXT(""));
-		GameSandboxPath.ReplaceInline(TEXT("[Platform]"), *Target->PlatformName());
-
-		FString LocalGamePath = FPaths::ProjectDir();
-		if (FPaths::IsProjectFilePathSet())
-		{
-			LocalGamePath = FPaths::GetPath(FPaths::GetProjectFilePath()) + TEXT("/");
-		}
-
-		FString LocalEnginePath = FPaths::EngineDir();
-
-		// Registry generator already exists
-		FAssetRegistryGenerator* PlatformAssetRegistry = PlatformData->RegistryGenerator.Get();
-		check(PlatformAssetRegistry);
-
-		// Load the platform cooked asset registry file
-		const FString CookedAssetRegistry = FPaths::ProjectDir() / TEXT("Metadata") / GetDevelopmentAssetRegistryFilename();
-		const FString SandboxCookedAssetRegistryFilename = ConvertToFullSandboxPath(*CookedAssetRegistry, true, Target->PlatformName());
-
-		bool bIsIterateSharedBuild = IsCookFlagSet(ECookInitializationFlags::IterateSharedBuild);
-
-		if (bIsIterateSharedBuild)
-		{
-			// see if the shared build is newer then the current cooked content in the local directory
-			FDateTime CurrentLocalCookedBuild = IFileManager::Get().GetTimeStamp(*SandboxCookedAssetRegistryFilename);
-
-			// iterate on the shared build if the option is set
-			FString SharedCookedAssetRegistry = FPaths::Combine(*FPaths::ProjectSavedDir(), TEXT("SharedIterativeBuild"), *Target->PlatformName(), TEXT("Metadata"), GetDevelopmentAssetRegistryFilename());
-
-			FDateTime CurrentIterativeCookedBuild = IFileManager::Get().GetTimeStamp(*SharedCookedAssetRegistry);
-
-			if ( (CurrentIterativeCookedBuild >= CurrentLocalCookedBuild) && 
-				(CurrentIterativeCookedBuild != FDateTime::MinValue()) )
-			{
-				// clean the sandbox
-				ClearPlatformCookedData(Target);
-
-				// SaveCurrentIniSettings(Target); // use this if we don't care about ini safty.
-				// copy the ini settings from the shared cooked build. 
-				const FString PlatformName = Target->PlatformName();
-				const FString SharedCookedIniFile = FPaths::Combine(*FPaths::ProjectSavedDir(), TEXT("SharedIterativeBuild"), *PlatformName, TEXT("Metadata"), TEXT("CookedIniVersion.txt"));
-				const FString SandboxCookedIniFile = ConvertToFullSandboxPath(*(FPaths::ProjectDir() / TEXT("Metadata") / TEXT("CookedIniVersion.txt")), true).Replace(TEXT("[Platform]"), *PlatformName);
-
-				IFileManager::Get().Copy(*SandboxCookedIniFile, *SharedCookedIniFile);
-
-				bool bIniSettingsOutOfDate = IniSettingsOutOfDate(Target);
-				if (bIniSettingsOutOfDate && !IsCookFlagSet(ECookInitializationFlags::IgnoreIniSettingsOutOfDate))
-				{
-					UE_LOG(LogCook, Display, TEXT("Shared iterative build ini settings out of date, not using shared cooked build"));
-				}
-				else
-				{
-					if (bIniSettingsOutOfDate)
-					{
-						UE_LOG(LogCook, Display, TEXT("Shared iterative build ini settings out of date, but we don't care"));
-					}
-
-					UE_LOG(LogCook, Display, TEXT("Shared iterative build is newer then local cooked build, iteratively cooking from shared build "));
-					PlatformAssetRegistry->LoadPreviousAssetRegistry(SharedCookedAssetRegistry);
-				}
-			}
-			else
-			{
-				UE_LOG(LogCook, Display, TEXT("Local cook is newer then shared cooked build, iterativly cooking from local build"));
-				PlatformAssetRegistry->LoadPreviousAssetRegistry(SandboxCookedAssetRegistryFilename);
-			}
-		}
-		else
-		{
-			PlatformAssetRegistry->LoadPreviousAssetRegistry(SandboxCookedAssetRegistryFilename);
-		}
-
-		// Get list of changed packages
-		TSet<FName> ModifiedPackages, NewPackages, RemovedPackages, IdenticalCookedPackages, IdenticalUncookedPackages;
-
-		// We recurse modifications up the reference chain because it is safer, if this ends up being a significant issue in some games we can add a command line flag
-		bool bRecurseModifications = true;
-		bool bRecurseScriptModifications = !IsCookFlagSet(ECookInitializationFlags::IgnoreScriptPackagesOutOfDate);
-		PlatformAssetRegistry->ComputePackageDifferences(ModifiedPackages, NewPackages, RemovedPackages, IdenticalCookedPackages, IdenticalUncookedPackages, bRecurseModifications, bRecurseScriptModifications);
-
-		// check the files on disk 
-		TMap<FName, FName> UncookedPathToCookedPath;
-		// get all the on disk cooked files
-		GetAllCookedFiles(UncookedPathToCookedPath, SandboxPath);
-
-		const static FName NAME_DummyCookedFilename(TEXT("DummyCookedFilename")); // pls never name a package dummycookedfilename otherwise shit might go wonky
-		if (bIsIterateSharedBuild)
-		{
-			check(IFileManager::Get().FileExists(*NAME_DummyCookedFilename.ToString()) == false);
-
-			TSet<FName> ExistingPackages = ModifiedPackages;
-			ExistingPackages.Append(RemovedPackages);
-			ExistingPackages.Append(IdenticalCookedPackages);
-			ExistingPackages.Append(IdenticalUncookedPackages);
-
-			// if we are iterating of a shared build the cooked files might not exist in the cooked directory because we assume they are packaged in the pak file (which we don't want to extract)
-			for (FName PackageName : ExistingPackages)
-			{
-				FString Filename;
-				if (FPackageName::DoesPackageExist(PackageName.ToString(), nullptr, &Filename))
-				{
-					UncookedPathToCookedPath.Add(FName(*Filename), NAME_DummyCookedFilename);
-				}
-			}
-		}
-
-		uint32 NumPackagesConsidered = UncookedPathToCookedPath.Num();
-		uint32 NumPackagesUnableToFindCookedPackageInfo = 0;
-		uint32 NumPackagesFileHashMismatch = 0;
-		uint32 NumPackagesKept = 0;
-		uint32 NumMarkedFailedSaveKept = 0;
-		uint32 NumPackagesRemoved = 0;
-
-		TArray<FName> KeptPackages;
-
-		for (const auto& CookedPaths : UncookedPathToCookedPath)
-		{
-			const FName CookedFile = CookedPaths.Value;
-			const FName UncookedFilename = CookedPaths.Key;
-			const FName* FoundPackageName = GetPackageNameCache().GetCachedPackageNameFromStandardFileName(UncookedFilename);
-			bool bShouldKeep = true;
-			const FName SourcePackageName = FoundPackageName ? *FoundPackageName : NAME_None;
-			if ( !FoundPackageName )
-			{
-				// Source file no longer exists
-				++NumPackagesRemoved;
-				bShouldKeep = false;
-			}
-			else
-			{
-				if (ModifiedPackages.Contains(SourcePackageName))
-				{
-					++NumPackagesFileHashMismatch;
-					bShouldKeep = false;
-				}
-				else if (NewPackages.Contains(SourcePackageName) || RemovedPackages.Contains(SourcePackageName))
-				{
-					++NumPackagesUnableToFindCookedPackageInfo;
-					bShouldKeep = false;
-				}
-				else if (IdenticalUncookedPackages.Contains(SourcePackageName))
-				{
-					// These are packages which failed to save the first time 
-					// most likely because they are editor only packages
-					bShouldKeep = false;
-				}
-			}
-
-			TArray<const ITargetPlatform*> PlatformsForPackage;
-			PlatformsForPackage.Add(Target);
-
-			if (bShouldKeep)
-			{
-				// Mark this package as cooked so that we don't unnecessarily try to cook it again
-				if (IdenticalCookedPackages.Contains(SourcePackageName))
-				{
-					UE::Cook::FPackageData* PackageData = PackageDatas->TryAddPackageDataByPackageName(SourcePackageName);
-					if (PackageData)
-					{
-						PackageData->AddCookedPlatforms({ Target }, true /* bSucceeded */);
-						KeptPackages.Add(SourcePackageName);
-						++NumPackagesKept;
-					}
-				}
-			}
-			else
-			{
-				if (SourcePackageName != NAME_None && IsCookByTheBookMode()) // cook on the fly will requeue this package when it wants it 
-				{
-					// Force cook the modified file
-					ExternalRequests->EnqueueUnique(UE::Cook::FFilePlatformRequest(UncookedFilename, PlatformsForPackage));
-				}
-				if (CookedFile != NAME_DummyCookedFilename)
-				{
-					// delete the old package 
-					const FString CookedFullPath = FPaths::ConvertRelativePathToFull(CookedFile.ToString());
-					UE_LOG(LogCook, Verbose, TEXT("Deleting cooked package %s failed filehash test"), *CookedFullPath);
-					CookedPackagesToDelete.Add(CookedFullPath);
-				}
-				else
-				{
-					// the cooker should rebuild this package because it's not in the cooked package list
-					// the new package will have higher priority then the package in the original shared cooked build
-					const FString UncookedFilenameString = UncookedFilename.ToString();
-					UE_LOG(LogCook, Verbose, TEXT("Shared cooked build: Detected package is out of date %s"), *UncookedFilenameString);
-				}
-			}
-		}
-
-		// Register identical uncooked packages from previous run
-		for (FName UncookedPackage : IdenticalUncookedPackages)
-		{
-			UE::Cook::FPackageData* PackageData = PackageDatas->TryAddPackageDataByPackageName(UncookedPackage);
-			if (PackageData)
-			{
-				ensure(!PackageData->HasAnyCookedPlatforms({ Target }, /* bIncludeFailed */ false));
-				PackageData->AddCookedPlatforms({ Target }, false /* bSucceeded */);
-				KeptPackages.Add(UncookedPackage);
-				++NumMarkedFailedSaveKept;
-			}
-		}
-
-		PlatformAssetRegistry->UpdateKeptPackages(KeptPackages);
-
-		UE_LOG(LogCook, Display, TEXT("Iterative cooking summary for %s, \nConsidered: %d, \nFile Hash missmatch: %d, \nPackages Kept: %d, \nPackages failed save kept: %d, \nMissing Cooked Info(expected 0): %d"),
-			*Target->PlatformName(),
-			NumPackagesConsidered, NumPackagesFileHashMismatch,
-			NumPackagesKept, NumMarkedFailedSaveKept,
-			NumPackagesUnableToFindCookedPackageInfo);
+	check(!DefaultAsyncIODeletePlatformName.IsEmpty());
+	FString DeleteDirectory = GetSandboxDirectory(DefaultAsyncIODeletePlatformName);
+	FPaths::NormalizeDirectoryName(DeleteDirectory);
+	DeleteDirectory += TEXT("_Del");
 
-		auto DeletePackageLambda = [&CookedPackagesToDelete](int32 PackageIndex)
-		{
-			const FString& CookedFullPath = CookedPackagesToDelete[PackageIndex];
-			IFileManager::Get().Delete(*CookedFullPath, true, true, true);
-		};
-		ParallelFor(CookedPackagesToDelete.Num(), DeletePackageLambda);
-	}
+	AsyncIODelete = MakeUnique<FAsyncIODelete>(DeleteDirectory);
+	return *AsyncIODelete;
 }
 
-void UCookOnTheFlyServer::PopulateCookedPackagesFromPackageStore(const TArrayView<const ITargetPlatform* const>& Platforms)
+void UCookOnTheFlyServer::PopulateCookedPackages(TArrayView<const ITargetPlatform* const> TargetPlatforms)
 {
-	TRACE_CPUPROFILER_EVENT_SCOPE(UCookOnTheFlyServer::PopulateCookedPackagesFromPackgeStore);
+	TRACE_CPUPROFILER_EVENT_SCOPE(UCookOnTheFlyServer::PopulateCookedPackages);
 
 	struct FPreviousAssetPackageData
 	{
@@ -6261,20 +5815,19 @@ void UCookOnTheFlyServer::PopulateCookedPackagesFromPackageStore(const TArrayVie
 		TMap<FName, const FAssetPackageData*> AssetPackageDataMap;
 	};
 
-	for (const ITargetPlatform* Platform : Platforms)
+	for (const ITargetPlatform* TargetPlatform : TargetPlatforms)
 	{
-		UE_LOG(LogCook, Display, TEXT("Populating cooked package(s) from package store on platform '%s'"), *Platform->PlatformName());
-		
-		UE::Cook::FPlatformData* PlatformData = PlatformManager->GetPlatformData(Platform);
+		UE::Cook::FPlatformData* PlatformData = PlatformManager->GetPlatformData(TargetPlatform);
 		FAssetRegistryGenerator* PlatformAssetRegistry = PlatformData->RegistryGenerator.Get();
-		ICookedPackageWriter* PackageWriter = GetPackageWriter(FName(*Platform->PlatformName()));
-		
+		UE::Cook::FCookSavePackageContext& CookSavePackageContext = FindOrCreateSaveContext(TargetPlatform);
+		ICookedPackageWriter& PackageWriter = *CookSavePackageContext.PackageWriter;
 		check(PlatformData);
 		check(PlatformAssetRegistry);
-		check(PackageWriter);
+		UE_LOG(LogCook, Display, TEXT("Populating cooked package(s) from %s package store on platform '%s'"),
+			*CookSavePackageContext.WriterDebugName, *TargetPlatform->PlatformName());
 
 		TArray<ICookedPackageWriter::FCookedPackageInfo> CookedPackages;
-		PackageWriter->GetCookedPackages(CookedPackages);
+		PackageWriter.GetCookedPackages(CookedPackages);
 		
 		UE_LOG(LogCook, Display, TEXT("Found '%d' cooked package(s) in package store"), CookedPackages.Num());
 
@@ -6288,40 +5841,79 @@ void UCookOnTheFlyServer::PopulateCookedPackagesFromPackageStore(const TArrayVie
 		PlatformAssetRegistry->ComputePackageDifferences(Options, PreviousAssetPackageData.AssetPackageDataMap, Difference);
 
 		TArray<FName> PackagesToRemove;
-		PackagesToRemove.Reserve(Difference.ModifiedPackages.Num() + Difference.RemovedPackages.Num());
-		
-		for (const FName& ModifiedPackageName : Difference.ModifiedPackages)
+		PackagesToRemove.Reserve(Difference.ModifiedPackages.Num() + Difference.RemovedPackages.Num() + Difference.IdenticalUncookedPackages.Num());
+		for (FName PackageToRemove : Difference.ModifiedPackages)
 		{
-			PackagesToRemove.Add(ModifiedPackageName);
+			PackagesToRemove.Add(PackageToRemove);
 		}
-		
-		for (const FName& RemovedPackageName : Difference.RemovedPackages)
+		for (FName PackageToRemove : Difference.RemovedPackages)
+		{
+			PackagesToRemove.Add(PackageToRemove);
+		}
+		for (FName PackageToRemove : Difference.IdenticalUncookedPackages)
 		{
-			PackagesToRemove.Add(RemovedPackageName );
+			PackagesToRemove.Add(PackageToRemove);
 		}
 		
 		UE_LOG(LogCook, Display, TEXT("Keeping '%d' and removing '%d' cooked package(s)"), Difference.IdenticalCookedPackages.Num(), PackagesToRemove.Num());
 		
-		PackageWriter->RemoveCookedPackages(PackagesToRemove);
+		PackageWriter.RemoveCookedPackages(PackagesToRemove);
 		
 		TArray<FName> PackagesToKeep;
-		for (const FName& IdenticalPackageName : Difference.IdenticalCookedPackages)
+		for (const FName& IdenticalPackage : Difference.IdenticalCookedPackages)
+		{
+			// Mark this package as cooked so that we don't unnecessarily try to cook it again
+			UE::Cook::FPackageData* PackageData = PackageDatas->TryAddPackageDataByPackageName(IdenticalPackage);
+			if (PackageData)
+			{
+				PackageData->AddCookedPlatforms({ TargetPlatform }, true /* bSucceeded */);
+				PackagesToKeep.Add(IdenticalPackage);
+			}
+		}
+		for (FName UncookedPackage : Difference.IdenticalUncookedPackages)
 		{
-			UE::Cook::FPackageData* PackageData = PackageDatas->TryAddPackageDataByPackageName(IdenticalPackageName);
+			// Mark this failed-to-cook package as cooked so that we don't unnecessarily try to cook it again
+			UE::Cook::FPackageData* PackageData = PackageDatas->TryAddPackageDataByPackageName(UncookedPackage);
 			if (PackageData)
 			{
-				const bool bSucceded = true;
-				PackageData->AddCookedPlatforms({ Platform }, bSucceded);
-				PackagesToKeep.Add(IdenticalPackageName);
+				ensure(!PackageData->HasAnyCookedPlatforms({ TargetPlatform }, /* bIncludeFailed */ false));
+				PackageData->AddCookedPlatforms({ TargetPlatform }, false /* bSucceeded */);
+				PackagesToKeep.Add(UncookedPackage);
+			}
+		}
+		if (IsCookByTheBookMode())
+		{
+			for (const FName& RemovePackageName : Difference.ModifiedPackages)
+			{
+				// cook on the fly will requeue this package when it wants it, but for cook on the book we force cook the modified file
+				// so that the output set of packages is up to date (even if the user is currently cooking only a subset)
+				UE::Cook::FPackageData* PackageData = PackageDatas->TryAddPackageDataByPackageName(RemovePackageName);
+				if (PackageData)
+				{
+					ExternalRequests->EnqueueUnique(UE::Cook::FFilePlatformRequest(PackageData->GetFileName(),
+						TConstArrayView<const ITargetPlatform*>{ TargetPlatform }));
+				}
 			}
 		}
 
-		PlatformAssetRegistry->UpdateKeptPackages(PackagesToKeep, [&PreviousAssetPackageData](const FName& PackageName, FAssetPackageData* Current)
+		TUniquePtr<FAssetRegistryState> PreviousAssetRegistry(PackageWriter.ReleasePreviousAssetRegistry());
+		if (PreviousAssetRegistry)
+		{
+			// SetPreviousAssetRegistry is needed for UpdateKeptPackagesAssetData, which is called at the end
+			// of the cook for any packages that were kept from the previous cook and not cooked in this cook
+			PlatformAssetRegistry->SetPreviousAssetRegistry(MoveTemp(PreviousAssetRegistry));
+			PlatformAssetRegistry->UpdateKeptPackages(PackagesToKeep);
+		}
+		else
 		{
-			const FAssetPackageData* Previous = PreviousAssetPackageData.AssetPackageDataMap.FindRef(PackageName);
-			Current->CookedHash = Previous->CookedHash;
-			Current->DiskSize = Previous->DiskSize;
-		});
+			// TODO: Implementation ReleasePreviousAssetRegistry for all PackageWriters and add an assertion here instead of handling null.
+			PlatformAssetRegistry->UpdateKeptPackages(PackagesToKeep, [&PreviousAssetPackageData](const FName& PackageName, FAssetPackageData* Current)
+				{
+					const FAssetPackageData* Previous = PreviousAssetPackageData.AssetPackageDataMap.FindRef(PackageName);
+					Current->CookedHash = Previous->CookedHash;
+					Current->DiskSize = Previous->DiskSize;
+				});
+		}
 	}
 }
 
@@ -7616,8 +7208,6 @@ void UCookOnTheFlyServer::CancelCookByTheBook()
 
 		CancelAllQueues();
 
-		ClearPackageStoreContexts();
-
 		CookByTheBookOptions->bRunning = false;
 		SandboxFile = nullptr;
 
@@ -7697,20 +7287,24 @@ void UCookOnTheFlyServer::ClearPlatformCookedData(const ITargetPlatform* TargetP
 	{
 		return;
 	}
+	ResetCookResults({ TargetPlatform });
 
-	// if we are going to clear the cooked packages it is conceivable that we will recook the packages which we just cooked 
-	// that means it's also conceivable that we will recook the same package which currently has an outstanding async write request
-	UPackage::WaitForAsyncFileWrites();
+	FindOrCreatePackageWriter(TargetPlatform).RemoveCookedPackages();
+}
 
+void UCookOnTheFlyServer::ResetCookResults(TConstArrayView<const ITargetPlatform*> TargetPlatforms)
+{
 	for (UE::Cook::FPackageData* PackageData : *PackageDatas.Get())
 	{
-		PackageData->RemoveCookedPlatform(TargetPlatform);
+		PackageData->RemoveCookedPlatforms(TargetPlatforms);
 	}
 
 	TArray<FName> PackageNames;
-	PackageTracker->UnsolicitedCookedPackages.GetPackagesForPlatformAndRemove(TargetPlatform, PackageNames);
-
-	DeleteSandboxDirectory(TargetPlatform->PlatformName());
+	for (const ITargetPlatform* TargetPlatform : TargetPlatforms)
+	{
+		PackageNames.Reset();
+		PackageTracker->UnsolicitedCookedPackages.GetPackagesForPlatformAndRemove(TargetPlatform, PackageNames);
+	}
 }
 
 void UCookOnTheFlyServer::ClearPlatformCookedData(const FString& PlatformName)
@@ -7718,11 +7312,11 @@ void UCookOnTheFlyServer::ClearPlatformCookedData(const FString& PlatformName)
 	ClearPlatformCookedData(GetTargetPlatformManagerRef().FindTargetPlatform(PlatformName));
 }
 
-void UCookOnTheFlyServer::ClearCachedCookedPlatformDataForPlatform( const ITargetPlatform* TargetPlatform )
+void UCookOnTheFlyServer::ClearCachedCookedPlatformDataForPlatform(const ITargetPlatform* TargetPlatform)
 {
 	if (TargetPlatform)
 	{
-		for ( TObjectIterator<UObject> It; It; ++It )
+		for (TObjectIterator<UObject> It; It; ++It)
 		{
 			It->ClearCachedCookedPlatformData(TargetPlatform);
 		}
@@ -7751,7 +7345,7 @@ void UCookOnTheFlyServer::CreateSandboxFile()
 	// but will not be used to actually write/read files so we can safely
 	// use [Platform] token in the sandbox directory name and then replace it
 	// with the actual platform name.
-	check( SandboxFile == nullptr );
+	check(SandboxFile == nullptr);
 	SandboxFile = FSandboxPlatformFile::Create(false);
 
 	// Output directory override.	
@@ -7761,7 +7355,7 @@ void UCookOnTheFlyServer::CreateSandboxFile()
 	SandboxFile->Initialize(&FPlatformFileManager::Get().GetPlatformFile(), *FString::Printf(TEXT("-sandbox=\"%s\""), *OutputDirectory));
 }
 
-void UCookOnTheFlyServer::InitializeSandbox(const TArrayView<const ITargetPlatform* const>& TargetPlatforms)
+void UCookOnTheFlyServer::BeginCookSandbox(TConstArrayView<const ITargetPlatform*> TargetPlatforms)
 {
 #if OUTPUT_COOKTIMING
 	double CleanSandboxTime = 0.0;
@@ -7774,25 +7368,44 @@ void UCookOnTheFlyServer::InitializeSandbox(const TArrayView<const ITargetPlatfo
 			CreateSandboxFile();
 		}
 
-		// before we can delete any cooked files we need to make sure that we have finished writing them
-		UPackage::WaitForAsyncFileWrites();
-
-		// Skip markup of packages for reload during sandbox initialization.  TODO: Is this no longer required because it is no longer possible to load packages in this function?
-		bIsInitializingSandbox = true;
-		ON_SCOPE_EXIT
-		{
-			bIsInitializingSandbox = false;
-		};
-
-		TSet<const ITargetPlatform*> AlreadyInitializedPlatforms;
-		TArray<const ITargetPlatform*, TInlineAllocator<ExpectedMaxNumPlatforms>> RefreshPlatforms;
 		const bool bIsDiffOnly = FParse::Param(FCommandLine::Get(), TEXT("DIFFONLY"));
 		const bool bIsIterativeCook = IsCookFlagSet(ECookInitializationFlags::Iterative);
+		const bool bIsSharedIterativeCook = IsCookFlagSet(ECookInitializationFlags::IterateSharedBuild);
+		TArray<const ITargetPlatform*, TInlineAllocator<ExpectedMaxNumPlatforms>> ResetPlatforms;
+		TArray<const ITargetPlatform*, TInlineAllocator<ExpectedMaxNumPlatforms>> PopulatePlatforms;
 
-		for (const ITargetPlatform* Target : TargetPlatforms)
+		for (const ITargetPlatform* TargetPlatform : TargetPlatforms)
 		{
-			UE::Cook::FPlatformData* PlatformData = PlatformManager->GetPlatformData(Target);
-			const bool bIsIniSettingsOutOfDate = IniSettingsOutOfDate(Target); // needs to be executed for side effects even if non-iterative
+			UE::Cook::FPlatformData* PlatformData = PlatformManager->GetPlatformData(TargetPlatform);
+			ICookedPackageWriter& PackageWriter = FindOrCreatePackageWriter(TargetPlatform);
+			bool bIterateSharedBuild = false;
+			if (bIsIterativeCook && bIsSharedIterativeCook && !PlatformData->bIsSandboxInitialized)
+			{
+				// see if the shared build is newer then the current cooked content in the local directory
+				FString SharedCookedAssetRegistry = FPaths::Combine(*FPaths::ProjectSavedDir(), TEXT("SharedIterativeBuild"),
+					*TargetPlatform->PlatformName(), TEXT("Metadata"), GetDevelopmentAssetRegistryFilename());
+				
+				FDateTime PreviousLocalCookedBuild = PackageWriter.GetPreviousCookTime();
+				FDateTime PreviousSharedCookedBuild = IFileManager::Get().GetTimeStamp(*SharedCookedAssetRegistry);
+				if (PreviousSharedCookedBuild != FDateTime::MinValue() &&
+					PreviousSharedCookedBuild >= PreviousLocalCookedBuild)
+				{
+					// copy the ini settings from the shared cooked build. 
+					const FString SharedCookedIniFile = FPaths::Combine(*FPaths::ProjectSavedDir(), TEXT("SharedIterativeBuild"),
+						*TargetPlatform->PlatformName(), TEXT("Metadata"), TEXT("CookedIniVersion.txt"));
+					FString SandboxCookedIniFile = FPaths::ProjectDir() / TEXT("Metadata") / TEXT("CookedIniVersion.txt");
+					SandboxCookedIniFile = ConvertToFullSandboxPath(*SandboxCookedIniFile, true);
+					SandboxCookedIniFile.ReplaceInline(TEXT("[Platform]"), *TargetPlatform->PlatformName());
+					IFileManager::Get().Copy(*SandboxCookedIniFile, *SharedCookedIniFile);
+					bIterateSharedBuild = true;
+					UE_LOG(LogCook, Display, TEXT("Shared iterative build is newer then local cooked build, iteratively cooking from shared build."));
+				}
+				else
+				{
+					UE_LOG(LogCook, Display, TEXT("Local cook is newer then shared cooked build, iteratively cooking from local build."));
+				}
+			}
+			const bool bIsIniSettingsOutOfDate = IniSettingsOutOfDate(TargetPlatform); // needs to be executed for side effects even if non-iterative
 
 			bool bShouldClearCookedContent = true;
 			if (bIsDiffOnly)
@@ -7804,19 +7417,18 @@ void UCookOnTheFlyServer::InitializeSandbox(const TArrayView<const ITargetPlatfo
 			{
 				if (!bIsIniSettingsOutOfDate)
 				{
-					// We have constructed the sandbox in an earlier cook in this process (e.g. in the editor) and should not clear it again
 					bShouldClearCookedContent = false;
 				}
 				else
 				{
 					if (!IsCookFlagSet(ECookInitializationFlags::IgnoreIniSettingsOutOfDate))
 					{
-						UE_LOG(LogCook, Display, TEXT("Cook invalidated for platform %s ini settings don't match from last cook, clearing all cooked content"), *Target->PlatformName());
+						UE_LOG(LogCook, Display, TEXT("Cook invalidated for platform %s ini settings don't match from last cook, clearing all cooked content"), *TargetPlatform->PlatformName());
 						bShouldClearCookedContent = true;
 					}
 					else
 					{
-						UE_LOG(LogCook, Display, TEXT("Inisettings were out of date for platform %s but we are going with it anyway because IgnoreIniSettingsOutOfDate is set"), *Target->PlatformName());
+						UE_LOG(LogCook, Display, TEXT("Inisettings were out of date for platform %s but we are going with it anyway because IgnoreIniSettingsOutOfDate is set"), *TargetPlatform->PlatformName());
 						bShouldClearCookedContent = false;
 					}
 				}
@@ -7824,45 +7436,39 @@ void UCookOnTheFlyServer::InitializeSandbox(const TArrayView<const ITargetPlatfo
 			else
 			{
 				// In non-iterative cooks we will be replacing every cooked package and so should wipe the cooked directory
-				UE_LOG(LogCook, Display, TEXT("Clearing all cooked content for platform %s"), *Target->PlatformName());
+				UE_LOG(LogCook, Display, TEXT("Clearing all cooked content for platform %s"), *TargetPlatform->PlatformName());
 				bShouldClearCookedContent = true;
 			}
 
+			ICookedPackageWriter::FCookInfo CookInfo;
+			CookInfo.CookMode = IsCookOnTheFlyMode() ? ICookedPackageWriter::FCookInfo::CookOnTheFlyMode : ICookedPackageWriter::FCookInfo::CookByTheBookMode;
+			CookInfo.bCleanBuild = false;
+			CookInfo.bIterateSharedBuild = bIterateSharedBuild && !bShouldClearCookedContent;
 			if (bShouldClearCookedContent)
 			{
-				ClearPlatformCookedData(Target);
-				SaveCurrentIniSettings(Target);
+				CookInfo.bCleanBuild = true;
+				ResetPlatforms.Add(TargetPlatform);
 			}
-			else
+			else if (!PlatformData->bIsSandboxInitialized)
 			{
-				RefreshPlatforms.Add(Target);
-				if (PlatformData->bIsSandboxInitialized)
+				ResetPlatforms.Add(TargetPlatform);
+				// Reset but don't populate if we are looking for deterministic cooking differences; start from an empty list of cooked packages
+				if (!bIsDiffOnly)
 				{
-					AlreadyInitializedPlatforms.Add(Target);
+					PopulatePlatforms.Add(TargetPlatform);
 				}
 			}
-
+			PackageWriter.BeginCook(CookInfo);
 			PlatformData->bIsSandboxInitialized = true;
 		}
 
-		// Don't populate platforms that were already initialized; we already populated them when we first initialized them
-		RefreshPlatforms.RemoveAllSwap([&AlreadyInitializedPlatforms](const ITargetPlatform* TargetPlatform)
+		if (ResetPlatforms.Num())
 		{
-			return AlreadyInitializedPlatforms.Contains(TargetPlatform);
-		});
-
-		if (RefreshPlatforms.Num() != 0)
+			ResetCookResults(ResetPlatforms);
+		}
+		if (PopulatePlatforms.Num())
 		{
-			for (UE::Cook::FPackageData* PackageData : *PackageDatas.Get())
-			{
-				PackageData->RemoveCookedPlatforms(RefreshPlatforms);
-			}
-			// The asset registry makes populating cooked packages from disk fast, and populating is a performance benefit
-			// Don't populate however if we are looking for deterministic cooking differences; start from an empty list of cooked packages
-			if (!bIsDiffOnly && !IsUsingIoStore()) 
-			{
-				PopulateCookedPackagesFromDisk(RefreshPlatforms);
-			}
+			PopulateCookedPackages(PopulatePlatforms);
 		}
 	}
 
@@ -7877,8 +7483,13 @@ void UCookOnTheFlyServer::InitializeSandbox(const TArrayView<const ITargetPlatfo
 #endif
 }
 
-void UCookOnTheFlyServer::InitializePackageStore(const TArrayView<const ITargetPlatform* const>& TargetPlatforms)
+UE::Cook::FCookSavePackageContext* UCookOnTheFlyServer::CreateSaveContext(const ITargetPlatform* TargetPlatform) const
 {
+	if (SandboxFile == nullptr)
+	{
+		const_cast<UCookOnTheFlyServer&>(*this).CreateSandboxFile();
+	}
+
 	const FString RootPath = FPaths::RootDir();
 	const FString RootPathSandbox = ConvertToFullSandboxPath(*RootPath, true);
 
@@ -7888,59 +7499,39 @@ void UCookOnTheFlyServer::InitializePackageStore(const TArrayView<const ITargetP
 	const bool bIsDiffOnly = FParse::Param(FCommandLine::Get(), TEXT("DIFFONLY"));
 	const bool bIsCleanBuild = !IsCookFlagSet(ECookInitializationFlags::Iterative);
 
-	SavePackageContexts.Reserve(TargetPlatforms.Num());
-
-	for (const ITargetPlatform* TargetPlatform: TargetPlatforms)
-	{
-		const FString PlatformString = TargetPlatform->PlatformName();
-
-		const FString ResolvedRootPath = RootPathSandbox.Replace(TEXT("[Platform]"), *PlatformString);
-		const FString ResolvedProjectPath = ProjectPathSandbox.Replace(TEXT("[Platform]"), *PlatformString);
-
-		ICookedPackageWriter* PackageWriter = nullptr;
-		if (IsUsingIoStore())
-		{
-			FString MetadataDirectoryPath = ProjectPathSandbox / TEXT("Metadata");
-			MetadataDirectoryPath.ReplaceInline(TEXT("[Platform]"), *PlatformString);
-
-			if (IsUsingZenStore())
-			{
-				PackageWriter = new FZenStoreWriter(	ResolvedRootPath, 
-															MetadataDirectoryPath, 
-															TargetPlatform, 
-															bIsCleanBuild);
-			}
-			else
-			{
-				PackageWriter = new FFilePackageStoreWriter(ResolvedRootPath, MetadataDirectoryPath, TargetPlatform);
-			}
-		}
+	const FString PlatformString = TargetPlatform->PlatformName();
 
-		FConfigFile PlatformEngineIni;
-		FConfigCacheIni::LoadLocalIniFile(PlatformEngineIni, TEXT("Engine"), true, *TargetPlatform->IniPlatformName());
-		
-		bool bLegacyBulkDataOffsets = false;
-		PlatformEngineIni.GetBool(TEXT("Core.System"), TEXT("LegacyBulkDataOffsets"), bLegacyBulkDataOffsets);
+	const FString ResolvedRootPath = RootPathSandbox.Replace(TEXT("[Platform]"), *PlatformString);
+	const FString ResolvedProjectPath = ProjectPathSandbox.Replace(TEXT("[Platform]"), *PlatformString);
 
-		UE::Cook::FCookSavePackageContext* SavePackageContext 
-			= new UE::Cook::FCookSavePackageContext(TargetPlatform, PackageWriter, bLegacyBulkDataOffsets);
-		SavePackageContexts.Add(SavePackageContext);
+	ICookedPackageWriter* PackageWriter = nullptr;
+	FString MetadataDirectoryPath = ProjectPathSandbox / TEXT("Metadata");
+	MetadataDirectoryPath.ReplaceInline(TEXT("[Platform]"), *PlatformString);
+	FString WriterDebugName;
+	if (IsUsingZenStore())
+	{
+		PackageWriter = new FZenStoreWriter(ResolvedRootPath, MetadataDirectoryPath, TargetPlatform);
+		WriterDebugName = TEXT("ZenStore");
 	}
-
-	if (IsUsingIoStore())
+	else if (IsUsingIoStore())
 	{
-		if (!bIsCleanBuild)
-		{
-			PopulateCookedPackagesFromPackageStore(TargetPlatforms);
-		}
-
-		for (const UE::Cook::FCookSavePackageContext* Context : SavePackageContexts)
-		{
-			ICookedPackageWriter::FCookInfo CookInfo;
-			CookInfo.CookMode = IsCookOnTheFlyMode() ? ICookedPackageWriter::FCookInfo::CookOnTheFlyMode : ICookedPackageWriter::FCookInfo::CookByTheBookMode;
-			Context->PackageWriter->BeginCook(CookInfo);
-		}
+		PackageWriter = new FFilePackageStoreWriter(ResolvedRootPath, MetadataDirectoryPath, TargetPlatform);
+		WriterDebugName = TEXT("IoStoreFileWriter");
 	}
+	else
+	{
+		PackageWriter = new FLooseCookedPackageWriter(ResolvedRootPath, MetadataDirectoryPath, TargetPlatform,
+			const_cast<UCookOnTheFlyServer&>(*this).GetAsyncIODelete(), GetPackageNameCache(), PluginsToRemap);
+		WriterDebugName = TEXT("DirectoryWriter");
+	}
+
+	FConfigFile PlatformEngineIni;
+	FConfigCacheIni::LoadLocalIniFile(PlatformEngineIni, TEXT("Engine"), true, *TargetPlatform->IniPlatformName());
+		
+	bool bLegacyBulkDataOffsets = false;
+	PlatformEngineIni.GetBool(TEXT("Core.System"), TEXT("LegacyBulkDataOffsets"), bLegacyBulkDataOffsets);
+
+	return new UE::Cook::FCookSavePackageContext(TargetPlatform, PackageWriter, bLegacyBulkDataOffsets, WriterDebugName);
 }
 
 void UCookOnTheFlyServer::FinalizePackageStore()
@@ -7948,19 +7539,11 @@ void UCookOnTheFlyServer::FinalizePackageStore()
 	UE_SCOPED_HIERARCHICAL_COOKTIMER(FinalizePackageStore);
 
 	UE_LOG(LogCook, Display, TEXT("Finalize package store(s)..."));
-	for (UE::Cook::FCookSavePackageContext* PackageContext : SavePackageContexts)
+	for (const ITargetPlatform* TargetPlatform : PlatformManager->GetSessionPlatforms())
 	{
-		if (PackageContext != nullptr)
-		{
-			if (PackageContext->PackageWriter != nullptr)
-			{
-				PackageContext->PackageWriter->EndCook();
-			}
-		}
+		FindOrCreatePackageWriter(TargetPlatform).EndCook();
 	}
 	UE_LOG(LogCook, Display, TEXT("Done finalizing package store(s)"));
-
-	ClearPackageStoreContexts();
 }
 
 void UCookOnTheFlyServer::ClearPackageStoreContexts()
@@ -8218,11 +7801,10 @@ void UCookOnTheFlyServer::StartCookByTheBook( const FCookByTheBookStartupOptions
 		// If we changed the DLC State earlier on, we must clear out the package name cache
 		TermSandbox();
 	}
-
-	// This will either delete the sandbox or iteratively clean it
-	InitializeSandbox(TargetPlatforms);
+	
+	FindOrCreateSaveContexts(TargetPlatforms);
+	BeginCookSandbox(TargetPlatforms); // This will either delete the sandbox or iteratively clean it
 	InitializeTargetPlatforms(TargetPlatforms);
-	InitializePackageStore(TargetPlatforms);
 
 	if (CurrentCookMode == ECookMode::CookByTheBook && !IsCookFlagSet(ECookInitializationFlags::Iterative))
 	{
@@ -8591,11 +8173,6 @@ void UCookOnTheFlyServer::MaybeMarkPackageAsAlreadyLoaded(UPackage *Package)
 	check(IsCookingInEditor()==false);
 	check(IsCookByTheBookMode());
 
-	if (bIsInitializingSandbox)
-	{
-		return;
-	}
-
 	// if the package is already fully loaded then we are not going to mark it up anyway
 	if ( Package->IsFullyLoaded() )
 	{
@@ -9328,29 +8905,29 @@ uint32 UCookOnTheFlyServer::FullLoadAndSave(uint32& CookedPackageCount)
 	return Result;
 }
 
-ICookedPackageWriter* UCookOnTheFlyServer::GetPackageWriter(const FName& PlatformName) const
+ICookedPackageWriter& UCookOnTheFlyServer::FindOrCreatePackageWriter(const ITargetPlatform* TargetPlatform)
+{
+	return *FindOrCreateSaveContext(TargetPlatform).PackageWriter;
+}
+
+void UCookOnTheFlyServer::FindOrCreateSaveContexts(TConstArrayView<const ITargetPlatform*> TargetPlatforms)
 {
-	for (const UE::Cook::FCookSavePackageContext* Context : SavePackageContexts)
+	for (const ITargetPlatform* TargetPlatform : TargetPlatforms)
 	{
-		if (FName(*Context->SaveContext.TargetPlatform->PlatformName()) == PlatformName)
-		{
-			return Context->PackageWriter;
-		}
+		FindOrCreateSaveContext(TargetPlatform);
 	}
-
-	return nullptr;
 }
 
-UE::Cook::FCookSavePackageContext* UCookOnTheFlyServer::GetSavePackageContext(const ITargetPlatform* TargetPlatform) const
+UE::Cook::FCookSavePackageContext& UCookOnTheFlyServer::FindOrCreateSaveContext(const ITargetPlatform* TargetPlatform)
 {
 	for (UE::Cook::FCookSavePackageContext* Context : SavePackageContexts)
 	{
 		if (Context->SaveContext.TargetPlatform == TargetPlatform)
 		{
-			return Context;
+			return *Context;
 		}
 	}
-	return nullptr;
+	return *SavePackageContexts.Add_GetRef(CreateSaveContext(TargetPlatform));
 }
 
 
diff --git a/Engine/Source/Editor/UnrealEd/Private/Cooker/CookOnTheFlyServerInterface.h b/Engine/Source/Editor/UnrealEd/Private/Cooker/CookOnTheFlyServerInterface.h
--- a/Engine/Source/Editor/UnrealEd/Private/Cooker/CookOnTheFlyServerInterface.h
+++ b/Engine/Source/Editor/UnrealEd/Private/Cooker/CookOnTheFlyServerInterface.h
@@ -7,6 +7,7 @@
 
 struct FODSCRequestPayload;
 class ICookedPackageWriter;
+class ITargetPlatform;
 
 namespace UE { namespace Cook
 {
@@ -41,7 +42,7 @@ struct FRecompileShaderRequest
 	int32 ShaderPlatformToCompile = -1; // 
 	/** Materials to laod. */
 	TArray<FString> MaterialsToLoad;
-	/** On-deamnd shader compiler payload.  */
+	/** On-demand shader compiler payload.  */
 	TArray<FODSCRequestPayload> ShadersToRecompile;
 	/** Mesh materials, returned to the caller.  */
 	TArray<uint8>* MeshMaterialMaps = nullptr;
@@ -69,7 +70,7 @@ public:
 	 *
 	 * @param PlatformName The platform name.
 	 */
-	virtual bool AddPlatform(const FName& PlatformName) = 0;
+	virtual const ITargetPlatform* AddPlatform(const FName& PlatformName) = 0;
 
 	/**
 	 * Remove platform from the cook-on-the-fly session.
@@ -116,10 +117,8 @@ public:
 
 	/**
 	 * Returns the package store writer for the specified platform.
-	 *
-	 * @param PlatformName The platform name.
 	 */
-	virtual ICookedPackageWriter* GetPackageWriter(const FName& PlatformName) = 0;
+	virtual ICookedPackageWriter& GetPackageWriter(const ITargetPlatform* TargetPlatform) = 0;
 
 	/**
 	 * A somewhat temporary event triggered when all pending packages has been cooked.
diff --git a/Engine/Source/Editor/UnrealEd/Private/Cooker/CookPlatformManager.h b/Engine/Source/Editor/UnrealEd/Private/Cooker/CookPlatformManager.h
--- a/Engine/Source/Editor/UnrealEd/Private/Cooker/CookPlatformManager.h
+++ b/Engine/Source/Editor/UnrealEd/Private/Cooker/CookPlatformManager.h
@@ -44,7 +44,7 @@ namespace Cook
 		/* Pointer to the platform-specific RegistryGenerator for this platform.  If already constructed we can take a faster refresh path on future sessions. */
 		TUniquePtr<FAssetRegistryGenerator> RegistryGenerator;
 
-		/* Whether InitializeSandbox has been called for this platform.  If we have already initialized the sandbox we can take a faster refresh path on future sessions. */
+		/* Whether BeginCookSandbox has been called for this platform.  If we have already initialized the sandbox we can take a faster refresh path on future sessions. */
 		bool bIsSandboxInitialized = false;
 
 		/*
diff --git a/Engine/Source/Editor/UnrealEd/Private/Cooker/CookRequestCluster.cpp b/Engine/Source/Editor/UnrealEd/Private/Cooker/CookRequestCluster.cpp
--- a/Engine/Source/Editor/UnrealEd/Private/Cooker/CookRequestCluster.cpp
+++ b/Engine/Source/Editor/UnrealEd/Private/Cooker/CookRequestCluster.cpp
@@ -118,7 +118,7 @@ void FRequestCluster::Initialize(UCookOnTheFlyServer& COTFS)
 	PackageWriters.Reserve(Platforms.Num());
 	for (const ITargetPlatform* TargetPlatform : Platforms)
 	{
-		PackageWriters.Add(COTFS.GetPackageWriter(FName(*TargetPlatform->PlatformName())));
+		PackageWriters.Add(&COTFS.FindOrCreatePackageWriter(TargetPlatform));
 	}
 }
 
diff --git a/Engine/Source/Editor/UnrealEd/Private/Cooker/CookTypes.cpp b/Engine/Source/Editor/UnrealEd/Private/Cooker/CookTypes.cpp
--- a/Engine/Source/Editor/UnrealEd/Private/Cooker/CookTypes.cpp
+++ b/Engine/Source/Editor/UnrealEd/Private/Cooker/CookTypes.cpp
@@ -5,9 +5,7 @@
 #include "HAL/PlatformTime.h"
 #include "Math/NumericLimits.h"
 
-namespace UE
-{
-namespace Cook
+namespace UE::Cook
 {
 
 	FCookerTimer::FCookerTimer(const float& InTimeSlice, bool bInIsRealtimeMode, int InMaxNumPackagesToSave)
@@ -79,5 +77,22 @@ namespace Cook
 	{
 		FPlatformTLS::SetTlsValue(SchedulerThreadTlsSlot, bValue ? (void*)0x1 : (void*)0x0);
 	}
+
+	FCookSavePackageContext::FCookSavePackageContext(const ITargetPlatform* InTargetPlatform,
+		ICookedPackageWriter* InPackageWriter, bool bInForceLegacyOffsets, FStringView InWriterDebugName)
+		: SaveContext(InTargetPlatform, InPackageWriter->GetCookCapabilities().bSavePackageSupported ? InPackageWriter : nullptr, bInForceLegacyOffsets)
+		, WriterDebugName(InWriterDebugName)
+		, PackageWriter(InPackageWriter)
+	{
+		PackageWriterCapabilities = InPackageWriter->GetCookCapabilities();
+	}
+
+	FCookSavePackageContext::~FCookSavePackageContext()
+	{
+		// SaveContext destructor deletes the PackageWriter, so if we passed our writer into SaveContext, we do not delete it
+		if (!SaveContext.PackageWriter)
+		{
+			delete PackageWriter;
+		}
+	}
 }
-}
\ No newline at end of file
diff --git a/Engine/Source/Editor/UnrealEd/Private/Cooker/CookTypes.h b/Engine/Source/Editor/UnrealEd/Private/Cooker/CookTypes.h
--- a/Engine/Source/Editor/UnrealEd/Private/Cooker/CookTypes.h
+++ b/Engine/Source/Editor/UnrealEd/Private/Cooker/CookTypes.h
@@ -164,14 +164,13 @@ namespace UE::Cook
 	struct FCookSavePackageContext
 	{
 		FCookSavePackageContext(const ITargetPlatform* InTargetPlatform,
-			ICookedPackageWriter* InPackageWriter, bool bInForceLegacyOffsets)
-			: SaveContext(InTargetPlatform, InPackageWriter, bInForceLegacyOffsets)
-			, PackageWriter(InPackageWriter)
-		{
-		}
-		// Note the SaveContext takes responsibility for deleting PackageWriter
+			ICookedPackageWriter* InPackageWriter, bool bInForceLegacyOffsets, FStringView InWriterDebugName);
+		~FCookSavePackageContext();
+
 		FSavePackageContext SaveContext;
+		FString WriterDebugName;
 		ICookedPackageWriter* PackageWriter;
+		ICookedPackageWriter::FCookCapabilities PackageWriterCapabilities;
 	};
 
 	/* Thread Local Storage access to identify which thread is the SchedulerThread for cooking. */
@@ -230,3 +229,5 @@ public:
 };
 
 void LogCookerMessage(const FString& MessageText, EMessageSeverity::Type Severity);
+const FString& GetDevelopmentAssetRegistryFilename();
+#define REMAPPED_PLUGINS TEXT("RemappedPlugins")
diff --git a/Engine/Source/Editor/UnrealEd/Private/Cooker/IoStoreCookOnTheFlyRequestManager.cpp b/Engine/Source/Editor/UnrealEd/Private/Cooker/IoStoreCookOnTheFlyRequestManager.cpp
--- a/Engine/Source/Editor/UnrealEd/Private/Cooker/IoStoreCookOnTheFlyRequestManager.cpp
+++ b/Engine/Source/Editor/UnrealEd/Private/Cooker/IoStoreCookOnTheFlyRequestManager.cpp
@@ -438,14 +438,14 @@ private:
 
 		if (ConnectionStatus == UE::Cook::ECookOnTheFlyConnectionStatus::Connected)
 		{
-			if (CookOnTheFlyServer.AddPlatform(Client.PlatformName))
+			const ITargetPlatform* TargetPlatform = CookOnTheFlyServer.AddPlatform(Client.PlatformName);
+			if (TargetPlatform)
 			{
 				if (!PlatformContexts.Contains(Client.PlatformName))
 				{
 					TUniquePtr<FPlatformContext>& Context = PlatformContexts.Add(Client.PlatformName, MakeUnique<FPlatformContext>());
 					Context->PlatformName = Client.PlatformName;
-					
-					IPackageStoreWriter* PackageWriter = CookOnTheFlyServer.GetPackageWriter(Client.PlatformName)->AsPackageStoreWriter();
+					IPackageStoreWriter* PackageWriter = CookOnTheFlyServer.GetPackageWriter(TargetPlatform).AsPackageStoreWriter();
 					check(PackageWriter); // This class should not be used except when COTFS is using an IPackageStoreWriter
 					Context->PackageWriter = PackageWriter;
 					
diff --git a/Engine/Source/Editor/UnrealEd/Private/Cooker/LooseCookedPackageWriter.cpp b/Engine/Source/Editor/UnrealEd/Private/Cooker/LooseCookedPackageWriter.cpp
new file mode 100644
--- /dev/null
+++ b/Engine/Source/Editor/UnrealEd/Private/Cooker/LooseCookedPackageWriter.cpp
@@ -0,0 +1,363 @@
+// Copyright Epic Games, Inc. All Rights Reserved.
+
+#include "Cooker/LooseCookedPackageWriter.h"
+
+#include "AssetRegistry/AssetRegistryState.h"
+#include "Async/ParallelFor.h"
+#include "Cooker/AsyncIODelete.h"
+#include "Cooker/CookTypes.h"
+#include "Cooker/PackageNameCache.h"
+#include "HAL/FileManager.h"
+#include "HAL/PlatformFile.h"
+#include "HAL/PlatformFileManager.h"
+#include "Interfaces/IPluginManager.h"
+#include "Interfaces/ITargetPlatform.h"
+#include "Misc/AssertionMacros.h"
+#include "Misc/CString.h"
+#include "Misc/FileHelper.h"
+#include "Misc/PackageName.h"
+#include "Misc/Paths.h"
+#include "Misc/PathViews.h"
+#include "Misc/StringBuilder.h"
+#include "ProfilingDebugging/CpuProfilerTrace.h"
+#include "Serialization/ArrayReader.h"
+#include "UObject/Package.h"
+
+const static FName NAME_DummyCookedFilename(TEXT("DummyCookedFilename"));
+
+FLooseCookedPackageWriter::FLooseCookedPackageWriter(const FString& InOutputPath, const FString& InMetadataDirectoryPath, const ITargetPlatform* InTargetPlatform,
+	FAsyncIODelete& InAsyncIODelete, const FPackageNameCache& InPackageNameCache, const TArray<TSharedRef<IPlugin>>& InPluginsToRemap)
+	: OutputPath(InOutputPath)
+	, MetadataDirectoryPath(InMetadataDirectoryPath)
+	, TargetPlatform(*InTargetPlatform)
+	, PackageNameCache(InPackageNameCache)
+	, PluginsToRemap(InPluginsToRemap)
+	, AsyncIODelete(InAsyncIODelete)
+	, bIterateSharedBuild(false)
+{
+}
+
+FLooseCookedPackageWriter::~FLooseCookedPackageWriter()
+{
+}
+
+void FLooseCookedPackageWriter::BeginPackage(const FBeginPackageInfo& Info)
+{
+}
+
+void FLooseCookedPackageWriter::CommitPackage(const FCommitPackageInfo& Info)
+{
+}
+
+void FLooseCookedPackageWriter::WritePackageData(const FPackageInfo& Info, const FIoBuffer& PackageData, const TArray<FFileRegion>& FileRegions)
+{
+	// Not yet implemented
+	checkNoEntry();
+}
+
+void FLooseCookedPackageWriter::WriteBulkdata(const FBulkDataInfo& Info, const FIoBuffer& BulkData, const TArray<FFileRegion>& FileRegions)
+{
+	// Not yet implemented
+	checkNoEntry();
+}
+
+bool FLooseCookedPackageWriter::WriteAdditionalFile(const FAdditionalFileInfo& Info, const FIoBuffer& FileData)
+{
+	// Not yet implemented
+	checkNoEntry();
+	return false;
+}
+
+void FLooseCookedPackageWriter::WriteLinkerAdditionalData(const FLinkerAdditionalDataInfo& Info, const FIoBuffer& Data, const TArray<FFileRegion>& FileRegions)
+{
+	// Should not be called because IsLinkerAdditionalDataInSeparateArchive returned false
+	checkNoEntry();
+}
+
+FDateTime FLooseCookedPackageWriter::GetPreviousCookTime() const
+{
+	const FString PreviousAssetRegistry = FPaths::Combine(MetadataDirectoryPath, TEXT("DevelopmentAssetRegistry.bin"));
+	return IFileManager::Get().GetTimeStamp(*PreviousAssetRegistry);
+}
+
+void FLooseCookedPackageWriter::BeginCook(const FCookInfo& Info)
+{
+	bIterateSharedBuild = Info.bIterateSharedBuild;
+	if (Info.bCleanBuild)
+	{
+		DeleteSandboxDirectory();
+	}
+}
+
+void FLooseCookedPackageWriter::EndCook()
+{
+}
+
+void FLooseCookedPackageWriter::Flush()
+{
+	UPackage::WaitForAsyncFileWrites();
+}
+
+void FLooseCookedPackageWriter::GetCookedPackages(TArray<FCookedPackageInfo>& OutCookedPackages)
+{
+	// GetCookedPackages is currently only called in BeginCook to read the list of previous cooked files
+	// Read that list from the platform cooked asset registry file
+
+	// Report files from the shared build if the option is set
+	FString PreviousAssetRegistryFile;
+	if (bIterateSharedBuild)
+	{
+		// clean the sandbox
+		DeleteSandboxDirectory();
+		PreviousAssetRegistryFile = FPaths::Combine(*FPaths::ProjectSavedDir(), TEXT("SharedIterativeBuild"),
+			*TargetPlatform.PlatformName(), TEXT("Metadata"), GetDevelopmentAssetRegistryFilename());
+
+		// We use NAME_DummyCookedFilename as a placeholder in UncookedPathToCookedPath to indicate the file comes from the shared pak file
+		// If an actual filename of this name exists, our placeholder system will break
+		check(!IFileManager::Get().FileExists(*NAME_DummyCookedFilename.ToString()));
+	}
+	else
+	{
+		PreviousAssetRegistryFile = FPaths::Combine(MetadataDirectoryPath, GetDevelopmentAssetRegistryFilename());
+	}
+
+	UncookedPathToCookedPath.Reset();
+	GetAllCookedFiles();
+
+	FArrayReader SerializedAssetData;
+	if (IFileManager::Get().FileExists(*PreviousAssetRegistryFile) && FFileHelper::LoadFileToArray(SerializedAssetData, *PreviousAssetRegistryFile))
+	{
+		PreviousState = MakeUnique<FAssetRegistryState>();
+		PreviousState->Load(SerializedAssetData);
+
+		const TMap<FName, const FAssetPackageData*>& DataMap = PreviousState->GetAssetPackageDataMap();
+		OutCookedPackages.Reserve(OutCookedPackages.Num() + DataMap.Num());
+		FString UncookedFileNameStr;
+
+		for (const TPair<FName, const FAssetPackageData*>& Pair : DataMap)
+		{
+			FName PackageName = Pair.Key;
+			const FAssetPackageData* PackageData = Pair.Value;
+			bool bExistsOnDisk = false;
+			if (bIterateSharedBuild)
+			{
+				// if we are iterating of a shared build the cooked files might not exist in the cooked directory because we assume they are packaged in the pak file (which we don't want to extract)
+				// Overwrite the value in UncookedPathToCookpath with a placeholder that indicates the file exists even if it's not present in the Cooked directory
+				if (FPackageName::DoesPackageExist(PackageName.ToString(), nullptr, &UncookedFileNameStr))
+				{
+					UncookedPathToCookedPath.Add(FName(*UncookedFileNameStr), NAME_DummyCookedFilename);
+				}
+				bExistsOnDisk = true;
+			}
+			else
+			{
+				const FName UncookedFilename = PackageNameCache.GetCachedStandardFileName(PackageName);
+				if (!UncookedFilename.IsNone())
+				{
+					FName* CookedFilename = UncookedPathToCookedPath.Find(UncookedFilename);
+					bExistsOnDisk = CookedFilename != nullptr;
+				}
+			}
+			if (bExistsOnDisk)
+			{
+				PRAGMA_DISABLE_DEPRECATION_WARNINGS;
+				OutCookedPackages.Add(FCookedPackageInfo{ PackageName, PackageData->CookedHash,PackageData->PackageGuid,
+					PackageData->DiskSize, FIoHash() /** TargetDomainDependencies are not implemented by FLooseCookedPackageWriter */ });
+				PRAGMA_ENABLE_DEPRECATION_WARNINGS;
+			}
+		}
+	}
+}
+
+FCbObject FLooseCookedPackageWriter::GetTargetDomainDependencies(FName PackageName)
+{
+	/** TargetDomainDependencies are not implemented by FLooseCookedPackageWriter */
+	return FCbObject();
+}
+
+void FLooseCookedPackageWriter::RemoveCookedPackages(TArrayView<const FName> PackageNamesToRemove)
+{
+	if (UncookedPathToCookedPath.IsEmpty())
+	{
+		return;
+	}
+
+	// if we are going to clear the cooked packages it is conceivable that we will recook the packages which we just cooked 
+	// that means it's also conceivable that we will recook the same package which currently has an outstanding async write request
+	UPackage::WaitForAsyncFileWrites();
+	
+	if (PackageNamesToRemove.Num() > 0)
+	{
+		// PackageNameCache is read-game-thread-only, so we have to read it before calling parallel-for
+		TArray<FName> UncookedFileNamesToRemove;
+		UncookedFileNamesToRemove.Reserve(PackageNamesToRemove.Num());
+		for (FName PackageName : PackageNamesToRemove)
+		{
+			const FName UncookedFileName = PackageNameCache.GetCachedStandardFileName(PackageName);
+			if (!UncookedFileName.IsNone())
+			{
+				UncookedFileNamesToRemove.Add(UncookedFileName);
+			}
+		}
+
+		auto DeletePackageLambda = [&UncookedFileNamesToRemove, this](int32 PackageIndex)
+		{
+			FName UncookedFileName = UncookedFileNamesToRemove[PackageIndex];
+			FName* CookedFileName = UncookedPathToCookedPath.Find(UncookedFileName);
+			if (CookedFileName)
+			{
+				TStringBuilder<256> FilePath;
+				CookedFileName->ToString(FilePath);
+				IFileManager::Get().Delete(*FilePath, true, true, true);
+			}
+		};
+		ParallelFor(UncookedFileNamesToRemove.Num(), DeletePackageLambda);
+	}
+	UncookedPathToCookedPath.Empty();
+}
+
+void FLooseCookedPackageWriter::RemoveCookedPackages()
+{
+	DeleteSandboxDirectory();
+}
+
+FAssetRegistryState* FLooseCookedPackageWriter::ReleasePreviousAssetRegistry()
+{
+	return PreviousState.Release();
+}
+
+void FLooseCookedPackageWriter::DeleteSandboxDirectory()
+{
+	// if we are going to clear the cooked packages it is conceivable that we will recook the packages which we just cooked 
+	// that means it's also conceivable that we will recook the same package which currently has an outstanding async write request
+	UPackage::WaitForAsyncFileWrites();
+
+	FString SandboxDirectory = OutputPath;
+	FPaths::NormalizeDirectoryName(SandboxDirectory);
+
+	AsyncIODelete.DeleteDirectory(SandboxDirectory);
+}
+
+class FPackageSearchVisitor : public IPlatformFile::FDirectoryVisitor
+{
+	TArray<FString>& FoundFiles;
+public:
+	FPackageSearchVisitor(TArray<FString>& InFoundFiles)
+		: FoundFiles(InFoundFiles)
+	{}
+	virtual bool Visit(const TCHAR* FilenameOrDirectory, bool bIsDirectory)
+	{
+		if (bIsDirectory == false)
+		{
+			FString Filename(FilenameOrDirectory);
+			FStringView Extension(FPathViews::GetExtension(Filename, true /* bIncludeDot */));
+			const TCHAR* ExtensionStr = Extension.GetData();
+			check(ExtensionStr[Extension.Len()] == '\0'); // IsPackageExtension takes a null-terminated TCHAR; we should have it since GetExtension is from the end of the filename
+			if (FPackageName::IsPackageExtension(ExtensionStr))
+			{
+				FoundFiles.Add(Filename);
+			}
+		}
+		return true;
+	}
+};
+
+void FLooseCookedPackageWriter::GetAllCookedFiles()
+{
+	TRACE_CPUPROFILER_EVENT_SCOPE(FLooseCookedPackageWriter::GetAllCookedFiles);
+
+	const FString& SandboxRootDir = OutputPath;
+	TArray<FString> CookedFiles;
+	{
+		IPlatformFile& PlatformFile = FPlatformFileManager::Get().GetPlatformFile();
+		FPackageSearchVisitor PackageSearch(CookedFiles);
+		PlatformFile.IterateDirectoryRecursively(*SandboxRootDir, PackageSearch);
+	}
+
+	const FString SandboxProjectDir = FPaths::Combine(OutputPath, FApp::GetProjectName()) + TEXT("/");
+	const FString RelativeRootDir = FPaths::GetRelativePathToRoot();
+	const FString RelativeProjectDir = FPaths::ProjectDir();
+	FString UncookedFilename;
+	UncookedFilename.Reserve(1024);
+
+	for (const FString& CookedFile : CookedFiles)
+	{
+		const FName CookedFName(*CookedFile);
+		const FName UncookedFName = ConvertCookedPathToUncookedPath(
+			SandboxRootDir, RelativeRootDir,
+			SandboxProjectDir, RelativeProjectDir,
+			CookedFile, UncookedFilename);
+
+		UncookedPathToCookedPath.Add(UncookedFName, CookedFName);
+	}
+}
+
+FName FLooseCookedPackageWriter::ConvertCookedPathToUncookedPath(
+	const FString& SandboxRootDir, const FString& RelativeRootDir,
+	const FString& SandboxProjectDir, const FString& RelativeProjectDir,
+	const FString& CookedPath, FString& OutUncookedPath) const
+{
+	OutUncookedPath.Reset();
+
+	// Check for remapped plugins' cooked content
+	if (PluginsToRemap.Num() > 0 && CookedPath.Contains(REMAPPED_PLUGINS))
+	{
+		int32 RemappedIndex = CookedPath.Find(REMAPPED_PLUGINS);
+		check(RemappedIndex >= 0);
+		static uint32 RemappedPluginStrLen = FCString::Strlen(REMAPPED_PLUGINS);
+		// Snip everything up through the RemappedPlugins/ off so we can find the plugin it corresponds to
+		FString PluginPath = CookedPath.RightChop(RemappedIndex + RemappedPluginStrLen + 1);
+		// Find the plugin that owns this content
+		for (const TSharedRef<IPlugin>& Plugin : PluginsToRemap)
+		{
+			if (PluginPath.StartsWith(Plugin->GetName()))
+			{
+				OutUncookedPath = Plugin->GetContentDir();
+				static uint32 ContentStrLen = FCString::Strlen(TEXT("Content/"));
+				// Chop off the pluginName/Content since it's part of the full path
+				OutUncookedPath /= PluginPath.RightChop(Plugin->GetName().Len() + ContentStrLen);
+				break;
+			}
+		}
+
+		if (OutUncookedPath.Len() > 0)
+		{
+			return FName(*OutUncookedPath);
+		}
+		// Otherwise fall through to sandbox handling
+	}
+
+	auto BuildUncookedPath =
+		[&OutUncookedPath](const FString& CookedPath, const FString& CookedRoot, const FString& UncookedRoot)
+	{
+		OutUncookedPath.AppendChars(*UncookedRoot, UncookedRoot.Len());
+		OutUncookedPath.AppendChars(*CookedPath + CookedRoot.Len(), CookedPath.Len() - CookedRoot.Len());
+	};
+
+	if (CookedPath.StartsWith(SandboxRootDir))
+	{
+		// Optimized CookedPath.StartsWith(SandboxProjectDir) that does not compare all of SandboxRootDir again
+		if (CookedPath.Len() >= SandboxProjectDir.Len() &&
+			0 == FCString::Strnicmp(
+				*CookedPath + SandboxRootDir.Len(),
+				*SandboxProjectDir + SandboxRootDir.Len(),
+				SandboxProjectDir.Len() - SandboxRootDir.Len()))
+		{
+			BuildUncookedPath(CookedPath, SandboxProjectDir, RelativeProjectDir);
+		}
+		else
+		{
+			BuildUncookedPath(CookedPath, SandboxRootDir, RelativeRootDir);
+		}
+	}
+	else
+	{
+		FString FullCookedFilename = FPaths::ConvertRelativePathToFull(CookedPath);
+		BuildUncookedPath(FullCookedFilename, SandboxRootDir, RelativeRootDir);
+	}
+
+	// Convert to a standard filename as required by FPackageNameCache where this path is used.
+	FPaths::MakeStandardFilename(OutUncookedPath);
+
+	return FName(*OutUncookedPath);
+}
diff --git a/Engine/Source/Editor/UnrealEd/Private/Cooker/LooseCookedPackageWriter.h b/Engine/Source/Editor/UnrealEd/Private/Cooker/LooseCookedPackageWriter.h
new file mode 100644
--- /dev/null
+++ b/Engine/Source/Editor/UnrealEd/Private/Cooker/LooseCookedPackageWriter.h
@@ -0,0 +1,77 @@
+// Copyright Epic Games, Inc. All Rights Reserved.
+
+#pragma once
+
+#include "Containers/Map.h"
+#include "Serialization/PackageWriter.h"
+#include "Templates/SharedPointer.h"
+#include "Templates/UniquePtr.h"
+
+class FAsyncIODelete;
+class IPlugin;
+class ITargetPlatform;
+struct FPackageNameCache;
+
+/** A CookedPackageWriter that saves cooked packages in separate .uasset,.uexp,.ubulk files in the Saved\Cooked\[Platform] directory. */
+class FLooseCookedPackageWriter
+	: public ICookedPackageWriter
+{
+public:
+	FLooseCookedPackageWriter(const FString& OutputPath, const FString& MetadataDirectoryPath, const ITargetPlatform* TargetPlatform,
+		FAsyncIODelete& InAsyncIODelete, const FPackageNameCache& InPackageNameCache, const TArray<TSharedRef<IPlugin>>& InPluginsToRemap);
+	~FLooseCookedPackageWriter();
+
+	FCookCapabilities GetCookCapabilities() const override
+	{
+		FCookCapabilities Result;
+		Result.bDiffModeSupported = true;
+		// TODO: Finish implementing this class so we can turn bSavePackageSupported on (and delete bSavePackageSupported as everything will have it set to true)
+		Result.bSavePackageSupported = false;
+		return Result;
+	}
+
+	void BeginPackage(const FBeginPackageInfo& Info) override;
+	void CommitPackage(const FCommitPackageInfo& Info) override;
+	void WritePackageData(const FPackageInfo& Info, const FIoBuffer& PackageData, const TArray<FFileRegion>& FileRegions) override;
+	void WriteBulkdata(const FBulkDataInfo& Info, const FIoBuffer& BulkData, const TArray<FFileRegion>& FileRegions) override;
+	bool WriteAdditionalFile(const FAdditionalFileInfo& Info, const FIoBuffer& FileData) override;
+	void WriteLinkerAdditionalData(const FLinkerAdditionalDataInfo& Info, const FIoBuffer& Data, const TArray<FFileRegion>& FileRegions) override;
+
+	FDateTime GetPreviousCookTime() const override;
+	void BeginCook(const FCookInfo& Info) override;
+	void EndCook() override;
+	void Flush() override;
+	void GetCookedPackages(TArray<FCookedPackageInfo>& OutCookedPackages) override;
+	FCbObject GetTargetDomainDependencies(FName PackageName) override;
+	void RemoveCookedPackages(TArrayView<const FName> PackageNamesToRemove) override;
+	void RemoveCookedPackages() override;
+	FAssetRegistryState* ReleasePreviousAssetRegistry() override;
+
+private:
+
+	/* Delete the sandbox directory (asynchronously) in preparation for a clean cook */
+	void DeleteSandboxDirectory();
+	/**
+	* Searches the disk for all the cooked files in the sandbox path provided
+	* Returns a map of the uncooked file path matches to the cooked file path for each package which exists
+	*
+	* @param UncookedpathToCookedPath out Map of the uncooked path matched with the cooked package which exists
+	* @param SandboxRootDir root dir to search for cooked packages in
+	*/
+	void GetAllCookedFiles();
+
+	FName ConvertCookedPathToUncookedPath(
+		const FString& SandboxRootDir, const FString& RelativeRootDir,
+		const FString& SandboxProjectDir, const FString& RelativeProjectDir,
+		const FString& CookedPath, FString& OutUncookedPath) const;
+
+	TMap<FName, FName> UncookedPathToCookedPath;
+	FString OutputPath;
+	FString MetadataDirectoryPath;
+	const ITargetPlatform& TargetPlatform;
+	const FPackageNameCache& PackageNameCache;
+	TUniquePtr<FAssetRegistryState> PreviousState;
+	const TArray<TSharedRef<IPlugin>>& PluginsToRemap;
+	FAsyncIODelete& AsyncIODelete;
+	bool bIterateSharedBuild;
+};
diff --git a/Engine/Source/Editor/UnrealEd/Private/Cooker/NetworkFileCookOnTheFlyRequestManager.cpp b/Engine/Source/Editor/UnrealEd/Private/Cooker/NetworkFileCookOnTheFlyRequestManager.cpp
--- a/Engine/Source/Editor/UnrealEd/Private/Cooker/NetworkFileCookOnTheFlyRequestManager.cpp
+++ b/Engine/Source/Editor/UnrealEd/Private/Cooker/NetworkFileCookOnTheFlyRequestManager.cpp
@@ -47,7 +47,7 @@ public:
 private:
 	bool OnNewConnection(const FString& VersionInfo, const FString& PlatformName)
 	{
-		return CookOnTheFlyServer.AddPlatform(FName(*PlatformName));
+		return CookOnTheFlyServer.AddPlatform(FName(*PlatformName)) != nullptr;
 	}
 
 	void OnGetPrecookedList(const FString& PlatformName, TMap<FString, FDateTime>& PrecookedFileList)
diff --git a/Engine/Source/Editor/UnrealEd/Public/Commandlets/AssetRegistryGenerator.h b/Engine/Source/Editor/UnrealEd/Public/Commandlets/AssetRegistryGenerator.h
--- a/Engine/Source/Editor/UnrealEd/Public/Commandlets/AssetRegistryGenerator.h
+++ b/Engine/Source/Editor/UnrealEd/Public/Commandlets/AssetRegistryGenerator.h
@@ -3,10 +3,12 @@
 #pragma once
 
 #include "CoreMinimal.h"
-#include "Misc/Paths.h"
+
 #include "AssetData.h"
-#include "Misc/AssetRegistryInterface.h"
 #include "AssetRegistryState.h"
+#include "Misc/AssetRegistryInterface.h"
+#include "Misc/Paths.h"
+#include "Templates/UniquePtr.h"
 #include "UObject/UObjectHash.h"
 
 class FSandboxPlatformFile;
@@ -40,9 +42,9 @@ public:
 	const ITargetPlatform* GetTargetPlatform() const { return TargetPlatform; }
 
 	/** 
-	 * Loads asset registry from a previous run that is used for iterative or DLC cooking
+	 * Sets asset registry from a previous run that is used for iterative or DLC cooking
 	 */
-	bool LoadPreviousAssetRegistry(const FString& Filename);
+	void SetPreviousAssetRegistry(TUniquePtr<FAssetRegistryState>&& PreviousState);
 	
 	/**
 	 * Options when computing the differences between current and previous state.
diff --git a/Engine/Source/Runtime/Core/Public/Serialization/PackageWriter.h b/Engine/Source/Runtime/Core/Public/Serialization/PackageWriter.h
--- a/Engine/Source/Runtime/Core/Public/Serialization/PackageWriter.h
+++ b/Engine/Source/Runtime/Core/Public/Serialization/PackageWriter.h
@@ -4,9 +4,11 @@
 
 #include "Containers/StringView.h"
 #include "IO/IoDispatcher.h"
+#include "Misc/DateTime.h"
 #include "Misc/SecureHash.h"
 #include "Serialization/CompactBinary.h"
 
+class FAssetRegistryState;
 class IPackageStoreWriter;
 struct FPackageStoreEntryResource;
 
@@ -125,6 +127,27 @@ class ICookedPackageWriter : public IPackageWriter
 public:
 	virtual ~ICookedPackageWriter() = default;
 
+	struct FCookCapabilities
+	{
+		/** Whether this writer implements -diffonly and -linkerdiff. */
+		bool bDiffModeSupported = false;
+		/** Whether this writer implements the IPackageWriter interface and can be passed to SavePackage. */
+		bool bSavePackageSupported = true;
+	};
+
+	/** Return cook capabilities/settings this PackageWriter has/requires
+	  */
+	virtual FCookCapabilities GetCookCapabilities() const
+	{
+		return FCookCapabilities();
+	}
+
+	/** Return the timestamp of the previous cook, or FDateTime::MaxValue to indicate previous cook should be assumed newer than any other cook data. */
+	virtual FDateTime GetPreviousCookTime() const
+	{
+		return FDateTime::MaxValue();
+	}
+
 	struct FCookInfo
 	{
 		enum ECookMode
@@ -134,6 +157,8 @@ public:
 		};
 
 		ECookMode CookMode = CookByTheBookMode;
+		bool bCleanBuild = true;
+		bool bIterateSharedBuild = false;
 	};
 
 	/** Signal the start of a cooking pass
@@ -162,7 +187,7 @@ public:
 	};
 
 	/**
-	 * Returns a list of cooked package(s).
+	 * Appends list of cooked package(s) to the given output array.
 	 */
 	virtual void GetCookedPackages(TArray<FCookedPackageInfo>& OutCookedPackages) = 0;
 
@@ -173,10 +198,24 @@ public:
 	virtual FCbObject GetTargetDomainDependencies(FName PackageName) = 0;
 
 	/**
-	 * Remove cooked package(s) that has been modified since the last cook.
+	 * Remove the given cooked package(s) from storage; they have been modified since the last cook.
 	 */
 	virtual void RemoveCookedPackages(TArrayView<const FName> PackageNamesToRemove) = 0;
 
+	/**
+	 * Remove all cooked packages from storage.
+	 */
+	virtual void RemoveCookedPackages() = 0;
+
+	/**
+	 * If a previous-version AssetRegistryState was allocated during GetCookedPackages, drop the reference to it
+	 * and return it to the caller who takes responsibility for deleting it.
+	 */
+	virtual FAssetRegistryState* ReleasePreviousAssetRegistry()
+	{
+		return nullptr;
+	}
+
 	/** Downcast function for ICookedPackageWriters that implement the IPackageStoreWriter inherited interface. */
 	virtual IPackageStoreWriter* AsPackageStoreWriter()
 	{
