CookPackageData: Whitespace: Add newlines for lines over 120 columns.
#rnx

diff --git a/Engine/Source/Editor/UnrealEd/Private/Cooker/CookPackageData.cpp b/Engine/Source/Editor/UnrealEd/Private/Cooker/CookPackageData.cpp
--- a/Engine/Source/Editor/UnrealEd/Private/Cooker/CookPackageData.cpp
+++ b/Engine/Source/Editor/UnrealEd/Private/Cooker/CookPackageData.cpp
@@ -43,9 +43,12 @@ namespace Cook
 
 	FPackageData::~FPackageData()
 	{
-		ClearReferences(); // Should have been called earlier, but call it here in case it was missed
-		ClearCookedPlatforms(); // We need to send OnCookedPlatformRemoved message to the monitor, so it is not valid to destruct without calling ClearCookedPlatforms
-		SendToState(EPackageState::Idle, ESendFlags::QueueNone); // Update the monitor's counters and call exit functions
+		// ClearReferences should have been called earlier, but call it here in case it was missed
+		ClearReferences();
+		// We need to send OnCookedPlatformRemoved message to the monitor, so call ClearCookedPlatforms
+		ClearCookedPlatforms();
+		// Update the monitor's counters and call exit functions
+		SendToState(EPackageState::Idle, ESendFlags::QueueNone);
 	}
 
 	void FPackageData::ClearReferences()
@@ -123,8 +126,8 @@ namespace Cook
 		}
 	}
 
-	void FPackageData::UpdateRequestData(const TArrayView<const ITargetPlatform* const>& InRequestedPlatforms, bool bInIsUrgent,
-		FCompletionCallback&& InCompletionCallback, ESendFlags SendFlags)
+	void FPackageData::UpdateRequestData(const TArrayView<const ITargetPlatform* const>& InRequestedPlatforms,
+		bool bInIsUrgent, FCompletionCallback&& InCompletionCallback, ESendFlags SendFlags)
 	{
 		if (IsInProgress())
 		{
@@ -159,7 +162,8 @@ namespace Cook
 		}
 	}
 
-	void FPackageData::SetRequestData(const TArrayView<const ITargetPlatform* const>& InRequestedPlatforms, bool bInIsUrgent, FCompletionCallback&& InCompletionCallback)
+	void FPackageData::SetRequestData(const TArrayView<const ITargetPlatform* const>& InRequestedPlatforms,
+		bool bInIsUrgent, FCompletionCallback&& InCompletionCallback)
 	{
 		check(!CompletionCallback);
 		check(RequestedPlatforms.Num() == 0);
@@ -176,7 +180,8 @@ namespace Cook
 		CompletionCallback = FCompletionCallback();
 	}
 
-	void FPackageData::AddCookedPlatforms(const TArrayView<const ITargetPlatform* const>& New, const TArrayView<const bool>& Succeeded)
+	void FPackageData::AddCookedPlatforms(const TArrayView<const ITargetPlatform* const>& New,
+		const TArrayView<const bool>& Succeeded)
 	{
 		check(New.Num() == Succeeded.Num());
 		for (int32 n = 0; n < New.Num(); ++n)
@@ -256,7 +261,8 @@ namespace Cook
 		return CookedPlatforms.Num() > 0;
 	}
 
-	bool FPackageData::HasAnyCookedPlatforms(const TArrayView<const ITargetPlatform* const>& Platforms, bool bIncludeFailed) const
+	bool FPackageData::HasAnyCookedPlatforms(const TArrayView<const ITargetPlatform* const>& Platforms,
+		bool bIncludeFailed) const
 	{
 		if (CookedPlatforms.Num() == 0)
 		{
@@ -273,7 +279,8 @@ namespace Cook
 		return false;
 	}
 
-	bool FPackageData::HasAllCookedPlatforms(const TArrayView<const ITargetPlatform* const>& Platforms, bool bIncludeFailed) const
+	bool FPackageData::HasAllCookedPlatforms(const TArrayView<const ITargetPlatform* const>& Platforms,
+		bool bIncludeFailed) const
 	{
 		if (Platforms.Num() == 0)
 		{
@@ -328,7 +335,7 @@ namespace Cook
 		return static_cast<EPackageState>(State);
 	}
 
-	/** Boilerplate-reduction struct that defines all of the multi-state properties and sets them based on the given state */
+	/** Boilerplate-reduction struct that defines all multi-state properties and sets them based on the given state. */
 	struct FStateProperties
 	{
 		EPackageStateProperty Properties;
@@ -403,11 +410,15 @@ namespace Cook
 
 		FStateProperties OldProperties(OldState);
 		FStateProperties NewProperties(NextState);
-		// Exit state properties in order from highest to lowest, then enter state properties in order from lowest to highest
-		// This ensures that properties that rely on earlier properties are constructed later and torn down earlier than the earlier properties.
-		for (EPackageStateProperty Iterator = EPackageStateProperty::Max; Iterator >= EPackageStateProperty::Min; Iterator = static_cast<EPackageStateProperty>(static_cast<uint32>(Iterator) >> 1))
-		{
-			if (((OldProperties.Properties & Iterator) != EPackageStateProperty::None) & ((NewProperties.Properties & Iterator) == EPackageStateProperty::None))
+		// Exit state properties from highest to lowest; enter state properties from lowest to highest.
+		// This ensures that properties that rely on earlier properties are constructed later and torn down earlier
+		// than the earlier properties.
+		for (EPackageStateProperty Iterator = EPackageStateProperty::Max;
+			Iterator >= EPackageStateProperty::Min;
+			Iterator = static_cast<EPackageStateProperty>(static_cast<uint32>(Iterator) >> 1))
+		{
+			if (((OldProperties.Properties & Iterator) != EPackageStateProperty::None) &
+				((NewProperties.Properties & Iterator) == EPackageStateProperty::None))
 			{
 				switch (Iterator)
 				{
@@ -426,9 +437,12 @@ namespace Cook
 				}
 			}
 		}
-		for (EPackageStateProperty Iterator = EPackageStateProperty::Min; Iterator <= EPackageStateProperty::Max; Iterator = static_cast<EPackageStateProperty>(static_cast<uint32>(Iterator) << 1))
+		for (EPackageStateProperty Iterator = EPackageStateProperty::Min;
+			Iterator <= EPackageStateProperty::Max;
+			Iterator = static_cast<EPackageStateProperty>(static_cast<uint32>(Iterator) << 1))
 		{
-			if (((OldProperties.Properties & Iterator) == EPackageStateProperty::None) & ((NewProperties.Properties & Iterator) != EPackageStateProperty::None))
+			if (((OldProperties.Properties & Iterator) == EPackageStateProperty::None) &
+				((NewProperties.Properties & Iterator) != EPackageStateProperty::None))
 			{
 				switch (Iterator)
 				{
@@ -557,7 +571,9 @@ namespace Cook
 
 	void FPackageData::OnEnterRequest()
 	{
-		check(RequestedPlatforms.Num() > 0); // It is not valid to enter the request state without requested platforms; it indicates a bug due to e.g. calling SendToState without UpdateRequestData from Idle
+		// It is not valid to enter the request state without requested platforms; it indicates a bug due to e.g.
+		// calling SendToState without UpdateRequestData from Idle
+		check(RequestedPlatforms.Num() > 0);
 	}
 
 	void FPackageData::OnExitRequest()
@@ -645,7 +661,9 @@ namespace Cook
 	{
 		if (InCompletionCallback)
 		{
-			// We don't yet have a mechanism for calling two completion callbacks.  CompletionCallbacks only come from external requests, and it should not be possible to request twice, so a failed check here shouldn't happen.
+			// We don't yet have a mechanism for calling two completion callbacks.
+			// CompletionCallbacks only come from external requests, and it should not be possible to request twice,
+			// so a failed check here shouldn't happen.
 			check(!CompletionCallback);
 			CompletionCallback = MoveTemp(InCompletionCallback);
 		}
@@ -684,13 +702,16 @@ namespace Cook
 			PreloadableFile.Get()->InitializeAsync([this]()
 				{
 					TStringBuilder<NAME_SIZE> FileNameString;
-					// Note this is an unsynchronized read of this->GetFilename. It is allowed because GetFilename does not change until this is destructed, and the destructor of this waits for PreloadableFile to finish initialization
+					// Note this is an unsynchronized read of this->GetFilename. It is allowed because GetFilename does
+					// not change until this is destructed, and the destructor of this waits for PreloadableFile to
+					// finish initialization
 					this->GetFileName().ToString(FileNameString);
 					FPackagePath PackagePath = FPackagePath::FromLocalPath(FileNameString);
 					FOpenPackageResult Result = IPackageResourceManager::Get().OpenReadPackage(PackagePath);
 					if (Result.Archive)
 					{
-						this->PreloadableFileFormat = Result.Format; // PreloadableFileFormat is an atomic, so this write is safe with other reads
+						// PreloadableFileFormat is an atomic, so this write is safe with other reads
+						this->PreloadableFileFormat = Result.Format;
 					}
 					return Result.Archive.Release();
 				},
@@ -701,7 +722,8 @@ namespace Cook
 		{
 			if (GetIsUrgent())
 			{
-				// For urgent requests, wait on them to finish preloading rather than letting them run asynchronously and coming back to them later
+				// For urgent requests, wait on them to finish preloading rather than letting them run asynchronously
+				// and coming back to them later
 				FilePtr->WaitForInitialization();
 				check(FilePtr->IsInitialized());
 			}
@@ -721,7 +743,8 @@ namespace Cook
 
 		TStringBuilder<NAME_SIZE> FileNameString;
 		GetFileName().ToString(FileNameString);
-		if (!IPackageResourceManager::TryRegisterPreloadableArchive(FPackagePath::FromLocalPath(FileNameString), FilePtr, PreloadableFileFormat))
+		if (!IPackageResourceManager::TryRegisterPreloadableArchive(FPackagePath::FromLocalPath(FileNameString),
+			FilePtr, PreloadableFileFormat))
 		{
 			UE_LOG(LogCook, Warning, TEXT("Failed to register %s for preload."), *GetFileName().ToString());
 			SetIsPreloadAttempted(true);
@@ -764,7 +787,8 @@ namespace Cook
 			GetFileName().ToString(FileNameString);
 			if (IPackageResourceManager::UnRegisterPreloadableArchive(FPackagePath::FromLocalPath(FileNameString)))
 			{
-				UE_LOG(LogCook, Display, TEXT("PreloadableFile was created for %s but never used. This is wasteful and bad for cook performance."), *PackageName.ToString());
+				UE_LOG(LogCook, Display, TEXT("PreloadableFile was created for %s but never used. This is wasteful and bad for cook performance."),
+					*PackageName.ToString());
 			}
 			FilePtr->ReleaseCache(); // ReleaseCache to conserve memory if the Linker still has a pointer to it
 		}
@@ -814,7 +838,9 @@ namespace Cook
 			for (UObject* Object : ObjectsInOuter)
 			{
 				FWeakObjectPtr ObjectWeakPointer(Object);
-				if (!ObjectWeakPointer.Get()) // ignore pending kill objects; they will not be serialized out so we don't need to call BeginCacheForCookedPlatformData on them
+				// ignore pending kill objects; they will not be serialized out so we don't need to call
+				// BeginCacheForCookedPlatformData on them
+				if (!ObjectWeakPointer.Get())
 				{
 					continue;
 				}
@@ -903,8 +929,10 @@ namespace Cook
 		return GetPackage() == nullptr || !GetPackage()->IsFullyLoaded() ||
 			Algo::AnyOf(CachedObjectsInOuter, [](const FWeakObjectPtr& WeakPtr)
 				{
-					// TODO: Keep track of which objects were public, and only invalidate the save if the object that has been deleted or marked pending kill was public
-					// Until we make that change, we will unnecessarily invalidate and demote some packages after a garbage collect
+					// TODO: Keep track of which objects were public, and only invalidate the save if the object
+					// that has been deleted or marked pending kill was public
+					// Until we make that change, we will unnecessarily invalidate and demote some packages after a
+					// garbage collect
 					return WeakPtr.Get() == nullptr;
 				});
 	}
@@ -924,17 +952,20 @@ namespace Cook
 		return IsGenerating() && !GetHasBeginPrepareSaveFailed();
 	}
 
-	UE::Cook::FGeneratorPackage* FPackageData::CreateGeneratorPackage(const UObject* InSplitDataObject, ICookPackageSplitter* InCookPackageSplitterInstance)
+	UE::Cook::FGeneratorPackage* FPackageData::CreateGeneratorPackage(const UObject* InSplitDataObject,
+		ICookPackageSplitter* InCookPackageSplitterInstance)
 	{
 		check(!GetGeneratorPackage());
-		GeneratorPackage.Reset(new UE::Cook::FGeneratorPackage(*this, InSplitDataObject, InCookPackageSplitterInstance));
+		GeneratorPackage.Reset(new UE::Cook::FGeneratorPackage(*this, InSplitDataObject,
+			InCookPackageSplitterInstance));
 		return GetGeneratorPackage();
 	}
 	
 	//////////////////////////////////////////////////////////////////////////
 	// FGeneratorPackage
 
-	FGeneratorPackage::FGeneratorPackage(UE::Cook::FPackageData& InOwner, const UObject* InSplitDataObject, ICookPackageSplitter* InCookPackageSplitterInstance)
+	FGeneratorPackage::FGeneratorPackage(UE::Cook::FPackageData& InOwner, const UObject* InSplitDataObject,
+		ICookPackageSplitter* InCookPackageSplitterInstance)
 	: Owner(InOwner)
 	, SplitDataObjectName(*InSplitDataObject->GetFullName())
 	{
@@ -960,11 +991,13 @@ namespace Cook
 		}
 	}
 
-	bool FGeneratorPackage::TryGenerateList(UObject* OwnerObject, const FPackageNameCache& PackageNameCache, FPackageDatas& PackageDatas)
+	bool FGeneratorPackage::TryGenerateList(UObject* OwnerObject, const FPackageNameCache& PackageNameCache,
+		FPackageDatas& PackageDatas)
 	{
 		UPackage* OwnerPackage = Owner.GetPackage();
 		check(OwnerPackage);
-		TArray<ICookPackageSplitter::FGeneratedPackage> GeneratorDatas = CookPackageSplitterInstance->GetGenerateList(OwnerPackage, OwnerObject);
+		TArray<ICookPackageSplitter::FGeneratedPackage> GeneratorDatas =
+			CookPackageSplitterInstance->GetGenerateList(OwnerPackage, OwnerObject);
 		PackagesToGenerate.Reset(GeneratorDatas.Num());
 		for (ICookPackageSplitter::FGeneratedPackage& SplitterData : GeneratorDatas)
 		{
@@ -1000,8 +1033,9 @@ namespace Cook
 			}
 			GeneratedStruct.PackageData = &PackageData;
 			PackageData.SetGenerated(true);
-			// No package should be generated by two different splitters. If an earlier run of this splitter generated the package, the package's owner
-			// should have been reset to null when we called ClearGeneratedPackages between then and now
+			// No package should be generated by two different splitters. If an earlier run of this splitter generated
+			// the package, the package's owner should have been reset to null when we called ClearGeneratedPackages
+			// between then and now
 			check(PackageData.GetGeneratedOwner() == nullptr); 
 			PackageData.SetGeneratedOwner(this);
 		}
@@ -1025,7 +1059,8 @@ namespace Cook
 	{
 		FString ObjectPath = GetSplitDataObjectName().ToString();
 
-		// SplitDataObjectName is a FullObjectPath; strip off the leading <ClassName> in "<ClassName> <Package>.<Object>:<SubObject>"
+		// SplitDataObjectName is a FullObjectPath; strip off the leading <ClassName> in
+		// "<ClassName> <Package>.<Object>:<SubObject>"
 		int32 ClassDelimiterIndex = -1;
 		if (ObjectPath.FindChar(' ', ClassDelimiterIndex))
 		{
@@ -1042,7 +1077,8 @@ namespace Cook
 		}
 		if (Owner.GetState() == EPackageState::Save)
 		{
-			// UCookOnTheFlyServer::PreCollectGarbage adds references for the Generator package and all its public objects, so it should still be loaded
+			// UCookOnTheFlyServer::PreCollectGarbage adds references for the Generator package and all its public
+			// objects, so it should still be loaded
 			if (!Owner.GetPackage() || !FindSplitDataObject())
 			{
 				UE_LOG(LogCook, Error, TEXT("PackageSplitter object was deleted by garbage collection while generation was still ongoing. This will break the generation.")
@@ -1058,7 +1094,8 @@ namespace Cook
 			UPackage* OwnerPackage = FindObject<UPackage>(nullptr, *Owner.GetPackageName().ToString());
 			if (OwnerPackage)
 			{
-				if (RemainingToPopulate > 0 && !CookPackageSplitterInstance->UseInternalReferenceToAvoidGarbageCollect())
+				if (RemainingToPopulate > 0 &&
+					!CookPackageSplitterInstance->UseInternalReferenceToAvoidGarbageCollect())
 				{
 					UE_LOG(LogCook, Error, TEXT("PackageSplitter found the Generator package still in memory after it should have been deleted by GC.")
 						TEXT("\n\tThis is unexpected since garbage has been collected and the package should have been unreferenced so it should have been collected, and will break population of Generated packages.")
@@ -1109,7 +1146,8 @@ namespace Cook
 		if (!GeneratedStruct)
 		{
 			UE_LOG(LogCook, Warning, TEXT("PackageSplitter called SetGeneratedSaved on a package that does not belong to the splitter.")
-				TEXT("\n\tSplitter=%s, Generated=%s."), *GetSplitDataObjectName().ToString(), *PackageData.GetPackageName().ToString());
+				TEXT("\n\tSplitter=%s, Generated=%s."), *GetSplitDataObjectName().ToString(),
+				*PackageData.GetPackageName().ToString());
 			return;
 		}
 		if (GeneratedStruct->bHasSaved)
@@ -1140,7 +1178,8 @@ namespace Cook
 		FString UnusedPackagePath;
 		FString MountLocalFilePath;
 		GetIntermediateMountPoint(UnusedPackagePath, MountLocalFilePath);
-		FString Extension = FPaths::GetExtension(GeneratedStruct.PackageData->GetFileName().ToString(), true /* bIncludeDot */);
+		FString Extension = FPaths::GetExtension(GeneratedStruct.PackageData->GetFileName().ToString(),
+			true /* bIncludeDot */);
 		return FPaths::RemoveDuplicateSlashes(FString::Printf(TEXT("%s/%s%s"),
 			*MountLocalFilePath, *GeneratedStruct.RelativePath, *Extension));
 	}
@@ -1149,17 +1188,20 @@ namespace Cook
 	// FPendingCookedPlatformData
 
 
-	FPendingCookedPlatformData::FPendingCookedPlatformData(UObject* InObject, const ITargetPlatform* InTargetPlatform, FPackageData& InPackageData, bool bInNeedsResourceRelease, UCookOnTheFlyServer& InCookOnTheFlyServer)
-		:Object(InObject), TargetPlatform(InTargetPlatform), PackageData(InPackageData), CookOnTheFlyServer(InCookOnTheFlyServer),
-		CancelManager(nullptr), ClassName(InObject->GetClass()->GetFName()), bHasReleased(false), bNeedsResourceRelease(bInNeedsResourceRelease)
+	FPendingCookedPlatformData::FPendingCookedPlatformData(UObject* InObject, const ITargetPlatform* InTargetPlatform,
+		FPackageData& InPackageData, bool bInNeedsResourceRelease, UCookOnTheFlyServer& InCookOnTheFlyServer)
+		: Object(InObject), TargetPlatform(InTargetPlatform), PackageData(InPackageData)
+		, CookOnTheFlyServer(InCookOnTheFlyServer),	CancelManager(nullptr), ClassName(InObject->GetClass()->GetFName())
+		, bHasReleased(false), bNeedsResourceRelease(bInNeedsResourceRelease)
 	{
 		check(InObject);
 		PackageData.GetNumPendingCookedPlatformData() += 1;
 	}
 
 	FPendingCookedPlatformData::FPendingCookedPlatformData(FPendingCookedPlatformData&& Other)
-		:Object(Other.Object), TargetPlatform(Other.TargetPlatform), PackageData(Other.PackageData), CookOnTheFlyServer(Other.CookOnTheFlyServer),
-		CancelManager(Other.CancelManager), ClassName(Other.ClassName), bHasReleased(Other.bHasReleased), bNeedsResourceRelease(Other.bNeedsResourceRelease)
+		: Object(Other.Object), TargetPlatform(Other.TargetPlatform), PackageData(Other.PackageData)
+		, CookOnTheFlyServer(Other.CookOnTheFlyServer), CancelManager(Other.CancelManager), ClassName(Other.ClassName)
+		, bHasReleased(Other.bHasReleased), bNeedsResourceRelease(Other.bNeedsResourceRelease)
 	{
 		Other.Object = nullptr;
 	}
@@ -1191,7 +1233,8 @@ namespace Cook
 			{
 				if (GShaderCompilingManager->HasShaderJobs() == false)
 				{
-					UE_LOG(LogCook, Warning, TEXT("Shader compiler is in a bad state!  Shader %s is finished compile but shader compiling manager did not notify shader.  "), *LocalObject->GetPathName());
+					UE_LOG(LogCook, Warning, TEXT("Shader compiler is in a bad state!  Shader %s is finished compile but shader compiling manager did not notify shader.  "),
+						*LocalObject->GetPathName());
 				}
 			}*/
 			return false;
@@ -1208,7 +1251,8 @@ namespace Cook
 		if (bNeedsResourceRelease)
 		{
 			int32* CurrentAsyncCache = CookOnTheFlyServer.CurrentAsyncCacheForType.Find(ClassName);
-			check(CurrentAsyncCache != nullptr); // bNeedsRelease should not have been set if the AsyncCache does not have an entry for the class
+			// bNeedsRelease should not have been set if the AsyncCache does not have an entry for the class
+			check(CurrentAsyncCache != nullptr);
 			*CurrentAsyncCache += 1;
 		}
 
@@ -1260,7 +1304,9 @@ namespace Cook
 	int32 FPackageDataMonitor::GetNumUrgent() const
 	{
 		int32 NumUrgent = 0;
-		for (EPackageState State = EPackageState::Min; State <= EPackageState::Max; State = static_cast<EPackageState>(static_cast<uint32>(State) + 1))
+		for (EPackageState State = EPackageState::Min;
+			State <= EPackageState::Max;
+			State = static_cast<EPackageState>(static_cast<uint32>(State) + 1))
 		{
 			NumUrgent += NumUrgentInState[static_cast<uint32>(State) - static_cast<uint32>(EPackageState::Min)];
 		}
@@ -1405,8 +1451,10 @@ namespace Cook
 			return **PackageDataMapAddr;
 		}
 
-		checkf(FileNameToPackageData.Find(NormalizedFileName) == nullptr, TEXT("Package \"%s\" and package \"%s\" share the same filename \"%s\"."),
-			*PackageName.ToString(), *(*FileNameToPackageData.Find(NormalizedFileName))->GetPackageName().ToString(), *NormalizedFileName.ToString());
+		checkf(FileNameToPackageData.Find(NormalizedFileName) == nullptr,
+			TEXT("Package \"%s\" and package \"%s\" share the same filename \"%s\"."),
+			*PackageName.ToString(), *(*FileNameToPackageData.Find(NormalizedFileName))->GetPackageName().ToString(),
+			*NormalizedFileName.ToString());
 		return CreatePackageData(PackageName, NormalizedFileName);
 	}
 
@@ -1440,8 +1488,10 @@ namespace Cook
 			// This can happen if PackageName is a script package
 			return nullptr;
 		}
-		checkf(FileNameToPackageData.Find(FileName) == nullptr, TEXT("Package \"%s\" and package \"%s\" share the same filename \"%s\"."),
-			*PackageName.ToString(), *(*FileNameToPackageData.Find(FileName))->GetPackageName().ToString(), *FileName.ToString()); 
+		checkf(FileNameToPackageData.Find(FileName) == nullptr,
+			TEXT("Package \"%s\" and package \"%s\" share the same filename \"%s\"."),
+			*PackageName.ToString(), *(*FileNameToPackageData.Find(FileName))->GetPackageName().ToString(),
+			*FileName.ToString()); 
 		return &CreatePackageData(PackageName, FileName);
 	}
 
@@ -1517,13 +1567,15 @@ namespace Cook
 	}
 
 
-	void FPackageDatas::AddExistingPackageDatasForPlatform(const TArray<FName>& ExistingPackages, const ITargetPlatform* TargetPlatform)
+	void FPackageDatas::AddExistingPackageDatasForPlatform(const TArray<FName>& ExistingPackages,
+		const ITargetPlatform* TargetPlatform)
 	{
 		TArray<FPackageData*> NewPackageDataObjects;
 		NewPackageDataObjects.AddZeroed(ExistingPackages.Num());
 
-		// parallelize the read-only operations (and filling out NewPackageDataObjects by index which has no threading issues)
-		ParallelFor(ExistingPackages.Num(), [&ExistingPackages, TargetPlatform, &NewPackageDataObjects, this](int Index)
+		// parallelize the read-only operations (and write NewPackageDataObjects by index which has no threading issues)
+		ParallelFor(ExistingPackages.Num(),
+			[&ExistingPackages, TargetPlatform, &NewPackageDataObjects, this](int Index)
 		{
 			const FName& PackageFilename = ExistingPackages[Index];
 
@@ -1544,7 +1596,7 @@ namespace Cook
 			const FName* PackageName = PackageNameCache.FindExistingCachedPackageNameFromStandardFileName(FileName);
 			checkf(PackageName != nullptr, TEXT("AddExistingPackageDatasForPlatform should only be used when the PackageCache has already been filled out with the packages"));
 
-			// create the pacakge data and remember it for updating caches after the the ParallelFor
+			// create the package data and remember it for updating caches after the the ParallelFor
 			FPackageData* PackageData = new FPackageData(*this, *PackageName, FileName);
 			NewPackageDataObjects[Index] = PackageData;
 			PackageData->AddCookedPlatforms({ TargetPlatform }, true /* Succeeded */);
@@ -1621,7 +1673,8 @@ namespace Cook
 		return Monitor.GetNumCooked();
 	}
 
-	void FPackageDatas::GetCookedFileNamesForPlatform(const ITargetPlatform* Platform, TArray<FName>& CookedFiles, bool bGetFailedCookedPackages, bool bGetSuccessfulCookedPackages)
+	void FPackageDatas::GetCookedFileNamesForPlatform(const ITargetPlatform* Platform, TArray<FName>& CookedFiles,
+		bool bGetFailedCookedPackages, bool bGetSuccessfulCookedPackages)
 	{
 		for (const FPackageData* PackageData : PackageDatas)
 		{
@@ -1680,7 +1733,8 @@ namespace Cook
 			return;
 		}
 
-		GShaderCompilingManager->ProcessAsyncResults(true /* bLimitExecutionTime */, false /* bBlockOnGlobalShaderCompletion */);
+		GShaderCompilingManager->ProcessAsyncResults(true /* bLimitExecutionTime */,
+			false /* bBlockOnGlobalShaderCompletion */);
 		FAssetCompilingManager::Get().ProcessAsyncTasks(true);
 
 		FPendingCookedPlatformData* Datas = PendingCookedPlatformDatas.GetData();
@@ -1817,7 +1871,8 @@ namespace Cook
 
 	bool FLoadPrepareQueue::Contains(const FPackageData* PackageData) const
 	{
-		return (Algo::Find(PreloadingQueue, PackageData) != nullptr) || (Algo::Find(EntryQueue, PackageData) != nullptr);
+		return (Algo::Find(PreloadingQueue, PackageData) != nullptr) ||
+			(Algo::Find(EntryQueue, PackageData) != nullptr);
 	}
 
 	uint32 FLoadPrepareQueue::Remove(FPackageData* PackageData)
diff --git a/Engine/Source/Editor/UnrealEd/Private/Cooker/CookPackageData.h b/Engine/Source/Editor/UnrealEd/Private/Cooker/CookPackageData.h
--- a/Engine/Source/Editor/UnrealEd/Private/Cooker/CookPackageData.h
+++ b/Engine/Source/Editor/UnrealEd/Private/Cooker/CookPackageData.h
@@ -42,45 +42,75 @@ namespace Cook
 	/** Flags specifying the behavior of FPackageData::SendToState */
 	enum class ESendFlags
 	{
-		QueueNone = 0x0,                            /* PackageData will not be removed from queue for its old state and will not be added to queue for its new state.  Caller is responsible for remove and add. */
-		QueueRemove = 0x1,                          /* PackageData will be removed from the queue of its old state.  If this flag is missing, caller is responsible for removing from the old state's queue. */
-		QueueAdd = 0x2,                             /* PackageData will be added to queue for its next state.  If this flag is missing, caller is responsible for adding to queue. */
-		QueueAddAndRemove = QueueAdd | QueueRemove, /* PackageData will be removed from the queue of its old state and added to the queue of its new state. This may be wasteful, if the caller can add or remove more efficiently. */
+		/*
+		 * PackageData will not be removed from queue for its old state and will not be added to queue for its new state.
+		 * Caller is responsible for remove and add.
+		 */
+		QueueNone = 0x0,
+		/*
+		 * PackageData will be removed from the queue of its old state.
+		 * If this flag is missing, caller is responsible for removing from the old state's queue.
+		 */
+		QueueRemove = 0x1,
+		/*
+		 * PackageData will be added to queue for its next state.
+		 * If this flag is missing, caller is responsible for adding to queue.
+		 */
+		QueueAdd = 0x2,
+		/*
+		 * PackageData will be removed from the queue of its old state and added to the queue of its new state.
+		 * This may be wasteful, if the caller can add or remove more efficiently.
+		 */
+		QueueAddAndRemove = QueueAdd | QueueRemove,
 	};
 	ENUM_CLASS_FLAGS(ESendFlags);
 
 	/**
-	 * Contains all the information the cooker uses for a package, during request, load, or save.  Once allocated, this structure is never deallocated or moved for a given package; it is
-	 * deallocated only when the CookOnTheFlyServer is destroyed.
+	 * Contains all the information the cooker uses for a package, during request, load, or save.  Once allocated, this
+	 * structure is never deallocated or moved for a given package; it is deallocated only when the CookOnTheFlyServer
+	 * is destroyed.
 	 *
 	 * RequestedPlatforms - RequestedPlatforms are the Platforms that will be Saved during cooking.
-	 *   RequestedPlatforms is only non-empty when the Package is InProgress (e.g. is in the RequestQueue, LoadReadyQueue, or other state containers).
-	 *   Once a Package finishes cooking, the RequestedPlatforms are cleared, and may be set again later if the Package is requested for another platform.
-	 *   If multiple requests occur at the same time, their RequestedPlatforms are merged.
-	 *   Since modifying the RequestedPlatforms can put an InProgress package into an invalid state, direct write access is private; use UpdateRequestData or SetRequestData to write.
-
-	 * CookedPlatforms - CookedPlatforms are the platforms that have already been saved in the lifetime of the CookOnTheFlyServer; note this extends outside of the current CookOnTheFlyServer session.
+	 *   RequestedPlatforms is only non-empty when the Package is InProgress (e.g. is in the RequestQueue,
+	 *   LoadReadyQueue, or other state containers). Once a Package finishes cooking, the RequestedPlatforms are
+	 *   cleared, and may be set again later if the Package is requested for another platform. If multiple requests
+	 *   occur at the same time, their RequestedPlatforms are merged. Since modifying the RequestedPlatforms can put
+	 *   an InProgress package into an invalid state, direct write access is private; use UpdateRequestData or
+	 *   SetRequestData to write.
+
+	 * CookedPlatforms - CookedPlatforms are the platforms that have already been saved in the lifetime of the 
+	 *   CookOnTheFlyServer; note this extends outside of the current CookOnTheFlyServer session.
 	 *   CookedPlatforms also store a flag indicating whether the cook for the given platform was successful or not.
-	 *   CookedPlatforms can be added to a PackageData for reasons other than normal Save, such as when a Package is detected as not applicable for a Platform and its "Cooked" operation is therefore a noop.
-	 *   CookedPlatforms can be cleared for a PackageData if the Package is modified, e.g. during editor operations when the CookOnTheFlyServer is run from the editor.
+	 *   CookedPlatforms can be added to a PackageData for reasons other than normal Save, such as when a Package is
+	 *   detected as not applicable for a Platform and its "Cooked" operation is therefore a noop.
+	 *   CookedPlatforms can be cleared for a PackageData if the Package is modified, e.g. during editor operations
+	 *   when the CookOnTheFlyServer is run from the editor.
 	 *
 	 * Package - The package pointer corresponding to this PackageData.
-	 *   Contract for the lifetime of this field is that it is only set after the Package has passed through the load state, and it is cleared when the package returns to idle or is demoted to an earlier state.
+	 *   Contract for the lifetime of this field is that it is only set after the Package has passed through the load
+	 *   state, and it is cleared when the package returns to idle or is demoted to an earlier state.
 	 *   At all other times it is nullptr.
 	 *
-	 * State - the state of this PackageData in the CookOnTheFlyServer's current session. See the definition of EPackageState for a description of each state.
-	 *   Contract for the value of this field includes membership in the corresponding queue such as SaveQueue, and the presence or absence of state-specific data.
-	 *   Since modifying the State can put the Package into an invalid state, direct write access is private; SendToState handles enforcing the contracts.
-
- 	 * PendingCookedPlatformData - a counter for how many objects in the package have had BeginCacheForCookedPlatformData called and have not yet completed.
-	 *   This is used to block the package from saving until all objects have finished their cache
-	 *   It is also used to block the package from starting new BeginCacheForCookedPlatformData calls until all pending calls from a previous canceled save have completed.
-	 *   The lifetime of this counter extends for the lifetime of the PackageData; it is shared across CookOnTheFlyServer sessions.
+	 * State - the state of this PackageData in the CookOnTheFlyServer's current session. See the definition of
+	 *   EPackageState for a description of each state. Contract for the value of this field includes membership in the
+	 *   corresponding queue such as SaveQueue, and the presence or absence of state-specific data. Since modifying the
+	 *   State can put the Package into an invalid state, direct write access is private; SendToState handles enforcing
+	 *   the contracts.
+
+ 	 * PendingCookedPlatformData - a counter for how many objects in the package have had
+	 *   BeginCacheForCookedPlatformData called and have not yet completed. This is used to block the package from
+	 *   saving until all objects have finished their cache. It is also used to block the package from starting new
+	 *   BeginCacheForCookedPlatformData calls until all pending calls from a previous canceled save have completed.
+	 *   The lifetime of this counter extends for the lifetime of the PackageData; it is shared across
+	 *   CookOnTheFlyServer sessions.
 	 *
-	 * CookedPlatformDataNextIndex - Index for the next (Object, RequestedPlatform) pair that needs to have BeginCacheForCookedPlatformData called on it for the current PackageSave.
-	 *   This field is only non-zero during the save state; it is cleared when successfully or unsucessfully leaving the save state.
-	 *   The field is an object-major index into the two-dimensional array given by {Each Object in GetCachedObjectsInOuter} x {each Platform in RequestedPlatforms}
-	 *   e.g. the index into GetCachedObjectsInOuter is GetCookedPlatformDataNextIndex() / RequestedPlatforms.Num() and the index into RequestedPlatforms is GetCookedPlatformDataNextIndex() % RequestedPlatforms.Num()
+	 * CookedPlatformDataNextIndex - Index for the next (Object, RequestedPlatform) pair that needs to have
+	 *   BeginCacheForCookedPlatformData called on it for the current PackageSave. This field is only non-zero during
+	 *   the save state; it is cleared when successfully or unsucessfully leaving the save state. The field is an
+	 *   object-major index into the two-dimensional array given by
+	 *       {Each Object in GetCachedObjectsInOuter} x {each Platform in RequestedPlatforms}
+	 *   e.g. the index into GetCachedObjectsInOuter is GetCookedPlatformDataNextIndex() / RequestedPlatforms.Num()
+	 *   and the index into RequestedPlatforms is GetCookedPlatformDataNextIndex() % RequestedPlatforms.Num()
 	 *
 	 * Other fields with explanation inline
 	*/
@@ -100,44 +130,73 @@ namespace Cook
 		FPackageData& operator=(const FPackageData& other) = delete;
 		FPackageData& operator=(FPackageData&& other) = delete;
 
-		/** Return a copy of Package->GetName(). This field is derived from the FileName if necessary, and is never modified. */
+		/** Return a copy of Package->GetName(). It is derived from the FileName if necessary, and is never modified. */
 		const FName& GetPackageName() const;
-		/** Return a copy of the FileName containing the package, normalized as returned from MakeStandardFilename.  This field may change from e.g. *.umap to *.uasset if LoadPackage loads a different FileName for the requested PackageName. */
+		/**
+		 * Return a copy of the FileName containing the package, normalized as returned from MakeStandardFilename.
+		 * This field may change from e.g. *.umap to *.uasset if LoadPackage loads a different FileName for the
+		 * requested PackageName.
+		 */
 		const FName& GetFileName() const;
 
 		/** Get the current set of RequestedPlatforms. */
 		const TArray<const ITargetPlatform*>& GetRequestedPlatforms() const;
-		/** Return true if and only if every element of Platforms is present in RequestedPlatforms.  Returns true if Platforms is empty. */
+		/**
+		 * Return true if and only if every element of Platforms is present in RequestedPlatforms.
+		 * Returns true if Platforms is empty.
+		 */
 		bool ContainsAllRequestedPlatforms(const TArrayView<const ITargetPlatform* const>& Platforms) const;
 		/**
-		 * Get the flag for whether this InProgress PackageData has been marked as an urgent request (e.g. because it has been requested from the game client during cookonthefly.).
+		 * Get the flag for whether this InProgress PackageData has been marked as an urgent request
+		 * (e.g. because it has been requested from the game client during cookonthefly.).
 		 * Always false for PackageDatas that are not InProgress.
-		 * Since modifying the urgency can put an InProgress package into an invalid state, direct write access is private; use UpdateRequestData or SetRequestData to write. */
+		 * Since modifying the urgency can put an InProgress package into an invalid state, direct write access is
+		 * private; use UpdateRequestData or SetRequestData to write.
+		 */
 		bool GetIsUrgent() const { return static_cast<bool>(bIsUrgent); }
 		/**
-		 * Add the given data, which are all fields describing a request for the cook of this PackageData's Package, onto the existing request data.
+		 * Add the given data, which are all fields describing a request for the cook of this PackageData's Package,
+		 * onto the existing request data.
 		 * If the PackageData is not already in progress, this will be equivalent to calling SetRequestData.
-		 * If the PackageData is in progress and the given fields invalidate some of its inprogress state, the PackageData's current progress will be canceled and it will be demoted
-		 * back to an earlier state where the changes can be made.
-		 * Once the changes can be made without invalidating the PackageData's current state, the given request data will be added on to the PackageData's current request data;
-		 * the new request data will be the union of the two previous sets of data.
+		 * If the PackageData is in progress and the given fields invalidate some of its inprogress state, the
+		 * PackageData's current progress will be canceled and it will be demoted back to an earlier state where
+		 * the changes can be made.
+		 * Once the changes can be made without invalidating the PackageData's current state, the given request data
+		 * will be added on to the PackageData's current request data; the new request data will be the union of the
+		 * two previous sets of data.
 		 
-		 * @param SendFlags If the PackageData needs to be moved to a new state, these SendFlags will be used for whether to remove and add from the containers for the states.
+		 * @param SendFlags If the PackageData needs to be moved to a new state, these SendFlags will be used for
+		 * whether to remove and add from the containers for the states.
+		 */
+		void UpdateRequestData(const TArrayView<const ITargetPlatform* const>& InRequestedPlatforms, bool bInIsUrgent,
+			FCompletionCallback&& InCompletionCallback, ESendFlags SendFlags = ESendFlags::QueueAddAndRemove);
+		/**
+		 * Set the given data, which are all fields describing a request for the cook of this PackageData's Package,
+		 * onto the existing request data. It is invalid to call this on a PackageData that is already InProgress.
 		 */
-		void UpdateRequestData(const TArrayView<const ITargetPlatform* const>& InRequestedPlatforms, bool bInIsUrgent, FCompletionCallback&& InCompletionCallback, ESendFlags SendFlags = ESendFlags::QueueAddAndRemove);
-		/* Set the given data, which are all fields describing a request for the cook of this PackageData's Package, onto the existing request data. It is invalid to call this on a PackageData that is already InProgress. */
 		void SetRequestData(const TArrayView<const ITargetPlatform* const>& InRequestedPlatforms, bool bInIsUrgent, FCompletionCallback&& InCompletionCallback);
-		/* Clear all the inprogress variables from the current PackageData. It is invalid to call this except when the PackageData is transitioning out of InProgress. */
+		/**
+		 * Clear all the inprogress variables from the current PackageData. It is invalid to call this except when
+		 * the PackageData is transitioning out of InProgress.
+		 */
 		void ClearInProgressData();
 
 		/**
-		 * Add each element of New to CookedPlatforms if it is not already present, with the given succeeded flag.  New and Succeeded must be the same length; the succeeded flag for New[n] is Succeeded[n].
+		 * Add each element of New to CookedPlatforms if it is not already present, with the given succeeded flag.
+		 * New and Succeeded must be the same length; the succeeded flag for New[n] is Succeeded[n].
+		 * If a platform is already present in CookedPlatforms, its success flag is overwritten.
+		 */
+		void AddCookedPlatforms(const TArrayView<const ITargetPlatform* const>& New,
+			const TArrayView<const bool>& Succeeded);
+		/**
+		 * Add each element of New to CookedPlatforms if it is not already present, succeeded set to bSucceeded.
 		 * If a platform is already present in CookedPlatforms, its success flag is overwritten.
 		 */
-		void AddCookedPlatforms(const TArrayView<const ITargetPlatform* const>& New, const TArrayView<const bool>& Succeeded);
-		/** Add each element of New to CookedPlatforms if it is not already present, succeeded set to bSucceeded. If a platform is already present in CookedPlatforms, its success flag is overwritten. */
 		void AddCookedPlatforms(const TArrayView<const ITargetPlatform* const>& New, bool bSucceeded);
-		/** Add TargetPlatform to CookedPlatforms if it is not already present, succeeded set to bSucceeded. If TargetPlatform is already present in CookedPlatforms, its success flag is overwritten. */
+		/**
+		 * Add TargetPlatform to CookedPlatforms if it is not already present, succeeded set to bSucceeded.
+		 * If TargetPlatform is already present in CookedPlatforms, its success flag is overwritten.
+		 */
 		void AddCookedPlatform(const ITargetPlatform* TargetPlatform, bool bSucceeded);
 		/** Remove the given Platform and its succeeded flag from CookedPlatforms if it exists. */
 		void RemoveCookedPlatform(const ITargetPlatform* Platform);
@@ -151,15 +210,30 @@ namespace Cook
 		int32 GetNumCookedPlatforms() const;
 		/** Return true if and only if there is at least one element in CookedPlatforms */
 		bool HasAnyCookedPlatform() const;
-		/** Return true if and only if at least one element of Platforms is present in CookedPlatforms, and with its succeeded flag set to true if bIncludeFailed is false. Returns false if Platforms is empty. */
+		/**
+		 * Return true if and only if at least one element of Platforms is present in CookedPlatforms, and with its
+		 * succeeded flag set to true if bIncludeFailed is false. Returns false if Platforms is empty.
+		 */
 		bool HasAnyCookedPlatforms(const TArrayView<const ITargetPlatform* const>& Platforms, bool bIncludeFailed) const;
-		/** Return true if and only if ever element of Platforms is present in CookedPlatforms, and with its succeeded flag set to true if bIncludeFailed is false. Returns true if Platforms is empty. */
+		/**
+		 * Return true if and only if ever element of Platforms is present in CookedPlatforms, and with its succeeded
+		 * flag set to true if bIncludeFailed is false. Returns true if Platforms is empty.
+		 */
 		bool HasAllCookedPlatforms(const TArrayView<const ITargetPlatform* const>& Platforms, bool bIncludeFailed) const;
-		/** Return true if and only if the given Platform is present in CookedPlatforms, and with its succeeded flag set to true if bIncludeFailed is false. */
+		/**
+		 * Return true if and only if the given Platform is present in CookedPlatforms, and with its succeeded flag
+		 * set to true if bIncludeFailed is false.
+		 */
 		bool HasCookedPlatform(const ITargetPlatform* Platform, bool bIncludeFailed) const;
-		/** Return the CookResult for the given platform.  If the platform does not exist in CookedPlatforms, returns ECookResult::Unseen, otherwise returns Succeeded or Failed depending on its success flag. */
+		/**
+		 * Return the CookResult for the given platform.  If the platform does not exist in CookedPlatforms,
+		 * returns ECookResult::Unseen, otherwise returns Succeeded or Failed depending on its success flag.
+		 */
 		ECookResult GetCookResults(const ITargetPlatform* Platform) const;
-		/** Empties and then sets OutPlatforms to contain all elements of QueryPlatforms that are not present in CookedPlatforms. */
+		/**
+		 * Empties and then sets OutPlatforms to contain all elements of QueryPlatforms that are not present in
+		 * CookedPlatforms.
+		 */
 		template <typename ArrayType>
 		void GetUncookedPlatforms(const TArrayView<const ITargetPlatform* const>& QueryPlatforms, ArrayType& OutPlatforms)
 		{
@@ -173,7 +247,10 @@ namespace Cook
 			}
 		}
 
-		/* Return the package pointer. By contract it will be non-null if and only if the PackageData's state is >= EPackageState::Load. */
+		/**
+		 * Return the package pointer. By contract it will be non-null if and only if the PackageData's state is
+		 * >= EPackageState::Load.
+		 */
 		UPackage* GetPackage() const;
 		/* Set the package pointer. Caller is responsible for maintaining the contract for this field. */
 		void SetPackage(UPackage* InPackage);
@@ -181,19 +258,22 @@ namespace Cook
 		/** Return the current PackageState */
 		EPackageState GetState() const;
 		/**
-		 * Set the PackageData's state to the given state, remove and add of from the appropriate queues, and destroy, create, and verify the appropriate state-specific data.
+		 * Set the PackageData's state to the given state, remove and add of from the appropriate queues, and destroy,
+		 * create, and verify the appropriate state-specific data.
 
 		 * @param NextState The destination state
-		 * @param SendFlags Behavior for how the PackageData should be added/removed from the queues corresponding to the new and old states. Callers may want to manage queue membership
-		 *                  directly for better performance; removing from the middle of a queue is more expensive than popping from the front.
+		 * @param SendFlags Behavior for how the PackageData should be added/removed from the queues corresponding to
+		 *                  the new and old states. Callers may want to manage queue membership directly for better
+		 *                  performance; removing from the middle is more expensive than popping from the front.
 		 *                  See definition of ESendFlags for a description of the behavior controlled by SendFlags.
 		 */
 		void SendToState(EPackageState NextState, ESendFlags SendFlags);
 		/* Debug-only code to assert that this PackageData is contained by the container matching its current state. */
 		void CheckInContainer() const;
 		/**
-		 * Return true if and only if this PackageData is InProgress in the current CookOnTheFlyServer session. Some data is allocated/destroyed/verified when moving in and out of InProgress.
-		 * InProgress means the CookOnTheFlyServer will at some point inspect the PackageData and decide to cook, failtocook, or skip it.
+		 * Return true if and only if this PackageData is InProgress in the current CookOnTheFlyServer session.
+		 * Some data is allocated/destroyed/verified when moving in and out of InProgress.
+		 * InProgress means the CookOnTheFlyServer will in the future decide to cook, failtocook, or skip the PackageData.
 		 */
 		bool IsInProgress() const;
 
@@ -201,14 +281,21 @@ namespace Cook
 		bool IsInStateProperty(EPackageStateProperty Property) const;
 
 		/*
-		 * CompletionCallback - A callback that will be called when this PackageData next transitions from InProgress to not InProgress because of cook success, failure, skip, or cancel.
+		 * CompletionCallback - A callback that will be called when this PackageData next transitions from InProgress to not
+		 * InProgress because of cook success, failure, skip, or cancel.
 		 */
 		/** Get a reference to the currently set callback, to e.g. move it into a local variable during execution. */
 		FCompletionCallback& GetCompletionCallback();
-		/** Add the given callback into this PackageData's callback field. It is invalid to call this function with a non-empty callback if this PackageData already has a CompletionCallback. */
+		/**
+		 * Add the given callback into this PackageData's callback field. It is invalid to call this function with a
+		 * non-empty callback if this PackageData already has a CompletionCallback.
+		 */
 		void AddCompletionCallback(FCompletionCallback&& InCompletionCallback);
 
-		/** Get/Set a visited flag used when searching graphs of PackageData. User of the graph is responsible for setting the bIsVisited flag back to empty when graph operations are done. */
+		/**
+		 * Get/Set a visited flag used when searching graphs of PackageData. User of the graph is responsible for 
+		 * setting the bIsVisited flag back to empty when graph operations are done.
+		 */
 		bool GetIsVisited() const { return bIsVisited != 0; }
 		void SetIsVisited(bool bValue) { bIsVisited = static_cast<uint32>(bValue); }
 
@@ -225,11 +312,14 @@ namespace Cook
 		/** Issue check statements confirming that no preload data is allocated or flags are set. */
 		void CheckPreloadEmpty();
 
-		/** The list of objects inside the package.  Only non-empty during saving; it is populated on demand by TryCreateObjectCache and is cleared when leaving the save state. */
+		/**
+		 * The list of objects inside the package.  Only non-empty during saving; it is populated on demand by
+		 * TryCreateObjectCache and is cleared when leaving the save state.
+		 */
 		TArray<FWeakObjectPtr>& GetCachedObjectsInOuter();
-		/** Validate that the variables relying on the CachedObjectsInOuter are empty as required, when e.g. entering the save state. */
+		/** Validate that the CachedObjectsInOuter-dependent variables are empty, when entering save. */
 		void CheckObjectCacheEmpty() const;
-		/** Populate the CachedObjectsInOuter list if it is not already populated. Invalid to call except when in the save state. */
+		/** Populate CachedObjectsInOuter if not already populated. Invalid to call except when in the save state. */
 		void CreateObjectCache();
 		/** Clear the CachedObjectsInOuter list, when e.g. leaving the save state. */
 		void ClearObjectCache();
@@ -239,72 +329,84 @@ namespace Cook
 		const int32& GetCookedPlatformDataNextIndex() const;
 		int32& GetCookedPlatformDataNextIndex();
 
-		/** Get/Set the flag for whether this PackageData has populated GetCachedObjectsInOuter. Always false except during the save state. */
+		/** Get/Set the flag for whether CachedObjectsInOuter is populated. Always false except during save state. */
 		bool GetHasSaveCache() const { return static_cast<bool>(bHasSaveCache); }
 		void SetHasSaveCache(bool Value) { bHasSaveCache = Value != 0; }
-		/** Get/Set the flag for whether iteration has started during save over GetCachedObjectsInOuter to call BeginCacheForCookedPlatformData on each one. Always false except during the save state. */
+		/** Get/Set the flag for whether CachedObjectsInOuter iteration has started. Always false except during save. */
 		bool GetCookedPlatformDataStarted() const { return static_cast<bool>(bCookedPlatformDataStarted); }
 		void SetCookedPlatformDataStarted(bool Value) { bCookedPlatformDataStarted = (Value != 0); }
 		/**
-		 * Get/Set the flag for whether BeginCacheForCookedPlatformData has been called on every object in GetCachedObjectsInOuter.
-		 * Some objects might still be working asynchronously and have not yet returned true from IsCachedCookedPlatformDataLoaded, in which case GetNumPendingCookedPlatformData will be non-zero.
+		 * Get/Set the flag for whether BeginCacheForCookedPlatformData has been called on every CachedObjectsInOuter.
+		 * Some objects might still be working asynchronously and have not yet returned true from
+		 * IsCachedCookedPlatformDataLoaded, in which case GetNumPendingCookedPlatformData will be non-zero.
 		 * Always false except during the save state.
 		 */
 		bool GetCookedPlatformDataCalled() const { return static_cast<bool>(bCookedPlatformDataCalled); }
 		void SetCookedPlatformDataCalled(bool bValue) { bCookedPlatformDataCalled = bValue != 0; }
-		/** Get/Set the flag for whether BeginCacheForCookedPlatformData has been called and IsCachedCookedPlatformDataLoaded has subsequently returned true for every object in GetCachedObjectsInOuter.  Always false except during the save state. */
+		/**
+		 * Get/Set the flag for whether BeginCacheForCookedPlatformData has been called and
+		 * IsCachedCookedPlatformDataLoaded has subsequently returned true for every object in GetCachedObjectsInOuter.
+		 * Always false except during the save state.
+		 */
 		bool GetCookedPlatformDataComplete() const { return static_cast<bool>(bCookedPlatformDataComplete); }
 		void SetCookedPlatformDataComplete(bool bValue) { bCookedPlatformDataComplete = bValue != 0; }
-		/** Check whether savestate contracts on the PackageData were invalidated by by e.g. garbage collection of objects in its package. */
+		/** Check whether savestate contracts on the PackageData were invalidated by by e.g. garbage collection. */
 		bool IsSaveInvalidated() const;
 
-		/** Get/Set the flag for whether BeginPrepareSave has been called and returned false because an error occured. */
+		/** Get/Set the flag for whether BeginPrepareSave has been called and returned an error. */
 		bool GetHasBeginPrepareSaveFailed() const { return static_cast<bool>(bHasBeginPrepareSaveFailed); }
 		void SetHasBeginPrepareSaveFailed(bool bValue) { bHasBeginPrepareSaveFailed = bValue != 0; }
 
-		/** Validate that the fields related to the BeginCacheForCookedPlatformData calls are empty, as required when not in the save state. */
+		/** Validate that the BeginCacheForCookedPlatformData-dependent fields are empty, when entering save. */
 		void CheckCookedPlatformDataEmpty() const;
-		/** Clear the fields related to the BeginCacheForCookedPlatformData calls, when e.g. leaving the save state. Caller is responsible for having already executed any required cancellation steps to avoid dangling pending operations. */
+		/**
+		 * Clear the BeginCacheForCookedPlatformData-dependent fields, when leaving save.
+		 * Caller must have already executed any required cancellation steps to avoid dangling pending operations.
+		 */
 		void ClearCookedPlatformData();
 
-		/** Get/Set the flag managed by the PackageDataMonitor to count whether this PackageData has finished cooking. */
+		/** Get/Set the PackageDataMonitor flag that counts whether this PackageData has finished cooking. */
 		bool GetMonitorIsCooked() const { return static_cast<bool>(bMonitorIsCooked); }
 		void SetMonitorIsCooked(bool Value) { bMonitorIsCooked = Value != 0; }
 
 		/** Remove all request data about the given platform from all fields in this PackageData. */
 		void OnRemoveSessionPlatform(const ITargetPlatform* Platform);
 
-		/** Report whether this PackageData is holding any references to Objects and would therefore be affected by GarbageCollection. */
+		/** Report whether this PackageData holds references to UObjects and would be affected by GarbageCollection. */
 		bool HasReferencedObjects() const;
 
 		/** Swap all ITargetPlatform* stored on this instance according to the mapping in @param Remap. */
 		void RemapTargetPlatforms(const TMap<ITargetPlatform*, ITargetPlatform*>& Remap);
 
 		/** Create and return a GeneratorPackage helper object for a given CookPackageSplitter. */
-		FGeneratorPackage* CreateGeneratorPackage(const UObject* InSplitDataObject, ICookPackageSplitter* InCookPackageSplitterInstance);
+		FGeneratorPackage* CreateGeneratorPackage(const UObject* InSplitDataObject,
+			ICookPackageSplitter* InCookPackageSplitterInstance);
 		/** Destroy GeneratorPackage helper object. */
 		void DestroyGeneratorPackage() { GeneratorPackage.Reset(); }
 		/** Get GeneratorPackage helper object. */
 		FGeneratorPackage* GetGeneratorPackage() const { return GeneratorPackage.Get(); }
-		/** Get whether the PackageData has completed any necessary Generator steps and is ready for BeginCache calls */
+		/** Get whether the PackageData has done any necessary Generator steps and is ready for BeginCache calls. */
 		bool HasCompletedGeneration() const { return static_cast<bool>(bCompletedGeneration); }
-		/** Set whether the PackageData has completed any necessary Generator steps and is ready for BeginCache calls */
+		/** Set whether the PackageData has done any necessary Generator steps and is ready for BeginCache calls. */
 		void SetCompletedGeneration(bool Value) { bCompletedGeneration = Value != 0; }
 		/** Return whether the package is a Generator package and has started but not completed generation */
-		bool IsGenerating() const { return GetGeneratorPackage() != nullptr && HasInitializedGeneratorSave() && !HasCompletedGeneration(); }
-		/** Get whether the PackageData has completed the check for whether it is a Generator, used during Generation */
+		bool IsGenerating() const
+		{
+			return GetGeneratorPackage() != nullptr && HasInitializedGeneratorSave() && !HasCompletedGeneration();
+		}
+		/** Get whether the PackageData has completed the check for whether it is a Generator. */
 		bool HasInitializedGeneratorSave() const { return static_cast<bool>(bInitializedGeneratorSave); }
-		/** Set whether the PackageData has completed the check for whether it is a Generator, used during Generation */
+		/** Set whether the PackageData has completed the check for whether it is a Generator. */
 		void SetInitializedGeneratorSave(bool Value) { bInitializedGeneratorSave = Value != 0; }
-		/** Return whether the PackageData is a generated package created by a owning generator PackageData */
+		/** Return whether the PackageData is a generated package created by a owning generator PackageData. */
 		bool IsGenerated() const { return static_cast<bool>(bGenerated); }
-		/** Set whether the PackageData is a generated package created by a owning generator PackageData */
+		/** Set whether the PackageData is a generated package created by a owning generator PackageData. */
 		void SetGenerated(bool Value) { bGenerated = Value != 0; }
-		/** Set the owning generator PackageData. Only valid to call if SetGenerated(true) has been called */
+		/** Set the owning generator PackageData. Only valid to call if SetGenerated(true) has been called. */
 		void SetGeneratedOwner(FGeneratorPackage* InGeneratedOwner);
-		/** Return the owning generator PackageData. Will be null if not a generated package or if orphaned */
+		/** Return the owning generator PackageData. Will be null if not a generated package or if orphaned. */
 		FGeneratorPackage* GetGeneratedOwner() const { return GeneratedOwner; }
-		/** Mark that the package has finished generation and its save attempt, and should generate again if it needs to be resaved */
+		/** Mark that the package has finished generation+saving, and should generate again if it needs resaving. */
 		void ResetGenerationProgress();
 
 		/** Return whether a GC is required by a generator package as soon as possible. */
@@ -321,14 +423,23 @@ namespace Cook
 		void ClearRequestedPlatforms();
 		void SetIsUrgent(bool Value);
 
-		/** Set the FileName of the file that contains the package. This member is private because FPackageDatas keeps a map from FileName to PackageData that needs to be updated in sync with it. */
+		/**
+		 * Set the FileName of the file that contains the package. This member is private because FPackageDatas
+		 * keeps a map from FileName to PackageData that needs to be updated in sync with it.
+		 */
 		void SetFileName(const FName& InFileName);
 
-		/** Set the State of this PackageDAta in the CookOnTheFlyServer's session. This member is private because it needs to be updated in sync with other contract data. */
+		/**
+		 * Set the State of this PackageDAta in the CookOnTheFlyServer's session. This member is private because it
+		 * needs to be updated in sync with other contract data.
+		 */
 		void SetState(EPackageState NextState);
 
 	private:
-		/** Helper function to call the given EdgeFunction (e.g. OnExitInProgress) when a property changes from true to false. */
+		/**
+		 * Helper function to call the given EdgeFunction (e.g. OnExitInProgress)
+		 * when a property changes from true to false.
+		 */
 		typedef void (FPackageData::*FEdgeFunction)();
 		inline void UpdateDownEdge(bool bOld, bool bNew, const FEdgeFunction& EdgeFunction)
 		{
@@ -337,7 +448,10 @@ namespace Cook
 				(this->*EdgeFunction)();
 			}
 		}
-		/** Helper function to call the given EdgeFunction (e.g. OnEnterInProgress) when a property changes from false to true. */
+		/**
+		 * Helper function to call the given EdgeFunction (e.g. OnEnterInProgress)\
+		 * when a property changes from false to true.
+		 */
 		inline void UpdateUpEdge(bool bOld, bool bNew, const FEdgeFunction& EdgeFunction)
 		{
 			if ((bOld != bNew) & bNew)
@@ -368,14 +482,17 @@ namespace Cook
 		TUniquePtr<FGeneratorPackage> GeneratorPackage;
 		FGeneratorPackage* GeneratedOwner;
 		TArray<const ITargetPlatform*> RequestedPlatforms;
-		TArray<const ITargetPlatform*> CookedPlatforms; // Platform part of the CookedPlatforms set. Always the same length as CookSucceeded.
-		TArray<bool> CookSucceeded; // Success flag part of the CookedPlatforms set. Always the same length as CookedPlatforms.
+		/** Platform part of the CookedPlatforms set.Always the same length as CookSucceeded. */
+		TArray<const ITargetPlatform*> CookedPlatforms;
+		/** Success flag part of the CookedPlatforms set.Always the same length as CookedPlatforms. */
+		TArray<bool> CookSucceeded;
 		TArray<FWeakObjectPtr> CachedObjectsInOuter;
 		FCompletionCallback CompletionCallback;
 		FName PackageName;
 		FName FileName;
 		TWeakObjectPtr<UPackage> Package;
-		FPackageDatas& PackageDatas; // The one-per-CookOnTheFlyServer owner of this PackageData
+		/** The one-per-CookOnTheFlyServer owner of this PackageData. */
+		FPackageDatas& PackageDatas;
 		/**
 		* The number of active PreloadableFiles is tracked globally; wrap the PreloadableFile in a struct that
 		* guarantees we always update the counter when changing it
@@ -427,13 +544,15 @@ namespace Cook
 		};
 
 		/** Store the provided CookPackageSplitter and prepare the packages to generate. */
-		FGeneratorPackage(UE::Cook::FPackageData& InOwner, const UObject* InSplitDataObject, ICookPackageSplitter* InCookPackageSplitterInstance);
+		FGeneratorPackage(UE::Cook::FPackageData& InOwner, const UObject* InSplitDataObject,
+			ICookPackageSplitter* InCookPackageSplitterInstance);
 		~FGeneratorPackage();
 		/** Clear references to owned generated packages, and mark those packages as orphaned */
 		void ClearGeneratedPackages();
 
 		/** Call the Splitter's GetGenerateList and create the PackageDatas */
-		bool TryGenerateList(UObject* OwnerObject, const FPackageNameCache& PackageNameCache, FPackageDatas& PackageDatas);
+		bool TryGenerateList(UObject* OwnerObject, const FPackageNameCache& PackageNameCache,
+			FPackageDatas& PackageDatas);
 
 		/** Accessor for the packages to generate */
 		TArrayView<FGeneratedStruct> GetPackagesToGenerate() { return PackagesToGenerate; }
@@ -455,7 +574,7 @@ namespace Cook
 		 */
 		FString GetIntermediateLocalPath(const FGeneratorPackage::FGeneratedStruct& GeneratedStruct) const;
 
-		/** Splitter capability for whether the generated packages should populate from owner rather than load from disk */
+		/** Splitter capability for whether generated packages should populate from owner rather than load from disk. */
 		bool IsDeferredPopulate() const { return CookPackageSplitterInstance->UseDeferredPopulate(); }
 		/** Return the FGeneratedStruct with the given PackageData, or null if not found. */
 		FGeneratedStruct* FindGeneratedStruct(FPackageData* PackageData);
@@ -482,10 +601,10 @@ namespace Cook
 		/** Get/Clear whether the Owner was garbage collected since the previous call to TryPopulatePackage */
 		bool GetWasOwnerReloaded() const { return bWasOwnerReloaded; }
 		void ClearWasOwnerReloaded() { bWasOwnerReloaded = false; }
-		/** Call CreatePackage and set PackageData correctly for deterministic generated packages, and update status. */
+		/** Call CreatePackage and set PackageData for deterministic generated packages, and update status. */
 		UPackage* CreateGeneratedUPackage(FGeneratedStruct& GeneratedStruct,
 			const UPackage* OwnerPackage, const TCHAR* GeneratedPackageName);
-		/** Mark that the generated package has been populated, to keep track of when this GeneratorPackage is no longer needed. */
+		/** Mark that the generated package has been populated, to keep track of when *this is no longer needed. */
 		void SetGeneratedSaved(FPackageData& PackageData);
 		/** Return whether list has been generated and all generated packages have been populated */
 		bool IsComplete() const;
@@ -511,13 +630,16 @@ namespace Cook
 
 
 	/**
-	 * Stores information about the pending action in response to a single call to BeginCacheForCookedPlatformData that was made on a given object for the given platform, when saving the given PackageData.
+	 * Stores information about the pending action in response to a single call to BeginCacheForCookedPlatformData that
+	 * was made on a given object for the given platform, when saving the given PackageData.
 	 * This instance will remain alive until the object returns true from IsCachedCookedPlatformDataLoaded.
-	 * If the PackageData's save was canceled, this struct also becomes responsible for cleanup of the cached data by calling ClearAllCachedCookedPlatformData.
+	 * If the PackageData's save was canceled, this struct also becomes responsible for cleanup of the cached data by
+	 * calling ClearAllCachedCookedPlatformData.
 	 */
 	struct FPendingCookedPlatformData
 	{
-		FPendingCookedPlatformData(UObject* InObject, const ITargetPlatform* InTargetPlatform, FPackageData& InPackageData, bool bInNeedsResourceRelease, UCookOnTheFlyServer& InCookOnTheFlyServer);
+		FPendingCookedPlatformData(UObject* InObject, const ITargetPlatform* InTargetPlatform,
+			FPackageData& InPackageData, bool bInNeedsResourceRelease, UCookOnTheFlyServer& InCookOnTheFlyServer);
 		FPendingCookedPlatformData(FPendingCookedPlatformData&& Other);
 		FPendingCookedPlatformData(const FPendingCookedPlatformData& Other) = delete;
 		~FPendingCookedPlatformData();
@@ -526,7 +648,8 @@ namespace Cook
 
 		/**
 		 * Call IsCachedCookedPlatformDataLoaded on the object if it has not already returned true.
-		 * If IsCachedCookedPlatformDataLoaded returns true, this function releases all held resources related to the pending call, and returns true. Otherwise takes no action and returns false.
+		 * If IsCachedCookedPlatformDataLoaded returns true, this function releases all held resources related to the
+		 * pending call, and returns true. Otherwise takes no action and returns false.
 		 * Returns true and early exits if IsCachedCookedPlatformDataLoaded has already returned true.
 		 */
 		bool PollIsComplete();
@@ -544,17 +667,26 @@ namespace Cook
 		FPackageData& PackageData;
 		/** Backpointer to the CookOnTheFlyServer to allow releasing of resources for the pending call. */
 		UCookOnTheFlyServer& CookOnTheFlyServer;
-		/** Non-null only in the case of a cancel. Used to synchronize release of shared resources used by all FPendingCookedPlatformData for the various TargetPlatforms of a given object. */
+		/**
+		 * Non-null only in the case of a cancel. Used to synchronize release of shared resources used by all
+		 * FPendingCookedPlatformData for the various TargetPlatforms of a given object.
+		 */
 		FPendingCookedPlatformDataCancelManager* CancelManager;
 		/* Saved copy of the ClassName to use for resource releasing. */
 		FName ClassName;
 		/** Flag for whether we have executed the release. */
 		bool bHasReleased;
-		/** Flag for whether the CookOnTheFlyServer requires resource tracking for the object's BeginCacheForCookedPlatformData call. */
+		/**
+		 * Flag for whether the CookOnTheFlyServer requires resource tracking for the object's
+		 * BeginCacheForCookedPlatformData call.
+		 */
 		bool bNeedsResourceRelease;
 	};
 
-	/** Stores information about all of the FPendingCookedPlatformData for a given object, so that resources shared by all of the FPendingCookedPlatformData can be released after they are all released. */
+	/**
+	 * Stores information about all of the FPendingCookedPlatformData for a given object, so that resources shared by
+	 * all of the FPendingCookedPlatformData can be released after they are all released.
+	 */
 	struct FPendingCookedPlatformDataCancelManager
 	{
 		/** The number of FPendingCookedPlatformData for the given object that are still pending. */
@@ -563,7 +695,10 @@ namespace Cook
 		void Release(FPendingCookedPlatformData& Data);
 	};
 
-	/** The container class for PackageData pointers that are InProgress in a CookOnTheFlyServer. These containers most frequently do queue push/pop operations, but also commonly need to support iteration. */
+	/**
+	 * The container class for PackageData pointers that are InProgress in a CookOnTheFlyServer. These containers
+	 * most frequently do queue push/pop operations, but also commonly need to support iteration.
+	 */
 	class FPackageDataQueue : public TRingBuffer<FPackageData*>
 	{
 		using TRingBuffer<FPackageData*>::TRingBuffer;
@@ -574,21 +709,28 @@ namespace Cook
 		using TFastPointerSet<FPackageData*>::TFastPointerSet;
 	};
 
-	/** A monitor class held by an FPackageDatas to provide reporting and decision making based on aggregated-data across all InProgress or completed FPackageData. */
+	/**
+	 * A monitor class held by an FPackageDatas to provide reporting and decision making based on aggregated-data
+	 * across all InProgress or completed FPackageData.
+	 */
 	struct FPackageDataMonitor
 	{
 	public:
 		FPackageDataMonitor();
 
-		/** Report the number of FPackageData that are in any non-idle state and need to be acted on at some point by the CookOnTheFlyServer */
+		/** Report the number of FPackageData that are in any non-idle state and need action by CookOnTheFlyServer. */
 		int32 GetNumInProgress() const;
 		int32 GetNumPreloadAllocated() const;
-		/** Report the number of packages that have cooked any platform. Used by the cook commandlet for progress reporting to the user. */
+		/** Report the number of packages that have cooked any platform. Used by CookCommandlet progress reporting. */
 		int32 GetNumCooked() const;
-		/** Report the number of FPackageData that are currently marked as urgent. Used to check if a Pump function needs to exit to handle urgent PackageData in other states. */
+		/**
+		 * Report the number of FPackageData that are currently marked as urgent.
+		 * Used to check if a Pump function needs to exit to handle urgent PackageData in other states.
+		 */
 		int32 GetNumUrgent() const;
 		/**
-		 * Report the number of FPackageData that are in the given state and have been marked as urgent. Only valid to call on states that are in the InProgress set, such as Save.
+		 * Report the number of FPackageData that are in the given state and have been marked as urgent. Only valid to
+		 * call on states that are in the InProgress set, such as Save.
 		 * Used to prioritize scheduler actions.
 		 */
 		int32 GetNumUrgent(EPackageState InState) const;
@@ -639,7 +781,7 @@ namespace Cook
 	};
 
 	/**
-	 * A Container for FPackageDatas in the LoadPrepare state. A FIFO container with multiple substates that packages are moved along.
+	 * Container for FPackageDatas in the LoadPrepare state. A FIFO with multiple substates.
 	 */
 	class FLoadPrepareQueue
 	{
@@ -657,8 +799,10 @@ namespace Cook
 	};
 
 	/*
-	 * Class that manages the list of all PackageDatas for a CookOnTheFlyServer. PackageDatas is an associative array for extra data about a package (e.g. the RequestedPlatforms) that is needed by the CookOnTheFlyServer.
-	 * FPackageData are allocated once and never destroyed or moved until the CookOnTheFlyServer is destroyed. Memory on the FPackageData is allocated and deallocated as necessary for its current state.
+	 * Class that manages the list of all PackageDatas for a CookOnTheFlyServer. PackageDatas is an associative
+	 * array for extra data about a package (e.g. the RequestedPlatforms) that is needed by the CookOnTheFlyServer.
+	 * FPackageData are allocated once and never destroyed or moved until the CookOnTheFlyServer is destroyed.
+	 * Memory on the FPackageData is allocated and deallocated as necessary for its current state.
 	 * FPackageData are mapped by PackageName and by FileName.
 	 * This class also manages all non-temporary references to FPackageData such as the SaveQueue and RequeustQueue.
 	*/
@@ -670,72 +814,124 @@ namespace Cook
 
 		/** FGCObject interface function - return a debug name describing this FGCObject. */
 		virtual FString GetReferencerName() const override;
-		/** FGCObject interface function - add the objects referenced by this FGCObject to the ReferenceCollector. This class forwards the query on to the CookOnTheFlyServer. */
+		/**
+		 * FGCObject interface function - add the objects referenced by this FGCObject to the ReferenceCollector.
+		 * This class forwards the query on to the CookOnTheFlyServer.
+		 */
 		virtual void AddReferencedObjects(FReferenceCollector& Collector) override;
 
 		/** Return the PackageNameCache that is caching FileNames on disk for the CookOnTheFlyServer. */
 		const FPackageNameCache& GetPackageNameCache() const;
-		/** Return the Monitor used to report aggregated information about FPackageData managed by this FPackageDatas. */
+		/** Return the Monitor used to report aggregated information about FPackageDatas. */
 		FPackageDataMonitor& GetMonitor();
 		/** Return the backpointer to the CookOnTheFlyServer */
 		UCookOnTheFlyServer& GetCookOnTheFlyServer();
-		/** Return the RequestQueue used by the CookOnTheFlyServer. The RequestQueue is the mostly-FIFO list of PackageData that need to be cooked. */
+		/**
+		 * Return the RequestQueue used by the CookOnTheFlyServer. The RequestQueue is the mostly-FIFO list of
+		 * PackageData that need to be cooked.
+		 */
 		FRequestQueue& GetRequestQueue();
-		/** Return the LoadPrepareQueue used by the CookOnTheFlyServer. The LoadPrepareQueue is the dependency-ordered list of FPackageData that need to be preloaded before they can be loaded. */
+		/**
+		 * Return the LoadPrepareQueue used by the CookOnTheFlyServer. The LoadPrepareQueue is the dependency-ordered
+		 * list of FPackageData that need to be preloaded before they can be loaded.
+		 */
 		FLoadPrepareQueue& GetLoadPrepareQueue() { return LoadPrepareQueue; }
-		/** Return the LoadReadyQueue used by the CookOnTheFlyServer. The LoadReadyQueue is the dependency-ordered list of PackageData that need to be loaded. */
+		/**
+		 * Return the LoadReadyQueue used by the CookOnTheFlyServer. The LoadReadyQueue is the dependency-ordered list
+		 * of PackageData that need to be loaded.
+		 */
 		FPackageDataQueue& GetLoadReadyQueue() { return LoadReadyQueue; }
-		/** Return the SaveQueue used by the CookOnTheFlyServer. The SaveQueue is the performance-sorted list of PackageData that have been loaded and need to start or are only part way through saving. */
+		/**
+		 * Return the SaveQueue used by the CookOnTheFlyServer. The SaveQueue is the performance-sorted list of
+		 * PackageData that have been loaded and need to start or are only part way through saving.
+		 */
 		FPackageDataQueue& GetSaveQueue();
 
-		/** Return the PackageData for the given PackageName and FileName; no validation is done on the names. Creates the PackageData if it does not already exist. */
+		/**
+		Return the PackageData for the given PackageName and FileName; no validation is done on the names.
+		 * Creates the PackageData if it does not already exist.
+		 */
 		FPackageData& FindOrAddPackageData(const FName& PackageName, const FName& NormalizedFileName);
 
 		/** Return the PackageData with the given PackageName if one exists, otherwise return nullptr. */
 		FPackageData* FindPackageDataByPackageName(const FName& PackageName);
-		/** Return a pointer to the PackageData for the given PackageName. If one does not already exist, find its FileName on disk and create the PackageData. If no filename exists for the package on disk, return nullptr. */
+		/**
+		 * Return a pointer to the PackageData for the given PackageName. If one does not already exist,
+		 * find its FileName on disk and create the PackageData. If no filename exists, return nullptr.
+		 */
 		FPackageData* TryAddPackageDataByPackageName(const FName& PackageName);
-		/** Return a reference to the PackageData for the given PackageName. If one does not already exist, find its FileName on disk and create the PackageData. Asserts if FileName does not exist; caller is claiming it does. */
+		/**
+		 * Return a reference to the PackageData for the given PackageName. If one does not already exist,
+		 * find its FileName on disk and create the PackageData. Asserts if FileName does not exist;
+		 * caller is claiming it does.
+		 */
 		FPackageData& AddPackageDataByPackageNameChecked(const FName& PackageName);
 
-		/** Return the PackageData with the given FileName (or that has been marked as having the given FileName as an alias) if one exists, otherwise return nullptr. */
+		/**
+		 * Return the PackageData with the given FileName (or that has been marked as having the given FileName as an
+		 * alias) if one exists, otherwise return nullptr.
+		 */
 		FPackageData* FindPackageDataByFileName(const FName& InFileName);
-		/** Return a pointer to the PackageData for the given FileName (whether the actual name or an alias). If one does not already exist, verify the FileName on disk and create the PackageData. If no filename exists for the package on disk, return nullptr. */
+		/**
+		 * Return a pointer to the PackageData for the given FileName (whether the actual name or an alias).
+		 * If one does not already exist, verify the FileName on disk and create the PackageData.
+		 * If no filename exists for the package on disk, return nullptr.
+		 */
 		FPackageData* TryAddPackageDataByFileName(const FName& InFileName);
-		/** Return a reference to the PackageData for the given FileName (whether the actual name or an alias). If one does not already exist, verify the FileName on disk and create the PackageData. Asserts if FileName does not exist; caller is claiming it does. */
+		/**
+		 * Return a reference to the PackageData for the given FileName (whether the actual name or an alias).
+		 * If one does not already exist, verify the FileName on disk and create the PackageData. Asserts if FileName
+		 * does not exist; caller is claiming it does.
+		 */
 		FPackageData& AddPackageDataByFileNameChecked(const FName& FileName);
 
-		/** Optimized version of TryAddPackageDataByFileName for use by DLC when loading up the shipped packages. This expects the PackageNameCache to have already been updated (for instance with AppendCacheResults) */
-		void AddExistingPackageDatasForPlatform(const TArray<FName>& ExistingPackages, const ITargetPlatform* TargetPlatform);
+		/**
+		 * Optimized version of TryAddPackageDataByFileName for use by DLC when loading up the shipped packages.
+		 * This expects the PackageNameCache to have already been updated (for instance with AppendCacheResults)
+		 */
+		void AddExistingPackageDatasForPlatform(const TArray<FName>& ExistingPackages,
+			const ITargetPlatform* TargetPlatform);
 
 		/**
 		 * Try to find the PackageData for the given PackageName.
-		 * If it exists, change the PackageData's FileName if the current FileName is different and update the map to it.
-		 * This is called in response to the package being moved in the editor or if we attempted to load a FileName and got redirected to another FileName.
-		 * Returns the PackageData if it exists
+		 * If it exists, change the PackageData's FileName if current FileName is different and update the map to it.
+		 * This is called in response to the package being moved in the editor or if we attempted to load a FileName
+		 * and got redirected to another FileName.
+		 * Returns the PackageData if it exists.
 		 */
 		FPackageData* UpdateFileName(const FName& PackageName);
 		/**
-		 * Mark that the give PackageData should also be returned for searches for the given FileName in addition to searches for its own FileName.
+		 * Mark that the give PackageData should also be returned for searches for the given FileName in addition to
+		 * searches for its own FileName.
 		 * This is used for e.g. mapping AssetName.umap to AssetName.uasset.
 		 */
 		void RegisterFileNameAlias(FPackageData& PackageData, FName FileName);
 
-		/** Report the number of packages that have cooked any platform. Used by the cook commandlet for progress reporting to the user. */
+		/** Report the number of packages that have cooked any platform. Used by cook commandlet progress reporting. */
 		int32 GetNumCooked();
-		/** Add onto CookedFiles the list of all packages that have cooked any package, either successfully if bGetSuccessfulCookedPackages is true and/or unsuccessfully if bGetFailedCookedPackages is true. */
-		void GetCookedFileNamesForPlatform(const ITargetPlatform* Platform, TArray<FName>& CookedFiles, bool bGetFailedCookedPackages, bool bGetSuccessfulCookedPackages);
+		/**
+		 * Add onto CookedFiles the list of all packages that have cooked any package, either successfully if
+		 * bGetSuccessfulCookedPackages is true and/or unsuccessfully if bGetFailedCookedPackages is true.
+		 */
+		void GetCookedFileNamesForPlatform(const ITargetPlatform* Platform, TArray<FName>& CookedFiles,
+			bool bGetFailedCookedPackages, bool bGetSuccessfulCookedPackages);
 
-		/** Delete all PackageDatas and free all other memory used by this FPackageDatas. For performance reasons, should only be called on destruction. */
+		/**
+		 * Delete all PackageDatas and free all other memory used by this FPackageDatas.
+		 * For performance reasons, should only be called on destruction.
+		 */
 		void Clear();
 		/** Remove all CookedPlatforms from all PackageDatas. Used to e.g. invalidate previous cooks. */
 		void ClearCookedPlatforms();
-		/** Remove all request data about the given platform from all PackageDatas and other memory used by this FPackageDatas. */
+		/** Remove all request data about the given platform from all PackageDatas and other memory used by *this. */
 		void OnRemoveSessionPlatform(const ITargetPlatform* TargetPlatform);
 
-		/** Return the container that tracks pending calls to BeginCacheForCookedPlatformData during the save of a PackageData. */
+		/** Return the container that tracks pending calls to BeginCacheForCookedPlatformData. */
 		TArray<FPendingCookedPlatformData>& GetPendingCookedPlatformDatas();
-		/** Iterate over all elements in PendingCookedPlatformDatas and check whether they have completed, releasing their resources and pending count if so. */
+		/**
+		 * Iterate over all elements in PendingCookedPlatformDatas and check whether they have completed,
+		 * releasing their resources and pending count if so.
+		 */
 		void PollPendingCookedPlatformDatas();
 
 		/** RangedFor methods for iterating over all FPackageData managed by this FPackageDatas */
@@ -746,10 +942,13 @@ namespace Cook
 		void RemapTargetPlatforms(const TMap<ITargetPlatform*, ITargetPlatform*>& Remap);
 
 	private:
-		/** Construct a new FPackageData with the given PackageName and FileName and store references to it in the maps. New FPackageData are always created in the Idle state. */
+		/**
+		 * Construct a new FPackageData with the given PackageName and FileName and store references to it in the maps.
+		 * New FPackageData are always created in the Idle state.
+		 */
 		FPackageData& CreatePackageData(FName PackageName, FName FileName);
 
-		/** Collection of pointers to all FPackageData that have been constructed (each as a separate allocation) by this FPackageDatas. */
+		/** Collection of pointers to all FPackageData that have been constructed by *this. */
 		TArray<FPackageData*> PackageDatas;
 		FPackageDataMonitor Monitor;
 		FPackageNameCache PackageNameCache;
