Fix compiler error on missing header

diff --git a/Engine/Source/Editor/UnrealEd/Private/SourceCodeNavigation.cpp b/Engine/Source/Editor/UnrealEd/Private/SourceCodeNavigation.cpp
--- a/Engine/Source/Editor/UnrealEd/Private/SourceCodeNavigation.cpp
+++ b/Engine/Source/Editor/UnrealEd/Private/SourceCodeNavigation.cpp
@@ -7,6 +7,7 @@
 #include "HAL/FileManager.h"
 #include "Misc/Paths.h"
 #include "Misc/ScopeLock.h"
+#include "Misc/ConfigCacheIni.h"
 #include "Async/AsyncWork.h"
 #include "Modules/ModuleManager.h"
 #include "UObject/Class.h"
