Hybrid iterative cooking: Combine UE::Cook::FPackageData's platform data into a single container to make room to cache data about the package's target domain status for each platform.

diff --git a/Engine/Source/Editor/UnrealEd/Private/CookOnTheFlyServer.cpp b/Engine/Source/Editor/UnrealEd/Private/CookOnTheFlyServer.cpp
--- a/Engine/Source/Editor/UnrealEd/Private/CookOnTheFlyServer.cpp
+++ b/Engine/Source/Editor/UnrealEd/Private/CookOnTheFlyServer.cpp
@@ -1814,7 +1814,7 @@ void UCookOnTheFlyServer::ProcessRequest(UE::Cook::FPackageData& PackageData, UE
 	using namespace UE::Cook;
 	OutNumPushed = 0;
 
-	if (PackageData.HasAllCookedPlatforms(PackageData.GetRequestedPlatforms(), true /* bIncludeFailed */))
+	if (PackageData.AreAllRequestedPlatformsCooked(true /* bAllowFailedCooks */))
 	{
 #if DEBUG_COOKONTHEFLY
 		UE_LOG(LogCook, Display, TEXT("Package for platform already cooked %s, discarding request"), *PackageData.GetFileName().ToString());
@@ -2120,9 +2120,11 @@ void UCookOnTheFlyServer::LoadPackageInQueue(UE::Cook::FPackageData& PackageData
 
 			UE_LOG(LogCook, Verbose, TEXT("Request for %s received going to save %s"), *PackageFileName.ToString(), *NewPackageFileName.ToString());
 			FPackageData& OtherPackageData = PackageDatas->AddPackageDataByPackageNameChecked(LoadedPackage->GetFName());
-			OtherPackageData.UpdateRequestData(PackageData.GetRequestedPlatforms(), PackageData.GetIsUrgent(), FCompletionCallback());
+			TArray<const ITargetPlatform*, TInlineAllocator<ExpectedMaxNumPlatforms>> RequestedPlatforms;
+			PackageData.GetRequestedPlatforms(RequestedPlatforms);
+			OtherPackageData.UpdateRequestData(RequestedPlatforms, PackageData.GetIsUrgent(), FCompletionCallback());
 
-			PackageData.AddCookedPlatforms(PlatformManager->GetSessionPlatforms(), true);
+			PackageData.SetPlatformsCooked(PlatformManager->GetSessionPlatforms(), true);
 			RejectPackageToLoad(PackageData, TEXT("is redirected to another filename"));
 			return;
 		}
@@ -2203,7 +2205,7 @@ void UCookOnTheFlyServer::LoadPackageInQueue(UE::Cook::FPackageData& PackageData
 		}
 	}
 
-	if (PackageData.HasAllCookedPlatforms(PackageData.GetRequestedPlatforms(), true))
+	if (PackageData.AreAllRequestedPlatformsCooked(true))
 	{
 		// Already cooked. This can happen if we needed to load a package that was previously cooked and garbage collected because it is a loaddependency of a new request.
 		// Send the package back to idle, nothing further to do with it.
@@ -2220,8 +2222,14 @@ void UCookOnTheFlyServer::LoadPackageInQueue(UE::Cook::FPackageData& PackageData
 void UCookOnTheFlyServer::RejectPackageToLoad(UE::Cook::FPackageData& PackageData, const TCHAR* Reason)
 {
 	// make sure this package doesn't exist
-	for (const ITargetPlatform* TargetPlatform : PackageData.GetRequestedPlatforms())
+	for (const TPair<const ITargetPlatform*, UE::Cook::FPackageData::FPlatformData>& Pair : PackageData.GetPlatformDatas())
 	{
+		if (!Pair.Value.bRequested)
+		{
+			continue;
+		}
+		const ITargetPlatform* TargetPlatform = Pair.Key;
+
 		const FString SandboxFilename = ConvertToFullSandboxPath(PackageData.GetFileName().ToString(), true, TargetPlatform->PlatformName());
 		if (IFileManager::Get().FileExists(*SandboxFilename))
 		{
@@ -2745,7 +2753,8 @@ bool UCookOnTheFlyServer::BeginPrepareSave(UE::Cook::FPackageData& PackageData,
 	TArray<FWeakObjectPtr>& CachedObjectsInOuter = PackageData.GetCachedObjectsInOuter();
 	int32& CookedPlatformDataNextIndex = PackageData.GetCookedPlatformDataNextIndex();
 	FWeakObjectPtr* CachedObjectsInOuterData = CachedObjectsInOuter.GetData();
-	const TArray<const ITargetPlatform*>& TargetPlatforms = PackageData.GetRequestedPlatforms();
+	TArray<const ITargetPlatform*, TInlineAllocator<ExpectedMaxNumPlatforms>> TargetPlatforms;
+	PackageData.GetRequestedPlatforms(TargetPlatforms);
 	int NumPlatforms = TargetPlatforms.Num();
 	int NumIndexes = CachedObjectsInOuter.Num() * NumPlatforms;
 	while (CookedPlatformDataNextIndex < NumIndexes)
@@ -2876,7 +2885,7 @@ void UCookOnTheFlyServer::ReleaseCookedPlatformData(UE::Cook::FPackageData& Pack
 		// and this PackageData may therefore still have GetNumPendingCookedPlatformData > 0
 		if (!IsCookingInEditor()) // ClearAllCachedCookedPlatformData calls are only used when not in editor.
 		{
-			int32 NumPlatforms = PackageData.GetRequestedPlatforms().Num();
+			int32 NumPlatforms = PackageData.GetNumRequestedPlatforms();
 			if (NumPlatforms > 0) // Shouldn't happen because PumpSaves checks for this, but avoid a divide by 0 if it does.
 			{
 				// We have only called BeginCacheForCookedPlatformData on Object,Platform pairs up to GetCookedPlatformDataNextIndex.
@@ -3194,7 +3203,7 @@ void UCookOnTheFlyServer::PumpSaves(UE::Cook::FTickStackData& StackData, uint32
 		}
 
 		// Cook only the session platforms that have not yet been cooked for the given package
-		PackageData.GetUncookedPlatforms(PackageData.GetRequestedPlatforms(), PlatformsForPackage);
+		PackageData.GetUncookedPlatforms(PlatformsForPackage);
 		if (PlatformsForPackage.Num() == 0)
 		{
 			// We've already saved all possible platforms for this package; this should not be possible.
@@ -3264,7 +3273,7 @@ void UCookOnTheFlyServer::PumpSaves(UE::Cook::FTickStackData& StackData, uint32
 			if (PackageData.GetHasBeginPrepareSaveFailed())
 			{
 				ReleaseCookedPlatformData(PackageData, true /* bCompletedSave */);
-				PackageData.AddCookedPlatforms(PackageData.GetRequestedPlatforms(), false);
+				PackageData.SetPlatformsCooked(PlatformsForPackage, false /* bSucceeded */);
 				PackageData.SendToState(EPackageState::Idle, ESendFlags::QueueAdd);
 				++OutNumPushed;
 				continue;
@@ -3756,8 +3765,8 @@ void UCookOnTheFlyServer::MarkPackageDirtyForCookerFromSchedulerThread(const FNa
 	if ( PackageData && IsCookingInEditor() )
 	{
 		check(IsInGameThread()); // We're editing scheduler data, which is only allowable from the scheduler thread
-		bool bHadCookedPlatforms = PackageData->GetNumCookedPlatforms() > 0;
-		PackageData->ClearCookedPlatforms();
+		bool bHadCookedPlatforms = PackageData->HasAnyCookedPlatform();
+		PackageData->SetPlatformsNotCooked();
 		if (PackageData->IsInProgress())
 		{
 			PackageData->SendToState(UE::Cook::EPackageState::Request, UE::Cook::ESendFlags::QueueAddAndRemove);
@@ -4633,7 +4642,7 @@ void FSaveCookedPackageContext::FinishPlatform()
 	// by editor-only data; we may need to save it later when new content loads it through non editor-only references.
 	if (!bLocalReferencedOnlyByEditorOnlyData)
 	{
-		PackageData.AddCookedPlatform(TargetPlatform, bSuccessful);
+		PackageData.SetPlatformCooked(TargetPlatform, bSuccessful);
 	}
 
 	// Update flags used to determine garbage collection.
@@ -4682,7 +4691,7 @@ void FSaveCookedPackageContext::FinishPackage()
 					if (LocalizedPackageData)
 					{
 						bool bIsUrgent = false;
-						LocalizedPackageData->UpdateRequestData(PackageData.GetRequestedPlatforms(), bIsUrgent, UE::Cook::FCompletionCallback());
+						LocalizedPackageData->UpdateRequestData(PlatformsForPackage, bIsUrgent, UE::Cook::FCompletionCallback());
 					}
 				}
 			}
@@ -4710,7 +4719,7 @@ void FSaveCookedPackageContext::FinishPackage()
 				if (SoftObjectPackageData)
 				{
 					bool bIsUrgent = false;
-					SoftObjectPackageData->UpdateRequestData(PackageData.GetRequestedPlatforms(), bIsUrgent, UE::Cook::FCompletionCallback());
+					SoftObjectPackageData->UpdateRequestData(PlatformsForPackage, bIsUrgent, UE::Cook::FCompletionCallback());
 				}
 			}
 		}
@@ -5866,7 +5875,7 @@ void UCookOnTheFlyServer::PopulateCookedPackages(TArrayView<const ITargetPlatfor
 			UE::Cook::FPackageData* PackageData = PackageDatas->TryAddPackageDataByPackageName(IdenticalPackage);
 			if (PackageData)
 			{
-				PackageData->AddCookedPlatforms({ TargetPlatform }, true /* bSucceeded */);
+				PackageData->SetPlatformCooked(TargetPlatform, true /* bSucceeded */);
 				PackagesToKeep.Add(IdenticalPackage);
 			}
 		}
@@ -5877,7 +5886,7 @@ void UCookOnTheFlyServer::PopulateCookedPackages(TArrayView<const ITargetPlatfor
 			if (PackageData)
 			{
 				ensure(!PackageData->HasAnyCookedPlatforms({ TargetPlatform }, /* bIncludeFailed */ false));
-				PackageData->AddCookedPlatforms({ TargetPlatform }, false /* bSucceeded */);
+				PackageData->SetPlatformCooked(TargetPlatform, false /* bSucceeded */);
 				PackagesToKeep.Add(UncookedPackage);
 			}
 		}
@@ -7296,7 +7305,7 @@ void UCookOnTheFlyServer::ResetCookResults(TConstArrayView<const ITargetPlatform
 {
 	for (UE::Cook::FPackageData* PackageData : *PackageDatas.Get())
 	{
-		PackageData->RemoveCookedPlatforms(TargetPlatforms);
+		PackageData->SetPlatformsNotCooked(TargetPlatforms);
 	}
 
 	TArray<FName> PackageNames;
@@ -7926,7 +7935,7 @@ void UCookOnTheFlyServer::StartCookByTheBook( const FCookByTheBookStartupOptions
 			{
 				SCOPED_BOOT_TIMING("AddPackageDataByFileNamesForPlatform");
 				PackageDatas->AddExistingPackageDatasForPlatform(ActivePackageList, TargetPlatform);
-					}
+			}
 
 			if (OverridePackageList.Num() > 0)
 			{
@@ -8192,10 +8201,13 @@ void UCookOnTheFlyServer::MaybeMarkPackageAsAlreadyLoaded(UPackage *Package)
 		bShouldMarkAsAlreadyProcessed = PackageData->HasAllCookedPlatforms(PlatformManager->GetSessionPlatforms(), true /* bIncludeFailed */);
 
 		FString Platforms;
-		for (const ITargetPlatform* CookedPlatform : PackageData->GetCookedPlatforms())
+		for (const TPair<const ITargetPlatform*, UE::Cook::FPackageData::FPlatformData>& Pair : PackageData->GetPlatformDatas())
 		{
-			Platforms += TEXT(" ");
-			Platforms += CookedPlatform->PlatformName();
+			if (Pair.Value.bCookAttempted)
+			{
+				Platforms += TEXT(" ");
+				Platforms += Pair.Key->PlatformName();
+			}
 		}
 		if (IsCookFlagSet(ECookInitializationFlags::LogDebugInfo))
 		{
@@ -8857,7 +8869,7 @@ uint32 UCookOnTheFlyServer::FullLoadAndSave(uint32& CookedPackageCount)
 					}
 				}
 
-				PackageData.AddCookedPlatforms(TargetPlatforms, SavePackageSuccessPerPlatform);
+				PackageData.SetPlatformsCooked(TargetPlatforms, SavePackageSuccessPerPlatform);
 
 				if (SavePackageSuccessPerPlatform.Contains(false))
 				{
diff --git a/Engine/Source/Editor/UnrealEd/Private/Cooker/CookPackageData.cpp b/Engine/Source/Editor/UnrealEd/Private/Cooker/CookPackageData.cpp
--- a/Engine/Source/Editor/UnrealEd/Private/Cooker/CookPackageData.cpp
+++ b/Engine/Source/Editor/UnrealEd/Private/Cooker/CookPackageData.cpp
@@ -3,6 +3,7 @@
 #include "CookPackageData.h"
 
 #include "Algo/AnyOf.h"
+#include "Algo/Count.h"
 #include "AssetCompilingManager.h"
 #include "CookOnTheSide/CookOnTheFlyServer.h"
 #include "CookPlatformManager.h"
@@ -28,7 +29,11 @@ namespace Cook
 
 	//////////////////////////////////////////////////////////////////////////
 	// FPackageData
-
+	FPackageData::FPlatformData::FPlatformData()
+		: bRequested(false), bCookAttempted(false), bCookSucceeded(false), bTargetDomainLookupAttempted(false)
+		, bExistsInTargetDomain(false)
+	{
+	}
 
 	FPackageData::FPackageData(FPackageDatas& PackageDatas, const FName& InPackageName, const FName& InFileName)
 		: GeneratedOwner(nullptr), PackageName(InPackageName), FileName(InFileName), PackageDatas(PackageDatas)
@@ -45,8 +50,8 @@ namespace Cook
 	{
 		// ClearReferences should have been called earlier, but call it here in case it was missed
 		ClearReferences();
-		// We need to send OnCookedPlatformRemoved message to the monitor, so call ClearCookedPlatforms
-		ClearCookedPlatforms();
+		// We need to send OnLastCookedPlatformRemoved message to the monitor, so call SetPlatformsNotCooked
+		SetPlatformsNotCooked();
 		// Update the monitor's counters and call exit functions
 		SendToState(EPackageState::Idle, ESendFlags::QueueNone);
 	}
@@ -71,44 +76,59 @@ namespace Cook
 		FileName = InFileName;
 	}
 
-	const TArray<const ITargetPlatform*>& FPackageData::GetRequestedPlatforms() const
-	{
-		return RequestedPlatforms;
-	}
-
-	void FPackageData::SetRequestedPlatforms(const TArrayView<const ITargetPlatform* const>& Platforms)
+	int32 FPackageData::GetNumRequestedPlatforms() const
 	{
-		RequestedPlatforms.Empty(Platforms.Num());
-		AddRequestedPlatforms(Platforms);
+		int32 Result = 0;
+		for (const TPair<const ITargetPlatform*, FPlatformData>& Pair : PlatformDatas)
+		{
+			Result += Pair.Value.bRequested ? 1 : 0;
+		}
+		return Result;
 	}
 
-	void FPackageData::AddRequestedPlatforms(const TArrayView<const ITargetPlatform* const>& New)
+	void FPackageData::SetPlatformsRequested(TConstArrayView<const ITargetPlatform*> TargetPlatforms, bool bRequested)
 	{
-		for (const ITargetPlatform* TargetPlatform : New)
+		for (const ITargetPlatform* TargetPlatform : TargetPlatforms)
 		{
-			RequestedPlatforms.AddUnique(TargetPlatform);
+			PlatformDatas.FindOrAdd(TargetPlatform).bRequested = true;
 		}
 	}
 
 	void FPackageData::ClearRequestedPlatforms()
 	{
-		RequestedPlatforms.Empty();
+		for (TPair<const ITargetPlatform*, FPlatformData>& Pair : PlatformDatas)
+		{
+			Pair.Value.bRequested = false;
+		}
 	}
 
-	bool FPackageData::ContainsAllRequestedPlatforms(const TArrayView<const ITargetPlatform* const>& Platforms) const
+	bool FPackageData::HasAllRequestedPlatforms(const TArrayView<const ITargetPlatform* const>& Platforms) const
 	{
 		if (Platforms.Num() == 0)
 		{
 			return true;
 		}
-		if (RequestedPlatforms.Num() == 0)
+		if (PlatformDatas.Num() == 0)
 		{
 			return false;
 		}
 
 		for (const ITargetPlatform* QueryPlatform : Platforms)
 		{
-			if (!RequestedPlatforms.Contains(QueryPlatform))
+			const FPlatformData* PlatformData = PlatformDatas.Find(QueryPlatform);
+			if (!PlatformData || !PlatformData->bRequested)
+			{
+				return false;
+			}
+		}
+		return true;
+	}
+
+	bool FPackageData::AreAllRequestedPlatformsCooked(bool bAllowFailedCooks) const
+	{
+		for (const TPair<const ITargetPlatform*, FPlatformData>& Pair : PlatformDatas)
+		{
+			if (!Pair.Value.bCookAttempted || (!bAllowFailedCooks && !Pair.Value.bCookSucceeded))
 			{
 				return false;
 			}
@@ -126,7 +146,7 @@ namespace Cook
 		}
 	}
 
-	void FPackageData::UpdateRequestData(const TArrayView<const ITargetPlatform* const>& InRequestedPlatforms,
+	void FPackageData::UpdateRequestData(const TConstArrayView<const ITargetPlatform*> InRequestedPlatforms,
 		bool bInIsUrgent, FCompletionCallback&& InCompletionCallback, ESendFlags SendFlags)
 	{
 		if (IsInProgress())
@@ -140,7 +160,7 @@ namespace Cook
 				SetIsUrgent(true);
 			}
 
-			if (!ContainsAllRequestedPlatforms(InRequestedPlatforms))
+			if (!HasAllRequestedPlatforms(InRequestedPlatforms))
 			{
 				// Send back to the Request state (canceling any current operations) and then add the new platforms
 				if (GetState() != EPackageState::Request)
@@ -148,14 +168,14 @@ namespace Cook
 					check(SendFlags == ESendFlags::QueueAddAndRemove);
 					SendToState(EPackageState::Request, ESendFlags::QueueAddAndRemove);
 				}
-				AddRequestedPlatforms(InRequestedPlatforms);
+				SetPlatformsRequested(InRequestedPlatforms, true);
 			}
 			else if (bUrgencyChanged && SendFlags == ESendFlags::QueueAddAndRemove)
 			{
 				SendToState(GetState(), SendFlags);
 			}
 		}
-		else
+		else if (InRequestedPlatforms.Num() > 0)
 		{
 			SetRequestData(InRequestedPlatforms, bInIsUrgent, MoveTemp(InCompletionCallback));
 			SendToState(EPackageState::Request, SendFlags);
@@ -166,9 +186,11 @@ namespace Cook
 		bool bInIsUrgent, FCompletionCallback&& InCompletionCallback)
 	{
 		check(!CompletionCallback);
-		check(RequestedPlatforms.Num() == 0);
+		check(GetNumRequestedPlatforms() == 0)
 		check(!bIsUrgent);
-		SetRequestedPlatforms(InRequestedPlatforms);
+
+		check(InRequestedPlatforms.Num() != 0);
+		SetPlatformsRequested(InRequestedPlatforms, true);
 		SetIsUrgent(bInIsUrgent);
 		AddCompletionCallback(MoveTemp(InCompletionCallback));
 	}
@@ -180,91 +202,109 @@ namespace Cook
 		CompletionCallback = FCompletionCallback();
 	}
 
-	void FPackageData::AddCookedPlatforms(const TArrayView<const ITargetPlatform* const>& New,
-		const TArrayView<const bool>& Succeeded)
+	void FPackageData::SetPlatformsCooked(const TConstArrayView<const ITargetPlatform*> TargetPlatforms,
+		const TConstArrayView<bool> Succeeded)
 	{
-		check(New.Num() == Succeeded.Num());
-		for (int32 n = 0; n < New.Num(); ++n)
+		check(TargetPlatforms.Num() == Succeeded.Num());
+		for (int32 n = 0; n < TargetPlatforms.Num(); ++n)
 		{
-			AddCookedPlatform(New[n], Succeeded[n]);
+			SetPlatformCooked(TargetPlatforms[n], Succeeded[n]);
 		}
 	}
 
-	void FPackageData::AddCookedPlatforms(const TArrayView<const ITargetPlatform* const>& New, bool bSucceeded)
+	void FPackageData::SetPlatformsCooked(const TConstArrayView<const ITargetPlatform*> TargetPlatforms,
+		bool bSucceeded)
 	{
-		for (const ITargetPlatform* TargetPlatform : New)
+		for (const ITargetPlatform* TargetPlatform : TargetPlatforms)
 		{
-			AddCookedPlatform(TargetPlatform, bSucceeded);
+			SetPlatformCooked(TargetPlatform, bSucceeded);
 		}
 	}
 
-	void FPackageData::AddCookedPlatform(const ITargetPlatform* TargetPlatform, bool bSucceeded)
+	void FPackageData::SetPlatformCooked(const ITargetPlatform* TargetPlatform, bool bSucceeded)
 	{
-		int32 ExistingIndex = CookedPlatforms.IndexOfByKey(TargetPlatform);
-		if (ExistingIndex != INDEX_NONE)
+		bool bHasAnyOthers = false;
+		bool bModified = false;
+		for (TPair<const ITargetPlatform*, FPlatformData>& Pair : PlatformDatas)
 		{
-			CookSucceeded[ExistingIndex] = bSucceeded;
+			if (Pair.Key == TargetPlatform)
+			{
+				bModified = bModified | (Pair.Value.bCookAttempted == false);
+				Pair.Value.bCookAttempted = true;
+				Pair.Value.bCookSucceeded = bSucceeded;
+			}
+			else
+			{
+				bHasAnyOthers = bHasAnyOthers | (Pair.Value.bCookAttempted != false);
+			}
 		}
-		else
+		if (bModified && !bHasAnyOthers)
 		{
-			CookedPlatforms.Add(TargetPlatform);
-			CookSucceeded.Add(bSucceeded);
+			PackageDatas.GetMonitor().OnFirstCookedPlatformAdded(*this);
 		}
-		PackageDatas.GetMonitor().OnCookedPlatformAdded(*this);
 	}
 
-	void FPackageData::RemoveCookedPlatform(const ITargetPlatform* Platform)
+	void FPackageData::SetPlatformsNotCooked(const TConstArrayView<const ITargetPlatform*> TargetPlatforms)
 	{
-		int32 Num = CookedPlatforms.Num();
-		for (int32 n = CookedPlatforms.Num() - 1; n >= 0; --n)
+		for (const ITargetPlatform* TargetPlatform : TargetPlatforms)
 		{
-			if (CookedPlatforms[n] == Platform)
-			{
-				CookedPlatforms.RemoveAtSwap(n);
-				CookSucceeded.RemoveAtSwap(n);
-				PackageDatas.GetMonitor().OnCookedPlatformRemoved(*this);
-				break;
-			}
+			SetPlatformNotCooked(TargetPlatform);
 		}
 	}
 
-	void FPackageData::RemoveCookedPlatforms(const TArrayView<const ITargetPlatform* const>& Platforms)
+	void FPackageData::SetPlatformsNotCooked()
 	{
-		for (const ITargetPlatform* Platform : Platforms)
+		bool bModified = false;
+		for (TPair<const ITargetPlatform*, FPlatformData>& Pair : PlatformDatas)
 		{
-			RemoveCookedPlatform(Platform);
+			bModified = bModified | (Pair.Value.bCookAttempted != false);
+			Pair.Value.bCookAttempted = false;
+			Pair.Value.bCookSucceeded = false;
 		}
-	}
-
-	void FPackageData::ClearCookedPlatforms()
-	{
-		if (CookedPlatforms.Num() > 0)
+		if (bModified)
 		{
-			CookedPlatforms.Empty();
-			CookSucceeded.Empty();
-			PackageDatas.GetMonitor().OnCookedPlatformRemoved(*this);
+			PackageDatas.GetMonitor().OnLastCookedPlatformRemoved(*this);
 		}
 	}
 
-	const TArray<const ITargetPlatform*>& FPackageData::GetCookedPlatforms() const
+	void FPackageData::SetPlatformNotCooked(const ITargetPlatform* TargetPlatform)
 	{
-		return CookedPlatforms;
+		bool bHasAnyOthers = false;
+		bool bModified = false;
+		for (TPair<const ITargetPlatform*, FPlatformData>& Pair : PlatformDatas)
+		{
+			if (Pair.Key == TargetPlatform)
+			{
+				bModified = bModified | (Pair.Value.bCookAttempted != false);
+				Pair.Value.bCookAttempted = false;
+				Pair.Value.bCookSucceeded = false;
+			}
+			else
+			{
+				bHasAnyOthers = bHasAnyOthers | (Pair.Value.bCookAttempted != false);
+			}
+		}
+		if (bModified && !bHasAnyOthers)
+		{
+			PackageDatas.GetMonitor().OnLastCookedPlatformRemoved(*this);
+		}
 	}
 
-	int32 FPackageData::GetNumCookedPlatforms() const
+	const TSortedMap<const ITargetPlatform*, FPackageData::FPlatformData>& FPackageData::GetPlatformDatas() const
 	{
-		return CookedPlatforms.Num();
+		return PlatformDatas;
 	}
 
 	bool FPackageData::HasAnyCookedPlatform() const
 	{
-		return CookedPlatforms.Num() > 0;
+		return Algo::AnyOf(PlatformDatas,
+			[](const TPair<const ITargetPlatform*, FPlatformData>& Pair) { return Pair.Value.bCookAttempted; });
 	}
 
 	bool FPackageData::HasAnyCookedPlatforms(const TArrayView<const ITargetPlatform* const>& Platforms,
 		bool bIncludeFailed) const
 	{
-		if (CookedPlatforms.Num() == 0)
+		if (PlatformDatas.Num() == 0)
 		{
 			return false;
 		}
@@ -286,7 +326,7 @@ namespace Cook
 		{
 			return true;
 		}
-		if (CookedPlatforms.Num() == 0)
+		if (PlatformDatas.Num() == 0)
 		{
 			return false;
 		}
@@ -309,13 +349,10 @@ namespace Cook
 
 	ECookResult FPackageData::GetCookResults(const ITargetPlatform* Platform) const
 	{
-		int32 Num = CookedPlatforms.Num();
-		for (int n = 0; n < Num; ++n)
+		const FPlatformData* PlatformData = PlatformDatas.Find(Platform);
+		if (PlatformData && PlatformData->bCookAttempted)
 		{
-			if (CookedPlatforms[n] == Platform)
-			{
-				return CookSucceeded[n] ? ECookResult::Succeeded : ECookResult::Failed;
-			}
+			return PlatformData->bCookSucceeded ? ECookResult::Succeeded : ECookResult::Failed;
 		}
 		return ECookResult::Unseen;
 	}
@@ -573,7 +610,7 @@ namespace Cook
 	{
 		// It is not valid to enter the request state without requested platforms; it indicates a bug due to e.g.
 		// calling SendToState without UpdateRequestData from Idle
-		check(RequestedPlatforms.Num() > 0);
+		check(GetNumRequestedPlatforms() > 0);
 	}
 
 	void FPackageData::OnExitRequest()
@@ -905,7 +942,7 @@ namespace Cook
 
 	void FPackageData::OnRemoveSessionPlatform(const ITargetPlatform* Platform)
 	{
-		RequestedPlatforms.RemoveSwap(Platform);
+		PlatformDatas.Remove(Platform);
 	}
 
 	bool FPackageData::HasReferencedObjects() const
@@ -915,8 +952,36 @@ namespace Cook
 
 	void FPackageData::RemapTargetPlatforms(const TMap<ITargetPlatform*, ITargetPlatform*>& Remap)
 	{
-		RemapArrayElements(RequestedPlatforms, Remap);
-		RemapArrayElements(CookedPlatforms, Remap);
+		typedef TSortedMap<const ITargetPlatform*, FPlatformData> MapType;
+		MapType NewPlatformDatas;
+		NewPlatformDatas.Reserve(PlatformDatas.Num());
+		for (TPair<const ITargetPlatform*, FPlatformData>& ExistingPair : PlatformDatas)
+		{
+			ITargetPlatform* NewKey = Remap[ExistingPair.Key];
+			NewPlatformDatas.FindOrAdd(NewKey) = MoveTemp(ExistingPair.Value);
+		}
+
+		// The save state (and maybe more in the future) depend on the order of the request platforms remaining
+		// unchanged, due to CookedPlatformDataNextIndex. If we change that order due to the remap, we need to
+		// demote back to request.
+		if (IsInProgress() && GetState() != EPackageState::Request)
+		{
+			bool bDemote = true;
+			MapType::TConstIterator OldIter = PlatformDatas.CreateConstIterator();
+			MapType::TConstIterator NewIter = NewPlatformDatas.CreateConstIterator();
+			for (; OldIter; ++OldIter, ++NewIter)
+			{
+				if (OldIter.Key() != NewIter.Key())
+				{
+					bDemote = true;
+				}
+			}
+			if (bDemote)
+			{
+				SendToState(EPackageState::Request, ESendFlags::QueueAddAndRemove);
+			}
+		}
+		PlatformDatas = MoveTemp(NewPlatformDatas);
 	}
 
 	bool FPackageData::IsSaveInvalidated() const
@@ -1346,7 +1411,7 @@ namespace Cook
 		check(NumPreloadAllocated >= 0);
 	}
 
-	void FPackageDataMonitor::OnCookedPlatformAdded(FPackageData& PackageData)
+	void FPackageDataMonitor::OnFirstCookedPlatformAdded(FPackageData& PackageData)
 	{
 		if (!PackageData.GetMonitorIsCooked())
 		{
@@ -1355,15 +1420,12 @@ namespace Cook
 		}
 	}
 
-	void FPackageDataMonitor::OnCookedPlatformRemoved(FPackageData& PackageData)
+	void FPackageDataMonitor::OnLastCookedPlatformRemoved(FPackageData& PackageData)
 	{
-		if (PackageData.GetNumCookedPlatforms() == 0)
+		if (PackageData.GetMonitorIsCooked())
 		{
-			if (PackageData.GetMonitorIsCooked())
-			{
-				--NumCooked;
-				PackageData.SetMonitorIsCooked(false);
-			}
+			--NumCooked;
+			PackageData.SetMonitorIsCooked(false);
 		}
 	}
 
@@ -1586,23 +1648,25 @@ namespace Cook
 			}
 
 			FPackageData** PackageDataMapAddr = FileNameToPackageData.Find(FileName);
+			FPackageData* PackageData = nullptr;
 			if (PackageDataMapAddr != nullptr)
 			{
-				(*PackageDataMapAddr)->AddCookedPlatforms({ TargetPlatform }, true);
-				return;
+				PackageData = *PackageDataMapAddr;
 			}
+			else
+			{
+				// we expect all of these to be in PackageNameCache from earlier in the AR load process
+				const FName* PackageName = PackageNameCache.FindExistingCachedPackageNameFromStandardFileName(FileName);
+				checkf(PackageName != nullptr, TEXT("AddExistingPackageDatasForPlatform should only be used when the PackageCache has already been filled out with the packages"));
 
-			// we expect all of these to already be in the PackageNameCache from earlier in the AR load process
-			const FName* PackageName = PackageNameCache.FindExistingCachedPackageNameFromStandardFileName(FileName);
-			checkf(PackageName != nullptr, TEXT("AddExistingPackageDatasForPlatform should only be used when the PackageCache has already been filled out with the packages"));
-
-			// create the package data and remember it for updating caches after the the ParallelFor
-			FPackageData* PackageData = new FPackageData(*this, *PackageName, FileName);
-			NewPackageDataObjects[Index] = PackageData;
-			PackageData->AddCookedPlatforms({ TargetPlatform }, true /* Succeeded */);
+				// create the package data and remember it for updating caches after the the ParallelFor
+				PackageData = new FPackageData(*this, *PackageName, FileName);
+				NewPackageDataObjects[Index] = PackageData;
+			}
+			PackageData->SetPlatformCooked(TargetPlatform, true /* Succeeded */);
 		});
 
-		// update cache for all newly created objects (taken fropm CreatePackageData)
+		// update cache for all newly created objects (taken from CreatePackageData)
 		for (FPackageData* PackageData : NewPackageDataObjects)
 		{
 			if (PackageData)
@@ -1709,7 +1773,7 @@ namespace Cook
 	{
 		for (FPackageData* PackageData : PackageDatas)
 		{
-			PackageData->ClearCookedPlatforms();
+			PackageData->SetPlatformsNotCooked();
 		}
 	}
 
diff --git a/Engine/Source/Editor/UnrealEd/Private/Cooker/CookPackageData.h b/Engine/Source/Editor/UnrealEd/Private/Cooker/CookPackageData.h
--- a/Engine/Source/Editor/UnrealEd/Private/Cooker/CookPackageData.h
+++ b/Engine/Source/Editor/UnrealEd/Private/Cooker/CookPackageData.h
@@ -9,6 +9,7 @@
 #include "Containers/ArrayView.h"
 #include "Containers/Map.h"
 #include "Containers/RingBuffer.h"
+#include "Containers/SortedMap.h"
 #include "Containers/StringFwd.h"
 #include "Containers/UnrealString.h"
 #include "HAL/Platform.h"
@@ -70,20 +71,20 @@ namespace Cook
 	 * structure is never deallocated or moved for a given package; it is deallocated only when the CookOnTheFlyServer
 	 * is destroyed.
 	 *
-	 * RequestedPlatforms - RequestedPlatforms are the Platforms that will be Saved during cooking.
-	 *   RequestedPlatforms is only non-empty when the Package is InProgress (e.g. is in the RequestQueue,
-	 *   LoadReadyQueue, or other state containers). Once a Package finishes cooking, the RequestedPlatforms are
-	 *   cleared, and may be set again later if the Package is requested for another platform. If multiple requests
-	 *   occur at the same time, their RequestedPlatforms are merged. Since modifying the RequestedPlatforms can put
+	 * Requested platforms - Requested platforms are the Platforms that will be Saved during cooking.
+	 *   Platforms are only requested on the PackageData when the Package is InProgress (e.g. is in the RequestQueue,
+	 *   LoadReadyQueue, or other state containers). Once a Package finishes cooking, the platforms are set to
+	 *   unrequested, and may be set again later if the Package is requested for another platform. If multiple requests
+	 *   occur at the same time, their requested platforms are merged. Since modifying the requested platforms can put
 	 *   an InProgress package into an invalid state, direct write access is private; use UpdateRequestData or
 	 *   SetRequestData to write.
 
-	 * CookedPlatforms - CookedPlatforms are the platforms that have already been saved in the lifetime of the 
+	 * Cooked platforms - Cooked platforms are the platforms that have already been saved in the lifetime of the 
 	 *   CookOnTheFlyServer; note this extends outside of the current CookOnTheFlyServer session.
-	 *   CookedPlatforms also store a flag indicating whether the cook for the given platform was successful or not.
-	 *   CookedPlatforms can be added to a PackageData for reasons other than normal Save, such as when a Package is
+	 *   Cooked platforms also store a flag indicating whether the cook for the given platform was successful or not.
+	 *   Cooked platforms can be added to a PackageData for reasons other than normal Save, such as when a Package is
 	 *   detected as not applicable for a Platform and its "Cooked" operation is therefore a noop.
-	 *   CookedPlatforms can be cleared for a PackageData if the Package is modified, e.g. during editor operations
+	 *   Cooked platforms can be cleared for a PackageData if the Package is modified, e.g. during editor operations
 	 *   when the CookOnTheFlyServer is run from the editor.
 	 *
 	 * Package - The package pointer corresponding to this PackageData.
@@ -108,17 +109,29 @@ namespace Cook
 	 *   BeginCacheForCookedPlatformData called on it for the current PackageSave. This field is only non-zero during
 	 *   the save state; it is cleared when successfully or unsucessfully leaving the save state. The field is an
 	 *   object-major index into the two-dimensional array given by
-	 *       {Each Object in GetCachedObjectsInOuter} x {each Platform in RequestedPlatforms}
-	 *   e.g. the index into GetCachedObjectsInOuter is GetCookedPlatformDataNextIndex() / RequestedPlatforms.Num()
-	 *   and the index into RequestedPlatforms is GetCookedPlatformDataNextIndex() % RequestedPlatforms.Num()
+	 *       {Each Object in GetCachedObjectsInOuter} x {each Platform that is requested}
+	 *   e.g. the index into GetCachedObjectsInOuter is GetCookedPlatformDataNextIndex() / GetNumRequestedPlatforms()
+	 *   and the index into GetRequestedPlatforms() is GetCookedPlatformDataNextIndex() % GetNumRequestedPlatforms()
 	 *
 	 * Other fields with explanation inline
 	*/
 	struct FPackageData
 	{
 	public:
+		/** Data about a platform that has been interacted with (requested, etc) by the cook of a PackageData. */
+		struct FPlatformData
+		{
+			FPlatformData();
+			bool bRequested : 1;
+			bool bCookAttempted : 1;
+			bool bCookSucceeded : 1;
+			bool bTargetDomainLookupAttempted : 1;
+			bool bExistsInTargetDomain : 1;
+		};
+
 		FPackageData(FPackageDatas& PackageDatas, const FName& InPackageName, const FName& InFileName);
 		~FPackageData();
+
 		/**
 		 * ClearReferences is called on every PackageData before any packageDatas are deleted,
 		 * so references are still valid during ClearReferences
@@ -139,13 +152,34 @@ namespace Cook
 		 */
 		const FName& GetFileName() const;
 
-		/** Get the current set of RequestedPlatforms. */
-		const TArray<const ITargetPlatform*>& GetRequestedPlatforms() const;
+		/** Reset OutPlatforms and copy current set of requested platforms into it. */
+		template <typename ArrayType>
+		void GetRequestedPlatforms(ArrayType& OutPlatforms) const
+		{
+			OutPlatforms.Reset(PlatformDatas.Num());
+			for (const TPair<const ITargetPlatform*, FPlatformData>& Pair : PlatformDatas)
+			{
+				if (Pair.Value.bRequested)
+				{
+					OutPlatforms.Add(Pair.Key);
+				}
+			}
+		}
+
+		/** Get the number of currently requested platforms. */
+		int32 GetNumRequestedPlatforms() const;
+
 		/**
-		 * Return true if and only if every element of Platforms is present in RequestedPlatforms.
+		 * Return true if and only if every element of Platforms is currently requested.
 		 * Returns true if Platforms is empty.
 		 */
-		bool ContainsAllRequestedPlatforms(const TArrayView<const ITargetPlatform* const>& Platforms) const;
+		bool HasAllRequestedPlatforms(const TArrayView<const ITargetPlatform* const>& Platforms) const;
+
+		/** Return whether all Platforms that have been requested in this PackageData have also been cooked
+		 * 
+		 * @param bAllowFailedCooks Iff false, failed cooks cause function to return false.
+		 */
+		bool AreAllRequestedPlatformsCooked(bool bAllowFailedCooks) const;
 		/**
 		 * Get the flag for whether this InProgress PackageData has been marked as an urgent request
 		 * (e.g. because it has been requested from the game client during cookonthefly.).
@@ -155,26 +189,22 @@ namespace Cook
 		 */
 		bool GetIsUrgent() const { return static_cast<bool>(bIsUrgent); }
 		/**
-		 * Add the given data, which are all fields describing a request for the cook of this PackageData's Package,
-		 * onto the existing request data.
-		 * If the PackageData is not already in progress, this will be equivalent to calling SetRequestData.
-		 * If the PackageData is in progress and the given fields invalidate some of its inprogress state, the
-		 * PackageData's current progress will be canceled and it will be demoted back to an earlier state where
-		 * the changes can be made.
-		 * Once the changes can be made without invalidating the PackageData's current state, the given request data
-		 * will be added on to the PackageData's current request data; the new request data will be the union of the
-		 * two previous sets of data.
+		 * Add the given request data onto the existing request data.
+		 * If the PackageData is not in progress, just calls SetRequestData; skipped if InRequestedPlatforms is empty.
+		 * If the PackageData is in progress it is demoted back to an earlier state if necessary.
+		 * RequestData is set to the union of the previous and new values.
 		 
-		 * @param SendFlags If the PackageData needs to be moved to a new state, these SendFlags will be used for
-		 * whether to remove and add from the containers for the states.
+		 * @param SendFlags If the PackageData needs to change state, specify whether to remove/add to state containers.
 		 */
-		void UpdateRequestData(const TArrayView<const ITargetPlatform* const>& InRequestedPlatforms, bool bInIsUrgent,
+		void UpdateRequestData(const TConstArrayView<const ITargetPlatform*> InRequestedPlatforms, bool bInIsUrgent,
 			FCompletionCallback&& InCompletionCallback, ESendFlags SendFlags = ESendFlags::QueueAddAndRemove);
-		/**
-		 * Set the given data, which are all fields describing a request for the cook of this PackageData's Package,
-		 * onto the existing request data. It is invalid to call this on a PackageData that is already InProgress.
+		/*
+		 * Set the given data onto the existing request data.
+		 * It is invalid to call with empty InRequestedPlatforms.
+		 * It is invalid to call on a PackageData that is already InProgress. 
 		 */
-		void SetRequestData(const TArrayView<const ITargetPlatform* const>& InRequestedPlatforms, bool bInIsUrgent, FCompletionCallback&& InCompletionCallback);
+		void SetRequestData(const TArrayView<const ITargetPlatform* const>& InRequestedPlatforms, bool bInIsUrgent,
+			FCompletionCallback&& InCompletionCallback);
 		/**
 		 * Clear all the inprogress variables from the current PackageData. It is invalid to call this except when
 		 * the PackageData is transitioning out of InProgress.
@@ -182,67 +212,56 @@ namespace Cook
 		void ClearInProgressData();
 
 		/**
-		 * Add each element of New to CookedPlatforms if it is not already present, with the given succeeded flag.
-		 * New and Succeeded must be the same length; the succeeded flag for New[n] is Succeeded[n].
-		 * If a platform is already present in CookedPlatforms, its success flag is overwritten.
-		 */
-		void AddCookedPlatforms(const TArrayView<const ITargetPlatform* const>& New,
-			const TArrayView<const bool>& Succeeded);
-		/**
-		 * Add each element of New to CookedPlatforms if it is not already present, succeeded set to bSucceeded.
-		 * If a platform is already present in CookedPlatforms, its success flag is overwritten.
+		 * FindOrAdd each TargetPlatform and set its flags: CookAttempted=true, Succeeded=<given>.
+		 * In version that takes two arrays, TargetPlatforms and Succeeded must be the same length.
 		 */
-		void AddCookedPlatforms(const TArrayView<const ITargetPlatform* const>& New, bool bSucceeded);
+		void SetPlatformsCooked(const TConstArrayView<const ITargetPlatform*> TargetPlatforms,
+			const TConstArrayView<bool> Succeeded);
+		void SetPlatformsCooked(const TConstArrayView<const ITargetPlatform*> TargetPlatforms, bool bSucceeded);
+		void SetPlatformCooked(const ITargetPlatform* TargetPlatform, bool bSucceeded);
 		/**
-		 * Add TargetPlatform to CookedPlatforms if it is not already present, succeeded set to bSucceeded.
-		 * If TargetPlatform is already present in CookedPlatforms, its success flag is overwritten.
+		 * FindOrAdd each TargetPlatform and set its flags: CookAttempted=false.
+		 * In Version that takes no TargetPlatform, CookAttempted is cleared from all existing platforms.
 		 */
-		void AddCookedPlatform(const ITargetPlatform* TargetPlatform, bool bSucceeded);
-		/** Remove the given Platform and its succeeded flag from CookedPlatforms if it exists. */
-		void RemoveCookedPlatform(const ITargetPlatform* Platform);
-		/** Remove each element of Platforms and its succeeded flags from CookedPlatforms if it exists. */
-		void RemoveCookedPlatforms(const TArrayView<const ITargetPlatform* const>& Platforms);
-		/** Remove all platforms and their succeeded flags from CookedPlatforms. */
-		void ClearCookedPlatforms();
-		/** Return the array of CookedPlatforms as read-only. */
-		const TArray<const ITargetPlatform*>& GetCookedPlatforms() const;
-		/** Return the number of CookedPlatforms. */
-		int32 GetNumCookedPlatforms() const;
-		/** Return true if and only if there is at least one element in CookedPlatforms */
+		void SetPlatformsNotCooked(const TConstArrayView<const ITargetPlatform*> TargetPlatforms);
+		void SetPlatformsNotCooked();
+		void SetPlatformNotCooked(const ITargetPlatform* TargetPlatform);
+
+		/** Access the information about platforms interacted with by *this. */
+		const TSortedMap<const ITargetPlatform*, FPlatformData>& GetPlatformDatas() const;
+
+		/** Return true if and only if at least one platform has been cooked. */
 		bool HasAnyCookedPlatform() const;
 		/**
-		 * Return true if and only if at least one element of Platforms is present in CookedPlatforms, and with its
+		 * Return true if and only if at least one element of Platforms has been cooked, and with its
 		 * succeeded flag set to true if bIncludeFailed is false. Returns false if Platforms is empty.
 		 */
 		bool HasAnyCookedPlatforms(const TArrayView<const ITargetPlatform* const>& Platforms, bool bIncludeFailed) const;
 		/**
-		 * Return true if and only if ever element of Platforms is present in CookedPlatforms, and with its succeeded
+		 * Return true if and only if every element of Platforms has been cooked, and with its succeeded
 		 * flag set to true if bIncludeFailed is false. Returns true if Platforms is empty.
 		 */
 		bool HasAllCookedPlatforms(const TArrayView<const ITargetPlatform* const>& Platforms, bool bIncludeFailed) const;
 		/**
-		 * Return true if and only if the given Platform is present in CookedPlatforms, and with its succeeded flag
+		 * Return true if and only if the given Platform has been cooked, and with its succeeded flag
 		 * set to true if bIncludeFailed is false.
 		 */
 		bool HasCookedPlatform(const ITargetPlatform* Platform, bool bIncludeFailed) const;
 		/**
-		 * Return the CookResult for the given platform.  If the platform does not exist in CookedPlatforms,
+		 * Return the CookResult for the given platform.  If the platform has not been cooked,
 		 * returns ECookResult::Unseen, otherwise returns Succeeded or Failed depending on its success flag.
 		 */
 		ECookResult GetCookResults(const ITargetPlatform* Platform) const;
-		/**
-		 * Empties and then sets OutPlatforms to contain all elements of QueryPlatforms that are not present in
-		 * CookedPlatforms.
-		 */
+		/** Empties and then sets OutPlatforms to contain all platforms that have been requested and not cook-attempted. */
 		template <typename ArrayType>
-		void GetUncookedPlatforms(const TArrayView<const ITargetPlatform* const>& QueryPlatforms, ArrayType& OutPlatforms)
+		void GetUncookedPlatforms(ArrayType& OutPlatforms)
 		{
-			OutPlatforms.Reset();
-			for (const ITargetPlatform* Platform : QueryPlatforms)
+			OutPlatforms.Reset(PlatformDatas.Num());
+			for (const TPair<const ITargetPlatform*, FPlatformData>& Pair : PlatformDatas)
 			{
-				if (!CookedPlatforms.Contains(Platform))
+				if (Pair.Value.bRequested && !Pair.Value.bCookAttempted)
 				{
-					OutPlatforms.Add(Platform);
+					OutPlatforms.Add(Pair.Key);
 				}
 			}
 		}
@@ -415,11 +434,9 @@ namespace Cook
 	private:
 		friend struct UE::Cook::FPackageDatas;
 
-		/** Set the RequestedPlatforms, clearing the old values and replacing them with Platforms. */
-		void SetRequestedPlatforms(const TArrayView<const ITargetPlatform* const>& Platforms);
-		/** Add each element of New to RequestedPlatforms if it is not already present. */
-		void AddRequestedPlatforms(const TArrayView<const ITargetPlatform* const>& New);
-		/** Remove all elements from RequestedPlatforms. */
+		/** FindOrAdd each TargetPlatform and set its bRequested value. */
+		void SetPlatformsRequested(const TConstArrayView<const ITargetPlatform*> TargetPlatforms, bool bRequested);
+		/** Set all platforms to not requested. */
 		void ClearRequestedPlatforms();
 		void SetIsUrgent(bool Value);
 
@@ -481,11 +498,8 @@ namespace Cook
 
 		TUniquePtr<FGeneratorPackage> GeneratorPackage;
 		FGeneratorPackage* GeneratedOwner;
-		TArray<const ITargetPlatform*> RequestedPlatforms;
-		/** Platform part of the CookedPlatforms set.Always the same length as CookSucceeded. */
-		TArray<const ITargetPlatform*> CookedPlatforms;
-		/** Success flag part of the CookedPlatforms set.Always the same length as CookedPlatforms. */
-		TArray<bool> CookSucceeded;
+		/** Data for each platform that has been interacted with by *this. */
+		TSortedMap<const ITargetPlatform*, FPlatformData> PlatformDatas;
 		TArray<FWeakObjectPtr> CachedObjectsInOuter;
 		FCompletionCallback CompletionCallback;
 		FName PackageName;
@@ -738,10 +752,10 @@ namespace Cook
 		/** Callback called from FPackageData when it transitions to or from inprogress. */
 		void OnInProgressChanged(FPackageData& PackageData, bool bInProgress);
 		void OnPreloadAllocatedChanged(FPackageData& PackageData, bool bPreloadAllocated);
-		/** Callback called from FPackageData when it has added a platform to its CookedPackages list. */
-		void OnCookedPlatformAdded(FPackageData& PackageData);
-		/** Callback called from FPackageData when it has removed a platform from its CookedPackages list. */
-		void OnCookedPlatformRemoved(FPackageData& PackageData);
+		/** Callback called from FPackageData when it has set a platform to CookAttempted=true and it does not have any others cooked. */
+		void OnFirstCookedPlatformAdded(FPackageData& PackageData);
+		/** Callback called from FPackageData when it has set a platform to CookAttempted=false and it does not have any others cooked. */
+		void OnLastCookedPlatformRemoved(FPackageData& PackageData);
 		/** Callback called from FPackageData when it has changed its urgency. */
 		void OnUrgencyChanged(FPackageData& PackageData);
 		/** Callback called from FPackageData when it has changed its state. */
@@ -800,7 +814,7 @@ namespace Cook
 
 	/*
 	 * Class that manages the list of all PackageDatas for a CookOnTheFlyServer. PackageDatas is an associative
-	 * array for extra data about a package (e.g. the RequestedPlatforms) that is needed by the CookOnTheFlyServer.
+	 * array for extra data about a package (e.g. the requested platforms) that is needed by the CookOnTheFlyServer.
 	 * FPackageData are allocated once and never destroyed or moved until the CookOnTheFlyServer is destroyed.
 	 * Memory on the FPackageData is allocated and deallocated as necessary for its current state.
 	 * FPackageData are mapped by PackageName and by FileName.
@@ -921,7 +935,7 @@ namespace Cook
 		 * For performance reasons, should only be called on destruction.
 		 */
 		void Clear();
-		/** Remove all CookedPlatforms from all PackageDatas. Used to e.g. invalidate previous cooks. */
+		/** Set all platforms to not cooked in all PackageDatas. Used to e.g. invalidate previous cooks. */
 		void ClearCookedPlatforms();
 		/** Remove all request data about the given platform from all PackageDatas and other memory used by *this. */
 		void OnRemoveSessionPlatform(const ITargetPlatform* TargetPlatform);
