Fix static analysis warning for null check

#rnx

diff --git a/Engine/Source/Editor/UnrealEd/Private/Kismet2/BlueprintEditorUtils.cpp b/Engine/Source/Editor/UnrealEd/Private/Kismet2/BlueprintEditorUtils.cpp
--- a/Engine/Source/Editor/UnrealEd/Private/Kismet2/BlueprintEditorUtils.cpp
+++ b/Engine/Source/Editor/UnrealEd/Private/Kismet2/BlueprintEditorUtils.cpp
@@ -650,7 +650,7 @@ void FBlueprintEditorUtils::RefreshExternalBlueprintDependencyNodes(UBlueprint*
 						{
 							bShouldRefresh |= OwnerClass &&
 								(OwnerClass->IsChildOf(RefreshOnlyChild) || OwnerClass->GetAuthoritativeClass()->IsChildOf(RefreshOnlyChild));
-							if (!bShouldRefresh && Struct->IsA<UFunction>() && OwnerClass->HasAnyClassFlags(CLASS_Interface))
+							if (!bShouldRefresh && OwnerClass && Struct->IsA<UFunction>() && OwnerClass->HasAnyClassFlags(CLASS_Interface))
 							{
 								if (UClass* RefreshClass = Cast<UClass>(RefreshOnlyChild))
 								{
