Sequencer: On import, test for tangent equality before multiplying by key time diff to determine if user or broken tangents.

diff --git a/Engine/Source/Editor/UnrealEd/Private/Fbx/FbxCurveDataImport.cpp b/Engine/Source/Editor/UnrealEd/Private/Fbx/FbxCurveDataImport.cpp
--- a/Engine/Source/Editor/UnrealEd/Private/Fbx/FbxCurveDataImport.cpp
+++ b/Engine/Source/Editor/UnrealEd/Private/Fbx/FbxCurveDataImport.cpp
@@ -259,6 +259,7 @@ namespace UnFbx {
 				float ArriveTangentWeight = DefaultCurveWeight;
 				float ArriveTimeDiff = 0.f;
 				float LeaveTimeDiff = 0.f;
+				bool bEqualTangents = false;
 
 				switch (KeyInterpMode)
 				{
@@ -276,6 +277,8 @@ namespace UnFbx {
 						float LeftTangent = FbxCurve->KeyGetLeftDerivative(KeyIndex);
 						float RightTangent = FbxCurve->KeyGetRightDerivative(KeyIndex);
 						
+						bEqualTangents = FMath::IsNearlyEqual(LeftTangent, RightTangent);
+
 						if (KeyIndex > 0)
 						{
 							ArriveTimeDiff = CurKey.GetTime().GetSecondDouble() - FbxCurve->KeyGetTime(KeyIndex - 1).GetSecondDouble();
@@ -298,7 +301,7 @@ namespace UnFbx {
 				{
 					//no longer use KeyTangentMode & FbxAnimCurveDef::eTangentGenericBreak to see if tangent is broken since it maya broken/unify is just a manipulation state
 					//instead we manually check to see if the tangents are the same or different, if different then broken.
-					if (FMath::IsNearlyEqual(ArriveTangent, LeaveTangent))
+					if (bEqualTangents)
 					{
 						NewTangentMode = RCTM_User;
 					}
