Revert UEdMode.cpp changes from revisions 24 and 25. The undo of revision 24 did not put the file back in the same state it was in prior to the rolled back change.

diff --git a/Engine/Source/Editor/UnrealEd/Private/Tools/UEdMode.cpp b/Engine/Source/Editor/UnrealEd/Private/Tools/UEdMode.cpp
--- a/Engine/Source/Editor/UnrealEd/Private/Tools/UEdMode.cpp
+++ b/Engine/Source/Editor/UnrealEd/Private/Tools/UEdMode.cpp
@@ -96,9 +96,9 @@ void UEdMode::RegisterTool(TSharedPtr<FUICommandInfo> UICommand, FString ToolIde
 	CommandList->MapAction(UICommand,
 		FExecuteAction::CreateUObject(ToolsContext.Get(), &UEdModeInteractiveToolsContext::StartTool, ToolIdentifier),
 		FCanExecuteAction::CreateWeakLambda(ToolsContext.Get(), [this, ToolIdentifier]() {
-		return ShouldToolStartBeAllowed(ToolIdentifier) &&
-			ToolsContext->ToolManager->CanActivateTool(EToolSide::Mouse, ToolIdentifier);
-	}),
+			return ShouldToolStartBeAllowed(ToolIdentifier) && 
+				 ToolsContext->ToolManager->CanActivateTool(EToolSide::Mouse, ToolIdentifier);
+			}),
 		FIsActionChecked::CreateUObject(ToolsContext.Get(), &UEdModeInteractiveToolsContext::IsToolActive, EToolSide::Mouse, ToolIdentifier),
 		EUIActionRepeatMode::RepeatDisabled);
 
@@ -118,18 +118,13 @@ void UEdMode::Exit()
 		SettingsObject->SaveConfig();
 	}
 
-	//Tools can live without toolkit
-	for (auto& RegisteredTool : RegisteredTools)
-	{
-		ToolsContext->ToolManager->UnregisterToolType(RegisteredTool.Value);
-	}
-
 	if (Toolkit.IsValid())
 	{
 		const TSharedRef<FUICommandList>& CommandList = Toolkit->GetToolkitCommands();
 		for (auto& RegisteredTool : RegisteredTools)
 		{
 			CommandList->UnmapAction(RegisteredTool.Key);
+			ToolsContext->ToolManager->UnregisterToolType(RegisteredTool.Value);
 		}
 		FToolkitManager::Get().CloseToolkit(Toolkit.ToSharedRef());
 		Toolkit.Reset();
