Fix for "Golem does not play animations when running PIE for the first time"
The issue is that in BlueprintEditorUtils::UpdateStalePinWatches(), there is an uninitialized local variable. It was causing BP loading process to mark clean BPs as dirty, which then got treated as an error condition in AnimInstanceProxy::Initialization(), causing it to set RootNode of the proxy as nullptr. Consequently, FAnimInstanceProxy::EvaluateAnimationNode_WithRoot() just sets the output pose to the ref pose

diff --git a/Engine/Source/Editor/UnrealEd/Private/Kismet2/BlueprintEditorUtils.cpp b/Engine/Source/Editor/UnrealEd/Private/Kismet2/BlueprintEditorUtils.cpp
--- a/Engine/Source/Editor/UnrealEd/Private/Kismet2/BlueprintEditorUtils.cpp
+++ b/Engine/Source/Editor/UnrealEd/Private/Kismet2/BlueprintEditorUtils.cpp
@@ -7553,7 +7553,7 @@ void FBlueprintEditorUtils::UpdateTransactionalFlags(UBlueprint* Blueprint)
 void FBlueprintEditorUtils::UpdateStalePinWatches( UBlueprint* Blueprint )
 {
 	TSet<UEdGraphPin*> AllPins;
-	uint16 WatchCount;
+	uint16 WatchCount = 0;
 	
 	// Find all unique pins being watched
 	FKismetDebugUtilities::ForeachPinWatch(
