EditorDomain: Optimized the build input resolver when priority is Blocking

- Changed SetPayloads to wait on payloads explicitly when Priority is Blocking.
- Changed SetPayloads to invoke the callback without Begin/End calls on the Owner when payloads are complete by the end of the function.
- Fixed the callback arguments to use EStatus::Error when resolving failed.
- Added an appropriate memory order to every atomic operation.

#rnx

diff --git a/Engine/Source/Editor/UnrealEd/Private/EditorDomain/EditorBuildInputResolver.cpp b/Engine/Source/Editor/UnrealEd/Private/EditorDomain/EditorBuildInputResolver.cpp
--- a/Engine/Source/Editor/UnrealEd/Private/EditorDomain/EditorBuildInputResolver.cpp
+++ b/Engine/Source/Editor/UnrealEd/Private/EditorDomain/EditorBuildInputResolver.cpp
@@ -70,9 +70,6 @@ private:
 	// Data written during creation that is read-only during Cancel and CallOnResolved
 	FBuildDefinition Definition;
 	IRequestOwner& Owner;
-#if DO_CHECK
-	bool bCreationComplete = false;
-#endif
 
 	/**
 	 * The array of data that is written by the TFutures and passed into the IBuildInputResolver callback function.
@@ -88,6 +85,8 @@ private:
 	 */
 	TArray<FPayloadData> PayloadDatas;
 
+	/** Event for when the callback has finished executing; some functions must not return until it is done. */
+	FEventRef CallbackComplete;
 	/**
 	 * Tracks whether all TFutures are complete; the last one to complete calls the Resolve callback.
 	 * Value starts at 1 and is decremented when creation completes, so that if all TFutures are already complete
@@ -98,8 +97,9 @@ private:
 	std::atomic<bool> bAllSucceeded{ true };
 	/** Atomic to allow Cancel or CallOnResolved to resolve the race to see which one will call the callback. */
 	std::atomic<bool> bCallbackCalled{ false };
-	/** Event for when the callback has finished executing; some functions must not return until it is done. */
-	FEventRef CallbackComplete;
+#if DO_CHECK
+	bool bCreationComplete = false;
+#endif
 
 };
 typedef TResolveInputRequest<UE::BulkDataRegistry::FMetaData, FBuildInputMetaByKey, FOnBuildInputMetaResolved> FResolveInputMetaRequest;
@@ -211,7 +211,7 @@ template <class BulkPayloadType, class CallbackPayloadType, class CallbackType>
 void TResolveInputRequest<BulkPayloadType, CallbackPayloadType, CallbackType>::Cancel()
 {
 	check(bCreationComplete);
-	if (bCallbackCalled.exchange(true) == false)
+	if (bCallbackCalled.exchange(true, std::memory_order_relaxed) == false)
 	{
 		Owner.End(this, [this]
 		{
@@ -229,7 +229,7 @@ template <class BulkPayloadType, class CallbackPayloadType, class CallbackType>
 void TResolveInputRequest<BulkPayloadType, CallbackPayloadType, CallbackType>::Wait()
 {
 	check(bCreationComplete);
-	if (!bCallbackCalled)
+	if (!bCallbackCalled.load(std::memory_order_relaxed))
 	{
 		for (FPayloadData& PayloadData : PayloadDatas)
 		{
@@ -243,7 +243,7 @@ template <class BulkPayloadType, class CallbackPayloadType, class CallbackType>
 void TResolveInputRequest<BulkPayloadType, CallbackPayloadType, CallbackType>::SetPayloads(
 	TArrayView<TTuple<FString, TFuture<BulkPayloadType>>> InBulkPayloads, EStatus OtherStatus)
 {
-	RemainingInputCount += InBulkPayloads.Num();
+	RemainingInputCount.fetch_add(InBulkPayloads.Num(), std::memory_order_relaxed);
 
 	// Preallocate the Payloads array so that it will not be reallocated while we are adding futures to it
 	// This allows us to have each future write to its assigned element, potentially from another thread,
@@ -273,9 +273,9 @@ void TResolveInputRequest<BulkPayloadType, CallbackPayloadType, CallbackType>::S
 				else
 				{
 					pThis->LogBulkDataError(Payload.Key);
-					pThis->bAllSucceeded = false;
+					pThis->bAllSucceeded.store(false, std::memory_order_relaxed);
 				}
-				if (--pThis->RemainingInputCount == 0)
+				if (pThis->RemainingInputCount.fetch_sub(1, std::memory_order_acq_rel) == 1)
 				{
 					pThis->CallOnResolved();
 				}
@@ -286,25 +286,40 @@ void TResolveInputRequest<BulkPayloadType, CallbackPayloadType, CallbackType>::S
 	check(!bCreationComplete);
 	bCreationComplete = true;
 #endif
-	Owner.Begin(this);
-	if (--RemainingInputCount == 0)
+	const EPriority Priority = Owner.GetPriority();
+	if (Priority == EPriority::Blocking)
 	{
-		CallOnResolved();
+		for (FPayloadData& PayloadData : PayloadDatas)
+		{
+			PayloadData.Future.Wait();
+		}
+	}
+
+	if (Priority == EPriority::Blocking || RemainingInputCount.load(std::memory_order_acquire) == 1)
+	{
+		const EStatus Status = bAllSucceeded.load(std::memory_order_relaxed) ? EStatus::Ok : EStatus::Error;
+		OnResolved({Payloads, Status});
+		CallbackComplete->Trigger();
 	}
-	else if (Owner.GetPriority() == EPriority::Blocking)
+	else
 	{
-		Wait();
+		Owner.Begin(this);
+		if (RemainingInputCount.fetch_sub(1, std::memory_order_acq_rel) == 1)
+		{
+			CallOnResolved();
+		}
 	}
 }
 
 template <class BulkPayloadType, class CallbackPayloadType, class CallbackType>
 void TResolveInputRequest<BulkPayloadType, CallbackPayloadType, CallbackType>::CallOnResolved()
 {
-	if (bCallbackCalled.exchange(true) == false)
+	if (bCallbackCalled.exchange(true, std::memory_order_relaxed) == false)
 	{
 		Owner.End(this, [this]
 		{
-			OnResolved({Payloads, EStatus::Ok});
+			const EStatus Status = bAllSucceeded.load(std::memory_order_relaxed) ? EStatus::Ok : EStatus::Error;
+			OnResolved({Payloads, Status});
 			CallbackComplete->Trigger();
 		});
 	}
