fix animation bug caused by watchpin refactor

diff --git a/Engine/Source/Editor/UnrealEd/Private/Kismet2/BlueprintEditorUtils.cpp b/Engine/Source/Editor/UnrealEd/Private/Kismet2/BlueprintEditorUtils.cpp
--- a/Engine/Source/Editor/UnrealEd/Private/Kismet2/BlueprintEditorUtils.cpp
+++ b/Engine/Source/Editor/UnrealEd/Private/Kismet2/BlueprintEditorUtils.cpp
@@ -7591,8 +7591,6 @@ void FBlueprintEditorUtils::UpdateStalePinWatches( UBlueprint* Blueprint )
 		{
 			FKismetDebugUtilities::AddPinWatch(Blueprint, Pin);
 		}
-
-		Blueprint->Status = BS_Dirty;
 	}
 }
 
