Fix pin color in material editor

diff --git a/Engine/Source/Editor/UnrealEd/Private/MaterialGraph.cpp b/Engine/Source/Editor/UnrealEd/Private/MaterialGraph.cpp
--- a/Engine/Source/Editor/UnrealEd/Private/MaterialGraph.cpp
+++ b/Engine/Source/Editor/UnrealEd/Private/MaterialGraph.cpp
@@ -582,10 +582,12 @@ bool UMaterialGraph::IsInputActive(UEdGraphPin* GraphPin) const
 		RootNode && // No inputs without a root node
 		GraphPin->Direction == EGPD_Input &&
 		GraphPin->PinType.PinCategory != UMaterialGraphSchema::PC_Exec &&
-		GraphPin->SourceIndex < MaterialInputs.Num())
+		GraphPin->SourceIndex < MaterialInputs.Num() &&
+		GraphPin->GetOwningNode() == RootNode // we only care about pins from the root node here
+		) 
 	{
-		
-		return Material->IsPropertyActiveInEditor(MaterialInputs[GraphPin->SourceIndex].GetProperty());
+		const EMaterialProperty Property = MaterialInputs[GraphPin->SourceIndex].GetProperty();
+		return Material->IsPropertyActiveInEditor(Property);
 	}
 	return true;
 }
