EditorDomain: Fix crash in serialize, giving an error instead, when the segment is not fetchable from DDC.
Block log spam on repeated calls to serialize when the archive is in this error state.
#rnx

diff --git a/Engine/Source/Editor/UnrealEd/Private/EditorDomain/EditorDomainArchive.cpp b/Engine/Source/Editor/UnrealEd/Private/EditorDomain/EditorDomainArchive.cpp
--- a/Engine/Source/Editor/UnrealEd/Private/EditorDomain/EditorDomainArchive.cpp
+++ b/Engine/Source/Editor/UnrealEd/Private/EditorDomain/EditorDomainArchive.cpp
@@ -89,7 +89,7 @@ void FEditorDomainPackageSegments::Close()
 	Source = ESource::Closed;
 }
 
-bool FEditorDomainPackageSegments::Serialize(void* V, int64 Length)
+bool FEditorDomainPackageSegments::Serialize(void* V, int64 Length, bool bIsError)
 {
 	// Called from Interface-only
 	check(Source == ESource::Segments); // Caller should not call in any other case
@@ -102,7 +102,9 @@ bool FEditorDomainPackageSegments::Serialize(void* V, int64 Length)
 	}
 	if (End > Size)
 	{
-		UE_LOG(LogEditorDomain, Error, TEXT("Requested read of %" INT64_FMT " bytes when %" INT64_FMT " bytes remain (file = %s, size=%" INT64_FMT ")"),
+		// Log the error if this is the first error. Do not log it on subsequent errors, to prevent spam in e.g. array reading
+		UE_CLOG(!bIsError, LogEditorDomain, Error,
+			TEXT("Requested read of %" INT64_FMT " bytes when %" INT64_FMT " bytes remain (file = %s, size=%" INT64_FMT ")"),
 			Length, Size - Pos, *PackagePath.GetDebugName(), Size);
 		return false;
 	}
@@ -127,6 +129,15 @@ bool FEditorDomainPackageSegments::Serialize(void* V, int64 Length)
 		{
 			bool bIsReady;
 			MRUSegment = EnsureSegmentRange(Pos, End, true, bIsReady);
+			int64 SegmentEnd = MRUSegment < &Segments.Last() ? (MRUSegment + 1)->Start : Size;
+			int64 ExpectedSize = SegmentEnd - MRUSegment->Start;
+			if (MRUSegment->Data.GetSize() != ExpectedSize)
+			{
+				UE_CLOG(!bIsError, LogEditorDomain, Error, TEXT("Package %s is missing cache data in segment %d."),
+					*PackagePath.GetDebugName(), (int32)(MRUSegment - &Segments[0]));
+				MRUSegment = nullptr;
+				return false;
+			}
 		}
 	}
 	return true;
@@ -571,7 +582,7 @@ void FEditorDomainReadArchive::Serialize(void* V, int64 Length)
 		Serialize(V, Length);
 		break;
 	case FEditorDomainPackageSegments::ESource::Segments:
-		if (!Segments.Serialize(V, Length))
+		if (!Segments.Serialize(V, Length, IsError()))
 		{
 			SetError();
 		}
@@ -781,7 +792,7 @@ IAsyncReadRequest* FEditorDomainAsyncReadFileHandle::ReadRequest(int64 Offset, i
 		{
 			Memory = (uint8*)FMemory::Malloc(BytesToRead);
 		}
-		if (!Segments.Serialize(Memory, BytesToRead))
+		if (!Segments.Serialize(Memory, BytesToRead, false /* bIsError */))
 		{
 			if (!UserSuppliedMemory)
 			{
diff --git a/Engine/Source/Editor/UnrealEd/Private/EditorDomain/EditorDomainArchive.h b/Engine/Source/Editor/UnrealEd/Private/EditorDomain/EditorDomainArchive.h
--- a/Engine/Source/Editor/UnrealEd/Private/EditorDomain/EditorDomainArchive.h
+++ b/Engine/Source/Editor/UnrealEd/Private/EditorDomain/EditorDomainArchive.h
@@ -84,7 +84,7 @@ public:
 	 * Copy Length bytes from current Pos into V.
 	 * Must not be called unless GetSource is ESource::Segments.
 	 */
-	bool Serialize(void* V, int64 Length);
+	bool Serialize(void* V, int64 Length, bool bIsError);
 
 	/**
 	 * Ensure requests have been sent for segments touching the given range. Out of range requests are ignored.
