Option to auto solo audio for first PIE client

diff --git a/Engine/Source/Editor/UnrealEd/Classes/Settings/LevelEditorPlaySettings.h b/Engine/Source/Editor/UnrealEd/Classes/Settings/LevelEditorPlaySettings.h
--- a/Engine/Source/Editor/UnrealEd/Classes/Settings/LevelEditorPlaySettings.h
+++ b/Engine/Source/Editor/UnrealEd/Classes/Settings/LevelEditorPlaySettings.h
@@ -255,6 +255,10 @@ public:
 	UPROPERTY(config, EditAnywhere, Category=PlayInEditor, meta=(ToolTip="Whether to play sounds when in a Play In Editor session"))
 	bool EnableGameSound;
 
+	/** Whether to automatically solo audio in first PIE client. */
+	UPROPERTY(config, EditAnywhere, Category = PlayInEditor, meta = (ToolTip="Whether to automatically solo audio in first PIE client", EditCondition="EnableGameSound", DisplayName = "Solo Audio in First PIE Client"))
+	bool SoloAudioInFirstPIEClient;
+
 	/** Whether to play a sound when entering and exiting PIE */
 	UPROPERTY(config, EditAnywhere, AdvancedDisplay, Category = PlayInEditor, meta = (DisplayName = "Enable PIE Enter and Exit Sounds"))
 	bool EnablePIEEnterAndExitSounds;
diff --git a/Engine/Source/Editor/UnrealEd/Private/PlayLevel.cpp b/Engine/Source/Editor/UnrealEd/Private/PlayLevel.cpp
--- a/Engine/Source/Editor/UnrealEd/Private/PlayLevel.cpp
+++ b/Engine/Source/Editor/UnrealEd/Private/PlayLevel.cpp
@@ -675,6 +675,24 @@ void UEditorEngine::TeardownPlaySession(FWorldContext& PieWorldContext)
 		AudioDevice->ResetInterpolation();
 		AudioDevice->OnEndPIE(false); // TODO: Should this have been bWasSimulatingInEditor?
 		AudioDevice->SetTransientMasterVolume(1.0f);
+		// Reset solo audio
+		if (PlayInEditorSessionInfo.IsSet())
+		{
+			ULevelEditorPlaySettings* EditorPlaySettings = PlayInEditorSessionInfo->OriginalRequestParams.EditorPlaySettings;
+			if (EditorPlaySettings && EditorPlaySettings->SoloAudioInFirstPIEClient && GEngine)
+			{
+				if (FWorldContext* WorldContext = GEngine->GetWorldContextFromWorld(PlayWorld))
+				{
+					if (WorldContext->PIEInstance == 0)
+					{
+						if (FAudioDeviceManager* DeviceManager = AudioDevice->GetAudioDeviceManager())
+						{
+							DeviceManager->SetSoloDevice(INDEX_NONE);
+						}
+					}
+				}
+			}
+		}
 	}
 
 	// Clean up all streaming levels
@@ -2836,6 +2854,26 @@ UGameInstance* UEditorEngine::CreateInnerProcessPIEGameInstance(FRequestPlaySess
 				GameInstanceAudioDevice->SetTransientMasterVolume(0.0f);
 			}
 		}
+		if (InParams.EditorPlaySettings->SoloAudioInFirstPIEClient)
+		{
+			if (FAudioDeviceHandle GameInstanceAudioDevice = PlayWorld->GetAudioDevice())
+			{
+				if (GEngine)
+				{
+					if (FAudioDeviceManager* DeviceManager = GEngine->GetAudioDeviceManager())
+					{
+						if (InPIEInstanceIndex == 0)
+						{
+							DeviceManager->SetSoloDevice(GameInstanceAudioDevice->DeviceID);
+						}
+						else
+						{
+							GameInstanceAudioDevice->SetDeviceMuted(true);
+						}
+					}
+				}
+			}
+		}
 
 		GameViewport = ViewportClient;
 		GameViewport->bIsPlayInEditorViewport = true;
