Make hiding a level deselect all selected elements

diff --git a/Engine/Source/Editor/UnrealEd/Private/EditorLevelUtils.cpp b/Engine/Source/Editor/UnrealEd/Private/EditorLevelUtils.cpp
--- a/Engine/Source/Editor/UnrealEd/Private/EditorLevelUtils.cpp
+++ b/Engine/Source/Editor/UnrealEd/Private/EditorLevelUtils.cpp
@@ -53,6 +53,8 @@ EditorLevelUtils.cpp: Editor-specific level management routines
 #include "Dialogs/Dialogs.h"
 
 #include "Algo/AnyOf.h"
+#include "Elements/Framework/EngineElementsLibrary.h"
+#include "Elements/Interfaces/TypedElementWorldInterface.h"
 
 DEFINE_LOG_CATEGORY(LogLevelTools);
 
@@ -1206,18 +1208,22 @@ void SetLevelVisibilityNoGlobalUpdateInternal(ULevel* Level, const bool bShouldB
 	// If the level is being hidden, deselect actors and surfaces that belong to this level. (Part 1/2)
 	if (!bShouldBeVisible && ModifyMode == ELevelVisibilityDirtyMode::ModifyOnChange)
 	{
-		USelection* SelectedActors = GEditor->GetSelectedActors();
-		SelectedActors->Modify();
-		const TArray<AActor*>& Actors = Level->Actors;
-		for (int32 ActorIndex = 0; ActorIndex < Actors.Num(); ++ActorIndex)
+		if (UTypedElementSelectionSet* ActorSelectionSet = GEditor->GetSelectedActors()->GetElementSelectionSet())
 		{
-			AActor* Actor = Actors[ActorIndex];
-			if (Actor)
+			TArray<FTypedElementHandle> LevelElementHandles;
+			
+			ActorSelectionSet->ForEachSelectedElement<UTypedElementWorldInterface>(
+				[Level, &LevelElementHandles](const TTypedElement<UTypedElementWorldInterface>& Element)
 			{
-				SelectedActors->Deselect(Actor);
-			}
+				if (Element.GetOwnerLevel() == Level)
+				{
+					LevelElementHandles.Add(Element);
+				}
+				return true;
+			});
+			ActorSelectionSet->DeselectElements(LevelElementHandles, FTypedElementSelectionOptions());
 		}
-
+		
 		UEditorLevelUtils::DeselectAllSurfacesInLevel(Level);
 	}
 
