External actors: added delegates before and after saving any numbers of external actors.

diff --git a/Engine/Source/Editor/UnrealEd/Private/Editor.cpp b/Engine/Source/Editor/UnrealEd/Private/Editor.cpp
--- a/Engine/Source/Editor/UnrealEd/Private/Editor.cpp
+++ b/Engine/Source/Editor/UnrealEd/Private/Editor.cpp
@@ -117,6 +117,8 @@ FEditorDelegates::FOnPreSaveWorld						FEditorDelegates::PreSaveWorld;
 FEditorDelegates::FOnPostSaveWorld						FEditorDelegates::PostSaveWorld;
 FEditorDelegates::FOnPreSaveWorldWithContext			FEditorDelegates::PreSaveWorldWithContext;
 FEditorDelegates::FOnPostSaveWorldWithContext			FEditorDelegates::PostSaveWorldWithContext;
+FEditorDelegates::FOnPreSaveExternalActors				FEditorDelegates::PreSaveExternalActors;
+FEditorDelegates::FOnPostSaveExternalActors				FEditorDelegates::PostSaveExternalActors;
 FEditorDelegates::FOnFinishPickingBlueprintClass		FEditorDelegates::OnFinishPickingBlueprintClass;
 FEditorDelegates::FOnNewAssetCreation					FEditorDelegates::OnConfigureNewAssetProperties;
 FEditorDelegates::FOnNewAssetCreation					FEditorDelegates::OnNewAssetCreated;
@@ -137,7 +139,7 @@ FEditorDelegates::FOnDollyPerspectiveCamera				FEditorDelegates::OnDollyPerspect
 FSimpleMulticastDelegate								FEditorDelegates::OnShutdownPostPackagesSaved;
 FEditorDelegates::FOnPackageDeleted						FEditorDelegates::OnPackageDeleted;
 FEditorDelegates::FOnAssetsCanDelete					FEditorDelegates::OnAssetsCanDelete;
-FEditorDelegates::FOnAssetsAddExtraObjectsToDelete			FEditorDelegates::OnAssetsAddExtraObjectsToDelete;
+FEditorDelegates::FOnAssetsAddExtraObjectsToDelete		FEditorDelegates::OnAssetsAddExtraObjectsToDelete;
 FEditorDelegates::FOnAssetsPreDelete					FEditorDelegates::OnAssetsPreDelete;
 FEditorDelegates::FOnAssetsDeleted						FEditorDelegates::OnAssetsDeleted;
 FEditorDelegates::FOnAssetDragStarted					FEditorDelegates::OnAssetDragStarted;
diff --git a/Engine/Source/Editor/UnrealEd/Private/FileHelpers.cpp b/Engine/Source/Editor/UnrealEd/Private/FileHelpers.cpp
--- a/Engine/Source/Editor/UnrealEd/Private/FileHelpers.cpp
+++ b/Engine/Source/Editor/UnrealEd/Private/FileHelpers.cpp
@@ -2826,39 +2826,48 @@ EAutosaveContentPackagesResult::Type FEditorFileUtils::AutosaveMapEx(const FStri
 			// Now gather the world external packages and save them if needed
 			if (World->PersistentLevel)
 			{
+				TArray<UPackage*> ExternalPackagesToSave;
 				for (UPackage* ExternalPackage : World->PersistentLevel->GetLoadedExternalActorPackages())
 				{
 					if (ExternalPackage->IsDirty() && (bForceIfNotInList || DirtyPackagesForAutoSave.Contains(ExternalPackage))
 						&& FPackageName::IsValidLongPackageName(ExternalPackage->GetName(), /*bIncludeReadOnlyRoots=*/false))
 					{
 						// Don't try to save external packages that will get deleted
-						bool bActorPackageNeedsToSave = false;
-						ForEachObjectWithPackage(ExternalPackage, [&bActorPackageNeedsToSave](UObject* Object)
+						ForEachObjectWithPackage(ExternalPackage, [ExternalPackage, &ExternalPackagesToSave](UObject* Object)
 						{
 							if (Cast<AActor>(Object))
 							{
 								if (!Object->IsPendingKill())
 								{
-									bActorPackageNeedsToSave = true;
+									ExternalPackagesToSave.Add(ExternalPackage);
 									return false;
 								}
 							}
 							return true;
 						}, false);
+					}
+				}
 
-						if (bActorPackageNeedsToSave)
+				if (ExternalPackagesToSave.Num())
+				{
+					FEditorDelegates::PreSaveExternalActors.Broadcast(World);
+
+					for (UPackage* ExternalPackage : ExternalPackagesToSave)
+					{
+						const FString AutosaveFilename = GetAutoSaveFilename(ExternalPackage, AbsoluteAutosaveDir, AutosaveIndex, FPackageName::GetAssetPackageExtension());
+						if (!GEditor->Exec(nullptr, *FString::Printf(TEXT("OBJ SAVEPACKAGE PACKAGE=\"%s\" FILE=\"%s\" SILENT=false AUTOSAVING=true"), *ExternalPackage->GetName(), *AutosaveFilename)))
 						{
-							const FString AutosaveFilename = GetAutoSaveFilename(ExternalPackage, AbsoluteAutosaveDir, AutosaveIndex, FPackageName::GetAssetPackageExtension());
-							if (!GEditor->Exec(nullptr, *FString::Printf(TEXT("OBJ SAVEPACKAGE PACKAGE=\"%s\" FILE=\"%s\" SILENT=false AUTOSAVING=true"), *ExternalPackage->GetName(), *AutosaveFilename)))
-							{
-								return EAutosaveContentPackagesResult::Failure;
-							}
-							// We saved an actor
-							bResult = true;
-							// Re-mark the package as dirty, because autosaving it will have cleared the dirty flag
-							ExternalPackage->MarkPackageDirty();
+							return EAutosaveContentPackagesResult::Failure;
 						}
+
+						// We saved an actor
+						bResult = true;
+
+						// Re-mark the package as dirty, because autosaving it will have cleared the dirty flag
+						ExternalPackage->MarkPackageDirty();
 					}
+
+					FEditorDelegates::PostSaveExternalActors.Broadcast(World);
 				}
 			}
 		}
@@ -3755,6 +3764,21 @@ FEditorFileUtils::EPromptReturnCode InternalPromptForCheckoutAndSave(const TArra
 		FScopedSlowTask SlowTask(FinalSaveList.Num() * 2, NSLOCTEXT("UnrealEd", "SavingPackagesE", "Saving packages..."));
 		SlowTask.MakeDialog();
 
+		UWorld* ActorsWorld = nullptr;
+		for (UPackage* Package : FinalSaveList)
+		{
+			if (AActor* Actor = AActor::FindActorInPackage(Package))
+			{
+				ActorsWorld = Actor->GetWorld();
+				break;
+			}
+		}
+
+		if (ActorsWorld)
+		{
+			FEditorDelegates::PreSaveExternalActors.Broadcast(ActorsWorld);
+		}
+
 		for (UPackage* Package : FinalSaveList)
 		{
 			SlowTask.EnterProgressFrame(1);
@@ -3807,6 +3831,11 @@ FEditorFileUtils::EPromptReturnCode InternalPromptForCheckoutAndSave(const TArra
 				}
 			}
 		}
+
+		if (ActorsWorld)
+		{
+			FEditorDelegates::PostSaveExternalActors.Broadcast(ActorsWorld);
+		}
 	}
 
 	SaveErrors.Flush();
diff --git a/Engine/Source/Editor/UnrealEd/Public/Editor.h b/Engine/Source/Editor/UnrealEd/Public/Editor.h
--- a/Engine/Source/Editor/UnrealEd/Public/Editor.h
+++ b/Engine/Source/Editor/UnrealEd/Public/Editor.h
@@ -74,8 +74,12 @@ struct UNREALED_API FEditorDelegates
 	DECLARE_MULTICAST_DELEGATE_TwoParams(FOnPreSaveWorldWithContext, class UWorld*, FObjectPreSaveContext);
 	/** delegate type for post save world events ( uint32 SaveFlags, UWorld* World, bool bSuccess ) */
 	DECLARE_MULTICAST_DELEGATE_ThreeParams(FOnPostSaveWorld, uint32, class UWorld*, bool);
-	/** delegate type for post save world events ( UWorld* World, FObjectPostSaveContext ObjectSaveContext, bool bSuccess ) */
+	/** delegate type for post save world events ( UWorld* World, FObjectPostSaveContext ObjectSaveContext ) */
 	DECLARE_MULTICAST_DELEGATE_TwoParams(FOnPostSaveWorldWithContext, class UWorld*, FObjectPostSaveContext);
+	/** delegate type for pre save external actors event, called by editor save codepaths and auto saves ( UWorld* World )*/
+	DECLARE_MULTICAST_DELEGATE_OneParam(FOnPreSaveExternalActors, class UWorld*);
+	/** delegate type for post save external actors event, called by editor save codepaths and auto saves ( UWorld* World )*/
+	DECLARE_MULTICAST_DELEGATE_OneParam(FOnPostSaveExternalActors, class UWorld*);
 	/** delegate for a PIE event (begin, end, pause/resume, etc) (Params: bool bIsSimulating) */
 	DECLARE_MULTICAST_DELEGATE_OneParam(FOnPIEEvent, const bool);
 	/** delegate for a standalone local play event (Params: uint32 processID) */
@@ -225,6 +229,10 @@ struct UNREALED_API FEditorDelegates
 	static FOnPostSaveWorld PostSaveWorld;
 	/** Called after SaveWorld is processed */
 	static FOnPostSaveWorldWithContext PostSaveWorldWithContext;
+	/** Called before any number of external actors will be saved */
+	static FOnPreSaveExternalActors PreSaveExternalActors;
+	/** Called after any number of external actors has been saved */
+	static FOnPostSaveExternalActors PostSaveExternalActors;
 	/** Called when finishing picking a new blueprint class during construction */
 	static FOnFinishPickingBlueprintClass OnFinishPickingBlueprintClass;
 	/** Called when beginning configuration of a new asset */
