Refactor/cleanup surrounding material parameters.
 Primary goal is to reduce amount of code duplicated across different parameter types, and reduce amount of boilerplate required to add new parameters
The main part are new types FMaterialParameterValue, and FMaterialParameterMetadata (both in MaterialTypes.h).  FMaterialParameterValue is a tagged union/variant type, that can represent any material parameter value.  FMaterialParameterMetadata contains a value, along with various editor-only fields that describe how that value is used.
Using this new type, many places that previously included separate APIs for each parameter type collapse to a single API that operates on a FMaterialParameterValue/FMaterialParameterMetadata.  Code that previously had a block copy/pasted for each parameter type can be replaced with a loop that iterates over all parameter types.
Also includes some API cleanup that should improve (editor) performance.  Previously we had APIs GetAllScalarParameters, GetAllVectorParameters, etc, which would return a list of FMaterialParameterInfos for each parameter of the given type.  Code would often query all parameters, then loop over the results and lookup each parameter value.  Now there's a new API, GetAllParametersOfType, which returns a list of all FMaterialParameterMetadata of a given type.  It's easy to just collect the values at the time we're building the list of parameters, so this saves the work of looking up each value individually.

diff --git a/Engine/Source/Developer/MeshMergeUtilities/Private/ProxyMaterialUtilities.h b/Engine/Source/Developer/MeshMergeUtilities/Private/ProxyMaterialUtilities.h
--- a/Engine/Source/Developer/MeshMergeUtilities/Private/ProxyMaterialUtilities.h
+++ b/Engine/Source/Developer/MeshMergeUtilities/Private/ProxyMaterialUtilities.h
@@ -43,7 +43,7 @@ namespace ProxyMaterialUtilities
 
 	static EMaterialParameterType GetConstantParamType(EFlattenMaterialProperties InProperty)
 	{
-		EMaterialParameterType ParamType = EMaterialParameterType::Count;
+		EMaterialParameterType ParamType = EMaterialParameterType::None;
 
 		switch (InProperty)
 		{
@@ -106,7 +106,7 @@ namespace ProxyMaterialUtilities
 									  InBaseMaterial->GetVectorParameterDefaultValue(ConstName, DefaultVectorValue);
 				break;
 
-			case EMaterialParameterType::Count:
+			case EMaterialParameterType::None:
 				MissingNames += UseTexture.ToString();
 				break;
 			}
diff --git a/Engine/Source/Editor/MaterialEditor/Private/MaterialPropertyHelpers.cpp b/Engine/Source/Editor/MaterialEditor/Private/MaterialPropertyHelpers.cpp
--- a/Engine/Source/Editor/MaterialEditor/Private/MaterialPropertyHelpers.cpp
+++ b/Engine/Source/Editor/MaterialEditor/Private/MaterialPropertyHelpers.cpp
@@ -289,113 +289,38 @@ void FMaterialPropertyHelpers::CopyMaterialToInstance(UMaterialInstanceConstant*
 		{
 			ChildInstance->MarkPackageDirty();
 			ChildInstance->ClearParameterValuesEditorOnly();
+
+			FMaterialInstanceParameterUpdateContext UpdateContext(ChildInstance);
+
 			//propagate changes to the base material so the instance will be updated if it has a static permutation resource
-			FStaticParameterSet NewStaticParameters;
 			for (int32 GroupIdx = 0; GroupIdx < ParameterGroups.Num(); GroupIdx++)
 			{
 				FEditorParameterGroup& Group = ParameterGroups[GroupIdx];
 				for (int32 ParameterIdx = 0; ParameterIdx < Group.Parameters.Num(); ParameterIdx++)
 				{
-					if (Group.Parameters[ParameterIdx] == NULL)
-					{
-						continue;
-					}
-					UDEditorScalarParameterValue* ScalarParameterValue = Cast<UDEditorScalarParameterValue>(Group.Parameters[ParameterIdx]);
-					if (ScalarParameterValue)
-					{
-						if (ScalarParameterValue->bOverride)
-						{
-							ChildInstance->SetScalarParameterValueEditorOnly(ScalarParameterValue->ParameterInfo, ScalarParameterValue->ParameterValue);
-							continue;
-						}
-					}
-					UDEditorFontParameterValue* FontParameterValue = Cast<UDEditorFontParameterValue>(Group.Parameters[ParameterIdx]);
-					if (FontParameterValue)
-					{
-						if (FontParameterValue->bOverride)
-						{
-							ChildInstance->SetFontParameterValueEditorOnly(FontParameterValue->ParameterInfo, FontParameterValue->ParameterValue.FontValue, FontParameterValue->ParameterValue.FontPage);
-							continue;
-						}
-					}
-
-					UDEditorTextureParameterValue* TextureParameterValue = Cast<UDEditorTextureParameterValue>(Group.Parameters[ParameterIdx]);
-					if (TextureParameterValue)
+					UDEditorParameterValue* Parameter = Group.Parameters[ParameterIdx];
+					if (Parameter && Parameter->bOverride)
 					{
-						if (TextureParameterValue->bOverride)
+						FMaterialParameterMetadata ParameterResult;
+						if (Parameter->GetValue(ParameterResult))
 						{
-							ChildInstance->SetTextureParameterValueEditorOnly(TextureParameterValue->ParameterInfo, TextureParameterValue->ParameterValue);
-							continue;
+							UpdateContext.SetParameterValueEditorOnly(Parameter->ParameterInfo, ParameterResult);
 						}
-					}
-
-					UDEditorRuntimeVirtualTextureParameterValue* RuntimeVirtualTextureParameterValue = Cast<UDEditorRuntimeVirtualTextureParameterValue>(Group.Parameters[ParameterIdx]);
-					if (RuntimeVirtualTextureParameterValue)
-					{
-						if (RuntimeVirtualTextureParameterValue->bOverride)
+						else
 						{
-							ChildInstance->SetRuntimeVirtualTextureParameterValueEditorOnly(RuntimeVirtualTextureParameterValue->ParameterInfo, RuntimeVirtualTextureParameterValue->ParameterValue);
-							continue;
-						}
-					}
-
-					UDEditorVectorParameterValue* VectorParameterValue = Cast<UDEditorVectorParameterValue>(Group.Parameters[ParameterIdx]);
-					if (VectorParameterValue)
-					{
-						if (VectorParameterValue->bOverride)
-						{
-							ChildInstance->SetVectorParameterValueEditorOnly(VectorParameterValue->ParameterInfo, VectorParameterValue->ParameterValue);
-							continue;
-						}
-					}
-
-					UDEditorMaterialLayersParameterValue* LayersParameterValue = Cast<UDEditorMaterialLayersParameterValue>(Group.Parameters[ParameterIdx]);
-					if (LayersParameterValue)
-					{
-						FMaterialLayersFunctions LayerValue = LayersParameterValue->ParameterValue;
-						FGuid ExpressionIdValue = LayersParameterValue->ExpressionId;
-
-						if (LayersParameterValue->bOverride)
-						{
-							FStaticMaterialLayersParameter* NewParameter =
-								new(NewStaticParameters.MaterialLayersParameters) FStaticMaterialLayersParameter(LayersParameterValue->ParameterInfo, LayerValue, LayersParameterValue->bOverride, ExpressionIdValue);
-						}
-					}
-
-					UDEditorStaticSwitchParameterValue* StaticSwitchParameterValue = Cast<UDEditorStaticSwitchParameterValue>(Group.Parameters[ParameterIdx]);
-					if (StaticSwitchParameterValue)
-					{
-						bool SwitchValue = StaticSwitchParameterValue->ParameterValue;
-						FGuid ExpressionIdValue = StaticSwitchParameterValue->ExpressionId;
-
-						if (StaticSwitchParameterValue->bOverride)
-						{
-							FStaticSwitchParameter* NewParameter =
-								new(NewStaticParameters.StaticSwitchParameters) FStaticSwitchParameter(StaticSwitchParameterValue->ParameterInfo, SwitchValue, StaticSwitchParameterValue->bOverride, ExpressionIdValue);
-						}
-					}
-
-					// static component mask
-
-					UDEditorStaticComponentMaskParameterValue* StaticComponentMaskParameterValue = Cast<UDEditorStaticComponentMaskParameterValue>(Group.Parameters[ParameterIdx]);
-					if (StaticComponentMaskParameterValue)
-					{
-						bool MaskR = StaticComponentMaskParameterValue->ParameterValue.R;
-						bool MaskG = StaticComponentMaskParameterValue->ParameterValue.G;
-						bool MaskB = StaticComponentMaskParameterValue->ParameterValue.B;
-						bool MaskA = StaticComponentMaskParameterValue->ParameterValue.A;
-						FGuid ExpressionIdValue = StaticComponentMaskParameterValue->ExpressionId;
-
-						if (StaticComponentMaskParameterValue->bOverride)
-						{
-							FStaticComponentMaskParameter* NewParameter = new(NewStaticParameters.StaticComponentMaskParameters)
-								FStaticComponentMaskParameter(StaticComponentMaskParameterValue->ParameterInfo, MaskR, MaskG, MaskB, MaskA, StaticComponentMaskParameterValue->bOverride, ExpressionIdValue);
+							UDEditorMaterialLayersParameterValue* LayersParameterValue = Cast<UDEditorMaterialLayersParameterValue>(Parameter);
+							if (LayersParameterValue)
+							{
+								FMaterialLayersFunctions LayerValue = LayersParameterValue->ParameterValue;
+								FGuid ExpressionIdValue = LayersParameterValue->ExpressionId;
+
+								FStaticMaterialLayersParameter* NewParameter =
+									new(UpdateContext.GetStaticParameters().MaterialLayersParameters) FStaticMaterialLayersParameter(LayersParameterValue->ParameterInfo, LayerValue, LayersParameterValue->bOverride, ExpressionIdValue);
+							}
 						}
 					}
 				}
 			}
-
-			ChildInstance->UpdateStaticPermutation(NewStaticParameters);
 		}
 	}
 }
@@ -410,113 +335,26 @@ void FMaterialPropertyHelpers::TransitionAndCopyParameters(UMaterialInstanceCons
 			ChildInstance->MarkPackageDirty();
 			ChildInstance->ClearParameterValuesEditorOnly();
 			//propagate changes to the base material so the instance will be updated if it has a static permutation resource
-			FStaticParameterSet NewStaticParameters;
+			//FStaticParameterSet NewStaticParameters;
+			FMaterialInstanceParameterUpdateContext UpdateContext(ChildInstance);
 			for (int32 GroupIdx = 0; GroupIdx < ParameterGroups.Num(); GroupIdx++)
 			{
 				FEditorParameterGroup& Group = ParameterGroups[GroupIdx];
 				for (int32 ParameterIdx = 0; ParameterIdx < Group.Parameters.Num(); ParameterIdx++)
 				{
-					if (Group.Parameters[ParameterIdx] == NULL)
-					{
-						continue;
-					}
-					UDEditorScalarParameterValue* ScalarParameterValue = Cast<UDEditorScalarParameterValue>(Group.Parameters[ParameterIdx]);
-					if (ScalarParameterValue)
+					UDEditorParameterValue* Parameter = Group.Parameters[ParameterIdx];
+					if (Parameter && (Parameter->bOverride || bForceCopy))
 					{
-						if (ScalarParameterValue->bOverride || bForceCopy)
+						FMaterialParameterMetadata EditorValue;
+						if (Parameter->GetValue(EditorValue))
 						{
 							FMaterialParameterInfo TransitionedScalarInfo = FMaterialParameterInfo();
-							TransitionedScalarInfo.Name = ScalarParameterValue->ParameterInfo.Name;
-							ChildInstance->SetScalarParameterValueEditorOnly(TransitionedScalarInfo, ScalarParameterValue->ParameterValue);
-							continue;
-						}
-					}
-					UDEditorFontParameterValue* FontParameterValue = Cast<UDEditorFontParameterValue>(Group.Parameters[ParameterIdx]);
-					if (FontParameterValue)
-					{
-						if (FontParameterValue->bOverride || bForceCopy)
-						{
-							FMaterialParameterInfo TransitionedFontInfo = FMaterialParameterInfo();
-							TransitionedFontInfo.Name = FontParameterValue->ParameterInfo.Name;
-							ChildInstance->SetFontParameterValueEditorOnly(TransitionedFontInfo, FontParameterValue->ParameterValue.FontValue, FontParameterValue->ParameterValue.FontPage);
-							continue;
-						}
-					}
-
-					UDEditorTextureParameterValue* TextureParameterValue = Cast<UDEditorTextureParameterValue>(Group.Parameters[ParameterIdx]);
-					if (TextureParameterValue)
-					{
-						if (TextureParameterValue->bOverride || bForceCopy)
-						{
-							FMaterialParameterInfo TransitionedTextureInfo = FMaterialParameterInfo();
-							TransitionedTextureInfo.Name = TextureParameterValue->ParameterInfo.Name;
-							ChildInstance->SetTextureParameterValueEditorOnly(TransitionedTextureInfo, TextureParameterValue->ParameterValue);
-							continue;
-						}
-					}
-
-					UDEditorRuntimeVirtualTextureParameterValue* RVTParameterValue = Cast<UDEditorRuntimeVirtualTextureParameterValue>(Group.Parameters[ParameterIdx]);
-					if (RVTParameterValue)
-					{
-						if (RVTParameterValue->bOverride || bForceCopy)
-						{
-							FMaterialParameterInfo TransitionedTextureInfo = FMaterialParameterInfo();
-							TransitionedTextureInfo.Name = RVTParameterValue->ParameterInfo.Name;
-							ChildInstance->SetRuntimeVirtualTextureParameterValueEditorOnly(TransitionedTextureInfo, RVTParameterValue->ParameterValue);
-							continue;
-						}
-					}
-
-					UDEditorVectorParameterValue* VectorParameterValue = Cast<UDEditorVectorParameterValue>(Group.Parameters[ParameterIdx]);
-					if (VectorParameterValue)
-					{
-						if (VectorParameterValue->bOverride || bForceCopy)
-						{
-							FMaterialParameterInfo TransitionedVectorInfo = FMaterialParameterInfo();
-							TransitionedVectorInfo.Name = VectorParameterValue->ParameterInfo.Name;
-							ChildInstance->SetVectorParameterValueEditorOnly(TransitionedVectorInfo, VectorParameterValue->ParameterValue);
-							continue;
-						}
-					}
-
-					UDEditorStaticSwitchParameterValue* StaticSwitchParameterValue = Cast<UDEditorStaticSwitchParameterValue>(Group.Parameters[ParameterIdx]);
-					if (StaticSwitchParameterValue)
-					{
-						bool SwitchValue = StaticSwitchParameterValue->ParameterValue;
-						FGuid ExpressionIdValue = StaticSwitchParameterValue->ExpressionId;
-
-						if (StaticSwitchParameterValue->bOverride || bForceCopy)
-						{
-							FMaterialParameterInfo TransitionedSwitchInfo = FMaterialParameterInfo();
-							TransitionedSwitchInfo.Name = StaticSwitchParameterValue->ParameterInfo.Name;
-							FStaticSwitchParameter* NewParameter =
-								new(NewStaticParameters.StaticSwitchParameters) FStaticSwitchParameter(TransitionedSwitchInfo, SwitchValue, StaticSwitchParameterValue->bOverride, ExpressionIdValue);
-						}
-					}
-
-					// static component mask
-
-					UDEditorStaticComponentMaskParameterValue* StaticComponentMaskParameterValue = Cast<UDEditorStaticComponentMaskParameterValue>(Group.Parameters[ParameterIdx]);
-					if (StaticComponentMaskParameterValue)
-					{
-						bool MaskR = StaticComponentMaskParameterValue->ParameterValue.R;
-						bool MaskG = StaticComponentMaskParameterValue->ParameterValue.G;
-						bool MaskB = StaticComponentMaskParameterValue->ParameterValue.B;
-						bool MaskA = StaticComponentMaskParameterValue->ParameterValue.A;
-						FGuid ExpressionIdValue = StaticComponentMaskParameterValue->ExpressionId;
-
-						if (StaticComponentMaskParameterValue->bOverride || bForceCopy)
-						{
-							FMaterialParameterInfo TransitionedMaskInfo = FMaterialParameterInfo();
-							TransitionedMaskInfo.Name = StaticComponentMaskParameterValue->ParameterInfo.Name;
-							FStaticComponentMaskParameter* NewParameter = new(NewStaticParameters.StaticComponentMaskParameters)
-								FStaticComponentMaskParameter(TransitionedMaskInfo, MaskR, MaskG, MaskB, MaskA, StaticComponentMaskParameterValue->bOverride, ExpressionIdValue);
+							TransitionedScalarInfo.Name = Parameter->ParameterInfo.Name;
+							UpdateContext.SetParameterValueEditorOnly(TransitionedScalarInfo, EditorValue);
 						}
 					}
 				}
 			}
-
-			ChildInstance->UpdateStaticPermutation(NewStaticParameters);
 		}
 	}
 }
@@ -799,84 +637,13 @@ void FMaterialPropertyHelpers::ResetToDefault(class UDEditorParameterValue* Para
 	const FScopedTransaction Transaction( LOCTEXT( "ResetToDefault", "Reset To Default" ) );
 	Parameter->Modify();
 
-	const FMaterialParameterInfo& ParameterInfo = Parameter->ParameterInfo;
-	
-	UDEditorScalarParameterValue* ScalarParam = Cast<UDEditorScalarParameterValue>(Parameter);
-	UDEditorVectorParameterValue* VectorParam = Cast<UDEditorVectorParameterValue>(Parameter);
-	UDEditorTextureParameterValue* TextureParam = Cast<UDEditorTextureParameterValue>(Parameter);
-	UDEditorRuntimeVirtualTextureParameterValue* RuntimeVirtualTextureParam = Cast<UDEditorRuntimeVirtualTextureParameterValue>(Parameter);
-	UDEditorFontParameterValue* FontParam = Cast<UDEditorFontParameterValue>(Parameter);
-	UDEditorStaticSwitchParameterValue* SwitchParam = Cast<UDEditorStaticSwitchParameterValue>(Parameter);
-	UDEditorStaticComponentMaskParameterValue* CompMaskParam = Cast<UDEditorStaticComponentMaskParameterValue>(Parameter);
-	UDEditorMaterialLayersParameterValue* LayersParam = Cast<UDEditorMaterialLayersParameterValue>(Parameter);
-
-	if (ScalarParam)
+	FMaterialParameterMetadata EditorValue;
+	if (Parameter->GetValue(EditorValue))
 	{
-		float OutValue;
-		if (MaterialEditorInstance->SourceInstance->GetScalarParameterDefaultValue(ParameterInfo, OutValue))
+		FMaterialParameterMetadata DefaultValue;
+		if (MaterialEditorInstance->SourceInstance->GetParameterDefaultValue(EditorValue.Value.Type, Parameter->ParameterInfo, DefaultValue))
 		{
-			ScalarParam->ParameterValue = OutValue;
-			MaterialEditorInstance->CopyToSourceInstance();
-		}
-	}
-	else if (VectorParam)
-	{
-		FLinearColor OutValue;
-		if (MaterialEditorInstance->SourceInstance->GetVectorParameterDefaultValue(ParameterInfo, OutValue))
-		{
-			VectorParam->ParameterValue = OutValue;
-			MaterialEditorInstance->CopyToSourceInstance();
-		}
-	}
-	else if (TextureParam)
-	{
-		UTexture* OutValue;
-		if (MaterialEditorInstance->SourceInstance->GetTextureParameterDefaultValue(ParameterInfo, OutValue))
-		{
-			TextureParam->ParameterValue = OutValue;
-			MaterialEditorInstance->CopyToSourceInstance();
-		}
-	}
-	else if (RuntimeVirtualTextureParam)
-	{
-		URuntimeVirtualTexture* OutValue;
-		if (MaterialEditorInstance->SourceInstance->GetRuntimeVirtualTextureParameterDefaultValue(ParameterInfo, OutValue))
-		{
-			RuntimeVirtualTextureParam->ParameterValue = OutValue;
-			MaterialEditorInstance->CopyToSourceInstance();
-		}
-	}
-	else if (FontParam)
-	{
-		UFont* OutFontValue;
-		int32 OutFontPage;
-		if (MaterialEditorInstance->SourceInstance->GetFontParameterDefaultValue(ParameterInfo, OutFontValue, OutFontPage))
-		{
-			FontParam->ParameterValue.FontValue = OutFontValue;
-			FontParam->ParameterValue.FontPage = OutFontPage;
-			MaterialEditorInstance->CopyToSourceInstance();
-		}
-	}
-	else if (SwitchParam)
-	{
-		bool OutValue;
-		FGuid TempGuid(0,0,0,0);
-		if (MaterialEditorInstance->SourceInstance->GetStaticSwitchParameterDefaultValue(ParameterInfo, OutValue, TempGuid))
-		{
-			SwitchParam->ParameterValue = OutValue;
-			MaterialEditorInstance->CopyToSourceInstance();
-		}
-	}
-	else if (CompMaskParam)
-	{
-		bool OutValue[4];
-		FGuid TempGuid(0,0,0,0);
-		if (MaterialEditorInstance->SourceInstance->GetStaticComponentMaskParameterDefaultValue(ParameterInfo, OutValue[0], OutValue[1], OutValue[2], OutValue[3], TempGuid))
-		{
-			CompMaskParam->ParameterValue.R = OutValue[0];
-			CompMaskParam->ParameterValue.G = OutValue[1];
-			CompMaskParam->ParameterValue.B = OutValue[2];
-			CompMaskParam->ParameterValue.A = OutValue[3];
+			Parameter->SetValue(DefaultValue.Value);
 			MaterialEditorInstance->CopyToSourceInstance();
 		}
 	}
@@ -987,100 +754,19 @@ bool FMaterialPropertyHelpers::ShouldLayerAssetShowResetToDefault(TSharedPtr<FSo
 
 bool FMaterialPropertyHelpers::ShouldShowResetToDefault(UDEditorParameterValue* InParameter, UMaterialEditorInstanceConstant* MaterialEditorInstance)
 {
-	const FMaterialParameterInfo& ParameterInfo = InParameter->ParameterInfo;
-
-	UDEditorFontParameterValue* FontParam = Cast<UDEditorFontParameterValue>(InParameter);
-	UDEditorScalarParameterValue* ScalarParam = Cast<UDEditorScalarParameterValue>(InParameter);
-	UDEditorStaticComponentMaskParameterValue* CompMaskParam = Cast<UDEditorStaticComponentMaskParameterValue>(InParameter);
-	UDEditorStaticSwitchParameterValue* SwitchParam = Cast<UDEditorStaticSwitchParameterValue>(InParameter);
-	UDEditorTextureParameterValue* TextureParam = Cast<UDEditorTextureParameterValue>(InParameter);
-	UDEditorRuntimeVirtualTextureParameterValue* RuntimeVirtualTextureParam = Cast<UDEditorRuntimeVirtualTextureParameterValue>(InParameter);
-	UDEditorVectorParameterValue* VectorParam = Cast<UDEditorVectorParameterValue>(InParameter);
+	//const FMaterialParameterInfo& ParameterInfo = InParameter->ParameterInfo;
 
-	if (ScalarParam)
-	{
-		float OutValue;
-		if (MaterialEditorInstance->SourceInstance->GetScalarParameterDefaultValue(ParameterInfo, OutValue))
-		{
-			if (ScalarParam->ParameterValue != OutValue)
-			{
-				return true;
-			}
-		}
-	}
-	else if (FontParam)
+	FMaterialParameterMetadata EditorValue;
+	if (InParameter->GetValue(EditorValue))
 	{
-		UFont* OutFontValue;
-		int32 OutFontPage;
-		if (MaterialEditorInstance->SourceInstance->GetFontParameterDefaultValue(ParameterInfo, OutFontValue, OutFontPage))
+		FMaterialParameterMetadata SourceValue;
+		MaterialEditorInstance->SourceInstance->GetParameterDefaultValue(EditorValue.Value.Type, InParameter->ParameterInfo, SourceValue);
+		if (EditorValue.Value != SourceValue.Value)
 		{
-			if (FontParam->ParameterValue.FontValue != OutFontValue ||
-				FontParam->ParameterValue.FontPage != OutFontPage)
-			{
-				return true;
-			}
-		}
-	}
-	else if (TextureParam)
-	{
-		UTexture* OutValue;
-		if (MaterialEditorInstance->SourceInstance->GetTextureParameterDefaultValue(ParameterInfo, OutValue))
-		{
-			if (TextureParam->ParameterValue != OutValue)
-			{
-				return true;
-			}
-		}
-	}
-	else if (RuntimeVirtualTextureParam)
-	{
-		URuntimeVirtualTexture* OutValue;
-		if (MaterialEditorInstance->SourceInstance->GetRuntimeVirtualTextureParameterDefaultValue(ParameterInfo, OutValue))
-		{
-			if (RuntimeVirtualTextureParam->ParameterValue != OutValue)
-			{
-				return true;
-			}
-		}
-	}
-	else if (VectorParam)
-	{
-		FLinearColor OutValue;
-		if (MaterialEditorInstance->SourceInstance->GetVectorParameterDefaultValue(ParameterInfo, OutValue))
-		{
-			if (VectorParam->ParameterValue != OutValue)
-			{
-				return true;
-			}
-		}
-	}
-	else if (SwitchParam)
-	{
-		bool OutValue;
-		FGuid TempGuid(0, 0, 0, 0);
-		if (MaterialEditorInstance->SourceInstance->GetStaticSwitchParameterDefaultValue(ParameterInfo, OutValue, TempGuid))
-		{
-			if (SwitchParam->ParameterValue != OutValue)
-			{
-				return true;
-			}
-		}
-	}
-	else if (CompMaskParam)
-	{
-		bool OutValue[4];
-		FGuid TempGuid(0, 0, 0, 0);
-		if (MaterialEditorInstance->SourceInstance->GetStaticComponentMaskParameterDefaultValue(ParameterInfo, OutValue[0], OutValue[1], OutValue[2], OutValue[3], TempGuid))
-		{
-			if (CompMaskParam->ParameterValue.R != OutValue[0] ||
-				CompMaskParam->ParameterValue.G != OutValue[1] ||
-				CompMaskParam->ParameterValue.B != OutValue[2] ||
-				CompMaskParam->ParameterValue.A != OutValue[3])
-			{
-				return true;
-			}
+			return true;
 		}
 	}
+
 	return false;
 }
 
diff --git a/Engine/Source/Editor/UnrealEd/Classes/MaterialEditor/DEditorFontParameterValue.h b/Engine/Source/Editor/UnrealEd/Classes/MaterialEditor/DEditorFontParameterValue.h
--- a/Engine/Source/Editor/UnrealEd/Classes/MaterialEditor/DEditorFontParameterValue.h
+++ b/Engine/Source/Editor/UnrealEd/Classes/MaterialEditor/DEditorFontParameterValue.h
@@ -27,5 +27,25 @@ class UNREALED_API UDEditorFontParameterValue : public UDEditorParameterValue
 
 	UPROPERTY(EditAnywhere, Category = DEditorFontParameterValue)
 	struct FDFontParameters ParameterValue;
+
+	virtual FName GetDefaultGroupName() const override { return TEXT("Font Parameter Values"); }
+
+	virtual bool GetValue(FMaterialParameterMetadata& OutResult) const override
+	{
+		UDEditorParameterValue::GetValue(OutResult);
+		OutResult.Value = FMaterialParameterValue(ParameterValue.FontValue, ParameterValue.FontPage);
+		return true;
+	}
+
+	virtual bool SetValue(const FMaterialParameterValue& Value) override
+	{
+		if (Value.Type == EMaterialParameterType::Font)
+		{
+			ParameterValue.FontValue = Value.Font.Value;
+			ParameterValue.FontPage = Value.Font.Page;
+			return true;
+		}
+		return false;
+	}
 };
 
diff --git a/Engine/Source/Editor/UnrealEd/Classes/MaterialEditor/DEditorMaterialLayersParameterValue.h b/Engine/Source/Editor/UnrealEd/Classes/MaterialEditor/DEditorMaterialLayersParameterValue.h
--- a/Engine/Source/Editor/UnrealEd/Classes/MaterialEditor/DEditorMaterialLayersParameterValue.h
+++ b/Engine/Source/Editor/UnrealEd/Classes/MaterialEditor/DEditorMaterialLayersParameterValue.h
@@ -15,4 +15,6 @@ class UNREALED_API UDEditorMaterialLayersParameterValue : public UDEditorParamet
 
 	UPROPERTY(EditAnywhere, Category=DEditorMaterialLayersParameterValue)
 	struct FMaterialLayersFunctions ParameterValue;
+
+	virtual FName GetDefaultGroupName() const override { return TEXT("Material Layers Parameter Values"); }
 };
diff --git a/Engine/Source/Editor/UnrealEd/Classes/MaterialEditor/DEditorParameterValue.h b/Engine/Source/Editor/UnrealEd/Classes/MaterialEditor/DEditorParameterValue.h
--- a/Engine/Source/Editor/UnrealEd/Classes/MaterialEditor/DEditorParameterValue.h
+++ b/Engine/Source/Editor/UnrealEd/Classes/MaterialEditor/DEditorParameterValue.h
@@ -6,6 +6,7 @@
 #include "UObject/ObjectMacros.h"
 #include "UObject/Object.h"
 #include "Misc/Guid.h"
+#include "MaterialTypes.h"
 #include "Materials/MaterialExpressionParameter.h"
 
 #include "DEditorParameterValue.generated.h"
@@ -15,6 +16,11 @@ class UNREALED_API UDEditorParameterValue : public UObject
 {
 	GENERATED_UCLASS_BODY()
 
+	static UDEditorParameterValue* Create(UObject* Owner,
+		EMaterialParameterType Type,
+		const FMaterialParameterInfo& ParameterInfo,
+		const FMaterialParameterMetadata& Meta);
+
 	UPROPERTY(EditAnywhere, Category=DEditorParameterValue)
 	uint32 bOverride:1;
 
@@ -29,5 +35,25 @@ class UNREALED_API UDEditorParameterValue : public UObject
 	UPROPERTY()
 	int32 SortPriority = 32;
 #endif
+
+	virtual FName GetDefaultGroupName() const { return TEXT("None"); }
+	virtual bool SetValue(const FMaterialParameterValue& Value) { return false; }
+
+	virtual bool GetValue(FMaterialParameterMetadata& OutResult) const
+	{
+		OutResult.ExpressionGuid = ExpressionId;
+		OutResult.bOverride = bOverride;
+		return false;
+	}
+
+	EMaterialParameterType GetParameterType() const
+	{
+		FMaterialParameterMetadata Result;
+		if (GetValue(Result))
+		{
+			return Result.Value.Type;
+		}
+		return EMaterialParameterType::None;
+	}
 };
 
diff --git a/Engine/Source/Editor/UnrealEd/Classes/MaterialEditor/DEditorRuntimeVirtualTextureParameterValue.h b/Engine/Source/Editor/UnrealEd/Classes/MaterialEditor/DEditorRuntimeVirtualTextureParameterValue.h
--- a/Engine/Source/Editor/UnrealEd/Classes/MaterialEditor/DEditorRuntimeVirtualTextureParameterValue.h
+++ b/Engine/Source/Editor/UnrealEd/Classes/MaterialEditor/DEditorRuntimeVirtualTextureParameterValue.h
@@ -15,4 +15,23 @@ class UNREALED_API UDEditorRuntimeVirtualTextureParameterValue : public UDEditor
 
 	UPROPERTY(EditAnywhere, Category=DEditorTextureParameterValue)
 	TObjectPtr<class URuntimeVirtualTexture> ParameterValue;
+
+	virtual FName GetDefaultGroupName() const override { return TEXT("Texture Parameter Values"); }
+
+	virtual bool GetValue(FMaterialParameterMetadata& OutResult) const override
+	{
+		UDEditorParameterValue::GetValue(OutResult);
+		OutResult.Value = ParameterValue;
+		return true;
+	}
+
+	virtual bool SetValue(const FMaterialParameterValue& Value) override
+	{
+		if (Value.Type == EMaterialParameterType::RuntimeVirtualTexture)
+		{
+			ParameterValue = Value.RuntimeVirtualTexture;
+			return true;
+		}
+		return false;
+	}
 };
diff --git a/Engine/Source/Editor/UnrealEd/Classes/MaterialEditor/DEditorScalarParameterValue.h b/Engine/Source/Editor/UnrealEd/Classes/MaterialEditor/DEditorScalarParameterValue.h
--- a/Engine/Source/Editor/UnrealEd/Classes/MaterialEditor/DEditorScalarParameterValue.h
+++ b/Engine/Source/Editor/UnrealEd/Classes/MaterialEditor/DEditorScalarParameterValue.h
@@ -38,5 +38,28 @@ public:
 	UPROPERTY(Transient)
 	FScalarParameterAtlasData AtlasData;
 
+	virtual FName GetDefaultGroupName() const override { return TEXT("Scalar Parameter Values"); }
+
+	virtual bool GetValue(FMaterialParameterMetadata& OutResult) const override
+	{
+		UDEditorParameterValue::GetValue(OutResult);
+		OutResult.Value = ParameterValue;
+		OutResult.ScalarMin = SliderMin;
+		OutResult.ScalarMax = SliderMax;
+		OutResult.ScalarCurve = AtlasData.Curve;
+		OutResult.ScalarAtlas = AtlasData.Atlas;
+		OutResult.bUsedAsAtlasPosition = AtlasData.bIsUsedAsAtlasPosition;
+		return true;
+	}
+
+	virtual bool SetValue(const FMaterialParameterValue& Value) override
+	{
+		if (Value.Type == EMaterialParameterType::Scalar)
+		{
+			ParameterValue = Value.AsScalar();
+			return true;
+		}
+		return false;
+	}
 };
 
diff --git a/Engine/Source/Editor/UnrealEd/Classes/MaterialEditor/DEditorStaticComponentMaskParameterValue.h b/Engine/Source/Editor/UnrealEd/Classes/MaterialEditor/DEditorStaticComponentMaskParameterValue.h
--- a/Engine/Source/Editor/UnrealEd/Classes/MaterialEditor/DEditorStaticComponentMaskParameterValue.h
+++ b/Engine/Source/Editor/UnrealEd/Classes/MaterialEditor/DEditorStaticComponentMaskParameterValue.h
@@ -53,5 +53,26 @@ class UNREALED_API UDEditorStaticComponentMaskParameterValue : public UDEditorPa
 	UPROPERTY(EditAnywhere, Category=DEditorStaticComponentMaskParameterValue)
 	struct FDComponentMaskParameter ParameterValue;
 
+	virtual FName GetDefaultGroupName() const override { return TEXT("Static Component Mask Parameter Values"); }
+
+	virtual bool GetValue(FMaterialParameterMetadata& OutResult) const override
+	{
+		UDEditorParameterValue::GetValue(OutResult);
+		OutResult.Value = FMaterialParameterValue(ParameterValue.R, ParameterValue.G, ParameterValue.B, ParameterValue.A);
+		return true;
+	}
+
+	virtual bool SetValue(const FMaterialParameterValue& Value) override
+	{
+		if (Value.Type == EMaterialParameterType::StaticComponentMask)
+		{
+			ParameterValue.R = Value.Bool[0];
+			ParameterValue.G = Value.Bool[1];
+			ParameterValue.B = Value.Bool[2];
+			ParameterValue.A = Value.Bool[3];
+			return true;
+		}
+		return false;
+	}
 };
 
diff --git a/Engine/Source/Editor/UnrealEd/Classes/MaterialEditor/DEditorStaticSwitchParameterValue.h b/Engine/Source/Editor/UnrealEd/Classes/MaterialEditor/DEditorStaticSwitchParameterValue.h
--- a/Engine/Source/Editor/UnrealEd/Classes/MaterialEditor/DEditorStaticSwitchParameterValue.h
+++ b/Engine/Source/Editor/UnrealEd/Classes/MaterialEditor/DEditorStaticSwitchParameterValue.h
@@ -16,5 +16,23 @@ class UNREALED_API UDEditorStaticSwitchParameterValue : public UDEditorParameter
 	UPROPERTY(EditAnywhere, Category=DEditorStaticSwitchParameterValue)
 	uint32 ParameterValue:1;
 
+	virtual FName GetDefaultGroupName() const override { return TEXT("Static Switch Parameter Values"); }
+
+	virtual bool GetValue(FMaterialParameterMetadata& OutResult) const override
+	{
+		UDEditorParameterValue::GetValue(OutResult);
+		OutResult.Value = (bool)ParameterValue;
+		return true;
+	}
+
+	virtual bool SetValue(const FMaterialParameterValue& Value) override
+	{
+		if (Value.Type == EMaterialParameterType::StaticSwitch)
+		{
+			ParameterValue = Value.AsStaticSwitch();
+			return true;
+		}
+		return false;
+	}
 };
 
diff --git a/Engine/Source/Editor/UnrealEd/Classes/MaterialEditor/DEditorTextureParameterValue.h b/Engine/Source/Editor/UnrealEd/Classes/MaterialEditor/DEditorTextureParameterValue.h
--- a/Engine/Source/Editor/UnrealEd/Classes/MaterialEditor/DEditorTextureParameterValue.h
+++ b/Engine/Source/Editor/UnrealEd/Classes/MaterialEditor/DEditorTextureParameterValue.h
@@ -18,5 +18,28 @@ class UNREALED_API UDEditorTextureParameterValue : public UDEditorParameterValue
 
 	UPROPERTY(Transient)
 	FParameterChannelNames ChannelNames;
+
+	virtual FName GetDefaultGroupName() const override { return TEXT("Texture Parameter Values"); }
+
+	virtual bool GetValue(FMaterialParameterMetadata& OutResult) const override
+	{
+		UDEditorParameterValue::GetValue(OutResult);
+		OutResult.Value = ParameterValue;
+		OutResult.ChannelName[0] = ChannelNames.R;
+		OutResult.ChannelName[1] = ChannelNames.G;
+		OutResult.ChannelName[2] = ChannelNames.B;
+		OutResult.ChannelName[3] = ChannelNames.A;
+		return true;
+	}
+
+	virtual bool SetValue(const FMaterialParameterValue& Value) override
+	{
+		if (Value.Type == EMaterialParameterType::Texture)
+		{
+			ParameterValue = Value.Texture;
+			return true;
+		}
+		return false;
+	}
 };
 
diff --git a/Engine/Source/Editor/UnrealEd/Classes/MaterialEditor/DEditorVectorParameterValue.h b/Engine/Source/Editor/UnrealEd/Classes/MaterialEditor/DEditorVectorParameterValue.h
--- a/Engine/Source/Editor/UnrealEd/Classes/MaterialEditor/DEditorVectorParameterValue.h
+++ b/Engine/Source/Editor/UnrealEd/Classes/MaterialEditor/DEditorVectorParameterValue.h
@@ -21,5 +21,29 @@ class UNREALED_API UDEditorVectorParameterValue : public UDEditorParameterValue
 
 	UPROPERTY(Transient)
 	FParameterChannelNames ChannelNames;
+
+	virtual FName GetDefaultGroupName() const override { return TEXT("Vector Parameter Values"); }
+
+	virtual bool GetValue(FMaterialParameterMetadata& OutResult) const override
+	{
+		UDEditorParameterValue::GetValue(OutResult);
+		OutResult.Value = ParameterValue;
+		OutResult.bUsedAsChannelMask = bIsUsedAsChannelMask;
+		OutResult.ChannelName[0] = ChannelNames.R;
+		OutResult.ChannelName[1] = ChannelNames.G;
+		OutResult.ChannelName[2] = ChannelNames.B;
+		OutResult.ChannelName[3] = ChannelNames.A;
+		return true;
+	}
+
+	virtual bool SetValue(const FMaterialParameterValue& Value) override
+	{
+		if (Value.Type == EMaterialParameterType::Vector)
+		{
+			ParameterValue = Value.AsLinearColor();
+			return true;
+		}
+		return false;
+	}
 };
 
diff --git a/Engine/Source/Editor/UnrealEd/Classes/MaterialEditor/MaterialEditorInstanceConstant.h b/Engine/Source/Editor/UnrealEd/Classes/MaterialEditor/MaterialEditorInstanceConstant.h
--- a/Engine/Source/Editor/UnrealEd/Classes/MaterialEditor/MaterialEditorInstanceConstant.h
+++ b/Engine/Source/Editor/UnrealEd/Classes/MaterialEditor/MaterialEditorInstanceConstant.h
@@ -294,9 +294,6 @@ class UNREALED_API UMaterialEditorInstanceConstant : public UObject
 
 	void ApplySourceFunctionChanges();
 
-	/** Builds a FStaticParameterSet for the source UMaterialInstance to store.  The built set has only parameters overridden by this instance. */
-	void BuildStaticParametersForSourceInstance(FStaticParameterSet& OutStaticParameters);
-
 	/** 
 	 * Sets the source instance for this object and regenerates arrays. 
 	 *
diff --git a/Engine/Source/Editor/UnrealEd/Private/PreviewMaterial.cpp b/Engine/Source/Editor/UnrealEd/Private/PreviewMaterial.cpp
--- a/Engine/Source/Editor/UnrealEd/Private/PreviewMaterial.cpp
+++ b/Engine/Source/Editor/UnrealEd/Private/PreviewMaterial.cpp
@@ -505,6 +505,118 @@ void UMaterialEditorPreviewParameters::AssignParameterToGroup(UMaterial* ParentM
 	CurrentGroup.Parameters.Add(ParameterValue);
 }
 
+static UDEditorParameterValue* CreateParameter_Scalar(UObject* Owner, const FMaterialParameterMetadata& Meta)
+{
+	UDEditorScalarParameterValue* Parameter = NewObject<UDEditorScalarParameterValue>(Owner);
+	if (Meta.Value.Type == EMaterialParameterType::Scalar)
+	{
+		Parameter->ParameterValue = Meta.Value.AsScalar();
+		Parameter->SliderMin = Meta.ScalarMin;
+		Parameter->SliderMax = Meta.ScalarMax;
+		Parameter->AtlasData.bIsUsedAsAtlasPosition = Meta.bUsedAsAtlasPosition;
+		Parameter->AtlasData.Atlas = Meta.ScalarAtlas;
+		Parameter->AtlasData.Curve = Meta.ScalarCurve;
+	}
+	return Parameter;
+}
+
+static UDEditorParameterValue* CreateParameter_Vector(UObject* Owner, const FMaterialParameterMetadata& Meta)
+{
+	UDEditorVectorParameterValue* Parameter = NewObject<UDEditorVectorParameterValue>(Owner);
+	if (Meta.Value.Type == EMaterialParameterType::Vector)
+	{
+		Parameter->ParameterValue = Meta.Value.AsLinearColor();
+		Parameter->bIsUsedAsChannelMask = Meta.bUsedAsChannelMask;
+		Parameter->ChannelNames.R = Meta.ChannelName[0];
+		Parameter->ChannelNames.G = Meta.ChannelName[1];
+		Parameter->ChannelNames.B = Meta.ChannelName[2];
+		Parameter->ChannelNames.A = Meta.ChannelName[3];
+	}
+	return Parameter;
+}
+
+static UDEditorParameterValue* CreateParameter_Texture(UObject* Owner, const FMaterialParameterMetadata& Meta)
+{
+	UDEditorTextureParameterValue* Parameter = NewObject<UDEditorTextureParameterValue>(Owner);
+	if (Meta.Value.Type == EMaterialParameterType::Texture)
+	{
+		Parameter->ParameterValue = Meta.Value.Texture;
+		Parameter->ChannelNames.R = Meta.ChannelName[0];
+		Parameter->ChannelNames.G = Meta.ChannelName[1];
+		Parameter->ChannelNames.B = Meta.ChannelName[2];
+		Parameter->ChannelNames.A = Meta.ChannelName[3];
+	}
+	return Parameter;
+}
+
+static UDEditorParameterValue* CreateParameter_RuntimeVirtualTexture(UObject* Owner, const FMaterialParameterMetadata& Meta)
+{
+	UDEditorRuntimeVirtualTextureParameterValue* Parameter = NewObject<UDEditorRuntimeVirtualTextureParameterValue>(Owner);
+	if (Meta.Value.Type == EMaterialParameterType::RuntimeVirtualTexture)
+	{
+		Parameter->ParameterValue = Meta.Value.RuntimeVirtualTexture;
+	}
+	return Parameter;
+}
+
+static UDEditorParameterValue* CreateParameter_Font(UObject* Owner, const FMaterialParameterMetadata& Meta)
+{
+	UDEditorFontParameterValue* Parameter = NewObject<UDEditorFontParameterValue>(Owner);
+	if (Meta.Value.Type == EMaterialParameterType::Font)
+	{
+		Parameter->ParameterValue.FontValue = Meta.Value.Font.Value;
+		Parameter->ParameterValue.FontPage = Meta.Value.Font.Page;
+	}
+	return Parameter;
+}
+
+static UDEditorParameterValue* CreateParameter_StaticSwitch(UObject* Owner, const FMaterialParameterMetadata& Meta)
+{
+	UDEditorStaticSwitchParameterValue* Parameter = NewObject<UDEditorStaticSwitchParameterValue>(Owner);
+	if (Meta.Value.Type == EMaterialParameterType::StaticSwitch)
+	{
+		Parameter->ParameterValue = Meta.Value.AsStaticSwitch();
+	}
+	return Parameter;
+}
+
+static UDEditorParameterValue* CreateParameter_StaticComponentMask(UObject* Owner, const FMaterialParameterMetadata& Meta)
+{
+	UDEditorStaticComponentMaskParameterValue* Parameter = NewObject<UDEditorStaticComponentMaskParameterValue>(Owner);
+	if (Meta.Value.Type == EMaterialParameterType::StaticComponentMask)
+	{
+		Parameter->ParameterValue.R = Meta.Value.Bool[0];
+		Parameter->ParameterValue.G = Meta.Value.Bool[1];
+		Parameter->ParameterValue.B = Meta.Value.Bool[2];
+		Parameter->ParameterValue.A = Meta.Value.Bool[3];
+	}
+	return Parameter;
+}
+
+UDEditorParameterValue* UDEditorParameterValue::Create(UObject* Owner,
+	EMaterialParameterType Type,
+	const FMaterialParameterInfo& ParameterInfo,
+	const FMaterialParameterMetadata& Meta)
+{
+	UDEditorParameterValue* Parameter = nullptr;
+	switch (Type)
+	{
+	case EMaterialParameterType::Scalar: Parameter = CreateParameter_Scalar(Owner, Meta); break;
+	case EMaterialParameterType::Vector: Parameter = CreateParameter_Vector(Owner, Meta); break;
+	case EMaterialParameterType::Texture: Parameter = CreateParameter_Texture(Owner, Meta); break;
+	case EMaterialParameterType::RuntimeVirtualTexture: Parameter = CreateParameter_RuntimeVirtualTexture(Owner, Meta); break;
+	case EMaterialParameterType::Font: Parameter = CreateParameter_Font(Owner, Meta); break;
+	case EMaterialParameterType::StaticSwitch: Parameter = CreateParameter_StaticSwitch(Owner, Meta); break;
+	case EMaterialParameterType::StaticComponentMask: Parameter = CreateParameter_StaticComponentMask(Owner, Meta); break;
+	default: break;
+	}
+
+	check(Parameter);
+	Parameter->ParameterInfo = ParameterInfo;
+	Parameter->bOverride = Meta.bOverride;
+	return Parameter;
+}
+
 void UMaterialEditorPreviewParameters::RegenerateArrays()
 {
 	ParameterGroups.Empty();
@@ -552,220 +664,40 @@ void UMaterialEditorPreviewParameters::RegenerateArrays()
 		FMaterialExpressionParameterDataCache ExpressionParameterDataCache = CacheMaterialExpressionParameterData(PreviewMaterial, FStaticParameterSet());
 
 		// Loop through all types of parameters for this material and add them to the parameter arrays.
-		TArray<FMaterialParameterInfo> ParameterInfo;
-		TArray<FGuid> Guids;
-		ParentMaterial->GetAllVectorParameterInfo(ParameterInfo, Guids);
-
-		// Vector Parameters.
-		for (int32 ParameterIdx = 0; ParameterIdx < ParameterInfo.Num(); ParameterIdx++)
+		TMap<FMaterialParameterInfo, FMaterialParameterMetadata> ParameterValues;
+		for (int32 TypeIndex = 0; TypeIndex < NumMaterialParameterTypes; ++TypeIndex)
 		{
-			UDEditorVectorParameterValue& ParameterValue = *(NewObject<UDEditorVectorParameterValue>(this));
-			FName ParameterName = ParameterInfo[ParameterIdx].Name;
-			FLinearColor Value;
-			ParameterValue.bOverride = true;
-			ParameterValue.ExpressionId = Guids[ParameterIdx];
-			ParameterValue.ParameterInfo = ParameterInfo[ParameterIdx];
-			if (PreviewMaterial->GetVectorParameterValue(ParameterValue.ParameterInfo, Value))
-			{
-				ParameterValue.ParameterValue = Value;
-				PreviewMaterial->IsVectorParameterUsedAsChannelMask(ParameterValue.ParameterInfo, ParameterValue.bIsUsedAsChannelMask);		
-				PreviewMaterial->GetVectorParameterChannelNames(ParameterValue.ParameterInfo, ParameterValue.ChannelNames);
-			}
-			AssignGroupAndSortPriority(&ParameterValue, ExpressionParameterDataCache);
-		}
+			const EMaterialParameterType ParameterType = (EMaterialParameterType)TypeIndex;
 
-		// Scalar Parameters.
-		ParentMaterial->GetAllScalarParameterInfo(ParameterInfo, Guids);
-		for (int32 ParameterIdx = 0; ParameterIdx < ParameterInfo.Num(); ParameterIdx++)
-		{
-			UDEditorScalarParameterValue& ParameterValue = *(NewObject<UDEditorScalarParameterValue>(this));
-			FName ParameterName = ParameterInfo[ParameterIdx].Name;
-			float Value;
-
-			ParameterValue.bOverride = true;
-			ParameterValue.ParameterInfo = ParameterInfo[ParameterIdx];
-			ParameterValue.ExpressionId = Guids[ParameterIdx];
-
-			if (PreviewMaterial->GetScalarParameterValue(ParameterValue.ParameterInfo, Value))
+			ParentMaterial->GetAllParametersOfType(ParameterType, ParameterValues);
+			for (const auto& It : ParameterValues)
 			{
-				ParentMaterial->GetScalarParameterSliderMinMax(ParameterName, ParameterValue.SliderMin, ParameterValue.SliderMax);
-				ParentMaterial->IsScalarParameterUsedAsAtlasPosition(ParameterName, ParameterValue.AtlasData.bIsUsedAsAtlasPosition, ParameterValue.AtlasData.Curve, ParameterValue.AtlasData.Atlas);
-				ParameterValue.ParameterValue = Value;
+				UDEditorParameterValue* Parameter = UDEditorParameterValue::Create(this, ParameterType, It.Key, It.Value);
+				Parameter->bOverride = true;
+				AssignGroupAndSortPriority(Parameter, ExpressionParameterDataCache);
 			}
-			AssignGroupAndSortPriority(&ParameterValue, ExpressionParameterDataCache);
 		}
 
-		// Texture Parameters.
-		ParentMaterial->GetAllTextureParameterInfo(ParameterInfo, Guids);
-		for (int32 ParameterIdx = 0; ParameterIdx < ParameterInfo.Num(); ParameterIdx++)
-		{
-			UDEditorTextureParameterValue& ParameterValue = *(NewObject<UDEditorTextureParameterValue>(this));
-			FName ParameterName = ParameterInfo[ParameterIdx].Name;
-			UTexture* Value;
-
-			ParameterValue.bOverride = true;
-			ParameterValue.ParameterInfo = ParameterInfo[ParameterIdx];
-			ParameterValue.ExpressionId = Guids[ParameterIdx];
-
-			if (PreviewMaterial->GetTextureParameterValue(ParameterValue.ParameterInfo, Value))
-			{
-				ParameterValue.ParameterValue = Value;
-				PreviewMaterial->GetTextureParameterChannelNames(ParameterValue.ParameterInfo, ParameterValue.ChannelNames);
-			}
-			AssignGroupAndSortPriority(&ParameterValue, ExpressionParameterDataCache);
-		}
-
-		// Runtime Virtual Texture Parameters.
-		ParentMaterial->GetAllTextureParameterInfo(ParameterInfo, Guids);
-		for (int32 ParameterIdx = 0; ParameterIdx < ParameterInfo.Num(); ParameterIdx++)
-		{
-			UDEditorRuntimeVirtualTextureParameterValue& ParameterValue = *(NewObject<UDEditorRuntimeVirtualTextureParameterValue>(this));
-			FName ParameterName = ParameterInfo[ParameterIdx].Name;
-			URuntimeVirtualTexture* Value;
-
-			ParameterValue.bOverride = true;
-			ParameterValue.ParameterInfo = ParameterInfo[ParameterIdx];
-			ParameterValue.ExpressionId = Guids[ParameterIdx];
-
-			if (PreviewMaterial->GetRuntimeVirtualTextureParameterValue(ParameterValue.ParameterInfo, Value))
-			{
-				ParameterValue.ParameterValue = Value;
-			}
-			AssignGroupAndSortPriority(&ParameterValue, ExpressionParameterDataCache);
-		}
-
-		// Font Parameters.
-		ParentMaterial->GetAllFontParameterInfo(ParameterInfo, Guids);
-		for (int32 ParameterIdx = 0; ParameterIdx < ParameterInfo.Num(); ParameterIdx++)
-		{
-			UDEditorFontParameterValue& ParameterValue = *(NewObject<UDEditorFontParameterValue>(this));
-			FName ParameterName = ParameterInfo[ParameterIdx].Name;
-			UFont* FontValue;
-			int32 FontPage;
-
-			ParameterValue.bOverride = true;
-			ParameterValue.ParameterInfo = ParameterInfo[ParameterIdx];
-			ParameterValue.ExpressionId = Guids[ParameterIdx];
-
-			if (PreviewMaterial->GetFontParameterValue(ParameterValue.ParameterInfo, FontValue, FontPage))
-			{
-				ParameterValue.ParameterValue.FontValue = FontValue;
-				ParameterValue.ParameterValue.FontPage = FontPage;
-			}
-			AssignGroupAndSortPriority(&ParameterValue, ExpressionParameterDataCache);
-		}
-
-		// Get all static parameters from the source instance.  This will handle inheriting parent values.
-		FStaticParameterSet SourceStaticParameters;
 		// Static Material Layers Parameters
+		TArray<FMaterialParameterInfo> ParameterInfo;
+		TArray<FGuid> Guids;
 		ParentMaterial->GetAllMaterialLayersParameterInfo(ParameterInfo, Guids);
-		SourceStaticParameters.MaterialLayersParameters.AddZeroed(ParameterInfo.Num());
-
 		for (int32 ParameterIdx = 0; ParameterIdx < ParameterInfo.Num(); ParameterIdx++)
 		{
-			FStaticMaterialLayersParameter& ParameterValue = SourceStaticParameters.MaterialLayersParameters[ParameterIdx];
-			ParameterValue.ParameterInfo = ParameterInfo[ParameterIdx];
-			FMaterialLayersFunctions Value = FMaterialLayersFunctions();
-			FGuid ExpressionId = Guids[ParameterIdx];
-
-			ParameterValue.bOverride = true;
+			UDEditorMaterialLayersParameterValue* ParameterValue = NewObject<UDEditorMaterialLayersParameterValue>(this);
+			ParameterValue->bOverride = true;
+			ParameterValue->ParameterInfo = ParameterInfo[ParameterIdx];
+			ParameterValue->ExpressionId = Guids[ParameterIdx];
 
 			//get the settings from the parent in the MIC chain
-			if (PreviewMaterial->GetMaterialLayersParameterValue(ParameterValue.ParameterInfo, Value, ExpressionId))
+			FMaterialLayersFunctions Value;
+			if (PreviewMaterial->GetMaterialLayersParameterValue(ParameterValue->ParameterInfo, Value, ParameterValue->ExpressionId))
 			{
-				ParameterValue.Value = Value;
+				ParameterValue->ParameterValue = Value;
 			}
-			ParameterValue.ExpressionGUID = ExpressionId;
-		}
-
-		// Static Switch Parameters
-		ParentMaterial->GetAllStaticSwitchParameterInfo(ParameterInfo, Guids);
-		SourceStaticParameters.StaticSwitchParameters.AddZeroed(ParameterInfo.Num());
 
-		for (int32 ParameterIdx = 0; ParameterIdx < ParameterInfo.Num(); ParameterIdx++)
-		{
-			FStaticSwitchParameter& ParameterValue = SourceStaticParameters.StaticSwitchParameters[ParameterIdx];
-			ParameterValue.ParameterInfo = ParameterInfo[ParameterIdx];
-			bool Value = false;
-			FGuid ExpressionId = Guids[ParameterIdx];
-
-			ParameterValue.bOverride = true;
-
-			//get the settings from the parent in the MIC chain
-			if (PreviewMaterial->GetStaticSwitchParameterValue(ParameterValue.ParameterInfo, Value, ExpressionId))
-			{
-				ParameterValue.Value = Value;
-			}
-			ParameterValue.ExpressionGUID = ExpressionId;
+			AssignGroupAndSortPriority(ParameterValue, ExpressionParameterDataCache);
 		}
-
-		// Static Component Mask Parameters
-		ParentMaterial->GetAllStaticComponentMaskParameterInfo(ParameterInfo, Guids);
-		SourceStaticParameters.StaticComponentMaskParameters.AddZeroed(ParameterInfo.Num());
-		for (int32 ParameterIdx = 0; ParameterIdx < ParameterInfo.Num(); ParameterIdx++)
-		{
-			FStaticComponentMaskParameter& ParameterValue = SourceStaticParameters.StaticComponentMaskParameters[ParameterIdx];
-			bool R = false;
-			bool G = false;
-			bool B = false;
-			bool A = false;
-			FGuid ExpressionId = Guids[ParameterIdx];
-
-			ParameterValue.bOverride = true;
-			ParameterValue.ParameterInfo = ParameterInfo[ParameterIdx];
-
-			//get the settings from the parent in the MIC chain
-			if (PreviewMaterial->GetStaticComponentMaskParameterValue(ParameterValue.ParameterInfo, R, G, B, A, ExpressionId))
-			{
-				ParameterValue.R = R;
-				ParameterValue.G = G;
-				ParameterValue.B = B;
-				ParameterValue.A = A;
-			}
-			ParameterValue.ExpressionGUID = ExpressionId;
-		}
-
-		// Copy material layer Parameters
-		for (int32 ParameterIdx = 0; ParameterIdx < SourceStaticParameters.MaterialLayersParameters.Num(); ParameterIdx++)
-		{
-			FStaticMaterialLayersParameter MaterialLayersParameterValue = FStaticMaterialLayersParameter(SourceStaticParameters.MaterialLayersParameters[ParameterIdx]);
-			UDEditorMaterialLayersParameterValue& ParameterValue = *(NewObject<UDEditorMaterialLayersParameterValue>(this));
-			ParameterValue.ParameterValue = MaterialLayersParameterValue.Value;
-			ParameterValue.bOverride = MaterialLayersParameterValue.bOverride;
-			ParameterValue.ParameterInfo = MaterialLayersParameterValue.ParameterInfo;
-			ParameterValue.ExpressionId = MaterialLayersParameterValue.ExpressionGUID;
-
-			AssignGroupAndSortPriority(&ParameterValue, ExpressionParameterDataCache);
-		}
-
-		// Copy Static Switch Parameters
-		for (int32 ParameterIdx = 0; ParameterIdx < SourceStaticParameters.StaticSwitchParameters.Num(); ParameterIdx++)
-		{
-			FStaticSwitchParameter StaticSwitchParameterValue = FStaticSwitchParameter(SourceStaticParameters.StaticSwitchParameters[ParameterIdx]);
-			UDEditorStaticSwitchParameterValue& ParameterValue = *(NewObject<UDEditorStaticSwitchParameterValue>(this));
-			ParameterValue.ParameterValue = StaticSwitchParameterValue.Value;
-			ParameterValue.bOverride = StaticSwitchParameterValue.bOverride;
-			ParameterValue.ParameterInfo = StaticSwitchParameterValue.ParameterInfo;
-			ParameterValue.ExpressionId = StaticSwitchParameterValue.ExpressionGUID;
-
-			AssignGroupAndSortPriority(&ParameterValue, ExpressionParameterDataCache);
-		}
-
-		// Copy Static Component Mask Parameters
-		for (int32 ParameterIdx = 0; ParameterIdx < SourceStaticParameters.StaticComponentMaskParameters.Num(); ParameterIdx++)
-		{
-			FStaticComponentMaskParameter StaticComponentMaskParameterValue = FStaticComponentMaskParameter(SourceStaticParameters.StaticComponentMaskParameters[ParameterIdx]);
-			UDEditorStaticComponentMaskParameterValue& ParameterValue = *(NewObject<UDEditorStaticComponentMaskParameterValue>(this));
-			ParameterValue.ParameterValue.R = StaticComponentMaskParameterValue.R;
-			ParameterValue.ParameterValue.G = StaticComponentMaskParameterValue.G;
-			ParameterValue.ParameterValue.B = StaticComponentMaskParameterValue.B;
-			ParameterValue.ParameterValue.A = StaticComponentMaskParameterValue.A;
-			ParameterValue.bOverride = StaticComponentMaskParameterValue.bOverride;
-			ParameterValue.ParameterInfo = StaticComponentMaskParameterValue.ParameterInfo;
-			ParameterValue.ExpressionId = StaticComponentMaskParameterValue.ExpressionGUID;
-			AssignGroupAndSortPriority(&ParameterValue, ExpressionParameterDataCache);
-		}
-
 	}
 	// sort contents of groups
 	for (int32 ParameterIdx = 0; ParameterIdx < ParameterGroups.Num(); ParameterIdx++)
@@ -832,57 +764,14 @@ void UMaterialEditorPreviewParameters::CopyToSourceInstance()
 			FEditorParameterGroup & Group = ParameterGroups[GroupIdx];
 			for (int32 ParameterIdx = 0; ParameterIdx < Group.Parameters.Num(); ParameterIdx++)
 			{
-				if (Group.Parameters[ParameterIdx] == NULL)
+				UDEditorParameterValue* Parameter = Group.Parameters[ParameterIdx];
+				if (Parameter)
 				{
-					continue;
-				}
-				UDEditorScalarParameterValue* ScalarParameterValue = Cast<UDEditorScalarParameterValue>(Group.Parameters[ParameterIdx]);
-				if (ScalarParameterValue)
-				{
-					PreviewMaterial->SetScalarParameterValueEditorOnly(ScalarParameterValue->ParameterInfo.Name, ScalarParameterValue->ParameterValue);
-					continue;
-				}
-				UDEditorFontParameterValue* FontParameterValue = Cast<UDEditorFontParameterValue>(Group.Parameters[ParameterIdx]);
-				if (FontParameterValue)
-				{
-					PreviewMaterial->SetFontParameterValueEditorOnly(FontParameterValue->ParameterInfo.Name, FontParameterValue->ParameterValue.FontValue, FontParameterValue->ParameterValue.FontPage);
-					continue;
-				}
-				UDEditorTextureParameterValue* TextureParameterValue = Cast<UDEditorTextureParameterValue>(Group.Parameters[ParameterIdx]);
-				if (TextureParameterValue)
-				{
-					PreviewMaterial->SetTextureParameterValueEditorOnly(TextureParameterValue->ParameterInfo.Name, TextureParameterValue->ParameterValue);
-					continue;
-				}
-				UDEditorRuntimeVirtualTextureParameterValue* RuntimeVirtualTextureParameterValue = Cast<UDEditorRuntimeVirtualTextureParameterValue>(Group.Parameters[ParameterIdx]);
-				if (RuntimeVirtualTextureParameterValue)
-				{
-					PreviewMaterial->SetRuntimeVirtualTextureParameterValueEditorOnly(RuntimeVirtualTextureParameterValue->ParameterInfo.Name, RuntimeVirtualTextureParameterValue->ParameterValue);
-					continue;
-				}
-				UDEditorVectorParameterValue* VectorParameterValue = Cast<UDEditorVectorParameterValue>(Group.Parameters[ParameterIdx]);
-				if (VectorParameterValue)
-				{
-					PreviewMaterial->SetVectorParameterValueEditorOnly(VectorParameterValue->ParameterInfo.Name, VectorParameterValue->ParameterValue);
-						continue;
-				}
-				UDEditorStaticComponentMaskParameterValue* MaskParameterValue = Cast<UDEditorStaticComponentMaskParameterValue>(Group.Parameters[ParameterIdx]);
-				if (MaskParameterValue)
-				{
-					bool MaskR = MaskParameterValue->ParameterValue.R;
-					bool MaskG = MaskParameterValue->ParameterValue.G;
-					bool MaskB = MaskParameterValue->ParameterValue.B;
-					bool MaskA = MaskParameterValue->ParameterValue.A;
-					FGuid ExpressionIdValue = MaskParameterValue->ExpressionId;
-					PreviewMaterial->SetStaticComponentMaskParameterValueEditorOnly(MaskParameterValue->ParameterInfo.Name, MaskR, MaskG, MaskB, MaskA, ExpressionIdValue);
-					continue;
-				}
-				UDEditorStaticSwitchParameterValue* SwitchParameterValue = Cast<UDEditorStaticSwitchParameterValue>(Group.Parameters[ParameterIdx]);
-				if (SwitchParameterValue)
-				{
-					bool SwitchValue = SwitchParameterValue->ParameterValue;
-					PreviewMaterial->SetStaticSwitchParameterValueEditorOnly(SwitchParameterValue->ParameterInfo.Name, SwitchValue, SwitchParameterValue->ExpressionId);
-					continue;
+					FMaterialParameterMetadata EditorValue;
+					if (Parameter->GetValue(EditorValue))
+					{
+						PreviewMaterial->SetParameterValueEditorOnly(Parameter->ParameterInfo.Name, EditorValue);
+					}
 				}
 			}
 		}
@@ -904,57 +793,14 @@ void UMaterialEditorPreviewParameters::ApplySourceFunctionChanges()
 			FEditorParameterGroup & Group = ParameterGroups[GroupIdx];
 			for (int32 ParameterIdx = 0; ParameterIdx < Group.Parameters.Num(); ParameterIdx++)
 			{
-				if (Group.Parameters[ParameterIdx] == NULL)
-				{
-					continue;
-				}
-				UDEditorScalarParameterValue * ScalarParameterValue = Cast<UDEditorScalarParameterValue>(Group.Parameters[ParameterIdx]);
-				if (ScalarParameterValue)
-				{
-					OriginalFunction->SetScalarParameterValueEditorOnly(ScalarParameterValue->ParameterInfo.Name, ScalarParameterValue->ParameterValue);
-					continue;
-				}
-				UDEditorFontParameterValue * FontParameterValue = Cast<UDEditorFontParameterValue>(Group.Parameters[ParameterIdx]);
-				if (FontParameterValue)
-				{
-					OriginalFunction->SetFontParameterValueEditorOnly(FontParameterValue->ParameterInfo.Name, FontParameterValue->ParameterValue.FontValue, FontParameterValue->ParameterValue.FontPage);
-					continue;
-				}
-				UDEditorTextureParameterValue * TextureParameterValue = Cast<UDEditorTextureParameterValue>(Group.Parameters[ParameterIdx]);
-				if (TextureParameterValue)
-				{
-					OriginalFunction->SetTextureParameterValueEditorOnly(TextureParameterValue->ParameterInfo.Name, TextureParameterValue->ParameterValue);
-					continue;
-				}
-				UDEditorRuntimeVirtualTextureParameterValue * RuntimeVirtualTextureParameterValue = Cast<UDEditorRuntimeVirtualTextureParameterValue>(Group.Parameters[ParameterIdx]);
-				if (RuntimeVirtualTextureParameterValue)
-				{
-					OriginalFunction->SetRuntimeVirtualTextureParameterValueEditorOnly(RuntimeVirtualTextureParameterValue->ParameterInfo.Name, RuntimeVirtualTextureParameterValue->ParameterValue);
-					continue;
-				}
-				UDEditorVectorParameterValue * VectorParameterValue = Cast<UDEditorVectorParameterValue>(Group.Parameters[ParameterIdx]);
-				if (VectorParameterValue)
-				{
-					OriginalFunction->SetVectorParameterValueEditorOnly(VectorParameterValue->ParameterInfo.Name, VectorParameterValue->ParameterValue);
-					continue;
-				}
-				UDEditorStaticComponentMaskParameterValue * MaskParameterValue = Cast<UDEditorStaticComponentMaskParameterValue>(Group.Parameters[ParameterIdx]);
-				if (MaskParameterValue)
+				UDEditorParameterValue* Parameter = Group.Parameters[ParameterIdx];
+				if (Parameter)
 				{
-					bool MaskR = MaskParameterValue->ParameterValue.R;
-					bool MaskG = MaskParameterValue->ParameterValue.G;
-					bool MaskB = MaskParameterValue->ParameterValue.B;
-					bool MaskA = MaskParameterValue->ParameterValue.A;
-					FGuid ExpressionIdValue = MaskParameterValue->ExpressionId;
-					OriginalFunction->SetStaticComponentMaskParameterValueEditorOnly(MaskParameterValue->ParameterInfo.Name, MaskR, MaskG, MaskB, MaskA, ExpressionIdValue);
-					continue;
-				}
-				UDEditorStaticSwitchParameterValue * SwitchParameterValue = Cast<UDEditorStaticSwitchParameterValue>(Group.Parameters[ParameterIdx]);
-				if (SwitchParameterValue)
-				{
-					bool SwitchValue = SwitchParameterValue->ParameterValue;
-					OriginalFunction->SetStaticSwitchParameterValueEditorOnly(SwitchParameterValue->ParameterInfo.Name, SwitchValue, SwitchParameterValue->ExpressionId);
-					continue;
+					FMaterialParameterMetadata EditorValue;
+					if (Parameter->GetValue(EditorValue))
+					{
+						OriginalFunction->SetParameterValueEditorOnly(Parameter->ParameterInfo.Name, EditorValue);
+					}
 				}
 			}
 		}
@@ -1061,42 +907,7 @@ void  UMaterialEditorInstanceConstant::AssignParameterToGroup(UMaterial*, UDEdit
 	{
 		if (bUseOldStyleMICEditorGroups == true)
 		{
-			if (Cast<UDEditorVectorParameterValue>(ParameterValue))
-			{
-				ParameterGroupName = TEXT("Vector Parameter Values");
-			}
-			else if (Cast<UDEditorTextureParameterValue>(ParameterValue))
-			{
-				ParameterGroupName = TEXT("Texture Parameter Values");
-			}
-			else if (Cast<UDEditorRuntimeVirtualTextureParameterValue>(ParameterValue))
-			{
-				ParameterGroupName = TEXT("Texture Parameter Values");
-			}
-			else if (Cast<UDEditorScalarParameterValue>(ParameterValue))
-			{
-				ParameterGroupName = TEXT("Scalar Parameter Values");
-			}
-			else if (Cast<UDEditorStaticSwitchParameterValue>(ParameterValue))
-			{
-				ParameterGroupName = TEXT("Static Switch Parameter Values");
-			}
-			else if (Cast<UDEditorStaticComponentMaskParameterValue>(ParameterValue))
-			{
-				ParameterGroupName = TEXT("Static Component Mask Parameter Values");
-			}
-			else if (Cast<UDEditorFontParameterValue>(ParameterValue))
-			{
-				ParameterGroupName = TEXT("Font Parameter Values");
-			}
-			else if (Cast<UDEditorMaterialLayersParameterValue>(ParameterValue))
-			{
-				ParameterGroupName = TEXT("Material Layers Parameter Values");
-			}
-			else
-			{
-				ParameterGroupName = TEXT("None");
-			}
+			ParameterGroupName = ParameterValue->GetDefaultGroupName();
 		}
 		else
 		{
@@ -1173,223 +984,74 @@ void UMaterialEditorInstanceConstant::RegenerateArrays()
  		SourceInstance->GetStaticParameterValues(SourceStaticParameters);
 
 		// Loop through all types of parameters for this material and add them to the parameter arrays.
-		TArray<FMaterialParameterInfo> OutParameterInfo;
-		TArray<FGuid> Guids;
+		
 		// Need to get layer info first as other params are collected from layers
-		SourceInstance->GetAllMaterialLayersParameterInfo(OutParameterInfo, Guids);
-		// Copy Static Material Layers Parameters
-		for (int32 ParameterIdx = 0; ParameterIdx < SourceStaticParameters.MaterialLayersParameters.Num(); ParameterIdx++)
 		{
-			FStaticMaterialLayersParameter MaterialLayersParameterParameterValue = FStaticMaterialLayersParameter(SourceStaticParameters.MaterialLayersParameters[ParameterIdx]);
-			UDEditorMaterialLayersParameterValue& ParameterValue = *(NewObject<UDEditorMaterialLayersParameterValue>(this));
+			TArray<FMaterialParameterInfo> OutParameterInfo;
+			TArray<FGuid> Guids;
+			SourceInstance->GetAllMaterialLayersParameterInfo(OutParameterInfo, Guids);
+			// Copy Static Material Layers Parameters
+			for (int32 ParameterIdx = 0; ParameterIdx < SourceStaticParameters.MaterialLayersParameters.Num(); ParameterIdx++)
+			{
+				FStaticMaterialLayersParameter MaterialLayersParameterParameterValue = FStaticMaterialLayersParameter(SourceStaticParameters.MaterialLayersParameters[ParameterIdx]);
+				UDEditorMaterialLayersParameterValue& ParameterValue = *(NewObject<UDEditorMaterialLayersParameterValue>(this));
 
-			ParameterValue.ParameterValue = MaterialLayersParameterParameterValue.Value;
-			ParameterValue.bOverride = MaterialLayersParameterParameterValue.bOverride;
-			ParameterValue.ParameterInfo = MaterialLayersParameterParameterValue.ParameterInfo;
-			ParameterValue.ExpressionId = MaterialLayersParameterParameterValue.ExpressionGUID;
+				ParameterValue.ParameterValue = MaterialLayersParameterParameterValue.Value;
+				ParameterValue.bOverride = MaterialLayersParameterParameterValue.bOverride;
+				ParameterValue.ParameterInfo = MaterialLayersParameterParameterValue.ParameterInfo;
+				ParameterValue.ExpressionId = MaterialLayersParameterParameterValue.ExpressionGUID;
 
-			AssignParameterToGroup(ParentMaterial, &ParameterValue);
+				AssignParameterToGroup(ParentMaterial, &ParameterValue);
+			}
 		}
 
 		// Cache relevant material expression data to resolve editor param value info
 		FMaterialExpressionParameterDataCache ExpressionParameterDataCache = CacheMaterialExpressionParameterData(SourceInstance->GetMaterial(), SourceInstance->GetStaticParameters());
 
-		// Scalar Parameters.
-		SourceInstance->GetAllScalarParameterInfo(OutParameterInfo, Guids);		
-		for (int32 ParameterIdx = 0; ParameterIdx < OutParameterInfo.Num(); ParameterIdx++)
-		{			
-			UDEditorScalarParameterValue* ParamValue = NewObject<UDEditorScalarParameterValue>(this);
-
-			UDEditorScalarParameterValue& ParameterValue = *ParamValue;
-			const FMaterialParameterInfo& ParameterInfo = OutParameterInfo[ParameterIdx];
-
-			ParameterValue.bOverride = false;
-			ParameterValue.ParameterInfo = ParameterInfo;
-			ParameterValue.ExpressionId = Guids[ParameterIdx];
-
-			if (SourceInstance->GetScalarParameterValue(ParameterInfo, ParameterValue.ParameterValue))
-			{
-				SourceInstance->IsScalarParameterUsedAsAtlasPosition(ParameterInfo, ParameterValue.AtlasData.bIsUsedAsAtlasPosition, ParameterValue.AtlasData.Curve, ParameterValue.AtlasData.Atlas);
-				SourceInstance->GetScalarParameterSliderMinMax(ParameterInfo, ParameterValue.SliderMin, ParameterValue.SliderMax);		
-			}
-
-			// @todo: This is kind of slow, maybe store these in a map for lookup?
-			// See if this keyname exists in the source instance.
-			for(int32 ScalarParameterIdx = 0; ScalarParameterIdx < SourceInstance->ScalarParameterValues.Num(); ScalarParameterIdx++)
-			{
-				FScalarParameterValue& SourceParam = SourceInstance->ScalarParameterValues[ScalarParameterIdx];
-				if(ParameterInfo == SourceParam.ParameterInfo)
-				{
-					ParameterValue.bOverride = true;
-					ParameterValue.ParameterValue = SourceParam.ParameterValue;
-				}
-			}
-			AssignGroupAndSortPriority(ParamValue, ExpressionParameterDataCache);
-		}
-		
-
-		// Vector Parameters.
-		SourceInstance->GetAllVectorParameterInfo(OutParameterInfo, Guids);
-		for(int32 ParameterIdx = 0; ParameterIdx < OutParameterInfo.Num(); ParameterIdx++)
+		TMap<FMaterialParameterInfo, FMaterialParameterMetadata> ParameterValues;
+		for (int32 TypeIndex = 0; TypeIndex < NumMaterialParameterTypes; ++TypeIndex)
 		{
-			UDEditorVectorParameterValue* ParamValue = NewObject<UDEditorVectorParameterValue>(this);
-
-			UDEditorVectorParameterValue& ParameterValue = *ParamValue;
-			const FMaterialParameterInfo& ParameterInfo = OutParameterInfo[ParameterIdx];
-			
-			ParameterValue.bOverride = false;
-			ParameterValue.ParameterInfo = ParameterInfo;
-			ParameterValue.ExpressionId = Guids[ParameterIdx];
-
-			SourceInstance->GetVectorParameterValue(ParameterInfo, ParameterValue.ParameterValue);
-			SourceInstance->IsVectorParameterUsedAsChannelMask(ParameterInfo, ParameterValue.bIsUsedAsChannelMask);
-			SourceInstance->GetVectorParameterChannelNames(ParameterInfo, ParameterValue.ChannelNames);
-
-			// @todo: This is kind of slow, maybe store these in a map for lookup?
-			// See if this keyname exists in the source instance.
-			for(int32 VectorParameterIdx = 0; VectorParameterIdx < SourceInstance->VectorParameterValues.Num(); VectorParameterIdx++)
+			const EMaterialParameterType ParameterType = (EMaterialParameterType)TypeIndex;
+			if (IsStaticMaterialParameter(ParameterType))
 			{
-				FVectorParameterValue& SourceParam = SourceInstance->VectorParameterValues[VectorParameterIdx];
-				if(ParameterInfo == SourceParam.ParameterInfo)
-				{
-					ParameterValue.bOverride = true;
-					ParameterValue.ParameterValue = SourceParam.ParameterValue;
-				}
+				// Logic for static parameters is slightly different, handled both by GetStaticParameterValues above, and code below
+				continue;
 			}
-			AssignGroupAndSortPriority(ParamValue, ExpressionParameterDataCache);
-		}
-		
 
-		// Texture Parameters.
-		SourceInstance->GetAllTextureParameterInfo(OutParameterInfo, Guids);
-		for(int32 ParameterIdx=0; ParameterIdx<OutParameterInfo.Num(); ParameterIdx++)
-		{		
-			UDEditorTextureParameterValue* ParamValue = NewObject<UDEditorTextureParameterValue>(this);
-
-			UDEditorTextureParameterValue& ParameterValue = *ParamValue;
-			const FMaterialParameterInfo& ParameterInfo = OutParameterInfo[ParameterIdx];
-
-			ParameterValue.bOverride = false;
-			ParameterValue.ParameterInfo = ParameterInfo;
-			ParameterValue.ExpressionId = Guids[ParameterIdx];
-
-			ParameterValue.ParameterValue = nullptr;
-			SourceInstance->GetTextureParameterValue(ParameterInfo, static_cast<UTexture*&>(ParameterValue.ParameterValue));
-			SourceInstance->GetTextureParameterChannelNames(ParameterInfo, ParameterValue.ChannelNames);
-
-			// @todo: This is kind of slow, maybe store these in a map for lookup?
-			// See if this keyname exists in the source instance.
-			for(int32 TextureParameterIdx=0; TextureParameterIdx<SourceInstance->TextureParameterValues.Num(); TextureParameterIdx++)
+			SourceInstance->GetAllParametersOfType(ParameterType, ParameterValues);
+			for (const auto& It : ParameterValues)
 			{
-				FTextureParameterValue& SourceParam = SourceInstance->TextureParameterValues[TextureParameterIdx];
-				if(ParameterInfo == SourceParam.ParameterInfo)
-				{
-					ParameterValue.bOverride = true;
-					ParameterValue.ParameterValue = SourceParam.ParameterValue;
-				}
-			}			
-			AssignGroupAndSortPriority(ParamValue, ExpressionParameterDataCache);
-		}
-
-		// Runtime Virtual Texture Parameters.
-		SourceInstance->GetAllRuntimeVirtualTextureParameterInfo(OutParameterInfo, Guids);
-		for (int32 ParameterIdx = 0; ParameterIdx < OutParameterInfo.Num(); ParameterIdx++)
-		{
-			UDEditorRuntimeVirtualTextureParameterValue* ParamValue = NewObject<UDEditorRuntimeVirtualTextureParameterValue>(this);
-
-			UDEditorRuntimeVirtualTextureParameterValue& ParameterValue = *ParamValue;
-			const FMaterialParameterInfo& ParameterInfo = OutParameterInfo[ParameterIdx];
-
-			ParameterValue.bOverride = false;
-			ParameterValue.ParameterInfo = ParameterInfo;
-			ParameterValue.ExpressionId = Guids[ParameterIdx];
-
-			ParameterValue.ParameterValue = nullptr;
-			SourceInstance->GetRuntimeVirtualTextureParameterValue(ParameterInfo, static_cast<URuntimeVirtualTexture*&>(ParameterValue.ParameterValue));
-
-			// @todo: This is kind of slow, maybe store these in a map for lookup?
-			// See if this keyname exists in the source instance.
-			for (int32 TextureParameterIdx = 0; TextureParameterIdx < SourceInstance->RuntimeVirtualTextureParameterValues.Num(); TextureParameterIdx++)
-			{
-				FRuntimeVirtualTextureParameterValue& SourceParam = SourceInstance->RuntimeVirtualTextureParameterValues[TextureParameterIdx];
-				if (ParameterInfo == SourceParam.ParameterInfo)
-				{
-					ParameterValue.bOverride = true;
-					ParameterValue.ParameterValue = SourceParam.ParameterValue;
-				}
-				if (ParameterInfo.Name.IsEqual(SourceParam.ParameterInfo.Name) && ParameterInfo.Association == SourceParam.ParameterInfo.Association && ParameterInfo.Index == SourceParam.ParameterInfo.Index)
-				{
-					ParameterValue.bOverride = true;
-					ParameterValue.ParameterValue = SourceParam.ParameterValue;
-				}
+				UDEditorParameterValue* Parameter = UDEditorParameterValue::Create(this, ParameterType, It.Key, It.Value);
+				AssignGroupAndSortPriority(Parameter, ExpressionParameterDataCache);
 			}
-			AssignGroupAndSortPriority(ParamValue, ExpressionParameterDataCache);
-		}
-
-		// Font Parameters.
-		SourceInstance->GetAllFontParameterInfo(OutParameterInfo, Guids);
-		for(int32 ParameterIdx=0; ParameterIdx<OutParameterInfo.Num(); ParameterIdx++)
-		{
-			UDEditorFontParameterValue* ParamValue = NewObject<UDEditorFontParameterValue>(this);
-
-			UDEditorFontParameterValue& ParameterValue = *ParamValue;
-			const FMaterialParameterInfo& ParameterInfo = OutParameterInfo[ParameterIdx];
-
-			ParameterValue.bOverride = false;
-			ParameterValue.ParameterInfo = ParameterInfo;
-			ParameterValue.ExpressionId = Guids[ParameterIdx];
-
-			ParameterValue.ParameterValue.FontValue = nullptr;
-			ParameterValue.ParameterValue.FontPage = 0;
-			SourceInstance->GetFontParameterValue(ParameterInfo, static_cast<UFont*&>(ParameterValue.ParameterValue.FontValue), ParameterValue.ParameterValue.FontPage);
-
-			// @todo: This is kind of slow, maybe store these in a map for lookup?
-			// See if this keyname exists in the source instance.
-			for(int32 FontParameterIdx = 0; FontParameterIdx < SourceInstance->FontParameterValues.Num(); FontParameterIdx++)
-			{
-				FFontParameterValue& SourceParam = SourceInstance->FontParameterValues[FontParameterIdx];
-				if(ParameterInfo == SourceParam.ParameterInfo)
-				{
-					ParameterValue.bOverride = true;
-					ParameterValue.ParameterValue.FontValue = SourceParam.FontValue;
-					ParameterValue.ParameterValue.FontPage = SourceParam.FontPage;
-				}
-			}	
-			AssignGroupAndSortPriority(ParamValue, ExpressionParameterDataCache);
 		}
 
 		// Copy Static Switch Parameters
-		SourceInstance->GetAllStaticSwitchParameterInfo(OutParameterInfo, Guids);
 		for(int32 ParameterIdx = 0; ParameterIdx < SourceStaticParameters.StaticSwitchParameters.Num(); ParameterIdx++)
 		{	
-			FStaticSwitchParameter StaticSwitchParameterValue = FStaticSwitchParameter(SourceStaticParameters.StaticSwitchParameters[ParameterIdx]);
+			const FStaticSwitchParameter& StaticSwitchParameterValue = SourceStaticParameters.StaticSwitchParameters[ParameterIdx];
 			UDEditorStaticSwitchParameterValue* ParamValue = NewObject<UDEditorStaticSwitchParameterValue>(this);
 
-			UDEditorStaticSwitchParameterValue& ParameterValue = *ParamValue;
-
-			ParameterValue.ParameterValue = StaticSwitchParameterValue.Value;
-			ParameterValue.bOverride = StaticSwitchParameterValue.bOverride;
-			ParameterValue.ParameterInfo = StaticSwitchParameterValue.ParameterInfo;
-			ParameterValue.ExpressionId = StaticSwitchParameterValue.ExpressionGUID;
+			ParamValue->ParameterValue = StaticSwitchParameterValue.Value;
+			ParamValue->bOverride = StaticSwitchParameterValue.bOverride;
+			ParamValue->ParameterInfo = StaticSwitchParameterValue.ParameterInfo;
+			ParamValue->ExpressionId = StaticSwitchParameterValue.ExpressionGUID;
 			AssignGroupAndSortPriority(ParamValue, ExpressionParameterDataCache);
 		}
 
 		// Copy Static Component Mask Parameters
-		SourceInstance->GetAllStaticComponentMaskParameterInfo(OutParameterInfo, Guids);
 		for(int32 ParameterIdx=0; ParameterIdx<SourceStaticParameters.StaticComponentMaskParameters.Num(); ParameterIdx++)
 		{
-			FStaticComponentMaskParameter StaticComponentMaskParameterValue = FStaticComponentMaskParameter(SourceStaticParameters.StaticComponentMaskParameters[ParameterIdx]);
+			const FStaticComponentMaskParameter& StaticComponentMaskParameterValue = SourceStaticParameters.StaticComponentMaskParameters[ParameterIdx];
 			UDEditorStaticComponentMaskParameterValue* ParamValue = NewObject<UDEditorStaticComponentMaskParameterValue>(this);
 
-			UDEditorStaticComponentMaskParameterValue& ParameterValue = *ParamValue;
-
-			ParameterValue.ParameterValue.R = StaticComponentMaskParameterValue.R;
-			ParameterValue.ParameterValue.G = StaticComponentMaskParameterValue.G;
-			ParameterValue.ParameterValue.B = StaticComponentMaskParameterValue.B;
-			ParameterValue.ParameterValue.A = StaticComponentMaskParameterValue.A;
-			ParameterValue.bOverride = StaticComponentMaskParameterValue.bOverride;
-			ParameterValue.ParameterInfo = StaticComponentMaskParameterValue.ParameterInfo;
-			ParameterValue.ExpressionId = StaticComponentMaskParameterValue.ExpressionGUID;
+			ParamValue->ParameterValue.R = StaticComponentMaskParameterValue.R;
+			ParamValue->ParameterValue.G = StaticComponentMaskParameterValue.G;
+			ParamValue->ParameterValue.B = StaticComponentMaskParameterValue.B;
+			ParamValue->ParameterValue.A = StaticComponentMaskParameterValue.A;
+			ParamValue->bOverride = StaticComponentMaskParameterValue.bOverride;
+			ParamValue->ParameterInfo = StaticComponentMaskParameterValue.ParameterInfo;
+			ParamValue->ExpressionId = StaticComponentMaskParameterValue.ExpressionGUID;
 			AssignGroupAndSortPriority(ParamValue, ExpressionParameterDataCache);
 		}
 
@@ -1503,60 +1165,21 @@ void UMaterialEditorInstanceConstant::ResetOverrides(int32 Index, EMaterialParam
 {
 	check(GIsEditor);
 
-	for (FEditorParameterGroup Group : ParameterGroups)
+	for (const FEditorParameterGroup& Group : ParameterGroups)
 	{
 		for (UDEditorParameterValue* Parameter : Group.Parameters)
 		{
 			if (Parameter->ParameterInfo.Association == MaterialType
 				&& Parameter->ParameterInfo.Index == Index)
 			{
-				UDEditorScalarParameterValue* ScalarParameterValue = Cast<UDEditorScalarParameterValue>(Parameter);
-				UDEditorVectorParameterValue* VectorParameterValue = Cast<UDEditorVectorParameterValue>(Parameter);
-				UDEditorTextureParameterValue* TextureParameterValue = Cast<UDEditorTextureParameterValue>(Parameter);
-				UDEditorRuntimeVirtualTextureParameterValue* RuntimeVirtualTextureParameterValue = Cast<UDEditorRuntimeVirtualTextureParameterValue>(Parameter);
-				UDEditorFontParameterValue* FontParameterValue = Cast<UDEditorFontParameterValue>(Parameter);
-				UDEditorStaticSwitchParameterValue* StaticSwitchParameterValue = Cast<UDEditorStaticSwitchParameterValue>(Parameter);
-				UDEditorStaticComponentMaskParameterValue* StaticMaskParameterValue = Cast<UDEditorStaticComponentMaskParameterValue>(Parameter);
-				if (ScalarParameterValue)
-				{
-					float Value;
-					Parameter->bOverride = SourceInstance->GetScalarParameterValue(Parameter->ParameterInfo, Value, true);
-				}
-				if (VectorParameterValue)
-				{
-					FLinearColor Value;
-					Parameter->bOverride = SourceInstance->GetVectorParameterValue(Parameter->ParameterInfo, Value, true);
-				}
-				if (TextureParameterValue)
-				{
-					UTexture* Value;
-					Parameter->bOverride = SourceInstance->GetTextureParameterValue(Parameter->ParameterInfo, Value, true);
-				}
-				if (RuntimeVirtualTextureParameterValue)
-				{
-					URuntimeVirtualTexture* Value;
-					Parameter->bOverride = SourceInstance->GetRuntimeVirtualTextureParameterValue(Parameter->ParameterInfo, Value, true);
-				}
-				if (FontParameterValue)
+				const EMaterialParameterType ParameterType = Parameter->GetParameterType();
+				if (ParameterType != EMaterialParameterType::None)
 				{
-					UFont* FontValue;
-					int32 FontPage;
-					Parameter->bOverride = SourceInstance->GetFontParameterValue(Parameter->ParameterInfo, FontValue, FontPage, true);
-				}
-				if (StaticSwitchParameterValue)
-				{
-					bool Value;
-					FGuid ExpressionId;
-					Parameter->bOverride = SourceInstance->GetStaticSwitchParameterValue(Parameter->ParameterInfo, Value, ExpressionId, true);
-				}
-				if (StaticMaskParameterValue)
-				{
-					bool R;
-					bool G;
-					bool B;
-					bool A;
-					FGuid ExpressionId;
-					Parameter->bOverride = SourceInstance->GetStaticComponentMaskParameterValue(Parameter->ParameterInfo, R, G, B, A, ExpressionId, true);
+					FMaterialParameterMetadata SourceValue;
+					if (SourceInstance->GetParameterValue(ParameterType, Parameter->ParameterInfo, SourceValue, EMaterialGetParameterValueFlags::CheckInstanceOverrides))
+					{
+						Parameter->bOverride = SourceValue.bOverride;
+					}
 				}
 			}
 		}
@@ -1581,57 +1204,42 @@ void UMaterialEditorInstanceConstant::CopyToSourceInstance(const bool bForceStat
 
 		SourceInstance->ClearParameterValuesEditorOnly();
 
-		for (int32 GroupIdx=0; GroupIdx<ParameterGroups.Num(); GroupIdx++)
 		{
-			FEditorParameterGroup & Group = ParameterGroups[GroupIdx];
-			for (int32 ParameterIdx=0; ParameterIdx<Group.Parameters.Num(); ParameterIdx++)
-			{
-				if (Group.Parameters[ParameterIdx] == NULL)
-				{
-					continue;
-				}
+			FMaterialInstanceParameterUpdateContext UpdateContext(SourceInstance);
+			UpdateContext.SetBasePropertyOverrides(BasePropertyOverrides);
+			UpdateContext.SetForceStaticPermutationUpdate(bForceStaticPermutationUpdate);
 
-				UDEditorScalarParameterValue* ScalarParameterValue = Cast<UDEditorScalarParameterValue>(Group.Parameters[ParameterIdx]);
-				UDEditorVectorParameterValue* VectorParameterValue = Cast<UDEditorVectorParameterValue>(Group.Parameters[ParameterIdx]);
-				UDEditorTextureParameterValue* TextureParameterValue = Cast<UDEditorTextureParameterValue>(Group.Parameters[ParameterIdx]);
-				UDEditorRuntimeVirtualTextureParameterValue* RuntimeVirtualTextureParameterValue = Cast<UDEditorRuntimeVirtualTextureParameterValue>(Group.Parameters[ParameterIdx]);
-				UDEditorFontParameterValue* FontParameterValue = Cast<UDEditorFontParameterValue>(Group.Parameters[ParameterIdx]);
-				if (ScalarParameterValue && ScalarParameterValue->bOverride)
+			for (int32 GroupIdx = 0; GroupIdx < ParameterGroups.Num(); GroupIdx++)
+			{
+				FEditorParameterGroup& Group = ParameterGroups[GroupIdx];
+				for (int32 ParameterIdx = 0; ParameterIdx < Group.Parameters.Num(); ParameterIdx++)
 				{
-					SourceInstance->SetScalarParameterValueEditorOnly(ScalarParameterValue->ParameterInfo, ScalarParameterValue->ParameterValue);
-					// Copy from editor parameter to saved FParameter
-					if (ScalarParameterValue->AtlasData.bIsUsedAsAtlasPosition)
+					UDEditorParameterValue* Parameter = Group.Parameters[ParameterIdx];
+					if (Parameter && Parameter->bOverride)
 					{
-						FScalarParameterAtlasInstanceData InAtlasData = FScalarParameterAtlasInstanceData();
-						InAtlasData.bIsUsedAsAtlasPosition = ScalarParameterValue->AtlasData.bIsUsedAsAtlasPosition;
-						InAtlasData.Curve = ScalarParameterValue->AtlasData.Curve;
-						InAtlasData.Atlas = ScalarParameterValue->AtlasData.Atlas;
-						SourceInstance->SetScalarParameterAtlasEditorOnly(ScalarParameterValue->ParameterInfo, InAtlasData);
+						FMaterialParameterMetadata EditorValue;
+						if (Parameter->GetValue(EditorValue))
+						{
+							UpdateContext.SetParameterValueEditorOnly(Parameter->ParameterInfo, EditorValue, EMaterialSetParameterValueFlags::SetCurveAtlas);
+						}
+						else
+						{
+							// Material layers param
+							UDEditorMaterialLayersParameterValue* MaterialLayersParameterValue = Cast<UDEditorMaterialLayersParameterValue>(Parameter);
+							if (MaterialLayersParameterValue)
+							{
+								const FMaterialLayersFunctions& MaterialLayers = MaterialLayersParameterValue->ParameterValue;
+								FGuid ExpressionIdValue = MaterialLayersParameterValue->ExpressionId;
+
+								FStaticMaterialLayersParameter* NewParameter = new(UpdateContext.GetStaticParameters().MaterialLayersParameters)
+									FStaticMaterialLayersParameter(MaterialLayersParameterValue->ParameterInfo, MaterialLayers, MaterialLayersParameterValue->bOverride, ExpressionIdValue);
+							}
+						}
 					}
 				}
-				else if (VectorParameterValue && VectorParameterValue->bOverride)
-				{
-					SourceInstance->SetVectorParameterValueEditorOnly(VectorParameterValue->ParameterInfo, VectorParameterValue->ParameterValue);
-				}
-				else if (TextureParameterValue && TextureParameterValue->bOverride)
-				{
-					SourceInstance->SetTextureParameterValueEditorOnly(TextureParameterValue->ParameterInfo, TextureParameterValue->ParameterValue);
-				}
-				else if (RuntimeVirtualTextureParameterValue && RuntimeVirtualTextureParameterValue->bOverride)
-				{
-					SourceInstance->SetRuntimeVirtualTextureParameterValueEditorOnly(RuntimeVirtualTextureParameterValue->ParameterInfo, RuntimeVirtualTextureParameterValue->ParameterValue);
-				}
-				else if (FontParameterValue && FontParameterValue->bOverride)
-				{
-					SourceInstance->SetFontParameterValueEditorOnly(FontParameterValue->ParameterInfo, FontParameterValue->ParameterValue.FontValue, FontParameterValue->ParameterValue.FontPage);
-				}
 			}
 		}
 
-		FStaticParameterSet NewStaticParameters;
-		BuildStaticParametersForSourceInstance(NewStaticParameters);
-		SourceInstance->UpdateStaticPermutation(NewStaticParameters, BasePropertyOverrides, bForceStaticPermutationUpdate);
-
 		// Copy phys material back to source instance
 		SourceInstance->PhysMaterial = PhysMaterial;
 
@@ -1690,59 +1298,6 @@ void UMaterialEditorInstanceConstant::ApplySourceFunctionChanges()
 	}
 }
 
-void UMaterialEditorInstanceConstant::BuildStaticParametersForSourceInstance(FStaticParameterSet& OutStaticParameters)
-{
-	for(int32 GroupIdx=0; GroupIdx<ParameterGroups.Num(); GroupIdx++)
-	{
-		FEditorParameterGroup& Group = ParameterGroups[GroupIdx];
-
-		for(int32 ParameterIdx=0; ParameterIdx<Group.Parameters.Num(); ParameterIdx++)
-		{
-			if (Group.Parameters[ParameterIdx] == NULL)
-			{
-				continue;
-			}
-
-			// Static switch
-			UDEditorStaticSwitchParameterValue* StaticSwitchParameterValue = Cast<UDEditorStaticSwitchParameterValue>(Group.Parameters[ParameterIdx]);
-			if (StaticSwitchParameterValue && StaticSwitchParameterValue->bOverride)
-			{
-				bool SwitchValue = StaticSwitchParameterValue->ParameterValue;
-				FGuid ExpressionIdValue = StaticSwitchParameterValue->ExpressionId;
-
-				FStaticSwitchParameter* NewParameter = new(OutStaticParameters.StaticSwitchParameters)
-					FStaticSwitchParameter(StaticSwitchParameterValue->ParameterInfo, SwitchValue, StaticSwitchParameterValue->bOverride, ExpressionIdValue);
-			}
-
-			// Static component mask
-			UDEditorStaticComponentMaskParameterValue* StaticComponentMaskParameterValue = Cast<UDEditorStaticComponentMaskParameterValue>(Group.Parameters[ParameterIdx]);
-			if (StaticComponentMaskParameterValue && StaticComponentMaskParameterValue->bOverride)
-			{
-				bool MaskR = StaticComponentMaskParameterValue->ParameterValue.R;
-				bool MaskG = StaticComponentMaskParameterValue->ParameterValue.G;
-				bool MaskB = StaticComponentMaskParameterValue->ParameterValue.B;
-				bool MaskA = StaticComponentMaskParameterValue->ParameterValue.A;
-				FGuid ExpressionIdValue = StaticComponentMaskParameterValue->ExpressionId;
-
-				FStaticComponentMaskParameter* NewParameter = new(OutStaticParameters.StaticComponentMaskParameters) 
-					FStaticComponentMaskParameter(StaticComponentMaskParameterValue->ParameterInfo, MaskR, MaskG, MaskB, MaskA, StaticComponentMaskParameterValue->bOverride, ExpressionIdValue);
-			}
-
-			// Material layers param
-			UDEditorMaterialLayersParameterValue* MaterialLayersParameterValue = Cast<UDEditorMaterialLayersParameterValue>(Group.Parameters[ParameterIdx]);
-			if (MaterialLayersParameterValue && MaterialLayersParameterValue->bOverride)
-			{
-				const FMaterialLayersFunctions& MaterialLayers = MaterialLayersParameterValue->ParameterValue;
-				FGuid ExpressionIdValue = MaterialLayersParameterValue->ExpressionId;
-
-				FStaticMaterialLayersParameter* NewParameter = new(OutStaticParameters.MaterialLayersParameters)
-					FStaticMaterialLayersParameter(MaterialLayersParameterValue->ParameterInfo, MaterialLayers, MaterialLayersParameterValue->bOverride, ExpressionIdValue);
-			}
-		}
-	}
-}
-
-
 void UMaterialEditorInstanceConstant::SetSourceInstance(UMaterialInstanceConstant* MaterialInterface)
 {
 	check(MaterialInterface);
@@ -1755,9 +1310,40 @@ void UMaterialEditorInstanceConstant::SetSourceInstance(UMaterialInstanceConstan
 	RegenerateArrays();
 
 	//propagate changes to the base material so the instance will be updated if it has a static permutation resource
-	FStaticParameterSet NewStaticParameters;
-	BuildStaticParametersForSourceInstance(NewStaticParameters);
-	SourceInstance->UpdateStaticPermutation(NewStaticParameters);
+	FMaterialInstanceParameterUpdateContext UpdateContext(SourceInstance);
+	for (int32 GroupIdx = 0; GroupIdx < ParameterGroups.Num(); GroupIdx++)
+	{
+		FEditorParameterGroup& Group = ParameterGroups[GroupIdx];
+		for (int32 ParameterIdx = 0; ParameterIdx < Group.Parameters.Num(); ParameterIdx++)
+		{
+			UDEditorParameterValue* Parameter = Group.Parameters[ParameterIdx];
+			if (Parameter && Parameter->bOverride)
+			{
+				FMaterialParameterMetadata EditorValue;
+				if (Parameter->GetValue(EditorValue))
+				{
+					// Only want to update static parameters here
+					if (IsStaticMaterialParameter(EditorValue.Value.Type))
+					{
+						UpdateContext.SetParameterValueEditorOnly(Parameter->ParameterInfo, EditorValue);
+					}
+				}
+				else
+				{
+					// Material layers param
+					UDEditorMaterialLayersParameterValue* MaterialLayersParameterValue = Cast<UDEditorMaterialLayersParameterValue>(Group.Parameters[ParameterIdx]);
+					if (MaterialLayersParameterValue)
+					{
+						const FMaterialLayersFunctions& MaterialLayers = MaterialLayersParameterValue->ParameterValue;
+						FGuid ExpressionIdValue = MaterialLayersParameterValue->ExpressionId;
+
+						FStaticMaterialLayersParameter* NewParameter = new(UpdateContext.GetStaticParameters().MaterialLayersParameters)
+							FStaticMaterialLayersParameter(MaterialLayersParameterValue->ParameterInfo, MaterialLayers, MaterialLayersParameterValue->bOverride, ExpressionIdValue);
+					}
+				}
+			}
+		}
+	}
 }
 
 void UMaterialEditorInstanceConstant::SetSourceFunction(UMaterialFunctionInstance* MaterialFunction)
diff --git a/Engine/Source/Runtime/Engine/Classes/Materials/Material.h b/Engine/Source/Runtime/Engine/Classes/Materials/Material.h
--- a/Engine/Source/Runtime/Engine/Classes/Materials/Material.h
+++ b/Engine/Source/Runtime/Engine/Classes/Materials/Material.h
@@ -1035,22 +1035,9 @@ public:
 	ENGINE_API virtual UMaterial* GetMaterial() override;
 	ENGINE_API virtual const UMaterial* GetMaterial() const override;
 	ENGINE_API virtual const UMaterial* GetMaterial_Concurrent(TMicRecursionGuard RecursionGuard = TMicRecursionGuard()) const override;
-	ENGINE_API virtual bool GetScalarParameterValue(const FHashedMaterialParameterInfo& ParameterInfo,float& OutValue, bool bOveriddenOnly = false) const override;
-#if WITH_EDITOR
-	ENGINE_API virtual bool IsScalarParameterUsedAsAtlasPosition(const FHashedMaterialParameterInfo& ParameterInfo, bool& OutValue, TSoftObjectPtr<class UCurveLinearColor>& Curve, TSoftObjectPtr<class UCurveLinearColorAtlas>& Atlas) const override;
-	ENGINE_API virtual bool GetScalarParameterSliderMinMax(const FHashedMaterialParameterInfo& ParameterInfo, float& OutMinSlider, float& OutMaxSlider) const override;
-#endif
-	ENGINE_API virtual bool GetVectorParameterValue(const FHashedMaterialParameterInfo& ParameterInfo,FLinearColor& OutValue, bool bOveriddenOnly = false) const override;
-#if WITH_EDITOR
-	ENGINE_API virtual bool IsVectorParameterUsedAsChannelMask(const FHashedMaterialParameterInfo& ParameterInfo, bool& OutValue) const override;
-	ENGINE_API virtual bool GetVectorParameterChannelNames(const FHashedMaterialParameterInfo& ParameterInfo, FParameterChannelNames& OutValue) const override;
-#endif
-	ENGINE_API virtual bool GetTextureParameterValue(const FHashedMaterialParameterInfo& ParameterInfo,class UTexture*& OutValue, bool bOveriddenOnly = false) const override;
-	ENGINE_API virtual bool GetRuntimeVirtualTextureParameterValue(const FHashedMaterialParameterInfo& ParameterInfo, class URuntimeVirtualTexture*& OutValue, bool bOveriddenOnly = false) const override;
-#if WITH_EDITOR
-	ENGINE_API virtual bool GetTextureParameterChannelNames(const FHashedMaterialParameterInfo& ParameterInfo, FParameterChannelNames& OutValue) const override;
-#endif
-	ENGINE_API virtual bool GetFontParameterValue(const FHashedMaterialParameterInfo& ParameterInfo,class UFont*& OutFontValue,int32& OutFontPage, bool bOveriddenOnly = false) const override;
+
+	ENGINE_API virtual bool GetParameterValue(EMaterialParameterType Type, const FMemoryImageMaterialParameterInfo& ParameterInfo, FMaterialParameterMetadata& OutValue, EMaterialGetParameterValueFlags Flags) const override;
+
 	ENGINE_API virtual bool GetRefractionSettings(float& OutBiasValue) const override;
 	ENGINE_API virtual void GetDependencies(TSet<UMaterialInterface*>& Dependencies) override;
 	ENGINE_API virtual FMaterialRenderProxy* GetRenderProxy() const override;
@@ -1068,8 +1055,6 @@ public:
 	ENGINE_API virtual FMaterialResource* GetMaterialResource(ERHIFeatureLevel::Type InFeatureLevel, EMaterialQualityLevel::Type QualityLevel = EMaterialQualityLevel::Num) override;
 	ENGINE_API virtual const FMaterialResource* GetMaterialResource(ERHIFeatureLevel::Type InFeatureLevel, EMaterialQualityLevel::Type QualityLevel = EMaterialQualityLevel::Num) const override;
 #if WITH_EDITORONLY_DATA
-	ENGINE_API virtual bool GetStaticSwitchParameterValue(const FHashedMaterialParameterInfo& ParameterInfo, bool& OutValue, FGuid& OutExpressionGuid, bool bOveriddenOnly = false, bool bCheckParent = true) const override;
-	ENGINE_API virtual bool GetStaticComponentMaskParameterValue(const FHashedMaterialParameterInfo& ParameterInfo, bool& R, bool& G, bool& B, bool& A, FGuid& OutExpressionGuid, bool bOveriddenOnly = false, bool bCheckParent = true) const override;
 	ENGINE_API virtual bool GetMaterialLayersParameterValue(const FHashedMaterialParameterInfo& ParameterInfo, FMaterialLayersFunctions& OutLayers, FGuid& OutExpressionGuid, bool bCheckParent = true) const override;
 #endif // WITH_EDITORONLY_DATA
 	ENGINE_API virtual bool GetTerrainLayerWeightParameterValue(const FHashedMaterialParameterInfo& ParameterInfo, int32& OutWeightmapIndex, FGuid& OutExpressionGuid) const override;
@@ -1199,12 +1184,12 @@ public:
 	ENGINE_API bool IsCompilingOrHadCompileError(ERHIFeatureLevel::Type InFeatureLevel);
 
 #if WITH_EDITOR
+	ENGINE_API bool SetParameterValueEditorOnly(const FName& ParameterName, const FMaterialParameterMetadata& Meta);
 	ENGINE_API bool SetVectorParameterValueEditorOnly(FName ParameterName, FLinearColor InValue);
 	ENGINE_API bool SetScalarParameterValueEditorOnly(FName ParameterName, float InValue);
 	ENGINE_API bool SetTextureParameterValueEditorOnly(FName ParameterName, class UTexture* InValue);
 	ENGINE_API bool SetRuntimeVirtualTextureParameterValueEditorOnly(FName ParameterName, class URuntimeVirtualTexture* InValue);
 	ENGINE_API bool SetFontParameterValueEditorOnly(FName ParameterName, class UFont* InFontValue, int32 InFontPage);
-	ENGINE_API bool SetMaterialLayersParameterValueEditorOnly(FName ParameterName, struct FMaterialLayersFunctions InValue, FGuid OutExpressionGuid);
 	ENGINE_API bool SetStaticComponentMaskParameterValueEditorOnly(FName ParameterName, bool R, bool G, bool B, bool A, FGuid OutExpressionGuid);
 	ENGINE_API bool SetStaticSwitchParameterValueEditorOnly(FName ParameterName, bool OutValue, FGuid OutExpressionGuid);
 #endif // WITH_EDITOR
@@ -1228,19 +1213,6 @@ public:
 	ENGINE_API virtual void GetLightingGuidChain(bool bIncludeTextures, TArray<FGuid>& OutGuids) const override;
 
 private:
-#if WITH_EDITOR
-	bool GetScalarParameterValue_Legacy(const FHashedMaterialParameterInfo& ParameterInfo, float& OutValue, bool bOveriddenOnly) const;
-	bool GetVectorParameterValue_Legacy(const FHashedMaterialParameterInfo& ParameterInfo, FLinearColor& OutValue, bool bOveriddenOnly) const;
-	bool GetTextureParameterValue_Legacy(const FHashedMaterialParameterInfo& ParameterInfo, class UTexture*& OutValue, bool bOveriddenOnly) const;
-	bool GetRuntimeVirtualTextureParameterValue_Legacy(const FHashedMaterialParameterInfo& ParameterInfo, class URuntimeVirtualTexture*& OutValue, bool bOveriddenOnly) const;
-	bool GetFontParameterValue_Legacy(const FHashedMaterialParameterInfo& ParameterInfo, class UFont*& OutFontValue, int32& OutFontPage, bool bOveriddenOnly) const;
-#endif // WITH_EDITOR
-	bool GetScalarParameterValue_New(const FHashedMaterialParameterInfo& ParameterInfo, float& OutValue, bool bOveriddenOnly) const;
-	bool GetVectorParameterValue_New(const FHashedMaterialParameterInfo& ParameterInfo, FLinearColor& OutValue, bool bOveriddenOnly) const;
-	bool GetTextureParameterValue_New(const FHashedMaterialParameterInfo& ParameterInfo, class UTexture*& OutValue, bool bOveriddenOnly) const;
-	bool GetRuntimeVirtualTextureParameterValue_New(const FHashedMaterialParameterInfo& ParameterInfo, class URuntimeVirtualTexture*& OutValue, bool bOveriddenOnly) const;
-	bool GetFontParameterValue_New(const FHashedMaterialParameterInfo& ParameterInfo, class UFont*& OutFontValue, int32& OutFontPage, bool bOveriddenOnly) const;
-
 	void BackwardsCompatibilityInputConversion();
 	void BackwardsCompatibilityVirtualTextureOutputConversion();
 	void BackwardsCompatibilityDecalConversion();
@@ -1371,30 +1343,14 @@ public:
 	}
 #endif // WITH_EDITORONLY_DATA
 
-	ENGINE_API virtual void GetAllScalarParameterInfo(TArray<FMaterialParameterInfo>& OutParameterInfo, TArray<FGuid>& OutParameterIds) const override;
-	ENGINE_API virtual void GetAllVectorParameterInfo(TArray<FMaterialParameterInfo>& OutParameterInfo, TArray<FGuid>& OutParameterIds) const override;
-	ENGINE_API virtual void GetAllTextureParameterInfo(TArray<FMaterialParameterInfo>& OutParameterInfo, TArray<FGuid>& OutParameterIds) const override;
-	ENGINE_API virtual void GetAllFontParameterInfo(TArray<FMaterialParameterInfo>& OutParameterInfo, TArray<FGuid>& OutParameterIds) const override;
-	ENGINE_API virtual void GetAllRuntimeVirtualTextureParameterInfo(TArray<FMaterialParameterInfo>& OutParameterInfo, TArray<FGuid>& OutParameterIds) const override;
+	ENGINE_API virtual void GetAllParameterInfoOfType(EMaterialParameterType Type, TArray<FMaterialParameterInfo>& OutParameterInfo, TArray<FGuid>& OutParameterIds) const override;
+	ENGINE_API virtual void GetAllParametersOfType(EMaterialParameterType Type, TMap<FMaterialParameterInfo, FMaterialParameterMetadata>& OutParameters) const override;
 #if WITH_EDITORONLY_DATA
 	ENGINE_API virtual void GetAllMaterialLayersParameterInfo(TArray<FMaterialParameterInfo>& OutParameterInfo, TArray<FGuid>& OutParameterIds) const override;
-	ENGINE_API virtual void GetAllStaticSwitchParameterInfo(TArray<FMaterialParameterInfo>& OutParameterInfo, TArray<FGuid>& OutParameterIds) const override;
-	ENGINE_API virtual void GetAllStaticComponentMaskParameterInfo(TArray<FMaterialParameterInfo>& OutParameterInfo, TArray<FGuid>& OutParameterIds) const override;
-
 	ENGINE_API virtual bool IterateDependentFunctions(TFunctionRef<bool(UMaterialFunctionInterface*)> Predicate) const override;
 	ENGINE_API virtual void GetDependentFunctions(TArray<class UMaterialFunctionInterface*>& DependentFunctions) const override;
 #endif // WITH_EDITORONLY_DATA
 
-	ENGINE_API virtual bool GetScalarParameterDefaultValue(const FHashedMaterialParameterInfo& ParameterInfo, float& OutValue, bool bOveriddenOnly = false, bool bCheckOwnedGlobalOverrides = false) const override;
-	ENGINE_API virtual bool GetVectorParameterDefaultValue(const FHashedMaterialParameterInfo& ParameterInfo, FLinearColor& OutValue, bool bOveriddenOnly = false, bool bCheckOwnedGlobalOverrides = false) const override;
-	ENGINE_API virtual bool GetTextureParameterDefaultValue(const FHashedMaterialParameterInfo& ParameterInfo, class UTexture*& OutValue, bool bCheckOwnedGlobalOverrides = false) const override;
-	ENGINE_API virtual bool GetRuntimeVirtualTextureParameterDefaultValue(const FHashedMaterialParameterInfo& ParameterInfo, class URuntimeVirtualTexture*& OutValue, bool bCheckOwnedGlobalOverrides = false) const override;
-	ENGINE_API virtual bool GetFontParameterDefaultValue(const FHashedMaterialParameterInfo& ParameterInfo, class UFont*& OutFontValue, int32& OutFontPage, bool bCheckOwnedGlobalOverrides = false) const override;
-#if WITH_EDITOR
-	ENGINE_API virtual bool GetStaticComponentMaskParameterDefaultValue(const FHashedMaterialParameterInfo& ParameterInfo, bool& OutR, bool& OutG, bool& OutB, bool& OutA, FGuid& OutExpressionGuid, bool bCheckOwnedGlobalOverrides = false) const override;
-	ENGINE_API virtual bool GetStaticSwitchParameterDefaultValue(const FHashedMaterialParameterInfo& ParameterInfo, bool& OutValue, FGuid& OutExpressionGuid, bool bCheckOwnedGlobalOverrides = false) const override;
-#endif // WITH_EDITOR
-
 	UE_DEPRECATED(5.0, "No longer used.")
 	uint32 GetDecalBlendMode() const { return DecalBlendMode; }
 
@@ -1948,66 +1904,6 @@ private:
 		return nullptr;
 	}
 #endif // WITH_EDITORONLY_DATA
-
-#if WITH_EDITOR
-
-	template<typename ParameterType, typename ...Arguments>
-	bool SetParameterValueEditorOnly(FName InParameterName, Arguments... Args)
-	{
-		// Warning: in the case of duplicate parameters with different default values, this will find the first in the expression array, not necessarily the one that's used for rendering
-
-		// Lambda which calls SetParameterValue on a given parameter and triggers a PostEditChange event if successful
-		auto TrySetParameterValue = [&](ParameterType* Parameter) -> bool
-		{
-			if (Parameter && Parameter->SetParameterValue(InParameterName, Args...))
-			{
-				if (FProperty* ParamProperty = FindFProperty<FProperty>(ParameterType::StaticClass(), "DefaultValue"))
-				{
-					FPropertyChangedEvent PropertyChangedEvent(ParamProperty);
-					Parameter->PostEditChangeProperty(PropertyChangedEvent);
-					return true;
-				}
-			}
-			return false;
-		};
-
-		for (TObjectPtr<UMaterialExpression>& Expression : Expressions)
-		{
-			if (TrySetParameterValue(Cast<ParameterType>(Expression)))
-			{
-				return true;
-			}
-			else if (UMaterialExpressionMaterialFunctionCall* FunctionCall = Cast<UMaterialExpressionMaterialFunctionCall>(Expression))
-			{
-				if (FunctionCall->MaterialFunction)
-				{
-					TArray<UMaterialFunctionInterface*> Functions;
-					Functions.Add(FunctionCall->MaterialFunction);
-					FunctionCall->MaterialFunction->GetDependentFunctions(Functions);
-
-					for (UMaterialFunctionInterface* Function : Functions)
-					{
-						const TArray<TObjectPtr<UMaterialExpression>>* ExpressionPtr = Function->GetFunctionExpressions();
-						if (ExpressionPtr)
-						{
-							for (const TObjectPtr<UMaterialExpression>& FunctionExpression : *ExpressionPtr)
-							{
-								if (TrySetParameterValue(Cast<ParameterType>(FunctionExpression)))
-								{
-									return true;
-								}
-							}
-						}
-					}
-				}
-			}
-		}
-
-		return false;
-	}
-
-#endif // WITH_EDITOR
-
 };
 
 
diff --git a/Engine/Source/Runtime/Engine/Classes/Materials/MaterialExpression.h b/Engine/Source/Runtime/Engine/Classes/Materials/MaterialExpression.h
--- a/Engine/Source/Runtime/Engine/Classes/Materials/MaterialExpression.h
+++ b/Engine/Source/Runtime/Engine/Classes/Materials/MaterialExpression.h
@@ -519,6 +519,7 @@ class ENGINE_API UMaterialExpression : public UObject
 
 	virtual FName GetParameterName() const { return NAME_None; }
 	virtual void SetParameterName(const FName& Name) {}
+	virtual bool SetParameterValue(const FName& Name, const FMaterialParameterMetadata& Meta) { return false; }
 
 	/**
 	 * Called after a node copy, once the Material and Function properties are set correctly and that all new expressions are added to Material->Expressions
diff --git a/Engine/Source/Runtime/Engine/Classes/Materials/MaterialExpressionFontSampleParameter.h b/Engine/Source/Runtime/Engine/Classes/Materials/MaterialExpressionFontSampleParameter.h
--- a/Engine/Source/Runtime/Engine/Classes/Materials/MaterialExpressionFontSampleParameter.h
+++ b/Engine/Source/Runtime/Engine/Classes/Materials/MaterialExpressionFontSampleParameter.h
@@ -49,6 +49,14 @@ class UMaterialExpressionFontSampleParameter : public UMaterialExpressionFontSam
 	virtual void SetParameterName(const FName& Name) override { ParameterName = Name; }
 	virtual void ValidateParameterName(const bool bAllowDuplicateName) override;
 	virtual void SetValueToMatchingExpression(UMaterialExpression* OtherExpression) override;
+	virtual bool SetParameterValue(const FName& Name, const FMaterialParameterMetadata& Meta) override
+	{
+		if (Meta.Value.Type == EMaterialParameterType::Font)
+		{
+			return SetParameterValue(Name, Meta.Value.Font.Value, Meta.Value.Font.Page);
+		}
+		return Super::SetParameterValue(Name, Meta);
+	}
 #endif
 	//~ End UMaterialExpression Interface
 	
diff --git a/Engine/Source/Runtime/Engine/Classes/Materials/MaterialExpressionRuntimeVirtualTextureSampleParameter.h b/Engine/Source/Runtime/Engine/Classes/Materials/MaterialExpressionRuntimeVirtualTextureSampleParameter.h
--- a/Engine/Source/Runtime/Engine/Classes/Materials/MaterialExpressionRuntimeVirtualTextureSampleParameter.h
+++ b/Engine/Source/Runtime/Engine/Classes/Materials/MaterialExpressionRuntimeVirtualTextureSampleParameter.h
@@ -51,6 +51,14 @@ class UMaterialExpressionRuntimeVirtualTextureSampleParameter : public UMaterial
 	virtual void GetCaption(TArray<FString>& OutCaptions) const override;
 	virtual bool MatchesSearchQuery(const TCHAR* SearchQuery) override;
 	virtual void SetValueToMatchingExpression(UMaterialExpression* OtherExpression) override;
+	virtual bool SetParameterValue(const FName& Name, const FMaterialParameterMetadata& Meta) override
+	{
+		if (Meta.Value.Type == EMaterialParameterType::RuntimeVirtualTexture)
+		{
+			return SetParameterValue(Name, Meta.Value.RuntimeVirtualTexture);
+		}
+		return Super::SetParameterValue(Name, Meta);
+	}
 #endif
 	virtual FGuid& GetParameterExpressionId() override { return ExpressionGUID; }
 	//~ End UMaterialExpression Interface
diff --git a/Engine/Source/Runtime/Engine/Classes/Materials/MaterialExpressionScalarParameter.h b/Engine/Source/Runtime/Engine/Classes/Materials/MaterialExpressionScalarParameter.h
--- a/Engine/Source/Runtime/Engine/Classes/Materials/MaterialExpressionScalarParameter.h
+++ b/Engine/Source/Runtime/Engine/Classes/Materials/MaterialExpressionScalarParameter.h
@@ -44,6 +44,14 @@ class UMaterialExpressionScalarParameter : public UMaterialExpressionParameter
 	virtual int32 Compile(class FMaterialCompiler* Compiler, int32 OutputIndex) override;
 	virtual void GetCaption(TArray<FString>& OutCaptions) const override;
 	EMaterialGenerateHLSLStatus GenerateHLSLExpression(FMaterialHLSLGenerator& Generator, UE::HLSLTree::FScope& Scope, int32 OutputIndex, UE::HLSLTree::FExpression*& OutExpression) override;
+	virtual bool SetParameterValue(const FName& Name, const FMaterialParameterMetadata& Meta) override
+	{
+		if (Meta.Value.Type == EMaterialParameterType::Scalar)
+		{
+			return SetParameterValue(Name, Meta.Value.AsScalar());
+		}
+		return Super::SetParameterValue(Name, Meta);
+	}
 #endif
 	//~ End UMaterialExpression Interface
 
diff --git a/Engine/Source/Runtime/Engine/Classes/Materials/MaterialExpressionStaticBoolParameter.h b/Engine/Source/Runtime/Engine/Classes/Materials/MaterialExpressionStaticBoolParameter.h
--- a/Engine/Source/Runtime/Engine/Classes/Materials/MaterialExpressionStaticBoolParameter.h
+++ b/Engine/Source/Runtime/Engine/Classes/Materials/MaterialExpressionStaticBoolParameter.h
@@ -25,6 +25,14 @@ public:
 	virtual uint32 GetOutputType(int32 OutputIndex) override {return MCT_StaticBool;}
 	virtual void SetValueToMatchingExpression(UMaterialExpression* OtherExpression) override;
 	virtual EMaterialGenerateHLSLStatus GenerateHLSLExpression(FMaterialHLSLGenerator& Generator, UE::HLSLTree::FScope& Scope, int32 OutputIndex, UE::HLSLTree::FExpression*& OutExpression) override;
+	virtual bool SetParameterValue(const FName& Name, const FMaterialParameterMetadata& Meta) override
+	{
+		if (Meta.Value.Type == EMaterialParameterType::StaticSwitch)
+		{
+			return SetParameterValue(Name, Meta.Value.AsStaticSwitch(), Meta.ExpressionGuid);
+		}
+		return Super::SetParameterValue(Name, Meta);
+	}
 #endif
 	//~ End UMaterialExpression Interface
 
diff --git a/Engine/Source/Runtime/Engine/Classes/Materials/MaterialExpressionStaticComponentMaskParameter.h b/Engine/Source/Runtime/Engine/Classes/Materials/MaterialExpressionStaticComponentMaskParameter.h
--- a/Engine/Source/Runtime/Engine/Classes/Materials/MaterialExpressionStaticComponentMaskParameter.h
+++ b/Engine/Source/Runtime/Engine/Classes/Materials/MaterialExpressionStaticComponentMaskParameter.h
@@ -39,6 +39,19 @@ public:
 	virtual int32 Compile(class FMaterialCompiler* Compiler, int32 OutputIndex) override;
 	virtual void GetCaption(TArray<FString>& OutCaptions) const override;
 	virtual void SetValueToMatchingExpression(UMaterialExpression* OtherExpression) override;
+	virtual bool SetParameterValue(const FName& Name, const FMaterialParameterMetadata& Meta) override
+	{
+		if (Meta.Value.Type == EMaterialParameterType::StaticComponentMask)
+		{
+			return SetParameterValue(Name,
+				Meta.Value.Bool[0],
+				Meta.Value.Bool[1],
+				Meta.Value.Bool[2],
+				Meta.Value.Bool[3],
+				Meta.ExpressionGuid);
+		}
+		return Super::SetParameterValue(Name, Meta);
+	}
 #endif
 	//~ End UMaterialExpression Interface
 
diff --git a/Engine/Source/Runtime/Engine/Classes/Materials/MaterialExpressionTextureSampleParameter.h b/Engine/Source/Runtime/Engine/Classes/Materials/MaterialExpressionTextureSampleParameter.h
--- a/Engine/Source/Runtime/Engine/Classes/Materials/MaterialExpressionTextureSampleParameter.h
+++ b/Engine/Source/Runtime/Engine/Classes/Materials/MaterialExpressionTextureSampleParameter.h
@@ -52,6 +52,14 @@ class ENGINE_API UMaterialExpressionTextureSampleParameter : public UMaterialExp
 	virtual void ValidateParameterName(const bool bAllowDuplicateName) override;
 	virtual void SetValueToMatchingExpression(UMaterialExpression* OtherExpression) override;
 	virtual EMaterialGenerateHLSLStatus GenerateHLSLExpression(FMaterialHLSLGenerator& Generator, UE::HLSLTree::FScope& Scope, int32 OutputIndex, UE::HLSLTree::FExpression*& OutExpression) override;
+	virtual bool SetParameterValue(const FName& Name, const FMaterialParameterMetadata& Meta) override
+	{
+		if (Meta.Value.Type == EMaterialParameterType::Texture)
+		{
+			return SetParameterValue(Name, Meta.Value.Texture);
+		}
+		return Super::SetParameterValue(Name, Meta);
+	}
 #endif
 	//~ End UMaterialExpression Interface
 
diff --git a/Engine/Source/Runtime/Engine/Classes/Materials/MaterialExpressionVectorParameter.h b/Engine/Source/Runtime/Engine/Classes/Materials/MaterialExpressionVectorParameter.h
--- a/Engine/Source/Runtime/Engine/Classes/Materials/MaterialExpressionVectorParameter.h
+++ b/Engine/Source/Runtime/Engine/Classes/Materials/MaterialExpressionVectorParameter.h
@@ -33,6 +33,14 @@ class UMaterialExpressionVectorParameter : public UMaterialExpressionParameter
 	virtual int32 Compile(class FMaterialCompiler* Compiler, int32 OutputIndex) override;
 	virtual void GetCaption(TArray<FString>& OutCaptions) const override;
 	EMaterialGenerateHLSLStatus GenerateHLSLExpression(FMaterialHLSLGenerator& Generator, UE::HLSLTree::FScope& Scope, int32 OutputIndex, UE::HLSLTree::FExpression*& OutExpression) override;
+	virtual bool SetParameterValue(const FName& Name, const FMaterialParameterMetadata& Meta) override
+	{
+		if (Meta.Value.Type == EMaterialParameterType::Vector)
+		{
+			return SetParameterValue(Name, Meta.Value.AsLinearColor());
+		}
+		return Super::SetParameterValue(Name, Meta);
+	}
 #endif
 	//~ End UMaterialExpression Interface
 
diff --git a/Engine/Source/Runtime/Engine/Classes/Materials/MaterialFunction.h b/Engine/Source/Runtime/Engine/Classes/Materials/MaterialFunction.h
--- a/Engine/Source/Runtime/Engine/Classes/Materials/MaterialFunction.h
+++ b/Engine/Source/Runtime/Engine/Classes/Materials/MaterialFunction.h
@@ -164,6 +164,7 @@ public:
 
 
 #if WITH_EDITOR
+	ENGINE_API bool SetParameterValueEditorOnly(const FName& ParameterName, const FMaterialParameterMetadata& Meta);
 	ENGINE_API bool SetVectorParameterValueEditorOnly(FName ParameterName, FLinearColor InValue);
 	ENGINE_API bool SetScalarParameterValueEditorOnly(FName ParameterName, float InValue);
 	ENGINE_API bool SetTextureParameterValueEditorOnly(FName ParameterName, class UTexture* InValue);
diff --git a/Engine/Source/Runtime/Engine/Classes/Materials/MaterialInstance.h b/Engine/Source/Runtime/Engine/Classes/Materials/MaterialInstance.h
--- a/Engine/Source/Runtime/Engine/Classes/Materials/MaterialInstance.h
+++ b/Engine/Source/Runtime/Engine/Classes/Materials/MaterialInstance.h
@@ -93,6 +93,8 @@ struct FScalarParameterValue
 	{
 	}
 
+	bool IsOverride() const { return true; }
+
 	bool operator==(const FScalarParameterValue& Other) const
 	{
 		return
@@ -107,6 +109,17 @@ struct FScalarParameterValue
 	
 	typedef float ValueType;
 	static ValueType GetValue(const FScalarParameterValue& Parameter) { return Parameter.ParameterValue; }
+
+	void GetValue(FMaterialParameterMetadata& OutResult) const
+	{
+		OutResult.Value = ParameterValue;
+#if WITH_EDITORONLY_DATA
+		OutResult.ExpressionGuid = ExpressionGUID;
+		OutResult.ScalarCurve = AtlasData.Curve;
+		OutResult.ScalarAtlas = AtlasData.Atlas;
+		OutResult.bUsedAsAtlasPosition = AtlasData.bIsUsedAsAtlasPosition;
+#endif // WITH_EDITORONLY_DATA
+	}
 };
 
 /** Editable vector parameter. */
@@ -134,6 +147,8 @@ struct FVectorParameterValue
 	{
 	}
 
+	bool IsOverride() const { return true; }
+
 	bool operator==(const FVectorParameterValue& Other) const
 	{
 		return
@@ -148,6 +163,14 @@ struct FVectorParameterValue
 	
 	typedef FLinearColor ValueType;
 	static ValueType GetValue(const FVectorParameterValue& Parameter) { return Parameter.ParameterValue; }
+
+	void GetValue(FMaterialParameterMetadata& OutResult) const
+	{
+		OutResult.Value = ParameterValue;
+#if WITH_EDITORONLY_DATA
+		OutResult.ExpressionGuid = ExpressionGUID;
+#endif
+	}
 };
 
 /** Editable texture parameter. */
@@ -175,6 +198,8 @@ struct FTextureParameterValue
 	{
 	}
 
+	bool IsOverride() const { return true; }
+
 	bool operator==(const FTextureParameterValue& Other) const
 	{
 		return
@@ -189,6 +214,14 @@ struct FTextureParameterValue
 
 	typedef const UTexture* ValueType;
 	static ValueType GetValue(const FTextureParameterValue& Parameter) { return Parameter.ParameterValue; }
+
+	void GetValue(FMaterialParameterMetadata& OutResult) const
+	{
+		OutResult.Value = ParameterValue;
+#if WITH_EDITORONLY_DATA
+		OutResult.ExpressionGuid = ExpressionGUID;
+#endif
+	}
 };
 
 /** Editable runtime virtual texture parameter. */
@@ -211,6 +244,8 @@ struct FRuntimeVirtualTextureParameterValue
 	{
 	}
 
+	bool IsOverride() const { return true; }
+
 	bool operator==(const FRuntimeVirtualTextureParameterValue& Other) const
 	{
 		return
@@ -225,6 +260,14 @@ struct FRuntimeVirtualTextureParameterValue
 
 	typedef const URuntimeVirtualTexture* ValueType;
 	static ValueType GetValue(const FRuntimeVirtualTextureParameterValue& Parameter) { return Parameter.ParameterValue; }
+
+	void GetValue(FMaterialParameterMetadata& OutResult) const
+	{
+		OutResult.Value = ParameterValue;
+#if WITH_EDITORONLY_DATA
+		OutResult.ExpressionGuid = ExpressionGUID;
+#endif
+	}
 };
 
 /** Editable font parameter. */
@@ -256,6 +299,8 @@ struct FFontParameterValue
 	{
 	}
 
+	bool IsOverride() const { return true; }
+
 	bool operator==(const FFontParameterValue& Other) const
 	{
 		return
@@ -271,6 +316,14 @@ struct FFontParameterValue
 	
 	typedef const UTexture* ValueType;
 	static ValueType GetValue(const FFontParameterValue& Parameter);
+
+	void GetValue(FMaterialParameterMetadata& OutResult) const
+	{
+		OutResult.Value = FMaterialParameterValue(FontValue, FontPage);
+#if WITH_EDITORONLY_DATA
+		OutResult.ExpressionGuid = ExpressionGUID;
+#endif
+	}
 };
 
 template<class T>
@@ -291,6 +344,26 @@ bool CompareValueArraysByExpressionGUID(const TArray<T>& InA, const TArray<T>& I
 	return AA == BB;
 }
 
+#if WITH_EDITORONLY_DATA
+class FMaterialInstanceParameterUpdateContext
+{
+public:
+	ENGINE_API explicit FMaterialInstanceParameterUpdateContext(UMaterialInstance* InInstance);
+	ENGINE_API ~FMaterialInstanceParameterUpdateContext();
+
+	inline FStaticParameterSet& GetStaticParameters() { return StaticParameters; }
+
+	ENGINE_API void SetParameterValueEditorOnly(const FMaterialParameterInfo& ParameterInfo, const FMaterialParameterMetadata& Meta, EMaterialSetParameterValueFlags Flags = EMaterialSetParameterValueFlags::None);
+	ENGINE_API void SetForceStaticPermutationUpdate(bool bValue);
+	ENGINE_API void SetBasePropertyOverrides(const FMaterialInstanceBasePropertyOverrides& InValue);
+
+private:
+	UMaterialInstance* Instance;
+	FStaticParameterSet StaticParameters;
+	FMaterialInstanceBasePropertyOverrides BasePropertyOverrides;
+	bool bForceStaticPermutationUpdate;
+};
+#endif // WITH_EDITORONLY_DATA
 
 UCLASS(abstract, BlueprintType,MinimalAPI)
 class UMaterialInstance : public UMaterialInterface
@@ -487,24 +560,8 @@ public:
 	virtual ENGINE_API FMaterialResource* GetMaterialResource(ERHIFeatureLevel::Type InFeatureLevel, EMaterialQualityLevel::Type QualityLevel = EMaterialQualityLevel::Num) override;
 	virtual ENGINE_API const FMaterialResource* GetMaterialResource(ERHIFeatureLevel::Type InFeatureLevel, EMaterialQualityLevel::Type QualityLevel = EMaterialQualityLevel::Num) const override;
 
-#if WITH_EDITOR
-	virtual ENGINE_API bool GetScalarParameterSliderMinMax(const FHashedMaterialParameterInfo& ParameterInfo, float& OutSliderMin, float& OutSliderMax) const override;
-#endif
-	virtual ENGINE_API bool GetScalarParameterValue(const FHashedMaterialParameterInfo& ParameterInfo, float& OutValue, bool bOveriddenOnly = false) const override;
-#if WITH_EDITOR
-	virtual ENGINE_API bool IsScalarParameterUsedAsAtlasPosition(const FHashedMaterialParameterInfo& ParameterInfo, bool& OutValue, TSoftObjectPtr<class UCurveLinearColor>& Curve, TSoftObjectPtr<class UCurveLinearColorAtlas>& Atlas) const override;
-#endif
-	virtual ENGINE_API bool GetVectorParameterValue(const FHashedMaterialParameterInfo& ParameterInfo, FLinearColor& OutValue, bool bOveriddenOnly = false) const override;
-#if WITH_EDITOR
-	virtual ENGINE_API bool IsVectorParameterUsedAsChannelMask(const FHashedMaterialParameterInfo& ParameterInfo, bool& OutValue) const override;
-	virtual ENGINE_API bool GetVectorParameterChannelNames(const FHashedMaterialParameterInfo& ParameterInfo, FParameterChannelNames& OutValue) const override;
-#endif
-	virtual ENGINE_API bool GetTextureParameterValue(const FHashedMaterialParameterInfo& ParameterInfo, class UTexture*& OutValue, bool bOveriddenOnly = false) const override;
-	virtual ENGINE_API bool GetRuntimeVirtualTextureParameterValue(const FHashedMaterialParameterInfo& ParameterInfo, class URuntimeVirtualTexture*& OutValue, bool bOveriddenOnly = false) const override;
-#if WITH_EDITOR
-	virtual ENGINE_API bool GetTextureParameterChannelNames(const FHashedMaterialParameterInfo& ParameterInfo, FParameterChannelNames& OutValue) const override;
-#endif
-	virtual ENGINE_API bool GetFontParameterValue(const FHashedMaterialParameterInfo& ParameterInfo, class UFont*& OutFontValue, int32& OutFontPage, bool bOveriddenOnly = false) const override;
+	virtual ENGINE_API bool GetParameterValue(EMaterialParameterType Type, const FMemoryImageMaterialParameterInfo& ParameterInfo, FMaterialParameterMetadata& OutValue, EMaterialGetParameterValueFlags Flags) const override;
+
 	virtual ENGINE_API void GetUsedTextures(TArray<UTexture*>& OutTextures, EMaterialQualityLevel::Type QualityLevel, bool bAllQualityLevels, ERHIFeatureLevel::Type FeatureLevel, bool bAllFeatureLevels) const override;
 	virtual ENGINE_API void GetUsedTexturesAndIndices(TArray<UTexture*>& OutTextures, TArray< TArray<int32> >& OutIndices, EMaterialQualityLevel::Type QualityLevel, ERHIFeatureLevel::Type FeatureLevel) const;
 	virtual ENGINE_API void OverrideTexture(const UTexture* InTextureToOverride, UTexture* OverrideTexture, ERHIFeatureLevel::Type InFeatureLevel) override;
@@ -513,8 +570,6 @@ public:
 	virtual ENGINE_API bool CheckMaterialUsage(const EMaterialUsage Usage) override;
 	virtual ENGINE_API bool CheckMaterialUsage_Concurrent(const EMaterialUsage Usage) const override;
 #if WITH_EDITORONLY_DATA
-	virtual ENGINE_API bool GetStaticSwitchParameterValue(const FHashedMaterialParameterInfo& ParameterInfo, bool& OutValue, FGuid& OutExpressionGuid, bool bOveriddenOnly = false, bool bCheckParent = true) const override;
-	virtual ENGINE_API bool GetStaticComponentMaskParameterValue(const FHashedMaterialParameterInfo& ParameterInfo, bool& R, bool& G, bool& B, bool& A, FGuid& OutExpressionGuid, bool bOveriddenOnly = false, bool bCheckParent = true) const override;
 	virtual ENGINE_API bool GetMaterialLayersParameterValue(const FHashedMaterialParameterInfo& ParameterInfo, FMaterialLayersFunctions& OutLayers, FGuid &OutExpressionGuid, bool bCheckParent = true) const override;
 #endif // WITH_EDITORONLY_DATA
 	virtual ENGINE_API bool GetTerrainLayerWeightParameterValue(const FHashedMaterialParameterInfo& ParameterInfo, int32& OutWeightmapIndex, FGuid &OutExpressionGuid) const override;
@@ -650,31 +705,15 @@ public:
 	}
 #endif // WITH_EDITORONLY_DATA
 
-	void GetAllParametersOfType(EMaterialParameterType Type, TArray<FMaterialParameterInfo>& OutParameterInfo, TArray<FGuid>& OutParameterIds) const;
-
-	ENGINE_API virtual void GetAllScalarParameterInfo(TArray<FMaterialParameterInfo>& OutParameterInfo, TArray<FGuid>& OutParameterIds) const override;
-	ENGINE_API virtual void GetAllVectorParameterInfo(TArray<FMaterialParameterInfo>& OutParameterInfo, TArray<FGuid>& OutParameterIds) const override;
-	ENGINE_API virtual void GetAllTextureParameterInfo(TArray<FMaterialParameterInfo>& OutParameterInfo, TArray<FGuid>& OutParameterIds) const override;
-	ENGINE_API virtual void GetAllRuntimeVirtualTextureParameterInfo(TArray<FMaterialParameterInfo>& OutParameterInfo, TArray<FGuid>& OutParameterIds) const override;
-	ENGINE_API virtual void GetAllFontParameterInfo(TArray<FMaterialParameterInfo>& OutParameterInfo, TArray<FGuid>& OutParameterIds) const override;
-
+	ENGINE_API virtual void GetAllParameterInfoOfType(EMaterialParameterType Type, TArray<FMaterialParameterInfo>& OutParameterInfo, TArray<FGuid>& OutParameterIds) const override;
+	ENGINE_API virtual void GetAllParametersOfType(EMaterialParameterType Type, TMap<FMaterialParameterInfo, FMaterialParameterMetadata>& OutParameters) const override;
 #if WITH_EDITORONLY_DATA
 	ENGINE_API virtual void GetAllMaterialLayersParameterInfo(TArray<FMaterialParameterInfo>& OutParameterInfo, TArray<FGuid>& OutParameterIds) const override;
-	ENGINE_API virtual void GetAllStaticSwitchParameterInfo(TArray<FMaterialParameterInfo>& OutParameterInfo, TArray<FGuid>& OutParameterIds) const override;
-	ENGINE_API virtual void GetAllStaticComponentMaskParameterInfo(TArray<FMaterialParameterInfo>& OutParameterInfo, TArray<FGuid>& OutParameterIds) const override;
-
 	ENGINE_API virtual bool IterateDependentFunctions(TFunctionRef<bool(UMaterialFunctionInterface*)> Predicate) const override;
 	ENGINE_API virtual void GetDependentFunctions(TArray<class UMaterialFunctionInterface*>& DependentFunctions) const override;
 #endif // WITH_EDITORONLY_DATA
 
-	ENGINE_API virtual bool GetScalarParameterDefaultValue(const FHashedMaterialParameterInfo& ParameterInfo, float& OutValue, bool bOveriddenOnly = false, bool bCheckOwnedGlobalOverrides = false) const override;
-	ENGINE_API virtual bool GetVectorParameterDefaultValue(const FHashedMaterialParameterInfo& ParameterInfo, FLinearColor& OutValue, bool bOveriddenOnly = false, bool bCheckOwnedGlobalOverrides = false) const override;
-	ENGINE_API virtual bool GetTextureParameterDefaultValue(const FHashedMaterialParameterInfo& ParameterInfo, UTexture*& OutValue, bool bCheckOwnedGlobalOverrides = false) const override;
-	ENGINE_API virtual bool GetRuntimeVirtualTextureParameterDefaultValue(const FHashedMaterialParameterInfo& ParameterInfo, URuntimeVirtualTexture*& OutValue, bool bCheckOwnedGlobalOverrides = false) const override;
-	ENGINE_API virtual bool GetFontParameterDefaultValue(const FHashedMaterialParameterInfo& ParameterInfo, UFont*& OutFontValue, int32& OutFontPage, bool bCheckOwnedGlobalOverrides = false) const override;
 #if WITH_EDITOR
-	ENGINE_API virtual bool GetStaticSwitchParameterDefaultValue(const FHashedMaterialParameterInfo& ParameterInfo, bool& OutValue, FGuid& OutExpressionGuid, bool bCheckOwnedGlobalOverrides = false) const override;
-	ENGINE_API virtual bool GetStaticComponentMaskParameterDefaultValue(const FHashedMaterialParameterInfo& ParameterInfo, bool& OutR, bool& OutG, bool& OutB, bool& OutA, FGuid& OutExpressionGuid, bool bCheckOwnedGlobalOverrides = false) const override;
 	ENGINE_API virtual bool GetGroupName(const FHashedMaterialParameterInfo& ParameterInfo, FName& OutGroup) const override;
 #endif // WITH_EDITOR
 
@@ -780,7 +819,7 @@ protected:
 	void SetVectorParameterValueInternal(const FMaterialParameterInfo& ParameterInfo, FLinearColor Value);
 	bool SetVectorParameterByIndexInternal(int32 ParameterIndex, FLinearColor Value);
 	bool SetScalarParameterByIndexInternal(int32 ParameterIndex, float Value);
-	void SetScalarParameterValueInternal(const FMaterialParameterInfo& ParameterInfo, float Value);
+	void SetScalarParameterValueInternal(const FMaterialParameterInfo& ParameterInfo, float Value, bool bUseAtlas = false, FScalarParameterAtlasInstanceData AtlasData = FScalarParameterAtlasInstanceData());
 #if WITH_EDITOR
 	void SetScalarParameterAtlasInternal(const FMaterialParameterInfo& ParameterInfo, FScalarParameterAtlasInstanceData AtlasData);
 #endif
@@ -820,6 +859,7 @@ protected:
 	friend class UMaterialEditingLibrary;
 	/** Class that knows how to update MI's */
 	friend class FMaterialUpdateContext;
+	friend class FMaterialInstanceParameterUpdateContext;
 };
 
 #if WITH_EDITOR
@@ -990,7 +1030,7 @@ namespace MaterialInstance_Private
 
 		TArray<FMaterialParameterInfo> CachedParamInfos;
 		TArray<FGuid> CachedParamGuids;
-		ParentMaterial->GetCachedExpressionData().Parameters.GetAllParameterInfoOfType(ParamTypeEnum, false, CachedParamInfos, CachedParamGuids);
+		ParentMaterial->GetCachedExpressionData().Parameters.GetAllParameterInfoOfType(ParamTypeEnum, CachedParamInfos, CachedParamGuids);
 		int32 NumCachedParams = CachedParamGuids.Num();
 		check(NumCachedParams == CachedParamInfos.Num());
 
diff --git a/Engine/Source/Runtime/Engine/Classes/Materials/MaterialInterface.h b/Engine/Source/Runtime/Engine/Classes/Materials/MaterialInterface.h
--- a/Engine/Source/Runtime/Engine/Classes/Materials/MaterialInterface.h
+++ b/Engine/Source/Runtime/Engine/Classes/Materials/MaterialInterface.h
@@ -4,6 +4,7 @@
 #pragma once
 
 #include "CoreMinimal.h"
+#include "MaterialTypes.h"
 #include "Containers/ArrayView.h"
 #include "UObject/ObjectMacros.h"
 #include "UObject/Object.h"
@@ -445,7 +446,7 @@ public:
 	* @param	OutValue		Will contain the value of the parameter if successful
 	* @return					True if successful
 	*/
-	virtual ENGINE_API bool GetStaticSwitchParameterValue(const FHashedMaterialParameterInfo& ParameterInfo,bool &OutValue,FGuid &OutExpressionGuid, bool bOveriddenOnly = false, bool bCheckParent = true) const;
+	ENGINE_API bool GetStaticSwitchParameterValue(const FHashedMaterialParameterInfo& ParameterInfo,bool &OutValue,FGuid &OutExpressionGuid, bool bOveriddenOnly = false, bool bCheckParent = true) const;
 
 	/**
 	* Get the value of the given static component mask parameter
@@ -454,7 +455,7 @@ public:
 	* @param	R, G, B, A		Will contain the values of the parameter if successful
 	* @return					True if successful
 	*/
-	virtual ENGINE_API bool GetStaticComponentMaskParameterValue(const FHashedMaterialParameterInfo& ParameterInfo, bool& R, bool& G, bool& B, bool& A, FGuid& OutExpressionGuid, bool bOveriddenOnly = false, bool bCheckParent = true) const;
+	ENGINE_API bool GetStaticComponentMaskParameterValue(const FHashedMaterialParameterInfo& ParameterInfo, bool& R, bool& G, bool& B, bool& A, FGuid& OutExpressionGuid, bool bOveriddenOnly = false, bool bCheckParent = true) const;
 
 	/**
 	* Get the value of the given static material layers parameter
@@ -499,24 +500,21 @@ public:
 		PURE_VIRTUAL(UMaterialInterface::GetGroupSortPriority, return false;);
 #endif // WITH_EDITOR
 
-	virtual void GetAllScalarParameterInfo(TArray<FMaterialParameterInfo>& OutParameterInfo, TArray<FGuid>& OutParameterIds) const
-		PURE_VIRTUAL(UMaterialInterface::GetAllScalarParameterInfo,return;);
-	virtual void GetAllVectorParameterInfo(TArray<FMaterialParameterInfo>& OutParameterInfo, TArray<FGuid>& OutParameterIds) const
-		PURE_VIRTUAL(UMaterialInterface::GetAllVectorParameterInfo,return;);
-	virtual void GetAllTextureParameterInfo(TArray<FMaterialParameterInfo>& OutParameterInfo, TArray<FGuid>& OutParameterIds) const
-		PURE_VIRTUAL(UMaterialInterface::GetAllTextureParameterInfo,return;);
-	virtual void GetAllRuntimeVirtualTextureParameterInfo(TArray<FMaterialParameterInfo>& OutParameterInfo, TArray<FGuid>& OutParameterIds) const
-		PURE_VIRTUAL(UMaterialInterface::GetAllRuntimeVirtualTextureParameterInfo, return;);
-	virtual void GetAllFontParameterInfo(TArray<FMaterialParameterInfo>& OutParameterInfo, TArray<FGuid>& OutParameterIds) const
-		PURE_VIRTUAL(UMaterialInterface::GetAllFontParameterInfo,return;);
+	virtual void GetAllParameterInfoOfType(EMaterialParameterType Type, TArray<FMaterialParameterInfo>& OutParameterInfo, TArray<FGuid>& OutParameterIds) const
+		PURE_VIRTUAL(UMaterialInterface::GetAllParameterInfoOfType, return;);
+	virtual void GetAllParametersOfType(EMaterialParameterType Type, TMap<FMaterialParameterInfo, FMaterialParameterMetadata>& OutParameters) const
+		PURE_VIRTUAL(UMaterialInterface::GetAllParametersOfType, return;);
+	ENGINE_API void GetAllScalarParameterInfo(TArray<FMaterialParameterInfo>& OutParameterInfo, TArray<FGuid>& OutParameterIds) const;
+	ENGINE_API void GetAllVectorParameterInfo(TArray<FMaterialParameterInfo>& OutParameterInfo, TArray<FGuid>& OutParameterIds) const;
+	ENGINE_API void GetAllTextureParameterInfo(TArray<FMaterialParameterInfo>& OutParameterInfo, TArray<FGuid>& OutParameterIds) const;
+	ENGINE_API void GetAllRuntimeVirtualTextureParameterInfo(TArray<FMaterialParameterInfo>& OutParameterInfo, TArray<FGuid>& OutParameterIds) const;
+	ENGINE_API void GetAllFontParameterInfo(TArray<FMaterialParameterInfo>& OutParameterInfo, TArray<FGuid>& OutParameterIds) const;
 
 #if WITH_EDITORONLY_DATA
 	virtual void GetAllMaterialLayersParameterInfo(TArray<FMaterialParameterInfo>& OutParameterInfo, TArray<FGuid>& OutParameterIds) const
 		PURE_VIRTUAL(UMaterialInterface::GetAllMaterialLayersParameterInfo,return;);
-	virtual void GetAllStaticSwitchParameterInfo(TArray<FMaterialParameterInfo>& OutParameterInfo, TArray<FGuid>& OutParameterIds) const
-		PURE_VIRTUAL(UMaterialInterface::GetAllStaticSwitchParameterInfo,return;);
-	virtual void GetAllStaticComponentMaskParameterInfo(TArray<FMaterialParameterInfo>& OutParameterInfo, TArray<FGuid>& OutParameterIds) const
-		PURE_VIRTUAL(UMaterialInterface::GetAllStaticComponentMaskParameterInfo,return;);
+	ENGINE_API void GetAllStaticSwitchParameterInfo(TArray<FMaterialParameterInfo>& OutParameterInfo, TArray<FGuid>& OutParameterIds) const;
+	ENGINE_API void GetAllStaticComponentMaskParameterInfo(TArray<FMaterialParameterInfo>& OutParameterInfo, TArray<FGuid>& OutParameterIds) const;
 
 	virtual bool IterateDependentFunctions(TFunctionRef<bool(UMaterialFunctionInterface*)> Predicate) const
 		PURE_VIRTUAL(UMaterialInterface::IterateDependentFunctions,return false;);
@@ -524,22 +522,17 @@ public:
 		PURE_VIRTUAL(UMaterialInterface::GetDependentFunctions,return;);
 #endif // WITH_EDITORONLY_DATA
 
-	virtual bool GetScalarParameterDefaultValue(const FHashedMaterialParameterInfo& ParameterInfo, float& OutValue, bool bOveriddenOnly = false, bool bCheckOwnedGlobalOverrides = false) const
-		PURE_VIRTUAL(UMaterialInterface::GetScalarParameterDefaultValue,return false;);
-	virtual bool GetVectorParameterDefaultValue(const FHashedMaterialParameterInfo& ParameterInfo, FLinearColor& OutValue, bool bOveriddenOnly = false, bool bCheckOwnedGlobalOverrides = false) const
-		PURE_VIRTUAL(UMaterialInterface::GetVectorParameterDefaultValue,return false;);
-	virtual bool GetTextureParameterDefaultValue(const FHashedMaterialParameterInfo& ParameterInfo, class UTexture*& OutValue, bool bCheckOwnedGlobalOverrides = false) const
-		PURE_VIRTUAL(UMaterialInterface::GetTextureParameterDefaultValue,return false;);
-	virtual bool GetRuntimeVirtualTextureParameterDefaultValue(const FHashedMaterialParameterInfo& ParameterInfo, class URuntimeVirtualTexture*& OutValue, bool bCheckOwnedGlobalOverrides = false) const
-		PURE_VIRTUAL(UMaterialInterface::GetRuntimeVirtualTextureParameterDefaultValue, return false;);
-	virtual bool GetFontParameterDefaultValue(const FHashedMaterialParameterInfo& ParameterInfo, class UFont*& OutFontValue, int32& OutFontPage, bool bCheckOwnedGlobalOverrides = false) const
-		PURE_VIRTUAL(UMaterialInterface::GetFontParameterDefaultValue,return false;);
+	ENGINE_API bool GetParameterDefaultValue(EMaterialParameterType Type, const FMemoryImageMaterialParameterInfo& ParameterInfo, FMaterialParameterMetadata& OutValue) const;
+	ENGINE_API bool GetScalarParameterDefaultValue(const FHashedMaterialParameterInfo& ParameterInfo, float& OutValue) const;
+	ENGINE_API bool GetVectorParameterDefaultValue(const FHashedMaterialParameterInfo& ParameterInfo, FLinearColor& OutValue) const;
+	ENGINE_API bool GetTextureParameterDefaultValue(const FHashedMaterialParameterInfo& ParameterInfo, class UTexture*& OutValue) const;
+	ENGINE_API bool GetRuntimeVirtualTextureParameterDefaultValue(const FHashedMaterialParameterInfo& ParameterInfo, class URuntimeVirtualTexture*& OutValue) const;
+	ENGINE_API bool GetFontParameterDefaultValue(const FHashedMaterialParameterInfo& ParameterInfo, class UFont*& OutFontValue, int32& OutFontPage) const;
 	
 #if WITH_EDITOR
-	virtual bool GetStaticSwitchParameterDefaultValue(const FHashedMaterialParameterInfo& ParameterInfo, bool& OutValue, FGuid& OutExpressionGuid, bool bCheckOwnedGlobalOverrides = false) const
-		PURE_VIRTUAL(UMaterialInterface::GetStaticSwitchParameterDefaultValue,return false;);
-	virtual bool GetStaticComponentMaskParameterDefaultValue(const FHashedMaterialParameterInfo& ParameterInfo, bool& OutR, bool& OutG, bool& OutB, bool& OutA, FGuid& OutExpressionGuid, bool bCheckOwnedGlobalOverrides = false) const
-		PURE_VIRTUAL(UMaterialInterface::GetStaticComponentMaskParameterDefaultValue,return false;);
+	ENGINE_API bool GetStaticSwitchParameterDefaultValue(const FHashedMaterialParameterInfo& ParameterInfo, bool& OutValue, FGuid& OutExpressionGuid) const;
+	ENGINE_API bool GetStaticComponentMaskParameterDefaultValue(const FHashedMaterialParameterInfo& ParameterInfo, bool& OutR, bool& OutG, bool& OutB, bool& OutA, FGuid& OutExpressionGuid) const;
+
 	/** Add to the set any texture referenced by expressions, including nested functions, as well as any overrides from parameters. */
 	virtual void GetReferencedTexturesAndOverrides(TSet<const UTexture*>& InOutTextures) const
 		PURE_VIRTUAL(UMaterialInterface::GetReferencedTexturesAndOverrides, );
@@ -732,27 +725,26 @@ public:
 
 	ENGINE_API virtual bool GetGroupName(const FHashedMaterialParameterInfo& ParameterInfo, FName& GroupName) const;
 	ENGINE_API virtual bool GetParameterDesc(const FHashedMaterialParameterInfo& ParameterInfo, FString& OutDesc, const TArray<struct FStaticMaterialLayersParameter>* MaterialLayersParameters = nullptr) const;
-	ENGINE_API virtual bool GetScalarParameterSliderMinMax(const FHashedMaterialParameterInfo& ParameterInfo, float& OutSliderMin, float& OutSliderMax) const;
+	ENGINE_API bool GetScalarParameterSliderMinMax(const FHashedMaterialParameterInfo& ParameterInfo, float& OutSliderMin, float& OutSliderMax) const;
 #endif // WITH_EDITOR
-	ENGINE_API virtual bool GetScalarParameterValue(const FHashedMaterialParameterInfo& ParameterInfo, float& OutValue, bool bOveriddenOnly = false) const;
+
+	ENGINE_API virtual bool GetParameterValue(EMaterialParameterType Type, const FMemoryImageMaterialParameterInfo& ParameterInfo, FMaterialParameterMetadata& OutValue, EMaterialGetParameterValueFlags Flags = EMaterialGetParameterValueFlags::None) const;
+
+	ENGINE_API bool GetScalarParameterValue(const FHashedMaterialParameterInfo& ParameterInfo, float& OutValue, bool bOveriddenOnly = false) const;
 #if WITH_EDITOR
-	ENGINE_API virtual bool IsScalarParameterUsedAsAtlasPosition(const FHashedMaterialParameterInfo& ParameterInfo, bool& OutValue, TSoftObjectPtr<class UCurveLinearColor>& Curve, TSoftObjectPtr<class UCurveLinearColorAtlas>&  Atlas) const;
+	ENGINE_API bool IsScalarParameterUsedAsAtlasPosition(const FHashedMaterialParameterInfo& ParameterInfo, bool& OutValue, TSoftObjectPtr<class UCurveLinearColor>& Curve, TSoftObjectPtr<class UCurveLinearColorAtlas>&  Atlas) const;
 #endif // WITH_EDITOR
-	ENGINE_API virtual bool GetScalarCurveParameterValue(const FHashedMaterialParameterInfo& ParameterInfo, FInterpCurveFloat& OutValue) const;
-	ENGINE_API virtual bool GetVectorParameterValue(const FHashedMaterialParameterInfo& ParameterInfo, FLinearColor& OutValue, bool bOveriddenOnly = false) const;
+	ENGINE_API bool GetVectorParameterValue(const FHashedMaterialParameterInfo& ParameterInfo, FLinearColor& OutValue, bool bOveriddenOnly = false) const;
 #if WITH_EDITOR
-	ENGINE_API virtual bool IsVectorParameterUsedAsChannelMask(const FHashedMaterialParameterInfo& ParameterInfo, bool& OutValue) const;
-	ENGINE_API virtual bool GetVectorParameterChannelNames(const FHashedMaterialParameterInfo& ParameterInfo, FParameterChannelNames& OutValue) const;
+	ENGINE_API bool IsVectorParameterUsedAsChannelMask(const FHashedMaterialParameterInfo& ParameterInfo, bool& OutValue) const;
+	ENGINE_API bool GetVectorParameterChannelNames(const FHashedMaterialParameterInfo& ParameterInfo, FParameterChannelNames& OutValue) const;
 #endif
-	ENGINE_API virtual bool GetVectorCurveParameterValue(const FHashedMaterialParameterInfo& ParameterInfo, FInterpCurveVector& OutValue) const;
-	ENGINE_API virtual bool GetLinearColorParameterValue(const FHashedMaterialParameterInfo& ParameterInfo, FLinearColor& OutValue) const;
-	ENGINE_API virtual bool GetLinearColorCurveParameterValue(const FHashedMaterialParameterInfo& ParameterInfo, FInterpCurveLinearColor& OutValue) const;
-	ENGINE_API virtual bool GetTextureParameterValue(const FHashedMaterialParameterInfo& ParameterInfo, class UTexture*& OutValue, bool bOveriddenOnly = false) const;
-	ENGINE_API virtual bool GetRuntimeVirtualTextureParameterValue(const FHashedMaterialParameterInfo& ParameterInfo, class URuntimeVirtualTexture*& OutValue, bool bOveriddenOnly = false) const;
+	ENGINE_API bool GetTextureParameterValue(const FHashedMaterialParameterInfo& ParameterInfo, class UTexture*& OutValue, bool bOveriddenOnly = false) const;
+	ENGINE_API bool GetRuntimeVirtualTextureParameterValue(const FHashedMaterialParameterInfo& ParameterInfo, class URuntimeVirtualTexture*& OutValue, bool bOveriddenOnly = false) const;
 #if WITH_EDITOR
-	ENGINE_API virtual bool GetTextureParameterChannelNames(const FHashedMaterialParameterInfo& ParameterInfo, FParameterChannelNames& OutValue) const;
+	ENGINE_API bool GetTextureParameterChannelNames(const FHashedMaterialParameterInfo& ParameterInfo, FParameterChannelNames& OutValue) const;
 #endif
-	ENGINE_API virtual bool GetFontParameterValue(const FHashedMaterialParameterInfo& ParameterInfo,class UFont*& OutFontValue, int32& OutFontPage, bool bOveriddenOnly = false) const;
+	ENGINE_API bool GetFontParameterValue(const FHashedMaterialParameterInfo& ParameterInfo,class UFont*& OutFontValue, int32& OutFontPage, bool bOveriddenOnly = false) const;
 	ENGINE_API virtual bool GetRefractionSettings(float& OutBiasValue) const;
 
 	/**
diff --git a/Engine/Source/Runtime/Engine/Private/Materials/Material.cpp b/Engine/Source/Runtime/Engine/Private/Materials/Material.cpp
--- a/Engine/Source/Runtime/Engine/Private/Materials/Material.cpp
+++ b/Engine/Source/Runtime/Engine/Private/Materials/Material.cpp
@@ -1396,39 +1396,109 @@ bool UMaterial::IsCompilingOrHadCompileError(ERHIFeatureLevel::Type InFeatureLev
 }
 
 #if WITH_EDITOR
+bool UMaterial::SetParameterValueEditorOnly(const FName& ParameterName, const FMaterialParameterMetadata& Meta)
+{
+	// Warning: in the case of duplicate parameters with different default values, this will find the first in the expression array, not necessarily the one that's used for rendering
+	// Lambda which calls SetParameterValue on a given parameter and triggers a PostEditChange event if successful
+	auto TrySetParameterValue = [ParameterName, Meta](UMaterialExpression* Expression) -> bool
+	{
+		if (Expression && Expression->SetParameterValue(ParameterName, Meta))
+		{
+			FProperty* ParamProperty = Expression->GetClass()->FindPropertyByName(TEXT("DefaultValue"));
+			if (ParamProperty)
+			{
+				FPropertyChangedEvent PropertyChangedEvent(ParamProperty);
+				Expression->PostEditChangeProperty(PropertyChangedEvent);
+				return true;
+			}
+		}
+		return false;
+	};
+
+	for (TObjectPtr<UMaterialExpression>& Expression : Expressions)
+	{
+		if (TrySetParameterValue(Expression))
+		{
+			return true;
+		}
+		else if (UMaterialExpressionMaterialFunctionCall* FunctionCall = Cast<UMaterialExpressionMaterialFunctionCall>(Expression))
+		{
+			if (FunctionCall->MaterialFunction)
+			{
+				TArray<UMaterialFunctionInterface*> Functions;
+				Functions.Add(FunctionCall->MaterialFunction);
+				FunctionCall->MaterialFunction->GetDependentFunctions(Functions);
+
+				for (UMaterialFunctionInterface* Function : Functions)
+				{
+					const TArray<TObjectPtr<UMaterialExpression>>* ExpressionPtr = Function->GetFunctionExpressions();
+					if (ExpressionPtr)
+					{
+						for (const TObjectPtr<UMaterialExpression>& FunctionExpression : *ExpressionPtr)
+						{
+							if (TrySetParameterValue(FunctionExpression))
+							{
+								return true;
+							}
+						}
+					}
+				}
+			}
+		}
+	}
+
+	return false;
+}
+
 bool UMaterial::SetVectorParameterValueEditorOnly(FName ParameterName, FLinearColor InValue)
 {
-	return SetParameterValueEditorOnly<UMaterialExpressionVectorParameter>(ParameterName, InValue);
+	FMaterialParameterMetadata Meta;
+	Meta.Value = InValue;
+	return SetParameterValueEditorOnly(ParameterName, Meta);
 }
 
 bool UMaterial::SetScalarParameterValueEditorOnly(FName ParameterName, float InValue)
 {
-	return SetParameterValueEditorOnly<UMaterialExpressionScalarParameter>(ParameterName, InValue);
+	FMaterialParameterMetadata Meta;
+	Meta.Value = InValue;
+	return SetParameterValueEditorOnly(ParameterName, Meta);
 }
 
 bool UMaterial::SetTextureParameterValueEditorOnly(FName ParameterName, class UTexture* InValue)
 {
-	return SetParameterValueEditorOnly<UMaterialExpressionTextureSampleParameter>(ParameterName, InValue);
+	FMaterialParameterMetadata Meta;
+	Meta.Value = InValue;
+	return SetParameterValueEditorOnly(ParameterName, Meta);
 }
 
 bool UMaterial::SetRuntimeVirtualTextureParameterValueEditorOnly(FName ParameterName, class URuntimeVirtualTexture* InValue)
 {
-	return SetParameterValueEditorOnly<UMaterialExpressionRuntimeVirtualTextureSampleParameter>(ParameterName, InValue);
+	FMaterialParameterMetadata Meta;
+	Meta.Value = InValue;
+	return SetParameterValueEditorOnly(ParameterName, Meta);
 }
 
 bool UMaterial::SetFontParameterValueEditorOnly(FName ParameterName, class UFont* InFontValue, int32 InFontPage)
 {
-	return SetParameterValueEditorOnly<UMaterialExpressionFontSampleParameter>(ParameterName, InFontValue, InFontPage);
+	FMaterialParameterMetadata Meta;
+	Meta.Value = FMaterialParameterValue(InFontValue, InFontPage);
+	return SetParameterValueEditorOnly(ParameterName, Meta);
 }
 
 bool UMaterial::SetStaticSwitchParameterValueEditorOnly(FName ParameterName, bool InValue, FGuid InExpressionGuid)
 {
-	return SetParameterValueEditorOnly<UMaterialExpressionStaticSwitchParameter>(ParameterName, InValue, InExpressionGuid);
+	FMaterialParameterMetadata Meta;
+	Meta.Value = InValue;
+	Meta.ExpressionGuid = InExpressionGuid;
+	return SetParameterValueEditorOnly(ParameterName, Meta);
 }
 
 bool UMaterial::SetStaticComponentMaskParameterValueEditorOnly(FName ParameterName, bool R, bool G, bool B, bool A, FGuid InExpressionGuid)
 {
-	return SetParameterValueEditorOnly<UMaterialExpressionStaticComponentMaskParameter>(ParameterName, R, G, B, A, InExpressionGuid);
+	FMaterialParameterMetadata Meta;
+	Meta.Value = FMaterialParameterValue(R, G, B, A);
+	Meta.ExpressionGuid = InExpressionGuid;
+	return SetParameterValueEditorOnly(ParameterName, Meta);
 }
 #endif
 
@@ -1745,29 +1815,17 @@ void UMaterial::FixupMaterialUsageAfterLoad()
 }
 #endif // WITH_EDITOR
 
-void UMaterial::GetAllScalarParameterInfo(TArray<FMaterialParameterInfo>& OutParameterInfo, TArray<FGuid>& OutParameterIds) const
-{
-	GetCachedExpressionData().Parameters.GetAllParameterInfoOfType(EMaterialParameterType::Scalar, true, OutParameterInfo, OutParameterIds);
-}
-
-void UMaterial::GetAllVectorParameterInfo(TArray<FMaterialParameterInfo>& OutParameterInfo, TArray<FGuid>& OutParameterIds) const
-{
-	GetCachedExpressionData().Parameters.GetAllParameterInfoOfType(EMaterialParameterType::Vector, true, OutParameterInfo, OutParameterIds);
-}
-
-void UMaterial::GetAllTextureParameterInfo(TArray<FMaterialParameterInfo>& OutParameterInfo, TArray<FGuid>& OutParameterIds) const
+void UMaterial::GetAllParametersOfType(EMaterialParameterType Type, TMap<FMaterialParameterInfo, FMaterialParameterMetadata>& OutParameters) const
 {
-	GetCachedExpressionData().Parameters.GetAllParameterInfoOfType(EMaterialParameterType::Texture, true, OutParameterInfo, OutParameterIds);
+	OutParameters.Reset();
+	GetCachedExpressionData().Parameters.GetAllParametersOfType(Type, OutParameters);
 }
 
-void UMaterial::GetAllFontParameterInfo(TArray<FMaterialParameterInfo>& OutParameterInfo, TArray<FGuid>& OutParameterIds) const
+void UMaterial::GetAllParameterInfoOfType(EMaterialParameterType Type, TArray<FMaterialParameterInfo>& OutParameterInfo, TArray<FGuid>& OutParameterIds) const
 {
-	GetCachedExpressionData().Parameters.GetAllParameterInfoOfType(EMaterialParameterType::Font, true, OutParameterInfo, OutParameterIds);
-}
-
-void UMaterial::GetAllRuntimeVirtualTextureParameterInfo(TArray<FMaterialParameterInfo>& OutParameterInfo, TArray<FGuid>& OutParameterIds) const
-{
-	GetCachedExpressionData().Parameters.GetAllParameterInfoOfType(EMaterialParameterType::RuntimeVirtualTexture, true, OutParameterInfo, OutParameterIds);
+	OutParameterInfo.Reset();
+	OutParameterIds.Reset();
+	GetCachedExpressionData().Parameters.GetAllParameterInfoOfType(Type, OutParameterInfo, OutParameterIds);
 }
 
 #if WITH_EDITORONLY_DATA
@@ -1778,16 +1836,6 @@ void UMaterial::GetAllMaterialLayersParameterInfo(TArray<FMaterialParameterInfo>
 	GetAllParameterInfo<UMaterialExpressionMaterialAttributeLayers>(OutParameterInfo, OutParameterIds);
 }
 
-void UMaterial::GetAllStaticSwitchParameterInfo(TArray<FMaterialParameterInfo>& OutParameterInfo, TArray<FGuid>& OutParameterIds) const
-{
-	GetCachedExpressionData().Parameters.GetAllParameterInfoOfType(EMaterialParameterType::StaticSwitch, true, OutParameterInfo, OutParameterIds);
-}
-
-void UMaterial::GetAllStaticComponentMaskParameterInfo(TArray<FMaterialParameterInfo>& OutParameterInfo, TArray<FGuid>& OutParameterIds) const
-{
-	GetCachedExpressionData().Parameters.GetAllParameterInfoOfType(EMaterialParameterType::StaticComponentMask, true, OutParameterInfo, OutParameterIds);
-}
-
 bool UMaterial::IterateDependentFunctions(TFunctionRef<bool(UMaterialFunctionInterface*)> Predicate) const
 {
 	for (UMaterialExpression* Expression : Expressions)
@@ -1820,43 +1868,6 @@ void UMaterial::GetDependentFunctions(TArray<UMaterialFunctionInterface*>& Depen
 }
 #endif // WITH_EDITORONLY_DATA
 
-bool UMaterial::GetScalarParameterDefaultValue(const FHashedMaterialParameterInfo& ParameterInfo, float& OutValue, bool bOveriddenOnly, bool bCheckOwnedGlobalOverrides) const
-{
-	return GetScalarParameterValue(ParameterInfo, OutValue, bOveriddenOnly);
-}
-
-bool UMaterial::GetVectorParameterDefaultValue(const FHashedMaterialParameterInfo& ParameterInfo, FLinearColor& OutValue, bool bOveriddenOnly, bool bCheckOwnedGlobalOverrides) const
-{
-	return GetVectorParameterValue(ParameterInfo, OutValue, bOveriddenOnly);
-}
-
-bool UMaterial::GetTextureParameterDefaultValue(const FHashedMaterialParameterInfo& ParameterInfo, class UTexture*& OutValue, bool bCheckOwnedGlobalOverrides) const
-{
-	return GetTextureParameterValue(ParameterInfo, OutValue);
-}
-
-bool UMaterial::GetRuntimeVirtualTextureParameterDefaultValue(const FHashedMaterialParameterInfo& ParameterInfo, class URuntimeVirtualTexture*& OutValue, bool bCheckOwnedGlobalOverrides) const
-{
-	return GetRuntimeVirtualTextureParameterValue(ParameterInfo, OutValue);
-}
-
-bool UMaterial::GetFontParameterDefaultValue(const FHashedMaterialParameterInfo& ParameterInfo, class UFont*& OutFontValue, int32& OutFontPage, bool bCheckOwnedGlobalOverrides) const
-{
-	return GetFontParameterValue(ParameterInfo, OutFontValue, OutFontPage);
-}
-
-#if WITH_EDITOR
-bool UMaterial::GetStaticComponentMaskParameterDefaultValue(const FHashedMaterialParameterInfo& ParameterInfo, bool& OutR, bool& OutG, bool& OutB, bool& OutA, FGuid& OutExpressionGuid, bool bCheckOwnedGlobalOverrides) const
-{
-	return GetStaticComponentMaskParameterValue(ParameterInfo, OutR, OutG, OutB, OutA, OutExpressionGuid);
-}
-
-bool UMaterial::GetStaticSwitchParameterDefaultValue(const FHashedMaterialParameterInfo& ParameterInfo, bool& OutValue, FGuid& OutExpressionGuid, bool bCheckOwnedGlobalOverrides) const
-{
-	return GetStaticSwitchParameterValue(ParameterInfo, OutValue, OutExpressionGuid);
-}
-#endif // WITH_EDITOR
-
 extern FPostProcessMaterialNode* IteratePostProcessMaterialNodes(const FFinalPostProcessSettings& Dest, const UMaterial* Material, FBlendableEntry*& Iterator);
 
 void UMaterialInterface::OverrideBlendableSettings(class FSceneView& View, float Weight) const
@@ -2131,624 +2142,17 @@ void UMaterial::UpdateCachedExpressionData()
 }
 #endif // WITH_EDITOR
 
-bool UMaterial::GetScalarParameterValue(const FHashedMaterialParameterInfo& ParameterInfo, float& OutValue, bool bOveriddenOnly) const
-{
-	const bool bResult = GetScalarParameterValue_New(ParameterInfo, OutValue, bOveriddenOnly);
-#if WITH_EDITOR
-	if (GIsClient && CVarMaterialParameterLegacyChecks.GetValueOnAnyThread())
-	{
-		// Expressions may not be loaded on server, and cached data is not updated...results will be incorrect
-		float OldValue = 0.0f;
-		const bool bOldResult = GetScalarParameterValue_Legacy(ParameterInfo, OldValue, bOveriddenOnly);
-
-		ensureMsgf(bOldResult == bResult, TEXT("UMaterial::GetScalarParameterValue() mismatch, bOldResult: %d, bResult: %d"), bOldResult, bResult);
-		ensureMsgf(!bResult || OutValue == OldValue, TEXT("UMaterial::GetScalarParameterValue() mismatch, OutValue: %g, OldValue: %g"), OutValue, OldValue);
-	}
-#endif
-	return bResult;
-}
-
-bool UMaterial::GetScalarParameterValue_New(const FHashedMaterialParameterInfo& ParameterInfo, float& OutValue, bool bOveriddenOnly) const
+bool UMaterial::GetParameterValue(EMaterialParameterType Type, const FMemoryImageMaterialParameterInfo& ParameterInfo, FMaterialParameterMetadata& OutResult, EMaterialGetParameterValueFlags Flags) const
 {
-	if (!bOveriddenOnly && CachedExpressionData)
+	if (EnumHasAnyFlags(Flags, EMaterialGetParameterValueFlags::CheckNonOverrides) && CachedExpressionData)
 	{
-		const int32 Index = CachedExpressionData->Parameters.FindParameterIndex(EMaterialParameterType::Scalar, ParameterInfo);
-		if (Index != INDEX_NONE)
-		{
-			OutValue = CachedExpressionData->Parameters.ScalarValues[Index];
-			return true;
-		}
+		return CachedExpressionData->Parameters.GetParameterValue(Type, ParameterInfo, OutResult);
 	}
 
 	return false;
 }
 
-#if WITH_EDITOR
-bool UMaterial::GetScalarParameterValue_Legacy(const FHashedMaterialParameterInfo& ParameterInfo, float& OutValue, bool bOveriddenOnly) const
-{
-	// In the case of duplicate parameters with different values, this will return the
-	// first matching expression found, not necessarily the one that's used for rendering
-	UMaterialExpressionScalarParameter* Parameter = nullptr;
-
-	for (UMaterialExpression* Expression : Expressions)
-	{
-		// Only need to check parameters that match in associated scope
-		if (ParameterInfo.Association == EMaterialParameterAssociation::GlobalParameter)
-		{
-			if (UMaterialExpressionScalarParameter* ExpressionParameter = Cast<UMaterialExpressionScalarParameter>(Expression))
-			{
-				if (ExpressionParameter->IsNamedParameter(ParameterInfo, OutValue))
-				{
-					return !bOveriddenOnly;
-				}
-			}
-			else if (UMaterialExpressionMaterialFunctionCall* FunctionCall = Cast<UMaterialExpressionMaterialFunctionCall>(Expression))
-			{
-				UMaterialFunctionInterface* Function = FunctionCall->MaterialFunction;
-				UMaterialFunctionInterface* ParameterOwner = nullptr;
-
-				if (Function && Function->OverrideNamedScalarParameter(ParameterInfo, OutValue))
-				{
-					return true;
-				}
-
-				if (Function && Function->GetNamedParameterOfType(ParameterInfo, Parameter, &ParameterOwner))
-				{
-					if (ParameterOwner->OverrideNamedScalarParameter(ParameterInfo, OutValue))
-					{
-						return true;
-					}
-
-					if (Parameter)
-					{
-						Parameter->IsNamedParameter(ParameterInfo, OutValue);
-						return !bOveriddenOnly;
-					}
-				}
-			}
-		}
-		else
-		{
-			if (UMaterialExpressionMaterialAttributeLayers* LayersExpression = Cast<UMaterialExpressionMaterialAttributeLayers>(Expression))
-			{
-				UMaterialFunctionInterface* Function = LayersExpression->GetParameterAssociatedFunction(ParameterInfo);
-				UMaterialFunctionInterface* ParameterOwner = nullptr;
-
-				if (Function && Function->OverrideNamedScalarParameter(ParameterInfo, OutValue))
-				{
-					return true;
-				}
-
-				if (Function && Function->GetNamedParameterOfType(ParameterInfo, Parameter, &ParameterOwner))
-				{
-					if (ParameterOwner->OverrideNamedScalarParameter(ParameterInfo, OutValue))
-					{
-						return true;
-					}
-
-					Parameter->IsNamedParameter(ParameterInfo, OutValue);
-					return !bOveriddenOnly;
-				}
-			}
-		}
-	}
-
-	return false;
-}
-
-bool UMaterial::IsScalarParameterUsedAsAtlasPosition(const FHashedMaterialParameterInfo& ParameterInfo, bool& OutValue, TSoftObjectPtr<UCurveLinearColor>& OutCurve, TSoftObjectPtr<UCurveLinearColorAtlas>& OutAtlas) const
-{
-	if (CachedExpressionData)
-	{
-		return CachedExpressionData->Parameters.IsScalarParameterUsedAsAtlasPosition(ParameterInfo, OutValue, OutCurve, OutAtlas);
-	}
-
-	return false;
-}
-
-bool UMaterial::GetScalarParameterSliderMinMax(const FHashedMaterialParameterInfo& ParameterInfo, float& OutSliderMin, float& OutSliderMax) const
-{
-	if (CachedExpressionData)
-	{
-		return CachedExpressionData->Parameters.GetScalarParameterSliderMinMax(ParameterInfo, OutSliderMin, OutSliderMax);
-	}
-
-	return false;
-}
-#endif // WITH_EDITOR
-
-bool UMaterial::GetVectorParameterValue(const FHashedMaterialParameterInfo& ParameterInfo, FLinearColor& OutValue, bool bOveriddenOnly) const
-{
-	const bool bResult = GetVectorParameterValue_New(ParameterInfo, OutValue, bOveriddenOnly);
-#if WITH_EDITOR
-	if (GIsClient && CVarMaterialParameterLegacyChecks.GetValueOnAnyThread())
-	{
-		// Expressions may not be loaded on server, and cached data is not updated...results will be incorrect
-		FLinearColor OldValue(ForceInitToZero);
-		const bool bOldResult = GetVectorParameterValue_Legacy(ParameterInfo, OldValue, bOveriddenOnly);
-
-		ensureMsgf(bOldResult == bResult, TEXT("UMaterial::GetVectorParameterValue() mismatch, bOldResult: %d, bResult: %d"), bOldResult, bResult);
-		ensureMsgf(!bResult || OutValue == OldValue, TEXT("UMaterial::GetVectorParameterValue() mismatch, OutValue: %s, OldValue: %s"), *OutValue.ToString(), *OldValue.ToString());
-	}
-#endif
-	return bResult;
-}
-
-bool UMaterial::GetVectorParameterValue_New(const FHashedMaterialParameterInfo& ParameterInfo, FLinearColor& OutValue, bool bOveriddenOnly) const
-{
-	if (!bOveriddenOnly && CachedExpressionData)
-	{
-		const int32 Index = CachedExpressionData->Parameters.FindParameterIndex(EMaterialParameterType::Vector, ParameterInfo);
-		if (Index != INDEX_NONE)
-		{
-			OutValue = CachedExpressionData->Parameters.VectorValues[Index];
-			return true;
-		}
-	}
-	return false;
-}
-
-#if WITH_EDITOR
-bool UMaterial::GetVectorParameterValue_Legacy(const FHashedMaterialParameterInfo& ParameterInfo, FLinearColor& OutValue, bool bOveriddenOnly) const
-{
-	// In the case of duplicate parameters with different values, this will return the
-	// first matching expression found, not necessarily the one that's used for rendering
-	UMaterialExpressionVectorParameter* Parameter = nullptr;
-
-	for (UMaterialExpression* Expression : Expressions)
-	{
-		// Only need to check parameters that match in associated scope
-		if (ParameterInfo.Association == EMaterialParameterAssociation::GlobalParameter)
-		{
-			if (UMaterialExpressionVectorParameter* ExpressionParameter = Cast<UMaterialExpressionVectorParameter>(Expression))
-			{
-				if (ExpressionParameter->IsNamedParameter(ParameterInfo, OutValue))
-				{
-					return !bOveriddenOnly;
-				}
-			}
-			else if (UMaterialExpressionMaterialFunctionCall* FunctionCall = Cast<UMaterialExpressionMaterialFunctionCall>(Expression))
-			{
-				UMaterialFunctionInterface* Function = FunctionCall->MaterialFunction;
-				UMaterialFunctionInterface* ParameterOwner = nullptr;
-
-				if (Function && Function->OverrideNamedVectorParameter(ParameterInfo, OutValue))
-				{
-					return true;
-				}
-
-				if (Function && Function->GetNamedParameterOfType(ParameterInfo, Parameter, &ParameterOwner))
-				{
-					if (ParameterOwner->OverrideNamedVectorParameter(ParameterInfo, OutValue))
-					{
-						return true;
-					}
-					
-					Parameter->IsNamedParameter(ParameterInfo, OutValue);
-					return !bOveriddenOnly;
-				}
-			}
-		}
-		else
-		{
-			if (UMaterialExpressionMaterialAttributeLayers* LayersExpression = Cast<UMaterialExpressionMaterialAttributeLayers>(Expression))
-			{
-				UMaterialFunctionInterface* Function = LayersExpression->GetParameterAssociatedFunction(ParameterInfo);
-				UMaterialFunctionInterface* ParameterOwner = nullptr;
-
-				if (Function && Function->OverrideNamedVectorParameter(ParameterInfo, OutValue))
-				{
-					return true;
-				}
-
-				if (Function && Function->GetNamedParameterOfType(ParameterInfo, Parameter, &ParameterOwner))
-				{
-					if (ParameterOwner->OverrideNamedVectorParameter(ParameterInfo, OutValue))
-					{
-						return true;
-					}
-					
-					Parameter->IsNamedParameter(ParameterInfo, OutValue);
-					return !bOveriddenOnly;
-				}
-			}
-		}
-	}
-
-	return false;
-}
-
-bool UMaterial::IsVectorParameterUsedAsChannelMask(const FHashedMaterialParameterInfo& ParameterInfo, bool& OutValue) const
-{
-	if (CachedExpressionData)
-	{
-		return CachedExpressionData->Parameters.IsVectorParameterUsedAsChannelMask(ParameterInfo, OutValue);
-	}
-	return false;
-}
-
-bool UMaterial::GetVectorParameterChannelNames(const FHashedMaterialParameterInfo& ParameterInfo, FParameterChannelNames& OutValue) const
-{
-	if (CachedExpressionData)
-	{
-		return CachedExpressionData->Parameters.GetVectorParameterChannelNames(ParameterInfo, OutValue);
-	}
-	return false;
-}
-
-#endif // WITH_EDITOR
-
-
-bool UMaterial::GetTextureParameterValue(const FHashedMaterialParameterInfo& ParameterInfo, UTexture*& OutValue, bool bOveriddenOnly) const
-{
-	const bool bResult = GetTextureParameterValue_New(ParameterInfo, OutValue, bOveriddenOnly);
-#if WITH_EDITOR
-	if (GIsClient && CVarMaterialParameterLegacyChecks.GetValueOnAnyThread())
-	{
-		// Expressions may not be loaded on server, and cached data is not updated...results will be incorrect
-		UTexture* OldValue = nullptr;
-		const bool bOldResult = GetTextureParameterValue_Legacy(ParameterInfo, OldValue, bOveriddenOnly);
-
-		ensureMsgf(bOldResult == bResult, TEXT("UMaterial::GetTextureParameterValue() mismatch, bOldResult: %d, bResult: %d"), bOldResult, bResult);
-		ensureMsgf(!bResult || OutValue == OldValue, TEXT("UMaterial::GetTextureParameterValue() mismatch, OutValue: %s, OldValue: %s"),
-			OutValue ? *OutValue->GetName() : TEXT("nullptr"),
-			OldValue ? *OldValue->GetName() : TEXT("nullptr"));
-	}
-#endif
-	return bResult;
-}
-
-bool UMaterial::GetTextureParameterValue_New(const FHashedMaterialParameterInfo& ParameterInfo, UTexture*& OutValue, bool bOveriddenOnly) const
-{
-	if (!bOveriddenOnly && CachedExpressionData)
-	{
-		const int32 Index = CachedExpressionData->Parameters.FindParameterIndex(EMaterialParameterType::Texture, ParameterInfo);
-		if (Index != INDEX_NONE)
-		{
-			OutValue = CachedExpressionData->Parameters.TextureValues[Index];
-			return true;
-		}
-	}
-
-	OutValue = nullptr;
-	return false;
-}
-
-#if WITH_EDITOR
-bool UMaterial::GetTextureParameterValue_Legacy(const FHashedMaterialParameterInfo& ParameterInfo, UTexture*& OutValue, bool bOveriddenOnly) const
-{
-	// Actual legacy version of this functionality *ignored* the bOveriddenOnly parameter, which as far as I can tell was an oversight
-	// This version has been updated to match other GetXXXParameterValue_Legacy functions, in order to match current behavior
-
-	// In the case of duplicate parameters with different values, this will return the
-	// first matching expression found, not necessarily the one that's used for rendering
-	UMaterialExpressionTextureSampleParameter* Parameter = nullptr;
-
-	for (UMaterialExpression* Expression : Expressions)
-	{
-		// Only need to check parameters that match in associated scope
-		if (ParameterInfo.Association == EMaterialParameterAssociation::GlobalParameter)
-		{
-			if (UMaterialExpressionTextureSampleParameter* ExpressionParameter = Cast<UMaterialExpressionTextureSampleParameter>(Expression))
-			{
-				if (ExpressionParameter->IsNamedParameter(ParameterInfo, OutValue))
-				{
-					return !bOveriddenOnly;
-				}
-			}
-			else if (UMaterialExpressionMaterialFunctionCall* FunctionCall = Cast<UMaterialExpressionMaterialFunctionCall>(Expression))
-			{
-				UMaterialFunctionInterface* Function = FunctionCall->MaterialFunction;
-				UMaterialFunctionInterface* ParameterOwner = nullptr;
-
-				if (Function && Function->OverrideNamedTextureParameter(ParameterInfo, OutValue))
-				{
-					return true;
-				}
-
-				if (Function && Function->GetNamedParameterOfType(ParameterInfo, Parameter, &ParameterOwner))
-				{
-					if (ParameterOwner->OverrideNamedTextureParameter(ParameterInfo, OutValue))
-					{
-						return true;
-					}
-
-					Parameter->IsNamedParameter(ParameterInfo, OutValue);
-					return !bOveriddenOnly;
-				}
-			}
-		}
-		else
-		{
-			if (UMaterialExpressionMaterialAttributeLayers* LayersExpression = Cast<UMaterialExpressionMaterialAttributeLayers>(Expression))
-			{
-				UMaterialFunctionInterface* Function = LayersExpression->GetParameterAssociatedFunction(ParameterInfo);
-				UMaterialFunctionInterface* ParameterOwner = nullptr;
-
-				if (Function && Function->OverrideNamedTextureParameter(ParameterInfo, OutValue))
-				{
-					return true;
-				}
-
-				if (Function && Function->GetNamedParameterOfType(ParameterInfo, Parameter, &ParameterOwner))
-				{
-					if (ParameterOwner->OverrideNamedTextureParameter(ParameterInfo, OutValue))
-					{
-						return true;
-					}
-
-					Parameter->IsNamedParameter(ParameterInfo, OutValue);
-					return !bOveriddenOnly;
-				}
-			}
-		}
-	}
-
-	return false;
-}
-#endif // WITH_EDITOR
-
-bool UMaterial::GetRuntimeVirtualTextureParameterValue(const FHashedMaterialParameterInfo& ParameterInfo, URuntimeVirtualTexture*& OutValue, bool bOveriddenOnly) const
-{
-	const bool bResult = GetRuntimeVirtualTextureParameterValue_New(ParameterInfo, OutValue, bOveriddenOnly);
-#if WITH_EDITOR
-	if (GIsClient && CVarMaterialParameterLegacyChecks.GetValueOnAnyThread())
-	{
-		// Expressions may not be loaded on server, and cached data is not updated...results will be incorrect
-		URuntimeVirtualTexture* OldValue = nullptr;
-		const bool bOldResult = GetRuntimeVirtualTextureParameterValue_Legacy(ParameterInfo, OldValue, bOveriddenOnly);
-
-		ensureMsgf(bOldResult == bResult, TEXT("UMaterial::GetRuntimeVirtualTextureParameterValue() mismatch, bOldResult: %d, bResult: %d"), bOldResult, bResult);
-		ensureMsgf(!bResult || OutValue == OldValue, TEXT("UMaterial::GetRuntimeVirtualTextureParameterValue() mismatch, OutValue: %s, OldValue: %s"),
-			OutValue ? *OutValue->GetName() : TEXT("nullptr"),
-			OldValue ? *OldValue->GetName() : TEXT("nullptr"));
-	}
-#endif
-	return bResult;
-}
-
-bool UMaterial::GetRuntimeVirtualTextureParameterValue_New(const FHashedMaterialParameterInfo& ParameterInfo, URuntimeVirtualTexture*& OutValue, bool bOveriddenOnly) const
-{
-	if (!bOveriddenOnly && CachedExpressionData)
-	{
-		const int32 Index = CachedExpressionData->Parameters.FindParameterIndex(EMaterialParameterType::RuntimeVirtualTexture, ParameterInfo);
-		if (Index != INDEX_NONE)
-		{
-			OutValue = CachedExpressionData->Parameters.RuntimeVirtualTextureValues[Index];
-			return true;
-		}
-	}
-
-	OutValue = nullptr;
-	return false;
-}
-
-#if WITH_EDITOR
-bool UMaterial::GetRuntimeVirtualTextureParameterValue_Legacy(const FHashedMaterialParameterInfo& ParameterInfo, URuntimeVirtualTexture*& OutValue, bool bOveriddenOnly) const
-{
-	// Actual legacy version of this functionality *ignored* the bOveriddenOnly parameter, which as far as I can tell was an oversight
-	// This version has been updated to match other GetXXXParameterValue_Legacy functions, in order to match current behavior
-
-	// In the case of duplicate parameters with different values, this will return the
-	// first matching expression found, not necessarily the one that's used for rendering
-	UMaterialExpressionRuntimeVirtualTextureSampleParameter* Parameter = nullptr;
-
-	for (UMaterialExpression* Expression : Expressions)
-	{
-		// Only need to check parameters that match in associated scope
-		if (ParameterInfo.Association == EMaterialParameterAssociation::GlobalParameter)
-		{
-			if (UMaterialExpressionRuntimeVirtualTextureSampleParameter* ExpressionParameter = Cast<UMaterialExpressionRuntimeVirtualTextureSampleParameter>(Expression))
-			{
-				if (ExpressionParameter->IsNamedParameter(ParameterInfo, OutValue))
-				{
-					return !bOveriddenOnly;
-				}
-			}
-			else if (UMaterialExpressionMaterialFunctionCall* FunctionCall = Cast<UMaterialExpressionMaterialFunctionCall>(Expression))
-			{
-				UMaterialFunctionInterface* Function = FunctionCall->MaterialFunction;
-				UMaterialFunctionInterface* ParameterOwner = nullptr;
-
-				if (Function && Function->OverrideNamedRuntimeVirtualTextureParameter(ParameterInfo, OutValue))
-				{
-					return true;
-				}
-
-				if (Function && Function->GetNamedParameterOfType(ParameterInfo, Parameter, &ParameterOwner))
-				{
-					if (ParameterOwner->OverrideNamedRuntimeVirtualTextureParameter(ParameterInfo, OutValue))
-					{
-						return true;
-					}
-
-					Parameter->IsNamedParameter(ParameterInfo, OutValue);
-					return !bOveriddenOnly;
-				}
-			}
-		}
-		else
-		{
-			if (UMaterialExpressionMaterialAttributeLayers* LayersExpression = Cast<UMaterialExpressionMaterialAttributeLayers>(Expression))
-			{
-				UMaterialFunctionInterface* Function = LayersExpression->GetParameterAssociatedFunction(ParameterInfo);
-				UMaterialFunctionInterface* ParameterOwner = nullptr;
-
-				if (Function && Function->OverrideNamedRuntimeVirtualTextureParameter(ParameterInfo, OutValue))
-				{
-					return true;
-				}
-
-				if (Function && Function->GetNamedParameterOfType(ParameterInfo, Parameter, &ParameterOwner))
-				{
-					if (ParameterOwner->OverrideNamedRuntimeVirtualTextureParameter(ParameterInfo, OutValue))
-					{
-						return true;
-					}
-
-					Parameter->IsNamedParameter(ParameterInfo, OutValue);
-					return !bOveriddenOnly;
-				}
-			}
-		}
-	}
-
-	return false;
-}
-#endif // WITH_EDITOR
-
-#if WITH_EDITOR
-bool UMaterial::GetTextureParameterChannelNames(const FHashedMaterialParameterInfo& ParameterInfo, FParameterChannelNames& OutValue) const
-{
-	if (CachedExpressionData)
-	{
-		return CachedExpressionData->Parameters.GetTextureParameterChannelNames(ParameterInfo, OutValue);
-	}
-	return false;
-}
-#endif // WITH_EDITOR
-
-bool UMaterial::GetFontParameterValue(const FHashedMaterialParameterInfo& ParameterInfo, UFont*& OutFontValue, int32& OutFontPage, bool bOveriddenOnly) const
-{
-	const bool bResult = GetFontParameterValue_New(ParameterInfo, OutFontValue, OutFontPage, bOveriddenOnly);
-#if WITH_EDITOR
-	if (GIsClient && CVarMaterialParameterLegacyChecks.GetValueOnAnyThread())
-	{
-		// Expressions may not be loaded on server, and cached data is not updated...results will be incorrect
-		UFont* OldValue = nullptr;
-		int32 OldPage = INDEX_NONE;
-		const bool bOldResult = GetFontParameterValue_Legacy(ParameterInfo, OldValue, OldPage, bOveriddenOnly);
-
-		ensureMsgf(bOldResult == bResult, TEXT("UMaterial::GetFontParameterValue() mismatch, bOldResult: %d, bResult: %d"), bOldResult, bResult);
-		ensureMsgf(!bResult || (OutFontValue == OldValue && OldPage == OutFontPage), TEXT("UMaterial::GetFontParameterValue() mismatch, OutValue: %s, OutPage: %d OldValue: %s, OldPage: %d"),
-			OutFontValue ? *OutFontValue->GetName() : TEXT("nullptr"), OutFontPage,
-			OldValue ? *OldValue->GetName() : TEXT("nullptr"), OldPage);
-	}
-#endif
-	return bResult;
-}
-
-bool UMaterial::GetFontParameterValue_New(const FHashedMaterialParameterInfo& ParameterInfo, UFont*& OutFontValue, int32& OutFontPage, bool bOveriddenOnly) const
-{
-	if (!bOveriddenOnly && CachedExpressionData)
-	{
-		const int32 Index = CachedExpressionData->Parameters.FindParameterIndex(EMaterialParameterType::Font, ParameterInfo);
-		if (Index != INDEX_NONE)
-		{
-			OutFontValue = CachedExpressionData->Parameters.FontValues[Index];
-			OutFontPage = CachedExpressionData->Parameters.FontPageValues[Index];
-			return true;
-		}
-	}
-
-	OutFontValue = nullptr;
-	OutFontPage = INDEX_NONE;
-	return false;
-}
-
-#if WITH_EDITOR
-bool UMaterial::GetFontParameterValue_Legacy(const FHashedMaterialParameterInfo& ParameterInfo, UFont*& OutFontValue, int32& OutFontPage, bool bOveriddenOnly) const
-{
-	// In the case of duplicate parameters with different values, this will return the
-	// first matching expression found, not necessarily the one that's used for rendering
-	UMaterialExpressionFontSampleParameter* Parameter = nullptr;
-	
-	for (UMaterialExpression* Expression : Expressions)
-	{
-		// Only need to check parameters that match in associated scope
-		if (ParameterInfo.Association == EMaterialParameterAssociation::GlobalParameter)
-		{
-			if (UMaterialExpressionFontSampleParameter* ExpressionParameter = Cast<UMaterialExpressionFontSampleParameter>(Expression))
-			{
-				if (ExpressionParameter->IsNamedParameter(ParameterInfo, OutFontValue, OutFontPage))
-				{
-					// NOTE - actual legacy version returned 'true' here, updating to respect bOveriddenOnly to match other GetXXXParameterValue_Legacy implementations
-					return !bOveriddenOnly;
-				}
-			}
-			else if (UMaterialExpressionMaterialFunctionCall* FunctionCall = Cast<UMaterialExpressionMaterialFunctionCall>(Expression))
-			{
-				UMaterialFunctionInterface* Function = FunctionCall->MaterialFunction;
-				UMaterialFunctionInterface* ParameterOwner = nullptr;
-
-				if (Function && Function->OverrideNamedFontParameter(ParameterInfo, OutFontValue, OutFontPage))
-				{
-					return true;
-				}
-
-				if (Function && Function->GetNamedParameterOfType(ParameterInfo, Parameter, &ParameterOwner))
-				{
-					if (ParameterOwner->OverrideNamedFontParameter(ParameterInfo, OutFontValue, OutFontPage))
-					{
-						return true;
-					}
-					Parameter->IsNamedParameter(ParameterInfo, OutFontValue, OutFontPage);
-					return !bOveriddenOnly;
-				}
-			}
-		}
-		else
-		{
-			if (UMaterialExpressionMaterialAttributeLayers* LayersExpression = Cast<UMaterialExpressionMaterialAttributeLayers>(Expression))
-			{
-				UMaterialFunctionInterface* Function = LayersExpression->GetParameterAssociatedFunction(ParameterInfo);
-				UMaterialFunctionInterface* ParameterOwner = nullptr;
-
-				if (Function && Function->OverrideNamedFontParameter(ParameterInfo, OutFontValue, OutFontPage))
-				{
-					return true;
-				}
-
-				if (Function && Function->GetNamedParameterOfType(ParameterInfo, Parameter, &ParameterOwner))
-				{
-					if (ParameterOwner->OverrideNamedFontParameter(ParameterInfo, OutFontValue, OutFontPage))
-					{
-						return true;
-					}
-					Parameter->IsNamedParameter(ParameterInfo, OutFontValue, OutFontPage);
-					return !bOveriddenOnly;
-				}
-			}
-		}
-	}
-
-	return false;
-}
-#endif // WITH_EDITOR
-
 #if WITH_EDITORONLY_DATA
-bool UMaterial::GetStaticSwitchParameterValue(const FHashedMaterialParameterInfo& ParameterInfo, bool& OutValue, FGuid& OutExpressionGuid, bool bOveriddenOnly, bool bCheckParent) const
-{
-	if (!bOveriddenOnly && CachedExpressionData)
-	{
-		const int32 Index = CachedExpressionData->Parameters.FindParameterIndex(EMaterialParameterType::StaticSwitch, ParameterInfo);
-		if (Index != INDEX_NONE)
-		{
-			OutExpressionGuid = CachedExpressionData->Parameters.GetExpressionGuid(EMaterialParameterType::StaticSwitch, Index);
-			OutValue = CachedExpressionData->Parameters.StaticSwitchValues[Index];
-			return true;
-		}
-	}
-	return false;
-}
-
-bool UMaterial::GetStaticComponentMaskParameterValue(const FHashedMaterialParameterInfo& ParameterInfo, bool& R, bool& G, bool& B, bool& A, FGuid& OutExpressionGuid, bool bOveriddenOnly, bool bCheckParent) const
-{
-	if (!bOveriddenOnly && CachedExpressionData)
-	{
-		const int32 Index = CachedExpressionData->Parameters.FindParameterIndex(EMaterialParameterType::StaticComponentMask, ParameterInfo);
-		if (Index != INDEX_NONE)
-		{
-			OutExpressionGuid = CachedExpressionData->Parameters.GetExpressionGuid(EMaterialParameterType::StaticComponentMask, Index);
-			R = CachedExpressionData->Parameters.StaticComponentMaskValues[Index].R;
-			G = CachedExpressionData->Parameters.StaticComponentMaskValues[Index].G;
-			B = CachedExpressionData->Parameters.StaticComponentMaskValues[Index].B;
-			A = CachedExpressionData->Parameters.StaticComponentMaskValues[Index].A;
-		}
-	}
-	return false;
-}
-
 bool UMaterial::GetMaterialLayersParameterValue(const FHashedMaterialParameterInfo& ParameterInfo, FMaterialLayersFunctions& OutLayers, FGuid& OutExpressionGuid, bool bCheckParent /*= true*/) const
 {
 	UMaterialExpressionStaticComponentMaskParameter* Parameter = nullptr;
diff --git a/Engine/Source/Runtime/Engine/Private/Materials/MaterialCachedData.cpp b/Engine/Source/Runtime/Engine/Private/Materials/MaterialCachedData.cpp
--- a/Engine/Source/Runtime/Engine/Private/Materials/MaterialCachedData.cpp
+++ b/Engine/Source/Runtime/Engine/Private/Materials/MaterialCachedData.cpp
@@ -44,6 +44,7 @@
 #include "Curves/CurveLinearColor.h"
 #include "Curves/CurveLinearColorAtlas.h"
 
+const FMaterialCachedParameterEntry FMaterialCachedParameterEntry::EmptyData{};
 const FMaterialCachedExpressionData FMaterialCachedExpressionData::EmptyData{};
 
 void FMaterialCachedExpressionData::Reset()
@@ -103,10 +104,7 @@ void FMaterialCachedParameters::AddReferencedObjects(FReferenceCollector& Collec
 #if WITH_EDITOR
 static int32 TryAddParameter(FMaterialCachedParameters& CachedParameters, EMaterialParameterType Type, const FMaterialParameterInfo& ParameterInfo, const FGuid& ExpressionGuid)
 {
-	FMaterialCachedParameterEntry& Entry = Type >= EMaterialParameterType::RuntimeCount ?
-													CachedParameters.EditorOnlyEntries[static_cast<int32>(Type) - static_cast<int32>(EMaterialParameterType::RuntimeCount)] :
-													CachedParameters.RuntimeEntries[static_cast<int32>(Type)];
-
+	FMaterialCachedParameterEntry& Entry = CachedParameters.GetParameterTypeEntry(Type);
 	FSetElementId ElementId = Entry.ParameterInfoSet.FindId(ParameterInfo);
 	int32 Index = INDEX_NONE;
 	if (!ElementId.IsValidId())
@@ -701,7 +699,9 @@ bool FMaterialCachedExpressionData::UpdateForExpressions(const FMaterialCachedEx
 void FMaterialCachedParameterEntry::Reset()
 {
 	ParameterInfoSet.Reset();
+#if WITH_EDITORONLY_DATA
 	ExpressionGuids.Reset();
+#endif
 }
 
 void FMaterialCachedParameters::Reset()
@@ -743,140 +743,144 @@ int32 FMaterialCachedParameters::FindParameterIndex(EMaterialParameterType Type,
 	return ElementId.AsInteger();
 }
 
-const FGuid& FMaterialCachedParameters::GetExpressionGuid(EMaterialParameterType Type, int32 Index) const
+void FMaterialCachedParameters::GetParameterValueByIndex(EMaterialParameterType Type, int32 ParameterIndex, FMaterialParameterMetadata& OutResult) const
 {
 	const FMaterialCachedParameterEntry& Entry = GetParameterTypeEntry(Type);
-	return Entry.ExpressionGuids[Index];
-}
 
-void FMaterialCachedParameters::GetAllParameterInfoOfType(EMaterialParameterType Type, bool bEmptyOutput, TArray<FMaterialParameterInfo>& OutParameterInfo, TArray<FGuid>& OutParameterIds) const
-{
-	const FMaterialCachedParameterEntry& Entry = GetParameterTypeEntry(Type);
-	const int32 NumParameters = Entry.ParameterInfoSet.Num();
-	if (bEmptyOutput)
-	{
-		OutParameterInfo.Empty(NumParameters);
-		OutParameterIds.Empty(NumParameters);
-	}
-
-	for (TSet<FMaterialParameterInfo>::TConstIterator It(Entry.ParameterInfoSet); It; ++It)
-	{
-		OutParameterInfo.Add(*It);
-		OutParameterIds.Add(Entry.ExpressionGuids[It.GetId().AsInteger()]);
-	}
-}
+#if WITH_EDITORONLY_DATA
+	OutResult.ExpressionGuid = Entry.ExpressionGuids[ParameterIndex];
+#endif
 
-void FMaterialCachedParameters::GetAllGlobalParameterInfoOfType(EMaterialParameterType Type, bool bEmptyOutput, TArray<FMaterialParameterInfo>& OutParameterInfo, TArray<FGuid>& OutParameterIds) const
-{
-	const FMaterialCachedParameterEntry& Entry = GetParameterTypeEntry(Type);
-	const int32 NumParameters = Entry.ParameterInfoSet.Num();
-	if (bEmptyOutput)
+	switch (Type)
 	{
-		OutParameterInfo.Empty(NumParameters);
-		OutParameterIds.Empty(NumParameters);
-	}
-
-	for (TSet<FMaterialParameterInfo>::TConstIterator It(Entry.ParameterInfoSet); It; ++It)
-	{
-		const FMaterialParameterInfo& ParameterInfo = *It;
-		if (ParameterInfo.Association == GlobalParameter)
+	case EMaterialParameterType::Scalar:
+		OutResult.Value = ScalarValues[ParameterIndex];
+#if WITH_EDITORONLY_DATA
+		OutResult.ScalarMin = ScalarMinMaxValues[ParameterIndex].X;
+		OutResult.ScalarMax = ScalarMinMaxValues[ParameterIndex].Y;
 		{
-			OutParameterInfo.Add(ParameterInfo);
-			OutParameterIds.Add(Entry.ExpressionGuids[It.GetId().AsInteger()]);
+			UCurveLinearColor* Curve = ScalarCurveValues[ParameterIndex];
+			UCurveLinearColorAtlas* Atlas = ScalarCurveAtlasValues[ParameterIndex];
+			if (Curve && Atlas)
+			{
+				OutResult.ScalarCurve = Curve;
+				OutResult.ScalarAtlas = Atlas;
+				OutResult.bUsedAsAtlasPosition = true;
+			}
 		}
+#endif // WITH_EDITORONLY_DATA
+		break;
+	case EMaterialParameterType::Vector:
+		OutResult.Value = VectorValues[ParameterIndex];
+#if  WITH_EDITORONLY_DATA
+		OutResult.ChannelName[0] = VectorChannelNameValues[ParameterIndex].R;
+		OutResult.ChannelName[1] = VectorChannelNameValues[ParameterIndex].G;
+		OutResult.ChannelName[2] = VectorChannelNameValues[ParameterIndex].B;
+		OutResult.ChannelName[3] = VectorChannelNameValues[ParameterIndex].A;
+		OutResult.bUsedAsChannelMask = VectorUsedAsChannelMaskValues[ParameterIndex];
+#endif // WITH_EDITORONLY_DATA
+		break;
+	case EMaterialParameterType::Texture:
+		OutResult.Value = TextureValues[ParameterIndex];
+#if WITH_EDITORONLY_DATA
+		OutResult.ChannelName[0] = TextureChannelNameValues[ParameterIndex].R;
+		OutResult.ChannelName[1] = TextureChannelNameValues[ParameterIndex].G;
+		OutResult.ChannelName[2] = TextureChannelNameValues[ParameterIndex].B;
+		OutResult.ChannelName[3] = TextureChannelNameValues[ParameterIndex].A;
+#endif // WITH_EDITORONLY_DATA
+		break;
+	case EMaterialParameterType::RuntimeVirtualTexture:
+		OutResult.Value = RuntimeVirtualTextureValues[ParameterIndex];
+		break;
+	case EMaterialParameterType::Font:
+		OutResult.Value = FMaterialParameterValue(FontValues[ParameterIndex], FontPageValues[ParameterIndex]);
+		break;
+#if WITH_EDITORONLY_DATA
+	case EMaterialParameterType::StaticSwitch:
+		OutResult.Value = StaticSwitchValues[ParameterIndex];
+		break;
+	case EMaterialParameterType::StaticComponentMask:
+		OutResult.Value = StaticComponentMaskValues[ParameterIndex];
+		break;
+#endif // WITH_EDITORONLY_DATA
+	default:
+		checkNoEntry();
+		break;
 	}
 }
 
-bool FMaterialCachedParameters::GetScalarParameterSliderMinMax(const FMemoryImageMaterialParameterInfo& ParameterInfo, float& OutSliderMin, float& OutSliderMax) const
+bool FMaterialCachedParameters::GetParameterValue(EMaterialParameterType Type, const FMemoryImageMaterialParameterInfo& ParameterInfo, FMaterialParameterMetadata& OutResult) const
 {
-#if WITH_EDITORONLY_DATA
-	if (ScalarMinMaxValues.Num() != 0)
+	const int32 Index = FindParameterIndex(Type, ParameterInfo);
+	if (Index != INDEX_NONE)
 	{
-		const int32 ParameterIndex = FindParameterIndex(EMaterialParameterType::Scalar, ParameterInfo);
-		if (ParameterIndex != INDEX_NONE)
-		{
-			OutSliderMin = ScalarMinMaxValues[ParameterIndex].X;
-			OutSliderMax = ScalarMinMaxValues[ParameterIndex].Y;
-			return true;
-		}
+		GetParameterValueByIndex(Type, Index, OutResult);
+		return true;
 	}
-#endif // WITH_EDITORONLY_DATA
+
 	return false;
 }
 
-bool FMaterialCachedParameters::IsScalarParameterUsedAsAtlasPosition(const FMemoryImageMaterialParameterInfo& ParameterInfo, bool& OutValue, TSoftObjectPtr<UCurveLinearColor>& OutCurve, TSoftObjectPtr<UCurveLinearColorAtlas>& OutAtlas) const
-{
 #if WITH_EDITORONLY_DATA
-	if (ScalarCurveValues.Num() != 0)
-	{
-		const int32 ParameterIndex = FindParameterIndex(EMaterialParameterType::Scalar, ParameterInfo);
-		if (ParameterIndex != INDEX_NONE)
-		{
-			UCurveLinearColor* Curve = ScalarCurveValues[ParameterIndex];
-			UCurveLinearColorAtlas* Atlas = ScalarCurveAtlasValues[ParameterIndex];
-			if (Curve && Atlas)
-			{
-				OutCurve = Curve;
-				OutAtlas = Atlas;
-				OutValue = true;
-			}
-			else
-			{
-				OutValue = false;
-			}
-			return true;
-		}
-	}
-#endif // WITH_EDITORONLY_DATA
-	return false;
+const FGuid& FMaterialCachedParameters::GetExpressionGuid(EMaterialParameterType Type, int32 Index) const
+{
+	const FMaterialCachedParameterEntry& Entry = GetParameterTypeEntry(Type);
+	return Entry.ExpressionGuids[Index];
 }
+#endif // WITH_EDITORONLY_DATA
 
-bool FMaterialCachedParameters::IsVectorParameterUsedAsChannelMask(const FMemoryImageMaterialParameterInfo& ParameterInfo, bool& OutValue) const
+void FMaterialCachedParameters::GetAllParametersOfType(EMaterialParameterType Type, TMap<FMaterialParameterInfo, FMaterialParameterMetadata>& OutParameters) const
 {
-#if WITH_EDITORONLY_DATA
-	if (VectorUsedAsChannelMaskValues.Num() != 0)
+	const FMaterialCachedParameterEntry& Entry = GetParameterTypeEntry(Type);
+	const int32 NumParameters = Entry.ParameterInfoSet.Num();
+	OutParameters.Reserve(OutParameters.Num() + NumParameters);
+
+	for (int32 ParameterIndex = 0; ParameterIndex < NumParameters; ++ParameterIndex)
 	{
-		const int32 ParameterIndex = FindParameterIndex(EMaterialParameterType::Vector, ParameterInfo);
-		if (ParameterIndex != INDEX_NONE)
-		{
-			OutValue = VectorUsedAsChannelMaskValues[ParameterIndex];
-			return true;
-		}
+		const FMaterialParameterInfo& ParameterInfo = Entry.ParameterInfoSet[FSetElementId::FromInteger(ParameterIndex)];
+		FMaterialParameterMetadata& Result = OutParameters.Emplace(ParameterInfo);
+		GetParameterValueByIndex(Type, ParameterIndex, Result);
 	}
-#endif // WITH_EDITORONLY_DATA
-	return false;
 }
 
-bool FMaterialCachedParameters::GetVectorParameterChannelNames(const FMemoryImageMaterialParameterInfo& ParameterInfo, FParameterChannelNames& OutValue) const
+void FMaterialCachedParameters::GetAllParameterInfoOfType(EMaterialParameterType Type, TArray<FMaterialParameterInfo>& OutParameterInfo, TArray<FGuid>& OutParameterIds) const
 {
-#if WITH_EDITORONLY_DATA
-	if (VectorChannelNameValues.Num() != 0)
+	const FMaterialCachedParameterEntry& Entry = GetParameterTypeEntry(Type);
+	const int32 NumParameters = Entry.ParameterInfoSet.Num();
+	OutParameterInfo.Reserve(OutParameterInfo.Num() + NumParameters);
+	OutParameterIds.Reserve(OutParameterIds.Num() + NumParameters);
+
+	for (TSet<FMaterialParameterInfo>::TConstIterator It(Entry.ParameterInfoSet); It; ++It)
 	{
-		const int32 ParameterIndex = FindParameterIndex(EMaterialParameterType::Vector, ParameterInfo);
-		if (ParameterIndex != INDEX_NONE)
-		{
-			OutValue = VectorChannelNameValues[ParameterIndex];
-			return true;
-		}
+		OutParameterInfo.Add(*It);
+#if WITH_EDITORONLY_DATA
+		OutParameterIds.Add(Entry.ExpressionGuids[It.GetId().AsInteger()]);
+#else
+		OutParameterIds.Add(FGuid());
+#endif
 	}
-#endif // WITH_EDITORONLY_DATA
-	return false;
 }
 
-bool FMaterialCachedParameters::GetTextureParameterChannelNames(const FMemoryImageMaterialParameterInfo& ParameterInfo, FParameterChannelNames& OutValue) const
+void FMaterialCachedParameters::GetAllGlobalParameterInfoOfType(EMaterialParameterType Type, TArray<FMaterialParameterInfo>& OutParameterInfo, TArray<FGuid>& OutParameterIds) const
 {
-#if WITH_EDITORONLY_DATA
-	if (TextureChannelNameValues.Num() != 0)
+	const FMaterialCachedParameterEntry& Entry = GetParameterTypeEntry(Type);
+	const int32 NumParameters = Entry.ParameterInfoSet.Num();
+	OutParameterIds.Reserve(OutParameterInfo.Num() + NumParameters);
+	OutParameterIds.Reserve(OutParameterIds.Num() + NumParameters);
+
+	for (TSet<FMaterialParameterInfo>::TConstIterator It(Entry.ParameterInfoSet); It; ++It)
 	{
-		const int32 ParameterIndex = FindParameterIndex(EMaterialParameterType::Texture, ParameterInfo);
-		if (ParameterIndex != INDEX_NONE)
+		const FMaterialParameterInfo& ParameterInfo = *It;
+		if (ParameterInfo.Association == GlobalParameter)
 		{
-			OutValue = TextureChannelNameValues[ParameterIndex];
-			return true;
+			OutParameterInfo.Add(ParameterInfo);
+#if WITH_EDITORONLY_DATA
+			OutParameterIds.Add(Entry.ExpressionGuids[It.GetId().AsInteger()]);
+#else
+			OutParameterIds.Add(FGuid());
+#endif
 		}
 	}
-#endif // WITH_EDITORONLY_DATA
-	return false;
 }
 
 void FMaterialInstanceCachedData::Initialize(FMaterialCachedExpressionData&& InCachedExpressionData)
diff --git a/Engine/Source/Runtime/Engine/Private/Materials/MaterialExpressions.cpp b/Engine/Source/Runtime/Engine/Private/Materials/MaterialExpressions.cpp
--- a/Engine/Source/Runtime/Engine/Private/Materials/MaterialExpressions.cpp
+++ b/Engine/Source/Runtime/Engine/Private/Materials/MaterialExpressions.cpp
@@ -13043,17 +13043,14 @@ bool UMaterialFunction::HasFlippedCoordinates() const
 	return ReversedInputCount > StandardInputCount;
 }
 
-bool UMaterialFunction::SetVectorParameterValueEditorOnly(FName ParameterName, FLinearColor InValue)
+bool UMaterialFunction::SetParameterValueEditorOnly(const FName& ParameterName, const FMaterialParameterMetadata& Meta)
 {
 	for (UMaterialExpression* Expression : FunctionExpressions)
 	{
-		if (UMaterialExpressionVectorParameter* Parameter = Cast<UMaterialExpressionVectorParameter>(Expression))
+		if (Expression->SetParameterValue(ParameterName, Meta))
 		{
-			if (Parameter->SetParameterValue(ParameterName, InValue))
-			{
-				return true;
-				// Warning: in the case of duplicate parameters with different default values, this will find the first in the expression array, not necessarily the one that's used for rendering
-			}
+			return true;
+			// Warning: in the case of duplicate parameters with different default values, this will find the first in the expression array, not necessarily the one that's used for rendering
 		}
 		else if (UMaterialExpressionMaterialFunctionCall* FunctionCall = Cast<UMaterialExpressionMaterialFunctionCall>(Expression))
 		{
@@ -13068,14 +13065,11 @@ bool UMaterialFunction::SetVectorParameterValueEditorOnly(FName ParameterName, F
 					const TArray<TObjectPtr<UMaterialExpression>>* ExpressionPtr = Function->GetFunctionExpressions();
 					if (ExpressionPtr)
 					{
-						for (const TObjectPtr<UMaterialExpression>&FunctionExpression : *ExpressionPtr)
+						for (const TObjectPtr<UMaterialExpression>& FunctionExpression : *ExpressionPtr)
 						{
-							if (UMaterialExpressionVectorParameter* FunctionExpressionParameter = Cast<UMaterialExpressionVectorParameter>(FunctionExpression))
+							if (FunctionExpression->SetParameterValue(ParameterName, Meta))
 							{
-								if (FunctionExpressionParameter->SetParameterValue(ParameterName, InValue))
-								{
-									return true;
-								}
+								return true;
 							}
 						}
 					}
@@ -13086,265 +13080,55 @@ bool UMaterialFunction::SetVectorParameterValueEditorOnly(FName ParameterName, F
 	return false;
 }
 
-bool UMaterialFunction::SetScalarParameterValueEditorOnly(FName ParameterName, float InValue)
+bool UMaterialFunction::SetVectorParameterValueEditorOnly(FName ParameterName, FLinearColor InValue)
 {
-	for (UMaterialExpression* Expression : FunctionExpressions)
-	{
-		if (UMaterialExpressionScalarParameter* Parameter = Cast<UMaterialExpressionScalarParameter>(Expression))
-		{
-			if (Parameter->SetParameterValue(ParameterName, InValue))
-			{
-				return true;
-				// Warning: in the case of duplicate parameters with different default values, this will find the first in the expression array, not necessarily the one that's used for rendering
-			}
-		}
-		else if (UMaterialExpressionMaterialFunctionCall* FunctionCall = Cast<UMaterialExpressionMaterialFunctionCall>(Expression))
-		{
-			if (FunctionCall->MaterialFunction)
-			{
-				TArray<UMaterialFunctionInterface*> Functions;
-				Functions.Add(FunctionCall->MaterialFunction);
-				FunctionCall->MaterialFunction->GetDependentFunctions(Functions);
+	FMaterialParameterMetadata Meta;
+	Meta.Value = InValue;
+	return SetParameterValueEditorOnly(ParameterName, Meta);
+}
 
-				for (UMaterialFunctionInterface* Function : Functions)
-				{
-					const TArray<TObjectPtr<UMaterialExpression>>* ExpressionPtr = Function->GetFunctionExpressions();
-					if (ExpressionPtr)
-					{
-						for (const TObjectPtr<UMaterialExpression>&FunctionExpression : *ExpressionPtr)
-						{
-							if (UMaterialExpressionScalarParameter* FunctionExpressionParameter = Cast<UMaterialExpressionScalarParameter>(FunctionExpression))
-							{
-								if (FunctionExpressionParameter->SetParameterValue(ParameterName, InValue))
-								{
-									return true;
-								}
-							}
-						}
-					}
-				}
-			}
-		}
-	}
-	return false;
+bool UMaterialFunction::SetScalarParameterValueEditorOnly(FName ParameterName, float InValue)
+{
+	FMaterialParameterMetadata Meta;
+	Meta.Value = InValue;
+	return SetParameterValueEditorOnly(ParameterName, Meta);
 };
 
 bool UMaterialFunction::SetTextureParameterValueEditorOnly(FName ParameterName, class UTexture* InValue)
 {
-	for (UMaterialExpression* Expression : FunctionExpressions)
-	{
-		if (UMaterialExpressionTextureSampleParameter* Parameter = Cast<UMaterialExpressionTextureSampleParameter>(Expression))
-		{
-			if (Parameter->SetParameterValue(ParameterName, InValue))
-			{
-				return true;
-				// Warning: in the case of duplicate parameters with different default values, this will find the first in the expression array, not necessarily the one that's used for rendering
-			}
-		}
-		else if (UMaterialExpressionMaterialFunctionCall* FunctionCall = Cast<UMaterialExpressionMaterialFunctionCall>(Expression))
-		{
-			if (FunctionCall->MaterialFunction)
-			{
-				TArray<UMaterialFunctionInterface*> Functions;
-				Functions.Add(FunctionCall->MaterialFunction);
-				FunctionCall->MaterialFunction->GetDependentFunctions(Functions);
-
-				for (UMaterialFunctionInterface* Function : Functions)
-				{
-					const TArray<TObjectPtr<UMaterialExpression>>* ExpressionPtr = Function->GetFunctionExpressions();
-					if (ExpressionPtr)
-					{
-						for (const TObjectPtr<UMaterialExpression>&FunctionExpression : *ExpressionPtr)
-						{
-							if (UMaterialExpressionTextureSampleParameter* FunctionExpressionParameter = Cast<UMaterialExpressionTextureSampleParameter>(FunctionExpression))
-							{
-								if (FunctionExpressionParameter->SetParameterValue(ParameterName, InValue))
-								{
-									return true;
-								}
-							}
-						}
-					}
-				}
-			}
-		}
-	}
-	return false;
+	FMaterialParameterMetadata Meta;
+	Meta.Value = InValue;
+	return SetParameterValueEditorOnly(ParameterName, Meta);
 };
 
 bool UMaterialFunction::SetRuntimeVirtualTextureParameterValueEditorOnly(FName ParameterName, class URuntimeVirtualTexture* InValue)
 {
-	for (UMaterialExpression* Expression : FunctionExpressions)
-	{
-		if (UMaterialExpressionRuntimeVirtualTextureSampleParameter* Parameter = Cast<UMaterialExpressionRuntimeVirtualTextureSampleParameter>(Expression))
-		{
-			if (Parameter->SetParameterValue(ParameterName, InValue))
-			{
-				return true;
-				// Warning: in the case of duplicate parameters with different default values, this will find the first in the expression array, not necessarily the one that's used for rendering
-			}
-		}
-		else if (UMaterialExpressionMaterialFunctionCall* FunctionCall = Cast<UMaterialExpressionMaterialFunctionCall>(Expression))
-		{
-			if (FunctionCall->MaterialFunction)
-			{
-				TArray<UMaterialFunctionInterface*> Functions;
-				Functions.Add(FunctionCall->MaterialFunction);
-				FunctionCall->MaterialFunction->GetDependentFunctions(Functions);
-
-				for (UMaterialFunctionInterface* Function : Functions)
-				{
-					const TArray<TObjectPtr<UMaterialExpression>>* ExpressionPtr = Function->GetFunctionExpressions();
-					if (ExpressionPtr)
-					{
-						for (const TObjectPtr<UMaterialExpression>&FunctionExpression : *ExpressionPtr)
-						{
-							if (UMaterialExpressionRuntimeVirtualTextureSampleParameter* FunctionExpressionParameter = Cast<UMaterialExpressionRuntimeVirtualTextureSampleParameter>(FunctionExpression))
-							{
-								if (FunctionExpressionParameter->SetParameterValue(ParameterName, InValue))
-								{
-									return true;
-								}
-							}
-						}
-					}
-				}
-			}
-		}
-	}
-	return false;
+	FMaterialParameterMetadata Meta;
+	Meta.Value = InValue;
+	return SetParameterValueEditorOnly(ParameterName, Meta);
 };
 
 bool UMaterialFunction::SetFontParameterValueEditorOnly(FName ParameterName, class UFont* InFontValue, int32 InFontPage)
 {
-	for (UMaterialExpression* Expression : FunctionExpressions)
-	{
-		if (UMaterialExpressionFontSampleParameter* Parameter = Cast<UMaterialExpressionFontSampleParameter>(Expression))
-		{
-			if (Parameter->SetParameterValue(ParameterName, InFontValue, InFontPage))
-			{
-				return true;
-				// Warning: in the case of duplicate parameters with different default values, this will find the first in the expression array, not necessarily the one that's used for rendering
-			}
-		}
-		else if (UMaterialExpressionMaterialFunctionCall* FunctionCall = Cast<UMaterialExpressionMaterialFunctionCall>(Expression))
-		{
-			if (FunctionCall->MaterialFunction)
-			{
-				TArray<UMaterialFunctionInterface*> Functions;
-				Functions.Add(FunctionCall->MaterialFunction);
-				FunctionCall->MaterialFunction->GetDependentFunctions(Functions);
-
-				for (UMaterialFunctionInterface* Function : Functions)
-				{
-					const TArray<TObjectPtr<UMaterialExpression>>* ExpressionPtr = Function->GetFunctionExpressions();
-					if (ExpressionPtr)
-					{
-						for (const TObjectPtr<UMaterialExpression>&FunctionExpression : *ExpressionPtr)
-						{
-							if (FunctionExpression && FunctionExpression.IsA<UMaterialExpressionFontSampleParameter>())
-							{
-								UMaterialExpressionFontSampleParameter* FunctionExpressionParameter = (UMaterialExpressionFontSampleParameter*)FunctionExpression.Get();
-								if (FunctionExpressionParameter->SetParameterValue(ParameterName, InFontValue, InFontPage))
-								{
-									return true;
-								}
-							}
-						}
-					}
-				}
-			}
-		}
-	}
-	return false;
+	FMaterialParameterMetadata Meta;
+	Meta.Value = FMaterialParameterValue(InFontValue, InFontPage);
+	return SetParameterValueEditorOnly(ParameterName, Meta);
 };
 
 bool UMaterialFunction::SetStaticSwitchParameterValueEditorOnly(FName ParameterName, bool OutValue, FGuid OutExpressionGuid)
 {
-	for (UMaterialExpression* Expression : FunctionExpressions)
-	{
-		if (UMaterialExpressionStaticSwitchParameter* Parameter = Cast<UMaterialExpressionStaticSwitchParameter>(Expression))
-		{
-			if (Parameter->SetParameterValue(ParameterName, OutValue, OutExpressionGuid))
-			{
-				return true;
-				// Warning: in the case of duplicate parameters with different default values, this will find the first in the expression array, not necessarily the one that's used for rendering
-			}
-		}
-		else if (UMaterialExpressionMaterialFunctionCall* FunctionCall = Cast<UMaterialExpressionMaterialFunctionCall>(Expression))
-		{
-			if (FunctionCall->MaterialFunction)
-			{
-				TArray<UMaterialFunctionInterface*> Functions;
-				Functions.Add(FunctionCall->MaterialFunction);
-				FunctionCall->MaterialFunction->GetDependentFunctions(Functions);
-
-				for (UMaterialFunctionInterface* Function : Functions)
-				{
-					const TArray<TObjectPtr<UMaterialExpression>>* ExpressionPtr = Function->GetFunctionExpressions();
-					if (ExpressionPtr)
-					{
-						for (const TObjectPtr<UMaterialExpression>&FunctionExpression : *ExpressionPtr)
-						{
-							if (FunctionExpression && FunctionExpression.IsA<UMaterialExpressionStaticSwitchParameter>())
-							{
-								UMaterialExpressionStaticSwitchParameter* FunctionExpressionParameter = (UMaterialExpressionStaticSwitchParameter*)FunctionExpression.Get();
-								if (FunctionExpressionParameter->SetParameterValue(ParameterName, OutValue, OutExpressionGuid))
-								{
-									return true;
-								}
-							}
-						}
-					}
-				}
-			}
-		}
-	}
-	return false;
+	FMaterialParameterMetadata Meta;
+	Meta.Value = OutValue;
+	Meta.ExpressionGuid = OutExpressionGuid;
+	return SetParameterValueEditorOnly(ParameterName, Meta);
 };
 
 bool UMaterialFunction::SetStaticComponentMaskParameterValueEditorOnly(FName ParameterName, bool R, bool G, bool B, bool A, FGuid OutExpressionGuid)
 {
-	for (UMaterialExpression* Expression : FunctionExpressions)
-	{
-		if (UMaterialExpressionStaticComponentMaskParameter* Parameter = Cast<UMaterialExpressionStaticComponentMaskParameter>(Expression))
-		{
-			if (Parameter->SetParameterValue(ParameterName, R, G, B, A, OutExpressionGuid))
-			{
-				return true;
-				// Warning: in the case of duplicate parameters with different default values, this will find the first in the expression array, not necessarily the one that's used for rendering
-			}
-		}
-		else if (UMaterialExpressionMaterialFunctionCall* FunctionCall = Cast<UMaterialExpressionMaterialFunctionCall>(Expression))
-		{
-			if (FunctionCall->MaterialFunction)
-			{
-				TArray<UMaterialFunctionInterface*> Functions;
-				Functions.Add(FunctionCall->MaterialFunction);
-				FunctionCall->MaterialFunction->GetDependentFunctions(Functions);
-
-				for (UMaterialFunctionInterface* Function : Functions)
-				{
-					const TArray<TObjectPtr<UMaterialExpression>>* ExpressionPtr = Function->GetFunctionExpressions();
-					if (ExpressionPtr)
-					{
-						for (const TObjectPtr<UMaterialExpression>&FunctionExpression : *ExpressionPtr)
-						{
-							if (FunctionExpression && FunctionExpression.IsA<UMaterialExpressionStaticComponentMaskParameter>())
-							{
-								UMaterialExpressionStaticComponentMaskParameter* FunctionExpressionParameter = (UMaterialExpressionStaticComponentMaskParameter*)FunctionExpression.Get();
-								if (FunctionExpressionParameter->SetParameterValue(ParameterName, R, G, B, A, OutExpressionGuid))
-								{
-									return true;
-								}
-							}
-						}
-					}
-				}
-			}
-		}
-	}
-	return false;
+	FMaterialParameterMetadata Meta;
+	Meta.Value = FMaterialParameterValue(R, G, B, A);
+	Meta.ExpressionGuid = OutExpressionGuid;
+	return SetParameterValueEditorOnly(ParameterName, Meta);
 };
 #endif //WITH_EDITORONLY_DATA
 
diff --git a/Engine/Source/Runtime/Engine/Private/Materials/MaterialInstance.cpp b/Engine/Source/Runtime/Engine/Private/Materials/MaterialInstance.cpp
--- a/Engine/Source/Runtime/Engine/Private/Materials/MaterialInstance.cpp
+++ b/Engine/Source/Runtime/Engine/Private/Materials/MaterialInstance.cpp
@@ -49,6 +49,8 @@
 #include "UObject/UE5MainStreamObjectVersion.h"
 #include "ShaderCompilerCore.h"
 
+PRAGMA_DISABLE_OPTIMIZATION
+
 #if ENABLE_COOK_STATS
 #include "ProfilingDebugging/ScopedTimers.h"
 namespace MaterialInstanceCookStats
@@ -833,113 +835,7 @@ UMaterial* UMaterialInstance::GetMaterial()
 	}
 }
 
-#if WITH_EDITOR
-bool UMaterialInstance::GetScalarParameterSliderMinMax(const FHashedMaterialParameterInfo& ParameterInfo, float& OutSliderMin, float& OutSliderMax) const
-{
-	if (GetReentrantFlag())
-	{
-		return false;
-	}
-
-	if (ParameterInfo.Association != EMaterialParameterAssociation::GlobalParameter && CachedData)
-	{
-		if (CachedData->Parameters.GetScalarParameterSliderMinMax(ParameterInfo, OutSliderMin, OutSliderMax))
-		{
-			return true;
-		}
-	}
-
-	if (Parent)
-	{
-		FMICReentranceGuard	Guard(this);
-		return Parent->GetScalarParameterSliderMinMax(ParameterInfo, OutSliderMin, OutSliderMax);
-	}
-
-	return false;
-}
-#endif // WITH_EDITOR
-
-bool UMaterialInstance::GetScalarParameterValue(const FHashedMaterialParameterInfo& ParameterInfo, float& OutValue, bool bOveriddenOnly) const
-{
-	bool bFoundAValue = false;
-
-	if (GetReentrantFlag())
-	{
-		return false;
-	}
-
-	// Instance override
-	const FScalarParameterValue* ParameterValue = GameThread_FindParameterByName(ScalarParameterValues, ParameterInfo);
-	if (ParameterValue)
-	{
-		OutValue = ParameterValue->ParameterValue;
-		return true;
-	}
-	
-	// Instance-included default
-	if (!bOveriddenOnly && ParameterInfo.Association != EMaterialParameterAssociation::GlobalParameter && CachedData)
-	{
-		const int32 ParameterIndex = CachedData->Parameters.FindParameterIndex(EMaterialParameterType::Scalar, ParameterInfo);
-		if (ParameterIndex != INDEX_NONE)
-		{
-			OutValue = CachedData->Parameters.ScalarValues[ParameterIndex];
-			return true;
-		}
-	}
-	
-	// Next material in hierarchy
-	if (Parent)
-	{
-		FMICReentranceGuard	Guard(this);
-		return Parent->GetScalarParameterValue(ParameterInfo, OutValue, bOveriddenOnly);
-	}
-	
-	return false;
-}
-
-#if WITH_EDITOR
-bool UMaterialInstance::IsScalarParameterUsedAsAtlasPosition(const FHashedMaterialParameterInfo& ParameterInfo, bool& OutValue, TSoftObjectPtr<UCurveLinearColor>& OutCurve, TSoftObjectPtr<UCurveLinearColorAtlas>& OutAtlas) const
-{
-	bool bFoundAValue = false;
-
-	if (GetReentrantFlag())
-	{
-		return false;
-	}
-
-	// Instance override
-	const FScalarParameterValue* ParameterValue = GameThread_FindParameterByName(ScalarParameterValues, ParameterInfo);
-#if WITH_EDITOR
-	if (ParameterValue && ParameterValue->AtlasData.Curve.Get() != nullptr && ParameterValue->AtlasData.Atlas.Get() != nullptr)
-	{
-		OutValue = ParameterValue->AtlasData.bIsUsedAsAtlasPosition;
-		OutCurve = ParameterValue->AtlasData.Curve;
-		OutAtlas = ParameterValue->AtlasData.Atlas;
-		return true;
-	}
-#endif
-
-	// Instance-included default
-	if (ParameterInfo.Association != EMaterialParameterAssociation::GlobalParameter && CachedData)
-	{
-		if (CachedData->Parameters.IsScalarParameterUsedAsAtlasPosition(ParameterInfo, OutValue, OutCurve, OutAtlas))
-		{
-			return true;
-		}
-	}
-
-	// Next material in hierarchy
-	if (Parent)
-	{
-		FMICReentranceGuard	Guard(this);
-		return Parent->IsScalarParameterUsedAsAtlasPosition(ParameterInfo, OutValue, OutCurve, OutAtlas);
-	}
-
-	return false;
-}
-#endif // WITH_EDITOR
-
-bool UMaterialInstance::GetVectorParameterValue(const FHashedMaterialParameterInfo& ParameterInfo, FLinearColor& OutValue, bool bOveriddenOnly) const
+bool UMaterialInstance::GetParameterValue(EMaterialParameterType Type, const FMemoryImageMaterialParameterInfo& ParameterInfo, FMaterialParameterMetadata& OutResult, EMaterialGetParameterValueFlags Flags) const
 {
 	bool bFoundAValue = false;
 
@@ -949,244 +845,55 @@ bool UMaterialInstance::GetVectorParameterValue(const FHashedMaterialParameterIn
 	}
 
 	// Instance override
-	const FVectorParameterValue* ParameterValue = GameThread_FindParameterByName(VectorParameterValues, ParameterInfo);
-	if (ParameterValue)
-	{
-		OutValue = ParameterValue->ParameterValue;
-		return true;
-	}
-	
-	// Instance-included default
-	if (!bOveriddenOnly && ParameterInfo.Association != EMaterialParameterAssociation::GlobalParameter && CachedData)
-	{
-		const int32 ParameterIndex = CachedData->Parameters.FindParameterIndex(EMaterialParameterType::Vector, ParameterInfo);
-		if (ParameterIndex != INDEX_NONE)
-		{
-			OutValue = CachedData->Parameters.VectorValues[ParameterIndex];
-			return true;
-		}
-	}
-	
-	// Next material in hierarchy
-	if (Parent)
-	{
-		FMICReentranceGuard	Guard(this);
-		return Parent->GetVectorParameterValue(ParameterInfo, OutValue, bOveriddenOnly);
-	}
-	
-	return false;
-}
-
-#if WITH_EDITOR
-bool UMaterialInstance::IsVectorParameterUsedAsChannelMask(const FHashedMaterialParameterInfo& ParameterInfo, bool& OutValue) const
-{
-	bool bFoundAValue = false;
-
-	if (GetReentrantFlag())
-	{
-		return false;
-	}
-	
-	// Instance-included default
-	if (ParameterInfo.Association != EMaterialParameterAssociation::GlobalParameter && CachedData)
-	{
-		if (CachedData->Parameters.IsVectorParameterUsedAsChannelMask(ParameterInfo, OutValue))
-		{
-			return true;
-		}
-	}
-	
-	// Next material in hierarchy
-	if (Parent)
-	{
-		FMICReentranceGuard	Guard(this);
-		return Parent->IsVectorParameterUsedAsChannelMask(ParameterInfo, OutValue);
-	}
-	
-	return false;
-}
-
-bool  UMaterialInstance::GetVectorParameterChannelNames(const FHashedMaterialParameterInfo& ParameterInfo, FParameterChannelNames& OutValue) const
-{
-	bool bFoundAValue = false;
-
-	if (GetReentrantFlag())
-	{
-		return false;
-	}
-
-	// Instance-included default
-	if (ParameterInfo.Association != EMaterialParameterAssociation::GlobalParameter && CachedData)
-	{
-		if (CachedData->Parameters.GetVectorParameterChannelNames(ParameterInfo, OutValue))
-		{
-			return true;
+	if (EnumHasAnyFlags(Flags, EMaterialGetParameterValueFlags::CheckInstanceOverrides))
+	{
+		// If we're not a parent, set the flag to signify this parameter was a MI override
+		const bool bSetOverride = !EnumHasAnyFlags(Flags, EMaterialGetParameterValueFlags::IsParent);
+		bool bResult = false;
+		switch (Type)
+		{
+		case EMaterialParameterType::Scalar: bResult = GameThread_GetParameterValue(ScalarParameterValues, ParameterInfo, bSetOverride, OutResult); break;
+		case EMaterialParameterType::Vector: bResult = GameThread_GetParameterValue(VectorParameterValues, ParameterInfo, bSetOverride, OutResult); break;
+		case EMaterialParameterType::Texture: bResult = GameThread_GetParameterValue(TextureParameterValues, ParameterInfo, bSetOverride, OutResult); break;
+		case EMaterialParameterType::RuntimeVirtualTexture: bResult = GameThread_GetParameterValue(RuntimeVirtualTextureParameterValues, ParameterInfo, bSetOverride, OutResult); break;
+		case EMaterialParameterType::Font: bResult = GameThread_GetParameterValue(FontParameterValues, ParameterInfo, bSetOverride, OutResult); break;
+#if WITH_EDITORONLY_DATA
+		case EMaterialParameterType::StaticSwitch: bResult = GameThread_GetParameterValue(StaticParameters.StaticSwitchParameters, ParameterInfo, bSetOverride, OutResult); break;
+		case EMaterialParameterType::StaticComponentMask: bResult = GameThread_GetParameterValue(StaticParameters.StaticComponentMaskParameters, ParameterInfo, bSetOverride, OutResult); break;
+#endif // WITH_EDITORONLY_DATA
+		default: checkNoEntry(); break;
 		}
-	}
-
-	// Next material in hierarchy
-	if (Parent)
-	{
-		FMICReentranceGuard	Guard(this);
-		return Parent->GetVectorParameterChannelNames(ParameterInfo, OutValue);
-	}
-
-	return false;
-};
-#endif // WITH_EDITOR
-
-bool UMaterialInstance::GetTextureParameterValue(const FHashedMaterialParameterInfo& ParameterInfo, UTexture*& OutValue, bool bOveriddenOnly) const
-{
-	bool bFoundAValue = false;
-
-	if (GetReentrantFlag())
-	{
-		OutValue = nullptr;
-		return false;
-	}
-
-	// Instance override
-	const FTextureParameterValue* ParameterValue = GameThread_FindParameterByName(TextureParameterValues, ParameterInfo);
-	if (ParameterValue)
-	{
-		OutValue = ParameterValue->ParameterValue;
-		return true;
-	}
-	
-	// Instance-included default
-	if (!bOveriddenOnly && ParameterInfo.Association != EMaterialParameterAssociation::GlobalParameter && CachedData)
-	{
-		const int32 ParameterIndex = CachedData->Parameters.FindParameterIndex(EMaterialParameterType::Texture, ParameterInfo);
-		if (ParameterIndex != INDEX_NONE)
+		if (bResult)
 		{
-			OutValue = CachedData->Parameters.TextureValues[ParameterIndex];
 			return true;
 		}
 	}
-	
-	// Next material in hierarchy
-	if (Parent)
-	{
-		FMICReentranceGuard	Guard(this);
-		return Parent->GetTextureParameterValue(ParameterInfo,OutValue,bOveriddenOnly);
-	}
-	
-	OutValue = nullptr;
-	return false;
-}
-
-bool UMaterialInstance::GetRuntimeVirtualTextureParameterValue(const FHashedMaterialParameterInfo& ParameterInfo, URuntimeVirtualTexture*& OutValue, bool bOveriddenOnly) const
-{
-	bool bFoundAValue = false;
-
-	if (GetReentrantFlag())
-	{
-		OutValue = nullptr;
-		return false;
-	}
-
-	// Instance override
-	const FRuntimeVirtualTextureParameterValue* ParameterValue = GameThread_FindParameterByName(RuntimeVirtualTextureParameterValues, ParameterInfo);
-	if (ParameterValue)
-	{
-		OutValue = ParameterValue->ParameterValue;
-		return true;
-	}
 
 	// Instance-included default
-	if (!bOveriddenOnly && ParameterInfo.Association != EMaterialParameterAssociation::GlobalParameter && CachedData)
+	if (EnumHasAnyFlags(Flags, EMaterialGetParameterValueFlags::CheckNonOverrides) &&
+		ParameterInfo.Association != EMaterialParameterAssociation::GlobalParameter &&
+		CachedData)
 	{
-		const int32 ParameterIndex = CachedData->Parameters.FindParameterIndex(EMaterialParameterType::RuntimeVirtualTexture, ParameterInfo);
-		if (ParameterIndex != INDEX_NONE)
+		if (CachedData->Parameters.GetParameterValue(Type, ParameterInfo, OutResult))
 		{
-			OutValue = CachedData->Parameters.RuntimeVirtualTextureValues[ParameterIndex];
 			return true;
 		}
 	}
 
 	// Next material in hierarchy
-	if (Parent)
-	{
-		FMICReentranceGuard	Guard(this);
-		return Parent->GetRuntimeVirtualTextureParameterValue(ParameterInfo, OutValue, bOveriddenOnly);
-	}
-
-	OutValue = nullptr;
-	return false;
-}
-
-
-#if WITH_EDITOR
-bool  UMaterialInstance::GetTextureParameterChannelNames(const FHashedMaterialParameterInfo& ParameterInfo, FParameterChannelNames& OutValue) const
-{
-	bool bFoundAValue = false;
-
-	if (GetReentrantFlag())
-	{
-		return false;
-	}
-
-	// Instance-included default
-	if (ParameterInfo.Association != EMaterialParameterAssociation::GlobalParameter && CachedData)
+	if (EnumHasAnyFlags(Flags, EMaterialGetParameterValueFlags::CheckParent) && Parent)
 	{
-		if (CachedData->Parameters.GetTextureParameterChannelNames(ParameterInfo, OutValue))
+		EMaterialGetParameterValueFlags ParentFlags = Flags | EMaterialGetParameterValueFlags::IsParent;
+		if (EnumHasAnyFlags(Flags, EMaterialGetParameterValueFlags::CheckNonOverrides))
 		{
-			return true;
+			// Parent's instance overrides are non-overrides relative to this material
+			ParentFlags |= EMaterialGetParameterValueFlags::CheckInstanceOverrides;
 		}
-	}
 
-	// Next material in hierarchy
-	if (Parent)
-	{
 		FMICReentranceGuard	Guard(this);
-		return Parent->GetTextureParameterChannelNames(ParameterInfo, OutValue);
-	}
-
-	return false;
-};
-#endif // WITH_EDITOR
-
-bool UMaterialInstance::GetFontParameterValue(const FHashedMaterialParameterInfo& ParameterInfo,class UFont*& OutFontValue, int32& OutFontPage, bool bOveriddenOnly) const
-{
-	bool bFoundAValue = false;
-
-	if (GetReentrantFlag())
-	{
-		OutFontValue = nullptr;
-		OutFontPage = INDEX_NONE;
-		return false;
+		return Parent->GetParameterValue(Type, ParameterInfo, OutResult, ParentFlags);
 	}
 
-	// Instance override
-	const FFontParameterValue* ParameterValue = GameThread_FindParameterByName(FontParameterValues, ParameterInfo);
-	if (ParameterValue)
-	{
-		OutFontValue = ParameterValue->FontValue;
-		OutFontPage = ParameterValue->FontPage;
-		return true;
-	}
-	
-	// Instance-included default
-	if (!bOveriddenOnly && ParameterInfo.Association != EMaterialParameterAssociation::GlobalParameter && CachedData)
-	{
-		const int32 ParameterIndex = CachedData->Parameters.FindParameterIndex(EMaterialParameterType::Font, ParameterInfo);
-		if (ParameterIndex != INDEX_NONE)
-		{
-			OutFontValue = CachedData->Parameters.FontValues[ParameterIndex];
-			OutFontPage = CachedData->Parameters.FontPageValues[ParameterIndex];
-			return true;
-		}
-	}
-	
-	// Next material in hierarchy
-	if (Parent)
-	{
-		FMICReentranceGuard	Guard(this);
-		return Parent->GetFontParameterValue(ParameterInfo, OutFontValue, OutFontPage, bOveriddenOnly);
-	}
-	
-	OutFontValue = nullptr;
-	OutFontPage = INDEX_NONE;
 	return false;
 }
 
@@ -2139,538 +1846,190 @@ void UMaterialInstance::GetStaticParameterValues(FStaticParameterSet& OutStaticP
 				if (ParameterInfo == LayersParam.ParameterInfo)
 				{
 					ParentParameter.bOverride = LayersParam.bOverride;
-					if (LayersParam.bOverride)
-					{
-						ParentParameter.Value = LayersParam.Value;
-					}
-				}
-			}
-		}
-
-		// Static Switch Parameters
-		GetAllStaticSwitchParameterInfo(OutParameterInfo, Guids);
-		OutStaticParameters.StaticSwitchParameters.AddZeroed(OutParameterInfo.Num());
-
-		for(int32 ParameterIdx=0; ParameterIdx<OutParameterInfo.Num(); ParameterIdx++)
-		{
-			FStaticSwitchParameter& ParentParameter = OutStaticParameters.StaticSwitchParameters[ParameterIdx];
-			FMaterialParameterInfo& ParameterInfo = OutParameterInfo[ParameterIdx];
-			FGuid ExpressionId = Guids[ParameterIdx];
-
-			ParentParameter.bOverride = false;
-			ParentParameter.ParameterInfo = ParameterInfo;
-
-			GetStaticSwitchParameterValue(ParameterInfo, ParentParameter.Value, ExpressionId);
-
-			ParentParameter.ExpressionGUID = Guids[ParameterIdx];
-
-			// If the SourceInstance is overriding this parameter, use its settings
-			for(int32 SwitchParamIdx = 0; SwitchParamIdx < StaticParameters.StaticSwitchParameters.Num(); SwitchParamIdx++)
-			{
-				const FStaticSwitchParameter& StaticSwitchParam = StaticParameters.StaticSwitchParameters[SwitchParamIdx];
-
-				if(ParameterInfo == StaticSwitchParam.ParameterInfo)
-				{
-					ParentParameter.bOverride = StaticSwitchParam.bOverride;
-					if (StaticSwitchParam.bOverride)
-					{
-						ParentParameter.Value = StaticSwitchParam.Value;
-					}
-				}
-			}
-		}
-
-		// Static Component Mask Parameters
-		GetAllStaticComponentMaskParameterInfo(OutParameterInfo, Guids);
-		OutStaticParameters.StaticComponentMaskParameters.AddZeroed(OutParameterInfo.Num());
-
-		for(int32 ParameterIdx=0; ParameterIdx<OutParameterInfo.Num(); ParameterIdx++)
-		{
-			FStaticComponentMaskParameter& ParentParameter = OutStaticParameters.StaticComponentMaskParameters[ParameterIdx];
-			FMaterialParameterInfo& ParameterInfo = OutParameterInfo[ParameterIdx];
-			FGuid ExpressionId = Guids[ParameterIdx];
-
-			ParentParameter.bOverride = false;
-			ParentParameter.ParameterInfo = ParameterInfo;
-			
-			GetStaticComponentMaskParameterValue(ParameterInfo, ParentParameter.R, ParentParameter.G, ParentParameter.B, ParentParameter.A, ExpressionId);
-
-			ParentParameter.ExpressionGUID = Guids[ParameterIdx];
-			
-			// If the SourceInstance is overriding this parameter, use its settings
-			for(int32 MaskParamIdx = 0; MaskParamIdx < StaticParameters.StaticComponentMaskParameters.Num(); MaskParamIdx++)
-			{
-				const FStaticComponentMaskParameter &StaticComponentMaskParam = StaticParameters.StaticComponentMaskParameters[MaskParamIdx];
-
-				if(ParameterInfo == StaticComponentMaskParam.ParameterInfo)
-				{
-					ParentParameter.bOverride = StaticComponentMaskParam.bOverride;
-					if (StaticComponentMaskParam.bOverride)
-					{
-						ParentParameter.R = StaticComponentMaskParam.R;
-						ParentParameter.G = StaticComponentMaskParam.G;
-						ParentParameter.B = StaticComponentMaskParam.B;
-						ParentParameter.A = StaticComponentMaskParam.A;
-					}
-				}
-			}
-		}
-	}
-
-	// Custom parameters.
-	CustomStaticParametersGetters.Broadcast(OutStaticParameters, this);
-
-	if (AllowCachingStaticParameterValuesCounter > 0)
-	{
-		CachedStaticParameterValues = OutStaticParameters;
-	}
-}
-#endif // WITH_EDITORONLY_DATA
-
-void UMaterialInstance::GetAllParametersOfType(EMaterialParameterType Type, TArray<FMaterialParameterInfo>& OutParameterInfo, TArray<FGuid>& OutParameterIds) const
-{
-	const UMaterial* Material = GetMaterial_Concurrent();
-	int32 NumParameters = 0;
-	if (CachedData)
-	{
-		NumParameters += CachedData->Parameters.GetNumParameters(Type);
-	}
-	if (Material)
-	{
-		NumParameters += Material->GetCachedExpressionData().Parameters.GetNumParameters(Type);
-	}
-
-	OutParameterInfo.Empty(NumParameters);
-	OutParameterIds.Empty(NumParameters);
-	if (CachedData)
-	{
-		CachedData->Parameters.GetAllParameterInfoOfType(Type, false, OutParameterInfo, OutParameterIds);
-	}
-	if (Material)
-	{
-		Material->GetCachedExpressionData().Parameters.GetAllGlobalParameterInfoOfType(Type, false, OutParameterInfo, OutParameterIds);
-	}
-}
-
-void UMaterialInstance::GetAllScalarParameterInfo(TArray<FMaterialParameterInfo>& OutParameterInfo, TArray<FGuid>& OutParameterIds) const
-{
-	GetAllParametersOfType(EMaterialParameterType::Scalar, OutParameterInfo, OutParameterIds);
-}
-
-void UMaterialInstance::GetAllVectorParameterInfo(TArray<FMaterialParameterInfo>& OutParameterInfo, TArray<FGuid>& OutParameterIds) const
-{
-	GetAllParametersOfType(EMaterialParameterType::Vector, OutParameterInfo, OutParameterIds);
-}
-
-void UMaterialInstance::GetAllTextureParameterInfo(TArray<FMaterialParameterInfo>& OutParameterInfo, TArray<FGuid>& OutParameterIds) const
-{
-	GetAllParametersOfType(EMaterialParameterType::Texture, OutParameterInfo, OutParameterIds);
-}
-
-void UMaterialInstance::GetAllRuntimeVirtualTextureParameterInfo(TArray<FMaterialParameterInfo>& OutParameterInfo, TArray<FGuid>& OutParameterIds) const
-{
-	GetAllParametersOfType(EMaterialParameterType::RuntimeVirtualTexture, OutParameterInfo, OutParameterIds);
-}
-
-void UMaterialInstance::GetAllFontParameterInfo(TArray<FMaterialParameterInfo>& OutParameterInfo, TArray<FGuid>& OutParameterIds) const
-{
-	GetAllParametersOfType(EMaterialParameterType::Font, OutParameterInfo, OutParameterIds);
-}
-
-#if WITH_EDITORONLY_DATA
-void UMaterialInstance::GetAllMaterialLayersParameterInfo(TArray<FMaterialParameterInfo>& OutParameterInfo, TArray<FGuid>& OutParameterIds) const
-{
-	OutParameterInfo.Empty();
-	OutParameterIds.Empty();
-	if (const UMaterial* Material = GetMaterial())
-	{
-		Material->GetAllParameterInfo<UMaterialExpressionMaterialAttributeLayers>(OutParameterInfo, OutParameterIds);
-	}
-}
-
-void UMaterialInstance::GetAllStaticSwitchParameterInfo(TArray<FMaterialParameterInfo>& OutParameterInfo, TArray<FGuid>& OutParameterIds) const
-{
-	GetAllParametersOfType(EMaterialParameterType::StaticSwitch, OutParameterInfo, OutParameterIds);
-}
-
-void UMaterialInstance::GetAllStaticComponentMaskParameterInfo(TArray<FMaterialParameterInfo>& OutParameterInfo, TArray<FGuid>& OutParameterIds) const
-{
-	GetAllParametersOfType(EMaterialParameterType::StaticComponentMask, OutParameterInfo, OutParameterIds);
-}
-
-bool UMaterialInstance::IterateDependentFunctions(TFunctionRef<bool(UMaterialFunctionInterface*)> Predicate) const
-{
-	// Important that local function references are listed first so that traversing for a parameter
-	// value we always hit the highest material in the hierarchy that can give us a valid value
-	for (const FStaticMaterialLayersParameter& LayersParam : StaticParameters.MaterialLayersParameters)
-	{
-		if (LayersParam.bOverride)
-		{
-			for (UMaterialFunctionInterface* Layer : LayersParam.Value.Layers)
-			{
-				if (Layer)
-				{
-					if (!Layer->IterateDependentFunctions(Predicate))
-					{
-						return false;
-					}
-					if (!Predicate(Layer))
-					{
-						return false;
-					}
-				}
-			}
-
-			for (UMaterialFunctionInterface* Blend : LayersParam.Value.Blends)
-			{
-				if (Blend)
-				{
-					if (!Blend->IterateDependentFunctions(Predicate))
-					{
-						return false;
-					}
-					if (!Predicate(Blend))
-					{
-						return false;
-					}
-				}
-			}
-		}
-	}
-
-	return Parent ? Parent->IterateDependentFunctions(Predicate) : true;
-}
-
-void UMaterialInstance::GetDependentFunctions(TArray<UMaterialFunctionInterface*>& DependentFunctions) const
-{
-	IterateDependentFunctions([&DependentFunctions](UMaterialFunctionInterface* MaterialFunction) -> bool
-		{
-			DependentFunctions.AddUnique(MaterialFunction);
-			return true;
-		});
-}
-#endif // WITH_EDITORONLY_DATA
-
-bool UMaterialInstance::GetScalarParameterDefaultValue(const FHashedMaterialParameterInfo& ParameterInfo, float& OutValue, bool bOveriddenOnly, bool bCheckOwnedGlobalOverrides) const
-{
-#if WITH_EDITOR
-	if (GetReentrantFlag())
-	{
-		return false;
-	}
-#endif
-
-	if (bCheckOwnedGlobalOverrides)
-	{
-		// Parameters overridden by this instance
-		for (const FScalarParameterValue& ScalarParam : ScalarParameterValues)
-		{
-			if (ScalarParam.ParameterInfo == ParameterInfo)
-			{
-				OutValue = ScalarParam.ParameterValue;
-				return true;
-			}
-		}
-	}
-
-	// In the case of duplicate parameters with different values, this will return the
-	// first matching expression found, not necessarily the one that's used for rendering
-	if (ParameterInfo.Association != EMaterialParameterAssociation::GlobalParameter && CachedData)
-	{
-		const int32 ParameterIndex = CachedData->Parameters.FindParameterIndex(EMaterialParameterType::Scalar, ParameterInfo);
-		if (ParameterIndex != INDEX_NONE)
-		{
-			OutValue = CachedData->Parameters.ScalarValues[ParameterIndex];
-			return true;
-		}
-	}
-
-	if (Parent)
-	{
-#if WITH_EDITOR
-		FMICReentranceGuard	Guard(this);
-#endif
-		return Parent->GetScalarParameterDefaultValue(ParameterInfo, OutValue, bOveriddenOnly, true);
-	}
-
-	return false;
-}
-
-bool UMaterialInstance::GetVectorParameterDefaultValue(const FHashedMaterialParameterInfo& ParameterInfo, FLinearColor& OutValue, bool bOveriddenOnly, bool bCheckOwnedGlobalOverrides) const
-{
-#if WITH_EDITOR
-	if (GetReentrantFlag())
-	{
-		return false;
-	}
-#endif
-
-	// In the case of duplicate parameters with different values, this will return the
-	// first matching expression found, not necessarily the one that's used for rendering
-	if (bCheckOwnedGlobalOverrides)
-	{
-		// Parameters overridden by this instance
-		for (const FVectorParameterValue& VectorParam : VectorParameterValues)
-		{
-			if (VectorParam.ParameterInfo == ParameterInfo)
-			{
-				OutValue = VectorParam.ParameterValue;
-				return true;
-			}
-		}
-	}
-
-	if (ParameterInfo.Association != EMaterialParameterAssociation::GlobalParameter && CachedData)
-	{
-		const int32 ParameterIndex = CachedData->Parameters.FindParameterIndex(EMaterialParameterType::Vector, ParameterInfo);
-		if (ParameterIndex != INDEX_NONE)
-		{
-			OutValue = CachedData->Parameters.VectorValues[ParameterIndex];
-			return true;
-		}
-	}
-
-	if (Parent)
-	{
-#if WITH_EDITOR
-		FMICReentranceGuard	Guard(this);
-#endif
-		return Parent->GetVectorParameterDefaultValue(ParameterInfo, OutValue, bOveriddenOnly, true);
-	}
-
-	return false;
-}
-
-bool UMaterialInstance::GetTextureParameterDefaultValue(const FHashedMaterialParameterInfo& ParameterInfo, UTexture*& OutValue, bool bCheckOwnedGlobalOverrides) const
-{
-#if WITH_EDITOR
-	if (GetReentrantFlag())
-	{
-		return false;
-	}
-#endif
-
-	// In the case of duplicate parameters with different values, this will return the
-	// first matching expression found, not necessarily the one that's used for rendering
-	if (bCheckOwnedGlobalOverrides)
-	{
-		// Parameters overridden by this instance
-		for (const FTextureParameterValue& TextureParam : TextureParameterValues)
-		{
-			if (TextureParam.ParameterInfo == ParameterInfo)
-			{
-				OutValue = TextureParam.ParameterValue;
-				return true;
-			}
-		}
-	}
-
-	if (ParameterInfo.Association != EMaterialParameterAssociation::GlobalParameter && CachedData)
-	{
-		const int32 ParameterIndex = CachedData->Parameters.FindParameterIndex(EMaterialParameterType::Texture, ParameterInfo);
-		if (ParameterIndex != INDEX_NONE)
-		{
-			OutValue = CachedData->Parameters.TextureValues[ParameterIndex];
-			return true;
-		}
-	}
-
-	if (Parent)
-	{
-#if WITH_EDITOR
-		FMICReentranceGuard	Guard(this);
-#endif
-		return Parent->GetTextureParameterDefaultValue(ParameterInfo, OutValue, true);
-	}
-
-	return false;
-}
-
-bool UMaterialInstance::GetRuntimeVirtualTextureParameterDefaultValue(const FHashedMaterialParameterInfo& ParameterInfo, URuntimeVirtualTexture*& OutValue, bool bCheckOwnedGlobalOverrides) const
-{
-#if WITH_EDITOR
-	if (GetReentrantFlag())
-	{
-		return false;
-	}
-#endif
+					if (LayersParam.bOverride)
+					{
+						ParentParameter.Value = LayersParam.Value;
+					}
+				}
+			}
+		}
 
-	// In the case of duplicate parameters with different values, this will return the
-	// first matching expression found, not necessarily the one that's used for rendering
-	if (bCheckOwnedGlobalOverrides)
-	{
-		// Parameters overridden by this instance
-		for (const FRuntimeVirtualTextureParameterValue& RuntimeVirtualTextureParam : RuntimeVirtualTextureParameterValues)
+		TMap<FMaterialParameterInfo, FMaterialParameterMetadata> ParameterValues;
+		for (int32 ParameterTypeIndex = 0; ParameterTypeIndex < NumMaterialParameterTypes; ++ParameterTypeIndex)
 		{
-			if (RuntimeVirtualTextureParam.ParameterInfo == ParameterInfo)
+			const EMaterialParameterType ParameterType = (EMaterialParameterType)ParameterTypeIndex;
+			if (IsStaticMaterialParameter(ParameterType))
 			{
-				OutValue = RuntimeVirtualTextureParam.ParameterValue;
-				return true;
+				ParameterValues.Reset();
+				GetAllParametersOfType(ParameterType, ParameterValues);
+				OutStaticParameters.AddParametersOfType(ParameterType, ParameterValues);
 			}
 		}
 	}
 
-	if (ParameterInfo.Association != EMaterialParameterAssociation::GlobalParameter && CachedData)
-	{
-		const int32 ParameterIndex = CachedData->Parameters.FindParameterIndex(EMaterialParameterType::RuntimeVirtualTexture, ParameterInfo);
-		if (ParameterIndex != INDEX_NONE)
-		{
-			OutValue = CachedData->Parameters.RuntimeVirtualTextureValues[ParameterIndex];
-			return true;
-		}
-	}
+	// Custom parameters.
+	CustomStaticParametersGetters.Broadcast(OutStaticParameters, this);
 
-	if (Parent)
+	if (AllowCachingStaticParameterValuesCounter > 0)
 	{
-#if WITH_EDITOR
-		FMICReentranceGuard	Guard(this);
-#endif
-		return Parent->GetRuntimeVirtualTextureParameterDefaultValue(ParameterInfo, OutValue, true);
+		CachedStaticParameterValues = OutStaticParameters;
 	}
-
-	return false;
 }
+#endif // WITH_EDITORONLY_DATA
 
-bool UMaterialInstance::GetFontParameterDefaultValue(const FHashedMaterialParameterInfo& ParameterInfo, UFont*& OutFontValue, int32& OutFontPage, bool bCheckOwnedGlobalOverrides) const
+void UMaterialInstance::GetAllParametersOfType(EMaterialParameterType Type, TMap<FMaterialParameterInfo, FMaterialParameterMetadata>& OutParameters) const
 {
-#if WITH_EDITOR
-	if (GetReentrantFlag())
+	const UMaterial* Material = GetMaterial_Concurrent();
+	int32 NumParameters = 0;
+	if (CachedData)
 	{
-		return false;
+		NumParameters += CachedData->Parameters.GetNumParameters(Type);
+	}
+	if (Material)
+	{
+		NumParameters += Material->GetCachedExpressionData().Parameters.GetNumParameters(Type);
+	}
+
+	OutParameters.Empty(NumParameters);
+	if (CachedData)
+	{
+		CachedData->Parameters.GetAllParametersOfType(Type, OutParameters);
+	}
+	if (Material)
+	{
+		Material->GetCachedExpressionData().Parameters.GetAllParametersOfType(Type, OutParameters);
 	}
-#endif
 
-	// In the case of duplicate parameters with different values, this will return the
-	// first matching expression found, not necessarily the one that's used for rendering
-	if (bCheckOwnedGlobalOverrides)
+	TArray<const UMaterialInstance*> InstanceChain;
+	InstanceChain.Add(this);
 	{
-		// Parameters overridden by this instance
-		for (const FFontParameterValue& FontParam : FontParameterValues)
+		const UMaterialInterface* CurrentParent = Parent;
+		while (CurrentParent)
 		{
-			if (FontParam.ParameterInfo == ParameterInfo)
+			const UMaterialInstance* ParentInstance = Cast<UMaterialInstance>(CurrentParent);
+			if (ParentInstance)
 			{
-				OutFontValue = FontParam.FontValue;
-				OutFontPage = FontParam.FontPage;
-				return true;
+				InstanceChain.Add(ParentInstance);
+				CurrentParent = ParentInstance->Parent;
+			}
+			else
+			{
+				check(CurrentParent == Material);
+				CurrentParent = nullptr;
 			}
 		}
 	}
 
-	if (ParameterInfo.Association != EMaterialParameterAssociation::GlobalParameter && CachedData)
+	for (int32 i = InstanceChain.Num() - 1; i >= 0; --i)
 	{
-		const int32 ParameterIndex = CachedData->Parameters.FindParameterIndex(EMaterialParameterType::Font, ParameterInfo);
-		if (ParameterIndex != INDEX_NONE)
+		const UMaterialInstance* Instance = InstanceChain[i];
+		const bool bSetOverride = (i == 0); // Only set the override flag for parameters overriden by the current material (always at slot0)
+		switch (Type)
 		{
-			OutFontValue = CachedData->Parameters.FontValues[ParameterIndex];
-			OutFontPage = CachedData->Parameters.FontPageValues[ParameterIndex];
-			return true;
+		case EMaterialParameterType::Scalar: GameThread_ApplyParameterOverrides(Instance->ScalarParameterValues, bSetOverride, OutParameters); break;
+		case EMaterialParameterType::Vector: GameThread_ApplyParameterOverrides(Instance->VectorParameterValues, bSetOverride, OutParameters); break;
+		case EMaterialParameterType::Texture: GameThread_ApplyParameterOverrides(Instance->TextureParameterValues, bSetOverride, OutParameters); break;
+		case EMaterialParameterType::RuntimeVirtualTexture: GameThread_ApplyParameterOverrides(Instance->RuntimeVirtualTextureParameterValues, bSetOverride, OutParameters); break;
+		case EMaterialParameterType::Font: GameThread_ApplyParameterOverrides(Instance->FontParameterValues, bSetOverride, OutParameters); break;
+#if WITH_EDITORONLY_DATA
+		case EMaterialParameterType::StaticSwitch: GameThread_ApplyParameterOverrides(Instance->StaticParameters.StaticSwitchParameters, bSetOverride, OutParameters); break;
+		case EMaterialParameterType::StaticComponentMask: GameThread_ApplyParameterOverrides(Instance->StaticParameters.StaticComponentMaskParameters, bSetOverride, OutParameters); break;
+#endif // WITH_EDITORONLY_DATA
+		default: checkNoEntry();
 		}
 	}
-
-	if (Parent)
-	{
-#if WITH_EDITOR
-		FMICReentranceGuard	Guard(this);
-#endif
-		return Parent->GetFontParameterDefaultValue(ParameterInfo, OutFontValue, OutFontPage, true);
-	}
-
-	return false;
 }
 
-#if WITH_EDITOR
-bool UMaterialInstance::GetStaticSwitchParameterDefaultValue(const FHashedMaterialParameterInfo& ParameterInfo, bool& OutValue, FGuid& OutExpressionGuid, bool bCheckOwnedGlobalOverrides) const
+void UMaterialInstance::GetAllParameterInfoOfType(EMaterialParameterType Type, TArray<FMaterialParameterInfo>& OutParameterInfo, TArray<FGuid>& OutParameterIds) const
 {
-	if (GetReentrantFlag())
+	const UMaterial* Material = GetMaterial_Concurrent();
+	int32 NumParameters = 0;
+	if (CachedData)
 	{
-		return false;
+		NumParameters += CachedData->Parameters.GetNumParameters(Type);
 	}
-
-	// In the case of duplicate parameters with different values, this will return the
-	// first matching expression found, not necessarily the one that's used for rendering
-	if (bCheckOwnedGlobalOverrides)
+	if (Material)
 	{
-		// Parameters overridden by this instance
-		for (const FStaticSwitchParameter& SwitchParam : StaticParameters.StaticSwitchParameters)
-		{
-			if (SwitchParam.bOverride && SwitchParam.ParameterInfo == ParameterInfo)
-			{
-				OutValue = SwitchParam.Value;
-				OutExpressionGuid = SwitchParam.ExpressionGUID;
-				return true;
-			}
-		}
+		NumParameters += Material->GetCachedExpressionData().Parameters.GetNumParameters(Type);
 	}
 
-	if (ParameterInfo.Association != EMaterialParameterAssociation::GlobalParameter)
+	OutParameterInfo.Empty(NumParameters);
+	OutParameterIds.Empty(NumParameters);
+	if (CachedData)
 	{
-		const int32 ParameterIndex = CachedData->Parameters.FindParameterIndex(EMaterialParameterType::StaticSwitch, ParameterInfo);
-		if (ParameterIndex != INDEX_NONE)
-		{
-			OutExpressionGuid = CachedData->Parameters.GetExpressionGuid(EMaterialParameterType::StaticSwitch, ParameterIndex);
-			OutValue = CachedData->Parameters.StaticSwitchValues[ParameterIndex];
-			return true;
-		}
+		CachedData->Parameters.GetAllParameterInfoOfType(Type, OutParameterInfo, OutParameterIds);
 	}
-	
-	if (Parent)
+	if (Material)
 	{
-		FMICReentranceGuard	Guard(this);
-		return Parent->GetStaticSwitchParameterDefaultValue(ParameterInfo, OutValue, OutExpressionGuid, true);
+		Material->GetCachedExpressionData().Parameters.GetAllGlobalParameterInfoOfType(Type, OutParameterInfo, OutParameterIds);
 	}
-
-	return false;
 }
 
-bool UMaterialInstance:: GetStaticComponentMaskParameterDefaultValue(const FHashedMaterialParameterInfo& ParameterInfo, bool& OutR, bool& OutG, bool& OutB, bool& OutA, FGuid& OutExpressionGuid, bool bCheckOwnedGlobalOverrides) const
+#if WITH_EDITORONLY_DATA
+void UMaterialInstance::GetAllMaterialLayersParameterInfo(TArray<FMaterialParameterInfo>& OutParameterInfo, TArray<FGuid>& OutParameterIds) const
 {
-	if (GetReentrantFlag())
+	OutParameterInfo.Empty();
+	OutParameterIds.Empty();
+	if (const UMaterial* Material = GetMaterial())
 	{
-		return false;
+		Material->GetAllParameterInfo<UMaterialExpressionMaterialAttributeLayers>(OutParameterInfo, OutParameterIds);
 	}
+}
 
-	// In the case of duplicate parameters with different values, this will return the
-	// first matching expression found, not necessarily the one that's used for rendering
-	if (bCheckOwnedGlobalOverrides)
+bool UMaterialInstance::IterateDependentFunctions(TFunctionRef<bool(UMaterialFunctionInterface*)> Predicate) const
+{
+	// Important that local function references are listed first so that traversing for a parameter
+	// value we always hit the highest material in the hierarchy that can give us a valid value
+	for (const FStaticMaterialLayersParameter& LayersParam : StaticParameters.MaterialLayersParameters)
 	{
-		// Parameters overridden by this instance
-		for (const FStaticComponentMaskParameter& ComponentMaskParam : StaticParameters.StaticComponentMaskParameters)
+		if (LayersParam.bOverride)
 		{
-			if (ComponentMaskParam.bOverride && ComponentMaskParam.ParameterInfo == ParameterInfo)
+			for (UMaterialFunctionInterface* Layer : LayersParam.Value.Layers)
 			{
-				OutR = ComponentMaskParam.R;
-				OutG = ComponentMaskParam.G;
-				OutB = ComponentMaskParam.B;
-				OutA = ComponentMaskParam.A;
-				OutExpressionGuid = ComponentMaskParam.ExpressionGUID;
-				return true;
+				if (Layer)
+				{
+					if (!Layer->IterateDependentFunctions(Predicate))
+					{
+						return false;
+					}
+					if (!Predicate(Layer))
+					{
+						return false;
+					}
+				}
+			}
+
+			for (UMaterialFunctionInterface* Blend : LayersParam.Value.Blends)
+			{
+				if (Blend)
+				{
+					if (!Blend->IterateDependentFunctions(Predicate))
+					{
+						return false;
+					}
+					if (!Predicate(Blend))
+					{
+						return false;
+					}
+				}
 			}
 		}
 	}
 
-	if (ParameterInfo.Association != EMaterialParameterAssociation::GlobalParameter)
-	{
-		const int32 ParameterIndex = CachedData->Parameters.FindParameterIndex(EMaterialParameterType::StaticComponentMask, ParameterInfo);
-		if (ParameterIndex != INDEX_NONE)
+	return Parent ? Parent->IterateDependentFunctions(Predicate) : true;
+}
+
+void UMaterialInstance::GetDependentFunctions(TArray<UMaterialFunctionInterface*>& DependentFunctions) const
+{
+	IterateDependentFunctions([&DependentFunctions](UMaterialFunctionInterface* MaterialFunction) -> bool
 		{
-			OutExpressionGuid = CachedData->Parameters.GetExpressionGuid(EMaterialParameterType::StaticComponentMask, ParameterIndex);
-			OutR = CachedData->Parameters.StaticComponentMaskValues[ParameterIndex].R;
-			OutG = CachedData->Parameters.StaticComponentMaskValues[ParameterIndex].G;
-			OutB = CachedData->Parameters.StaticComponentMaskValues[ParameterIndex].B;
-			OutA = CachedData->Parameters.StaticComponentMaskValues[ParameterIndex].A;
+			DependentFunctions.AddUnique(MaterialFunction);
 			return true;
-		}
-	}
-	
-	if (Parent)
-	{
-		FMICReentranceGuard	Guard(this);
-		return Parent->GetStaticComponentMaskParameterDefaultValue(ParameterInfo, OutR, OutG, OutB, OutA, OutExpressionGuid, true);
-	}
-
-	return false;
+		});
 }
+#endif // WITH_EDITORONLY_DATA
 
+#if WITH_EDITOR
 bool UMaterialInstance::GetGroupName(const FHashedMaterialParameterInfo& ParameterInfo, FName& OutGroup) const
 {
 	if (GetReentrantFlag())
@@ -3022,92 +2381,6 @@ void UMaterialInstance::CacheShadersForResources(EShaderPlatform ShaderPlatform,
 }
 
 #if WITH_EDITORONLY_DATA
-bool UMaterialInstance::GetStaticSwitchParameterValue(const FHashedMaterialParameterInfo& ParameterInfo, bool& OutValue, FGuid& OutExpressionGuid, bool bOveriddenOnly, bool bCheckParent) const
-{
-	if (GetReentrantFlag())
-	{
-		return false;
-	}
-
-	// Instance override
-	for (const FStaticSwitchParameter& Param : StaticParameters.StaticSwitchParameters)
-	{
-		if (Param.ParameterInfo == ParameterInfo)
-		{
-			OutExpressionGuid = Param.ExpressionGUID;
-			OutValue = Param.Value;
-			return true;
-		}
-	}
-
-	// Instance-included default
-	if (!bOveriddenOnly && ParameterInfo.Association != EMaterialParameterAssociation::GlobalParameter && CachedData)
-	{
-		const int32 ParameterIndex = CachedData->Parameters.FindParameterIndex(EMaterialParameterType::StaticSwitch, ParameterInfo);
-		if (ParameterIndex != INDEX_NONE)
-		{
-			OutExpressionGuid = CachedData->Parameters.GetExpressionGuid(EMaterialParameterType::StaticSwitch, ParameterIndex);
-			OutValue = CachedData->Parameters.StaticSwitchValues[ParameterIndex];
-			return true;
-		}
-	}
-
-	// Next material in hierarchy
-	if (Parent && bCheckParent)
-	{
-		FMICReentranceGuard	Guard(this);
-		return Parent->GetStaticSwitchParameterValue(ParameterInfo, OutValue, OutExpressionGuid, bOveriddenOnly, true);
-	}
-
-	return false;
-}
-
-bool UMaterialInstance::GetStaticComponentMaskParameterValue(const FHashedMaterialParameterInfo& ParameterInfo, bool& R, bool& G, bool& B, bool& A, FGuid& OutExpressionGuid, bool bOveriddenOnly, bool bCheckParent) const
-{
-	if (GetReentrantFlag())
-	{
-		return false;
-	}
-
-	// Instance override
-	for (const FStaticComponentMaskParameter& Param : StaticParameters.StaticComponentMaskParameters)
-	{
-		if (Param.ParameterInfo == ParameterInfo)
-		{
-			OutExpressionGuid = Param.ExpressionGUID;
-			R = Param.R;
-			G = Param.G;
-			B = Param.B;
-			A = Param.A;
-			return true;
-		}
-	}
-
-	// Instance-included default
-	if (!bOveriddenOnly && ParameterInfo.Association != EMaterialParameterAssociation::GlobalParameter && CachedData)
-	{
-		const int32 ParameterIndex = CachedData->Parameters.FindParameterIndex(EMaterialParameterType::StaticComponentMask, ParameterInfo);
-		if (ParameterIndex != INDEX_NONE)
-		{
-			OutExpressionGuid = CachedData->Parameters.GetExpressionGuid(EMaterialParameterType::StaticComponentMask, ParameterIndex);
-			R = CachedData->Parameters.StaticComponentMaskValues[ParameterIndex].R;
-			G = CachedData->Parameters.StaticComponentMaskValues[ParameterIndex].G;
-			B = CachedData->Parameters.StaticComponentMaskValues[ParameterIndex].B;
-			A = CachedData->Parameters.StaticComponentMaskValues[ParameterIndex].A;
-			return true;
-		}
-	}
-
-	// Next material in hierarchy
-	if (Parent && bCheckParent)
-	{
-		FMICReentranceGuard	Guard(this);
-		return Parent->GetStaticComponentMaskParameterValue(ParameterInfo, R, G, B, A, OutExpressionGuid, bOveriddenOnly, true);
-	}
-
-	return false;
-}
-
 bool UMaterialInstance::GetMaterialLayersParameterValue(const FHashedMaterialParameterInfo& ParameterInfo, FMaterialLayersFunctions& OutLayers, FGuid& OutExpressionGuid, bool bCheckParent /*= true*/) const
 {
 	if (GetReentrantFlag())
@@ -3789,6 +3062,65 @@ bool UMaterialInstance::SetVectorParameterByIndexInternal(int32 ParameterIndex,
 	return true;
 }
 
+#if WITH_EDITORONLY_DATA
+FMaterialInstanceParameterUpdateContext::FMaterialInstanceParameterUpdateContext(UMaterialInstance* InInstance)
+	: Instance(InInstance)
+	, bForceStaticPermutationUpdate(false)
+{
+	check(InInstance);
+	StaticParameters = InInstance->GetStaticParameters();
+	BasePropertyOverrides = InInstance->BasePropertyOverrides;
+}
+
+FMaterialInstanceParameterUpdateContext::~FMaterialInstanceParameterUpdateContext()
+{
+	Instance->UpdateStaticPermutation(StaticParameters, BasePropertyOverrides, bForceStaticPermutationUpdate);
+}
+
+void FMaterialInstanceParameterUpdateContext::SetParameterValueEditorOnly(const FMaterialParameterInfo& ParameterInfo, const FMaterialParameterMetadata& Meta, EMaterialSetParameterValueFlags Flags)
+{
+	const FMaterialParameterValue& Value = Meta.Value;
+	if (IsStaticMaterialParameter(Value.Type))
+	{
+		// Route static parameters to the static parameter set
+		StaticParameters.SetParameterValue(ParameterInfo, Meta, Flags);
+	}
+	else
+	{
+		const bool bUseAtlas = EnumHasAnyFlags(Flags, EMaterialSetParameterValueFlags::SetCurveAtlas);
+		FScalarParameterAtlasInstanceData AtlasData;
+
+		switch (Value.Type)
+		{
+		case EMaterialParameterType::Scalar:
+			if (bUseAtlas)
+			{
+				AtlasData.bIsUsedAsAtlasPosition = Meta.bUsedAsAtlasPosition;
+				AtlasData.Atlas = Meta.ScalarAtlas;
+				AtlasData.Curve = Meta.ScalarCurve;
+			}
+			Instance->SetScalarParameterValueInternal(ParameterInfo, Value.AsScalar(), bUseAtlas, AtlasData);
+			break;
+		case EMaterialParameterType::Vector: Instance->SetVectorParameterValueInternal(ParameterInfo, Value.AsLinearColor()); break;
+		case EMaterialParameterType::Texture: Instance->SetTextureParameterValueInternal(ParameterInfo, Value.Texture); break;
+		case EMaterialParameterType::Font: Instance->SetFontParameterValueInternal(ParameterInfo, Value.Font.Value, Value.Font.Page); break;
+		case EMaterialParameterType::RuntimeVirtualTexture: Instance->SetRuntimeVirtualTextureParameterValueInternal(ParameterInfo, Value.RuntimeVirtualTexture); break;
+		default: checkNoEntry();
+		}
+	}
+}
+
+void FMaterialInstanceParameterUpdateContext::SetForceStaticPermutationUpdate(bool bValue)
+{
+	bForceStaticPermutationUpdate = bValue;
+}
+
+void FMaterialInstanceParameterUpdateContext::SetBasePropertyOverrides(const FMaterialInstanceBasePropertyOverrides& InValue)
+{
+	BasePropertyOverrides = InValue;
+}
+#endif // WITH_EDITORONLY_DATA
+
 void UMaterialInstance::SetVectorParameterValueInternal(const FMaterialParameterInfo& ParameterInfo, FLinearColor Value)
 {
 	LLM_SCOPE(ELLMTag::MaterialInstance);
@@ -3832,7 +3164,7 @@ bool UMaterialInstance::SetScalarParameterByIndexInternal(int32 ParameterIndex,
 	return true;
 }
 
-void UMaterialInstance::SetScalarParameterValueInternal(const FMaterialParameterInfo& ParameterInfo, float Value)
+void UMaterialInstance::SetScalarParameterValueInternal(const FMaterialParameterInfo& ParameterInfo, float Value, bool bUseAtlas, FScalarParameterAtlasInstanceData AtlasData)
 {
 	LLM_SCOPE(ELLMTag::MaterialInstance);
 
@@ -3848,10 +3180,28 @@ void UMaterialInstance::SetScalarParameterValueInternal(const FMaterialParameter
 		bForceUpdate = true;
 	}
 
+	float ValueToSet = Value;
+#if WITH_EDITORONLY_DATA
+	if (bUseAtlas)
+	{
+		UCurveLinearColorAtlas* Atlas = Cast<UCurveLinearColorAtlas>(AtlasData.Atlas.Get());
+		UCurveLinearColor* Curve = Cast<UCurveLinearColor>(AtlasData.Curve.Get());
+		if (Atlas && Curve)
+		{
+			int32 Index = Atlas->GradientCurves.Find(Curve);
+			if (Index != INDEX_NONE)
+			{
+				ValueToSet = (float)Index;
+			}
+		}
+		ParameterValue->AtlasData = AtlasData;
+	}
+#endif // WITH_EDITORONLY_DATA
+
 	// Don't enqueue an update if it isn't needed
-	if (bForceUpdate || ParameterValue->ParameterValue != Value)
+	if (bForceUpdate || ParameterValue->ParameterValue != ValueToSet)
 	{
-		ParameterValue->ParameterValue = Value;
+		ParameterValue->ParameterValue = ValueToSet;
 		// Update the material instance data in the rendering thread.
 		GameThread_UpdateMIParameter(this, *ParameterValue);
 	}
@@ -4139,7 +3489,8 @@ void UMaterialInstance::UpdateStaticPermutation(const FStaticParameterSet& NewPa
 
 void UMaterialInstance::UpdateStaticPermutation(FMaterialUpdateContext* MaterialUpdateContext)
 {
-	UpdateStaticPermutation(StaticParameters, MaterialUpdateContext);
+	// Force the update, since we aren't technically changing anything
+	UpdateStaticPermutation(StaticParameters, BasePropertyOverrides, true, MaterialUpdateContext);
 }
 
 void UMaterialInstance::UpdateParameterNames()
@@ -5173,3 +4524,5 @@ void UMaterialInstance::AppendReferencedParameterCollectionIdsTo(TArray<FGuid>&
 UMaterialInstance::FCustomStaticParametersGetterDelegate UMaterialInstance::CustomStaticParametersGetters;
 TArray<UMaterialInstance::FCustomParameterSetUpdaterDelegate> UMaterialInstance::CustomParameterSetUpdaters;
 #endif // WITH_EDITORONLY_DATA
+
+PRAGMA_ENABLE_OPTIMIZATION
\ No newline at end of file
diff --git a/Engine/Source/Runtime/Engine/Private/Materials/MaterialInstanceSupport.h b/Engine/Source/Runtime/Engine/Private/Materials/MaterialInstanceSupport.h
--- a/Engine/Source/Runtime/Engine/Private/Materials/MaterialInstanceSupport.h
+++ b/Engine/Source/Runtime/Engine/Private/Materials/MaterialInstanceSupport.h
@@ -225,6 +225,7 @@ ParameterType* GameThread_FindParameterByName(TArray<ParameterType>& Parameters,
 	}
 	return NULL;
 }
+
 template <typename ParameterType>
 const ParameterType* GameThread_FindParameterByName(const TArray<ParameterType>& Parameters, const FHashedMaterialParameterInfo& ParameterInfo)
 {
@@ -276,3 +277,47 @@ const ParameterType* GameThread_FindParameterByIndex(const TArray<ParameterType>
 
 	return &Parameters[Index];
 }
+
+template <typename ParameterType>
+FORCEINLINE bool GameThread_GetParameterValue(const TArray<ParameterType>& Parameters, const FHashedMaterialParameterInfo& ParameterInfo, bool bSetOverride, FMaterialParameterMetadata& OutResult)
+{
+	for (int32 ParameterIndex = 0; ParameterIndex < Parameters.Num(); ParameterIndex++)
+	{
+		const ParameterType* Parameter = &Parameters[ParameterIndex];
+		if (Parameter->IsOverride() && Parameter->ParameterInfo == ParameterInfo)
+		{
+			Parameter->GetValue(OutResult);
+#if WITH_EDITORONLY_DATA
+			if (bSetOverride)
+			{
+				OutResult.bOverride = true;
+			}
+#endif // WITH_EDITORONLY_DATA
+			return true;
+		}
+	}
+	return false;
+}
+
+template <typename ParameterType>
+FORCEINLINE void GameThread_ApplyParameterOverrides(const TArray<ParameterType>& Parameters, bool bSetOverride, TMap<FMaterialParameterInfo, FMaterialParameterMetadata>& OutParameters)
+{
+	for (int32 ParameterIndex = 0; ParameterIndex < Parameters.Num(); ParameterIndex++)
+	{
+		const ParameterType* Parameter = &Parameters[ParameterIndex];
+		if (Parameter->IsOverride())
+		{
+			FMaterialParameterMetadata* Result = OutParameters.Find(Parameter->ParameterInfo);
+			if (Result)
+			{
+				Parameter->GetValue(*Result);
+#if WITH_EDITORONLY_DATA
+				if (bSetOverride)
+				{
+					Result->bOverride = true;
+				}
+#endif // WITH_EDITORONLY_DATA
+			}
+		}
+	}
+}
diff --git a/Engine/Source/Runtime/Engine/Private/Materials/MaterialInterface.cpp b/Engine/Source/Runtime/Engine/Private/Materials/MaterialInterface.cpp
--- a/Engine/Source/Runtime/Engine/Private/Materials/MaterialInterface.cpp
+++ b/Engine/Source/Runtime/Engine/Private/Materials/MaterialInterface.cpp
@@ -63,6 +63,20 @@ bool IsCompatibleWithHairStrands(EShaderPlatform Platform, const FMaterialShader
 		(Parameters.BlendMode == BLEND_Opaque || Parameters.BlendMode == BLEND_Masked);
 }
 
+static EMaterialGetParameterValueFlags MakeParameterValueFlags(bool bOveriddenOnly, bool bCheckParent = true)
+{
+	EMaterialGetParameterValueFlags Result = EMaterialGetParameterValueFlags::CheckInstanceOverrides;
+	if (bCheckParent)
+	{
+		Result |= EMaterialGetParameterValueFlags::CheckParent;
+	}
+	if (!bOveriddenOnly)
+	{
+		Result |= EMaterialGetParameterValueFlags::CheckNonOverrides;
+	}
+	return Result;
+}
+
 //////////////////////////////////////////////////////////////////////////
 
 /** Copies the material's relevance flags to a primitive's view relevance flags. */
@@ -124,11 +138,28 @@ void UMaterialInterface::GetUsedTexturesAndIndices(TArray<UTexture*>& OutTexture
 #if WITH_EDITORONLY_DATA
 bool UMaterialInterface::GetStaticSwitchParameterValue(const FHashedMaterialParameterInfo& ParameterInfo, bool& OutValue, FGuid& OutExpressionGuid, bool bOveriddenOnly /*= false*/, bool bCheckParent /*= true*/) const
 {
+	FMaterialParameterMetadata Result;
+	if (GetParameterValue(EMaterialParameterType::StaticSwitch, ParameterInfo, Result, MakeParameterValueFlags(bOveriddenOnly, bCheckParent)))
+	{
+		OutExpressionGuid = Result.ExpressionGuid;
+		OutValue = Result.Value.AsStaticSwitch();
+		return true;
+	}
 	return false;
 }
 
 bool UMaterialInterface::GetStaticComponentMaskParameterValue(const FHashedMaterialParameterInfo& ParameterInfo, bool& R, bool& G, bool& B, bool& A, FGuid& OutExpressionGuid, bool bOveriddenOnly /*= false*/, bool bCheckParent /*= true*/) const
 {
+	FMaterialParameterMetadata Result;
+	if (GetParameterValue(EMaterialParameterType::StaticSwitch, ParameterInfo, Result, MakeParameterValueFlags(bOveriddenOnly, bCheckParent)))
+	{
+		OutExpressionGuid = Result.ExpressionGuid;
+		R = Result.Value.Bool[0];
+		G = Result.Value.Bool[1];
+		B = Result.Value.Bool[2];
+		A = Result.Value.Bool[3];
+		return true;
+	}
 	return false;
 }
 #endif // WITH_EDITORONLY_DATA
@@ -384,82 +415,265 @@ void UMaterialInterface::GetLightingGuidChain(bool bIncludeTextures, TArray<FGui
 
 bool UMaterialInterface::GetVectorParameterValue(const FHashedMaterialParameterInfo& ParameterInfo, FLinearColor& OutValue, bool bOveriddenOnly) const
 {
-	// is never called but because our system wants a UMaterialInterface instance we cannot use "virtual =0"
+	FMaterialParameterMetadata Result;
+	if (GetParameterValue(EMaterialParameterType::Vector, ParameterInfo, Result, MakeParameterValueFlags(bOveriddenOnly)))
+	{
+		OutValue = Result.Value.AsLinearColor();
+		return true;
+	}
 	return false;
 }
 
 #if WITH_EDITOR
 bool UMaterialInterface::IsVectorParameterUsedAsChannelMask(const FHashedMaterialParameterInfo& ParameterInfo, bool& OutValue) const
 {
+	FMaterialParameterMetadata Result;
+	if (GetParameterValue(EMaterialParameterType::Vector, ParameterInfo, Result, EMaterialGetParameterValueFlags::CheckNonOverrides | EMaterialGetParameterValueFlags::CheckParent))
+	{
+		OutValue = Result.bUsedAsChannelMask;
+		return true;
+	}
 	return false;
 }
 
 bool UMaterialInterface::GetVectorParameterChannelNames(const FHashedMaterialParameterInfo& ParameterInfo, FParameterChannelNames& OutValue) const
 {
+	FMaterialParameterMetadata Result;
+	if (GetParameterValue(EMaterialParameterType::Vector, ParameterInfo, Result, EMaterialGetParameterValueFlags::CheckNonOverrides | EMaterialGetParameterValueFlags::CheckParent))
+	{
+		OutValue.R = Result.ChannelName[0];
+		OutValue.G = Result.ChannelName[1];
+		OutValue.B = Result.ChannelName[2];
+		OutValue.A = Result.ChannelName[3];
+		return true;
+	}
 	return false;
 }
 
 bool UMaterialInterface::GetScalarParameterSliderMinMax(const FHashedMaterialParameterInfo& ParameterInfo, float& OutSliderMin, float& OutSliderMax) const
 {
+	FMaterialParameterMetadata Result;
+	if (GetParameterValue(EMaterialParameterType::Scalar, ParameterInfo, Result, EMaterialGetParameterValueFlags::CheckNonOverrides | EMaterialGetParameterValueFlags::CheckParent))
+	{
+		OutSliderMin = Result.ScalarMin;
+		OutSliderMax = Result.ScalarMax;
+		return true;
+	}
 	return false;
 }
 #endif // WITH_EDITOR
 
 bool UMaterialInterface::GetScalarParameterValue(const FHashedMaterialParameterInfo& ParameterInfo, float& OutValue, bool bOveriddenOnly) const
 {
-	// is never called but because our system wants a UMaterialInterface instance we cannot use "virtual =0"
+	FMaterialParameterMetadata Result;
+	if (GetParameterValue(EMaterialParameterType::Scalar, ParameterInfo, Result, MakeParameterValueFlags(bOveriddenOnly)))
+	{
+		OutValue = Result.Value.AsScalar();
+		return true;
+	}
+	return false;
+}
+
+bool UMaterialInterface::GetParameterValue(EMaterialParameterType Type, const FMemoryImageMaterialParameterInfo& ParameterInfo, FMaterialParameterMetadata& OutValue, EMaterialGetParameterValueFlags Flags) const
+{
 	return false;
 }
 
 #if WITH_EDITOR
 bool UMaterialInterface::IsScalarParameterUsedAsAtlasPosition(const FHashedMaterialParameterInfo& ParameterInfo, bool& OutValue, TSoftObjectPtr<class UCurveLinearColor>& Curve, TSoftObjectPtr<class UCurveLinearColorAtlas>& Atlas) const
 {
+	FMaterialParameterMetadata Result;
+	if (GetParameterValue(EMaterialParameterType::Scalar, ParameterInfo, Result, EMaterialGetParameterValueFlags::CheckNonOverrides | EMaterialGetParameterValueFlags::CheckParent))
+	{
+		OutValue = Result.bUsedAsAtlasPosition;
+		Curve = Result.ScalarCurve;
+		Atlas = Result.ScalarAtlas;
+		return true;
+	}
 	return false;
 }
 #endif // WITH_EDITOR
 
-bool UMaterialInterface::GetScalarCurveParameterValue(const FHashedMaterialParameterInfo& ParameterInfo, FInterpCurveFloat& OutValue) const
+bool UMaterialInterface::GetTextureParameterValue(const FHashedMaterialParameterInfo& ParameterInfo, UTexture*& OutValue, bool bOveriddenOnly) const
 {
+	FMaterialParameterMetadata Result;
+	if (GetParameterValue(EMaterialParameterType::Texture, ParameterInfo, Result, MakeParameterValueFlags(bOveriddenOnly)))
+	{
+		OutValue = Result.Value.Texture;
+		return true;
+	}
 	return false;
 }
 
-bool UMaterialInterface::GetVectorCurveParameterValue(const FHashedMaterialParameterInfo& ParameterInfo, FInterpCurveVector& OutValue) const
+bool UMaterialInterface::GetRuntimeVirtualTextureParameterValue(const FHashedMaterialParameterInfo& ParameterInfo, URuntimeVirtualTexture*& OutValue, bool bOveriddenOnly) const
 {
+	FMaterialParameterMetadata Result;
+	if (GetParameterValue(EMaterialParameterType::RuntimeVirtualTexture, ParameterInfo, Result, MakeParameterValueFlags(bOveriddenOnly)))
+	{
+		OutValue = Result.Value.RuntimeVirtualTexture;
+		return true;
+	}
 	return false;
 }
 
-bool UMaterialInterface::GetLinearColorParameterValue(const FHashedMaterialParameterInfo& ParameterInfo, FLinearColor& OutValue) const
+#if WITH_EDITOR
+bool UMaterialInterface::GetTextureParameterChannelNames(const FHashedMaterialParameterInfo& ParameterInfo, FParameterChannelNames& OutValue) const
 {
+	FMaterialParameterMetadata Result;
+	if (GetParameterValue(EMaterialParameterType::Texture, ParameterInfo, Result, EMaterialGetParameterValueFlags::CheckNonOverrides | EMaterialGetParameterValueFlags::CheckParent))
+	{
+		OutValue.R = Result.ChannelName[0];
+		OutValue.G = Result.ChannelName[1];
+		OutValue.B = Result.ChannelName[2];
+		OutValue.A = Result.ChannelName[3];
+		return true;
+	}
 	return false;
 }
+#endif
 
-bool UMaterialInterface::GetLinearColorCurveParameterValue(const FHashedMaterialParameterInfo& ParameterInfo, FInterpCurveLinearColor& OutValue) const
+bool UMaterialInterface::GetFontParameterValue(const FHashedMaterialParameterInfo& ParameterInfo, class UFont*& OutFontValue, int32& OutFontPage, bool bOveriddenOnly) const
 {
+	FMaterialParameterMetadata Result;
+	if (GetParameterValue(EMaterialParameterType::RuntimeVirtualTexture, ParameterInfo, Result, MakeParameterValueFlags(bOveriddenOnly)))
+	{
+		OutFontValue = Result.Value.Font.Value;
+		OutFontPage = Result.Value.Font.Page;
+		return true;
+	}
 	return false;
 }
 
-bool UMaterialInterface::GetTextureParameterValue(const FHashedMaterialParameterInfo& ParameterInfo, UTexture*& OutValue, bool bOveriddenOnly) const
+bool UMaterialInterface::GetParameterDefaultValue(EMaterialParameterType Type, const FMemoryImageMaterialParameterInfo& ParameterInfo, FMaterialParameterMetadata& OutValue) const
 {
+	return GetParameterValue(Type, ParameterInfo, OutValue, EMaterialGetParameterValueFlags::CheckNonOverrides | EMaterialGetParameterValueFlags::CheckParent);
+}
+
+
+bool UMaterialInterface::GetScalarParameterDefaultValue(const FHashedMaterialParameterInfo& ParameterInfo, float& OutValue) const
+{
+	FMaterialParameterMetadata Result;
+	if (GetParameterDefaultValue(EMaterialParameterType::Scalar, ParameterInfo, Result))
+	{
+		OutValue = Result.Value.AsScalar();
+		return true;
+	}
 	return false;
 }
 
-bool UMaterialInterface::GetRuntimeVirtualTextureParameterValue(const FHashedMaterialParameterInfo& ParameterInfo, URuntimeVirtualTexture*& OutValue, bool bOveriddenOnly) const
+bool UMaterialInterface::GetVectorParameterDefaultValue(const FHashedMaterialParameterInfo& ParameterInfo, FLinearColor& OutValue) const
+{
+	FMaterialParameterMetadata Result;
+	if (GetParameterDefaultValue(EMaterialParameterType::Vector, ParameterInfo, Result))
+	{
+		OutValue = Result.Value.AsLinearColor();
+		return true;
+	}
+	return false;
+}
+
+bool UMaterialInterface::GetTextureParameterDefaultValue(const FHashedMaterialParameterInfo& ParameterInfo, class UTexture*& OutValue) const
+{
+	FMaterialParameterMetadata Result;
+	if (GetParameterDefaultValue(EMaterialParameterType::Texture, ParameterInfo, Result))
+	{
+		OutValue = Result.Value.Texture;
+		return true;
+	}
+	return false;
+}
+
+bool UMaterialInterface::GetRuntimeVirtualTextureParameterDefaultValue(const FHashedMaterialParameterInfo& ParameterInfo, class URuntimeVirtualTexture*& OutValue) const
 {
+	FMaterialParameterMetadata Result;
+	if (GetParameterDefaultValue(EMaterialParameterType::RuntimeVirtualTexture, ParameterInfo, Result))
+	{
+		OutValue = Result.Value.RuntimeVirtualTexture;
+		return true;
+	}
 	return false;
 }
 
+bool UMaterialInterface::GetFontParameterDefaultValue(const FHashedMaterialParameterInfo& ParameterInfo, class UFont*& OutFontValue, int32& OutFontPage) const
+{
+	FMaterialParameterMetadata Result;
+	if (GetParameterDefaultValue(EMaterialParameterType::RuntimeVirtualTexture, ParameterInfo, Result))
+	{
+		OutFontValue = Result.Value.Font.Value;
+		OutFontPage = Result.Value.Font.Page;
+		return true;
+	}
+	return false;
+}
+
+
 #if WITH_EDITOR
-bool UMaterialInterface::GetTextureParameterChannelNames(const FHashedMaterialParameterInfo& ParameterInfo, FParameterChannelNames& OutValue) const
+bool UMaterialInterface::GetStaticSwitchParameterDefaultValue(const FHashedMaterialParameterInfo& ParameterInfo, bool& OutValue, FGuid& OutExpressionGuid) const
 {
+	FMaterialParameterMetadata Result;
+	if (GetParameterDefaultValue(EMaterialParameterType::RuntimeVirtualTexture, ParameterInfo, Result))
+	{
+		OutExpressionGuid = Result.ExpressionGuid;
+		OutValue = Result.Value.AsStaticSwitch();
+		return true;
+	}
 	return false;
 }
-#endif
 
-bool UMaterialInterface::GetFontParameterValue(const FHashedMaterialParameterInfo& ParameterInfo, class UFont*& OutFontValue, int32& OutFontPage, bool bOveriddenOnly) const
+bool UMaterialInterface::GetStaticComponentMaskParameterDefaultValue(const FHashedMaterialParameterInfo& ParameterInfo, bool& OutR, bool& OutG, bool& OutB, bool& OutA, FGuid& OutExpressionGuid) const
 {
+	FMaterialParameterMetadata Result;
+	if (GetParameterDefaultValue(EMaterialParameterType::RuntimeVirtualTexture, ParameterInfo, Result))
+	{
+		OutExpressionGuid = Result.ExpressionGuid;
+		OutR = Result.Value.Bool[0];
+		OutG = Result.Value.Bool[1];
+		OutB = Result.Value.Bool[2];
+		OutA = Result.Value.Bool[3];
+		return true;
+	}
 	return false;
 }
 
+#endif // WITH_EDITOR
+
+void UMaterialInterface::GetAllScalarParameterInfo(TArray<FMaterialParameterInfo>& OutParameterInfo, TArray<FGuid>& OutParameterIds) const
+{
+	GetAllParameterInfoOfType(EMaterialParameterType::Scalar, OutParameterInfo, OutParameterIds);
+}
+
+void UMaterialInterface::GetAllVectorParameterInfo(TArray<FMaterialParameterInfo>& OutParameterInfo, TArray<FGuid>& OutParameterIds) const
+{
+	GetAllParameterInfoOfType(EMaterialParameterType::Vector, OutParameterInfo, OutParameterIds);
+}
+
+void UMaterialInterface::GetAllTextureParameterInfo(TArray<FMaterialParameterInfo>& OutParameterInfo, TArray<FGuid>& OutParameterIds) const
+{
+	GetAllParameterInfoOfType(EMaterialParameterType::Texture, OutParameterInfo, OutParameterIds);
+}
+
+void UMaterialInterface::GetAllRuntimeVirtualTextureParameterInfo(TArray<FMaterialParameterInfo>& OutParameterInfo, TArray<FGuid>& OutParameterIds) const
+{
+	GetAllParameterInfoOfType(EMaterialParameterType::RuntimeVirtualTexture, OutParameterInfo, OutParameterIds);
+}
+
+void UMaterialInterface::GetAllFontParameterInfo(TArray<FMaterialParameterInfo>& OutParameterInfo, TArray<FGuid>& OutParameterIds) const
+{
+	GetAllParameterInfoOfType(EMaterialParameterType::Font, OutParameterInfo, OutParameterIds);
+}
+
+#if WITH_EDITORONLY_DATA
+void UMaterialInterface::GetAllStaticSwitchParameterInfo(TArray<FMaterialParameterInfo>&OutParameterInfo, TArray<FGuid>&OutParameterIds) const
+{
+	GetAllParameterInfoOfType(EMaterialParameterType::StaticSwitch, OutParameterInfo, OutParameterIds);
+}
+
+void UMaterialInterface::GetAllStaticComponentMaskParameterInfo(TArray<FMaterialParameterInfo>&OutParameterInfo, TArray<FGuid>&OutParameterIds) const
+{
+	GetAllParameterInfoOfType(EMaterialParameterType::StaticComponentMask, OutParameterInfo, OutParameterIds);
+}
+#endif // WITH_EDITOR
+
 bool UMaterialInterface::GetRefractionSettings(float& OutBiasValue) const
 {
 	return false;
diff --git a/Engine/Source/Runtime/Engine/Private/Materials/MaterialShader.cpp b/Engine/Source/Runtime/Engine/Private/Materials/MaterialShader.cpp
--- a/Engine/Source/Runtime/Engine/Private/Materials/MaterialShader.cpp
+++ b/Engine/Source/Runtime/Engine/Private/Materials/MaterialShader.cpp
@@ -338,6 +338,84 @@ bool FStaticParameterSet::Equivalent(const FStaticParameterSet& ReferenceSet) co
 	return false;
 }
 
+#if WITH_EDITORONLY_DATA
+void FStaticParameterSet::SetParameterValue(const FMaterialParameterInfo& ParameterInfo, const FMaterialParameterMetadata& Meta, EMaterialSetParameterValueFlags Flags)
+{
+	const FMaterialParameterValue& Value = Meta.Value;
+	switch (Value.Type)
+	{
+	case EMaterialParameterType::StaticSwitch: SetStaticSwitchParameterValue(ParameterInfo, Meta.ExpressionGuid, Value.AsStaticSwitch()); break;
+	case EMaterialParameterType::StaticComponentMask: SetStaticComponentMaskParameterValue(ParameterInfo, Meta.ExpressionGuid, Value.Bool[0], Value.Bool[1], Value.Bool[2], Value.Bool[3]); break;
+	default: checkNoEntry();
+	}
+}
+
+void FStaticParameterSet::AddParametersOfType(EMaterialParameterType Type, const TMap<FMaterialParameterInfo, FMaterialParameterMetadata>& Values)
+{
+	switch (Type)
+	{
+	case EMaterialParameterType::StaticSwitch:
+		StaticSwitchParameters.Empty(Values.Num());
+		for (const auto& It : Values)
+		{
+			const FMaterialParameterMetadata& Meta = It.Value;
+			check(Meta.Value.Type == Type);
+			StaticSwitchParameters.Emplace(It.Key, Meta.Value.AsStaticSwitch(), Meta.bOverride, Meta.ExpressionGuid);
+		}
+		break;
+	case EMaterialParameterType::StaticComponentMask:
+		StaticComponentMaskParameters.Empty(Values.Num());
+		for (const auto& It : Values)
+		{
+			const FMaterialParameterMetadata& Meta = It.Value;
+			check(Meta.Value.Type == Type);
+			StaticComponentMaskParameters.Emplace(It.Key,
+				Meta.Value.Bool[0],
+				Meta.Value.Bool[1],
+				Meta.Value.Bool[2],
+				Meta.Value.Bool[3],
+				Meta.bOverride,
+				Meta.ExpressionGuid);
+		}
+		break;
+	default: checkNoEntry();
+	}
+}
+
+void FStaticParameterSet::SetStaticSwitchParameterValue(const FMaterialParameterInfo& ParameterInfo, const FGuid& ExpressionGuid, bool Value)
+{
+	for (FStaticSwitchParameter& Parameter : StaticSwitchParameters)
+	{
+		if (Parameter.ParameterInfo == ParameterInfo)
+		{
+			Parameter.bOverride = true;
+			Parameter.Value = Value;
+			return;
+		}
+	}
+
+	new(StaticSwitchParameters) FStaticSwitchParameter(ParameterInfo, Value, true, ExpressionGuid);
+}
+
+void FStaticParameterSet::SetStaticComponentMaskParameterValue(const FMaterialParameterInfo& ParameterInfo, const FGuid& ExpressionGuid, bool R, bool G, bool B, bool A)
+{
+	for (FStaticComponentMaskParameter& Parameter : StaticComponentMaskParameters)
+	{
+		if (Parameter.ParameterInfo == ParameterInfo)
+		{
+			Parameter.bOverride = true;
+			Parameter.R = R;
+			Parameter.G = G;
+			Parameter.B = B;
+			Parameter.A = A;
+			return;
+		}
+	}
+
+	new(StaticComponentMaskParameters) FStaticComponentMaskParameter(ParameterInfo, R, G, B, A, true, ExpressionGuid);
+}
+#endif // WITH_EDITORONLY_DATA
+
 void FMaterialShaderMapId::Serialize(FArchive& Ar, bool bLoadedByCookedMaterial)
 {
 	SCOPED_LOADTIMER(FMaterialShaderMapId_Serialize);
diff --git a/Engine/Source/Runtime/Engine/Public/MaterialCachedData.h b/Engine/Source/Runtime/Engine/Public/MaterialCachedData.h
--- a/Engine/Source/Runtime/Engine/Public/MaterialCachedData.h
+++ b/Engine/Source/Runtime/Engine/Public/MaterialCachedData.h
@@ -3,6 +3,7 @@
 #pragma once
 
 #include "CoreMinimal.h"
+#include "MaterialTypes.h"
 #include "Materials/MaterialLayersFunctions.h"
 #include "Misc/Guid.h"
 #include "UObject/ObjectMacros.h"
@@ -72,36 +73,13 @@ struct FParameterChannelNames
 	FText A;
 };
 
-enum class EMaterialParameterType : int32
-{
-	Scalar,
-	Vector,
-	Texture,
-	Font,
-	RuntimeVirtualTexture,
-
-	RuntimeCount, // Runtime parameter types must go above here, and editor-only ones below
-
-#if WITH_EDITORONLY_DATA
-	StaticSwitch = RuntimeCount,
-	StaticComponentMask,
-	// Excluding StaticMaterialLayer due to type specific complications
-
-	Count,
-#else
-	Count = RuntimeCount,
-#endif
-};
-static const int32 NumMaterialRuntimeParameterTypes = (int32)EMaterialParameterType::RuntimeCount;
-#if WITH_EDITORONLY_DATA
-static const int32 NumMaterialEditorOnlyParameterTypes = (int32)EMaterialParameterType::Count - (int32)EMaterialParameterType::RuntimeCount;
-#endif
-
 USTRUCT()
 struct FMaterialCachedParameterEntry
 {
 	GENERATED_USTRUCT_BODY()
 
+	ENGINE_API static const FMaterialCachedParameterEntry EmptyData;
+
 	void Reset();
 
 	// This is used to map FMaterialParameterInfos to indices, which are then used to index various TArrays containing values for each type of parameter
@@ -109,30 +87,11 @@ struct FMaterialCachedParameterEntry
 	UPROPERTY()
 	TSet<FMaterialParameterInfo> ParameterInfoSet;
 
+#if WITH_EDITORONLY_DATA
 	UPROPERTY()
-	TArray<FGuid> ExpressionGuids; // editor-only?
-
-};
-
-USTRUCT()
-struct FStaticComponentMaskValue
-{
-	GENERATED_USTRUCT_BODY();
-
-	FStaticComponentMaskValue() : R(false), G(false), B(false), A(false) {}
-	FStaticComponentMaskValue(bool InR, bool InG, bool InB, bool InA) : R(InR), G(InG), B(InB), A(InA) {}
-
-	UPROPERTY()
-	bool R = false;
-	
-	UPROPERTY()
-	bool G = false;
-
-	UPROPERTY()
-	bool B = false;
+	TArray<FGuid> ExpressionGuids;
+#endif // WITH_EDITORONLY_DATA
 
-	UPROPERTY()
-	bool A = false;
 };
 
 USTRUCT()
@@ -141,26 +100,26 @@ struct FMaterialCachedParameters
 	GENERATED_USTRUCT_BODY()
 
 #if WITH_EDITORONLY_DATA
-	inline const FMaterialCachedParameterEntry& GetParameterTypeEntry(EMaterialParameterType Type) const { return Type >= EMaterialParameterType::RuntimeCount ? EditorOnlyEntries[static_cast<int32>(Type) - static_cast<int32>(EMaterialParameterType::RuntimeCount)] : RuntimeEntries[static_cast<int32>(Type)]; }
+	inline const FMaterialCachedParameterEntry& GetParameterTypeEntry(EMaterialParameterType Type) const { return Type >= EMaterialParameterType::NumRuntime ? EditorOnlyEntries[static_cast<int32>(Type) - NumMaterialRuntimeParameterTypes] : RuntimeEntries[static_cast<int32>(Type)]; }
+	inline FMaterialCachedParameterEntry& GetParameterTypeEntry(EMaterialParameterType Type) { return Type >= EMaterialParameterType::NumRuntime ? EditorOnlyEntries[static_cast<int32>(Type) - NumMaterialRuntimeParameterTypes] : RuntimeEntries[static_cast<int32>(Type)]; }
 #else
-	inline const FMaterialCachedParameterEntry& GetParameterTypeEntry(EMaterialParameterType Type) const { return RuntimeEntries[static_cast<int32>(Type)]; }
+	inline const FMaterialCachedParameterEntry& GetParameterTypeEntry(EMaterialParameterType Type) const { return Type >= EMaterialParameterType::NumRuntime ? FMaterialCachedParameterEntry::EmptyData : RuntimeEntries[static_cast<int32>(Type)]; }
 #endif
 
 	inline int32 GetNumParameters(EMaterialParameterType Type) const { return GetParameterTypeEntry(Type).ParameterInfoSet.Num(); }
 	int32 FindParameterIndex(EMaterialParameterType Type, const FMemoryImageMaterialParameterInfo& HashedParameterInfo) const;
+	bool GetParameterValue(EMaterialParameterType Type, const FMemoryImageMaterialParameterInfo& ParameterInfo, FMaterialParameterMetadata& OutResult) const;
+	void GetParameterValueByIndex(EMaterialParameterType Type, int32 ParameterIndex, FMaterialParameterMetadata& OutResult) const;
+#if WITH_EDITORONLY_DATA
 	const FGuid& GetExpressionGuid(EMaterialParameterType Type, int32 Index) const;
-	void GetAllParameterInfoOfType(EMaterialParameterType Type, bool bEmptyOutput, TArray<FMaterialParameterInfo>& OutParameterInfo, TArray<FGuid>& OutParameterIds) const;
-	void GetAllGlobalParameterInfoOfType(EMaterialParameterType Type, bool bEmptyOutput, TArray<FMaterialParameterInfo>& OutParameterInfo, TArray<FGuid>& OutParameterIds) const;
+#endif
+	void GetAllParametersOfType(EMaterialParameterType Type, TMap<FMaterialParameterInfo, FMaterialParameterMetadata>& OutParameters) const;
+	void GetAllParameterInfoOfType(EMaterialParameterType Type, TArray<FMaterialParameterInfo>& OutParameterInfo, TArray<FGuid>& OutParameterIds) const;
+	void GetAllGlobalParameterInfoOfType(EMaterialParameterType Type, TArray<FMaterialParameterInfo>& OutParameterInfo, TArray<FGuid>& OutParameterIds) const;
 	void Reset();
 
 	void AddReferencedObjects(FReferenceCollector& Collector);
 
-	bool GetScalarParameterSliderMinMax(const FMemoryImageMaterialParameterInfo& ParameterInfo, float& OutSliderMin, float& OutSliderMax) const;
-	bool IsScalarParameterUsedAsAtlasPosition(const FMemoryImageMaterialParameterInfo& ParameterInfo, bool& OutValue, TSoftObjectPtr<UCurveLinearColor>& OutCurve, TSoftObjectPtr<UCurveLinearColorAtlas>& OutAtlas) const;
-	bool IsVectorParameterUsedAsChannelMask(const FMemoryImageMaterialParameterInfo& ParameterInfo, bool& OutValue) const;
-	bool GetVectorParameterChannelNames(const FMemoryImageMaterialParameterInfo& ParameterInfo, FParameterChannelNames& OutValue) const;
-	bool GetTextureParameterChannelNames(const FMemoryImageMaterialParameterInfo& ParameterInfo, FParameterChannelNames& OutValue) const;
-
 	UPROPERTY()
 	FMaterialCachedParameterEntry RuntimeEntries[NumMaterialRuntimeParameterTypes];
 
diff --git a/Engine/Source/Runtime/Engine/Public/MaterialTypes.h b/Engine/Source/Runtime/Engine/Public/MaterialTypes.h
new file mode 100644
--- /dev/null
+++ b/Engine/Source/Runtime/Engine/Public/MaterialTypes.h
@@ -0,0 +1,188 @@
+// Copyright Epic Games, Inc. All Rights Reserved.
+
+#pragma once
+
+#include "CoreMinimal.h"
+#include "Misc/Guid.h"
+#include "Misc/EnumClassFlags.h"
+#include "UObject/SoftObjectPtr.h"
+#include "MaterialTypes.generated.h"
+
+class UTexture;
+class UCurveLinearColor;
+class UCurveLinearColorAtlas;
+class UFont;
+class URuntimeVirtualTexture;
+
+enum class EMaterialParameterType : uint8
+{
+	Scalar = 0u,
+	Vector,
+	Texture,
+	Font,
+	RuntimeVirtualTexture,
+
+	NumRuntime, // Runtime parameter types must go above here, and editor-only ones below
+
+	// TODO - Would be nice to make static parameter values editor-only, but will save that for a future-refactor
+	StaticSwitch = NumRuntime,
+	StaticComponentMask,
+
+	Num,
+	None = 0xff,
+};
+
+static const int32 NumMaterialParameterTypes = (int32)EMaterialParameterType::Num;
+static const int32 NumMaterialRuntimeParameterTypes = (int32)EMaterialParameterType::NumRuntime;
+static const int32 NumMaterialEditorOnlyParameterTypes = NumMaterialParameterTypes - NumMaterialRuntimeParameterTypes;
+
+static inline bool IsStaticMaterialParameter(EMaterialParameterType InType)
+{
+	switch (InType)
+	{
+	case EMaterialParameterType::StaticSwitch:
+	case EMaterialParameterType::StaticComponentMask:
+		return true;
+	default:
+		return false;
+	}
+}
+
+enum class EMaterialGetParameterValueFlags : uint32
+{
+	None = 0u,
+	CheckNonOverrides = (1u << 0),
+	CheckInstanceOverrides = (1u << 1),
+	CheckParent = (1u << 2),
+	IsParent = (1u << 3), // automatically set when recursing into a parent
+
+	CheckAll = CheckNonOverrides | CheckInstanceOverrides | CheckParent,
+};
+ENUM_CLASS_FLAGS(EMaterialGetParameterValueFlags);
+
+enum class EMaterialSetParameterValueFlags : uint32
+{
+	None = 0u,
+	SetCurveAtlas = (1u << 0),
+};
+ENUM_CLASS_FLAGS(EMaterialSetParameterValueFlags);
+
+USTRUCT()
+struct FStaticComponentMaskValue
+{
+	GENERATED_USTRUCT_BODY();
+
+	FStaticComponentMaskValue() : R(false), G(false), B(false), A(false) {}
+	FStaticComponentMaskValue(bool InR, bool InG, bool InB, bool InA) : R(InR), G(InG), B(InB), A(InA) {}
+
+	UPROPERTY()
+	bool R = false;
+
+	UPROPERTY()
+	bool G = false;
+
+	UPROPERTY()
+	bool B = false;
+
+	UPROPERTY()
+	bool A = false;
+};
+
+struct FMaterialParameterValue
+{
+	FMaterialParameterValue() : Type(EMaterialParameterType::None) {}
+	FMaterialParameterValue(float InValue) : Type(EMaterialParameterType::Scalar) { Float[0] = InValue; }
+	FMaterialParameterValue(const FLinearColor& InValue) : Type(EMaterialParameterType::Vector) { Float[0] = InValue.R; Float[1] = InValue.G; Float[2] = InValue.B; Float[3] = InValue.A; }
+	FMaterialParameterValue(UTexture* InValue) : Type(EMaterialParameterType::Texture) { Texture = InValue; }
+	FMaterialParameterValue(const TObjectPtr<UTexture>& InValue) : Type(EMaterialParameterType::Texture) { Texture = InValue; }
+	FMaterialParameterValue(URuntimeVirtualTexture* InValue) : Type(EMaterialParameterType::RuntimeVirtualTexture) { RuntimeVirtualTexture = InValue; }
+	FMaterialParameterValue(const TObjectPtr<URuntimeVirtualTexture>& InValue) : Type(EMaterialParameterType::RuntimeVirtualTexture) { RuntimeVirtualTexture = InValue; }
+	FMaterialParameterValue(UFont* InValue, int32 InPage) : Type(EMaterialParameterType::Font) { Font.Value = InValue; Font.Page = InPage; }
+	FMaterialParameterValue(bool InValue) : Type(EMaterialParameterType::StaticSwitch) { Bool[0] = InValue; }
+	FMaterialParameterValue(const FStaticComponentMaskValue& InValue) : Type(EMaterialParameterType::StaticComponentMask) { Bool[0] = InValue.R; Bool[1] = InValue.G; Bool[2] = InValue.B; Bool[3] = InValue.A; }
+	FMaterialParameterValue(bool bMaskR, bool bMaskG, bool bMaskB, bool bMaskA) : Type(EMaterialParameterType::StaticComponentMask) { Bool[0] = bMaskR; Bool[1] = bMaskG; Bool[2] = bMaskB; Bool[3] = bMaskA; }
+
+	inline float AsScalar() const { check(Type == EMaterialParameterType::Scalar); return Float[0]; }
+	inline FLinearColor AsLinearColor() const { check(Type == EMaterialParameterType::Vector); return FLinearColor(Float[0], Float[1], Float[2], Float[3]); }
+	inline bool AsStaticSwitch() const { check(Type == EMaterialParameterType::StaticSwitch); return Bool[0]; }
+	inline FStaticComponentMaskValue AsStaticComponentMask() const { check(Type == EMaterialParameterType::StaticComponentMask); return FStaticComponentMaskValue(Bool[0], Bool[1], Bool[2], Bool[3]); }
+
+	union
+	{
+		float Float[4];
+		bool Bool[4];
+		UTexture* Texture;
+		URuntimeVirtualTexture* RuntimeVirtualTexture;
+		struct
+		{
+			UFont* Value;
+			int32 Page;
+		} Font;
+	};
+	EMaterialParameterType Type;
+};
+
+inline bool operator==(const FMaterialParameterValue& Lhs, const FMaterialParameterValue& Rhs)
+{
+	const EMaterialParameterType Type = Lhs.Type;
+	if (Type != Rhs.Type)
+	{
+		return false;
+	}
+	switch (Type)
+	{
+	case EMaterialParameterType::None: return true;
+	case EMaterialParameterType::Scalar: return Lhs.Float[0] == Rhs.Float[0];
+	case EMaterialParameterType::Vector: return
+		Lhs.Float[0] == Rhs.Float[0] &&
+		Lhs.Float[1] == Rhs.Float[1] &&
+		Lhs.Float[2] == Rhs.Float[2] &&
+		Lhs.Float[3] == Rhs.Float[3];
+	case EMaterialParameterType::Texture: return Lhs.Texture == Rhs.Texture;
+	case EMaterialParameterType::Font: return Lhs.Font.Value == Rhs.Font.Value && Lhs.Font.Page == Rhs.Font.Page;
+	case EMaterialParameterType::RuntimeVirtualTexture: return Lhs.RuntimeVirtualTexture == Rhs.RuntimeVirtualTexture;
+	case EMaterialParameterType::StaticSwitch: return Lhs.Bool[0] == Rhs.Bool[0];
+	case EMaterialParameterType::StaticComponentMask: return
+		Lhs.Bool[0] == Rhs.Bool[0] &&
+		Lhs.Bool[1] == Rhs.Bool[1] &&
+		Lhs.Bool[2] == Rhs.Bool[2] &&
+		Lhs.Bool[3] == Rhs.Bool[3];
+	default: checkNoEntry(); return false;
+	}
+}
+
+inline bool operator!=(const FMaterialParameterValue& Lhs, const FMaterialParameterValue& Rhs)
+{
+	return !operator==(Lhs, Rhs);
+}
+
+/** Holds a value, along with editor-only metadata that describes that value */
+struct FMaterialParameterMetadata
+{
+	FMaterialParameterValue Value;
+
+#if WITH_EDITORONLY_DATA
+	/** Name of channels, for Vectors/Textures */
+	FText ChannelName[4];
+
+	/** Curve/Atlas used to generate scalar value */
+	TSoftObjectPtr<class UCurveLinearColor> ScalarCurve;
+	TSoftObjectPtr<class UCurveLinearColorAtlas> ScalarAtlas;
+
+	/** UI range for scalar values */
+	float ScalarMin = 0.0f;
+	float ScalarMax = 0.0f;
+
+	/** GUID of the UMaterialExpression this parameter came from */
+	FGuid ExpressionGuid;
+
+	/** Should curves be used? */
+	bool bUsedAsAtlasPosition = false;
+
+	/** Valid for Vector parameters */
+	bool bUsedAsChannelMask = false;
+
+	/** Is the parameter overriden on the material it was queried from? */
+	bool bOverride = false;
+#endif // WITH_EDITORONLY_DATA
+};
\ No newline at end of file
diff --git a/Engine/Source/Runtime/Engine/Public/StaticParameterSet.h b/Engine/Source/Runtime/Engine/Public/StaticParameterSet.h
--- a/Engine/Source/Runtime/Engine/Public/StaticParameterSet.h
+++ b/Engine/Source/Runtime/Engine/Public/StaticParameterSet.h
@@ -8,6 +8,7 @@
 #include "UObject/FortniteMainBranchObjectVersion.h"
 #include "UObject/RenderingObjectVersion.h"
 #include "UObject/ReleaseObjectVersion.h"
+#include "MaterialTypes.h"
 #include "Materials/MaterialLayersFunctions.h"
 #include "StaticParameterSet.generated.h"
 
@@ -41,6 +42,8 @@ struct FStaticParameterBase
 		ExpressionGUID(InGuid)
 	{ }
 
+	bool IsOverride() const { return bOverride; }
+
 	friend FArchive& operator<<(FArchive& Ar, FStaticParameterBase& P)
 	{
 		// This method should never be called, derived structures need to implement their own code (to retain compatibility) or call SerializeBase (for new classes)
@@ -130,6 +133,14 @@ struct FStaticSwitchParameter : public FStaticParameterBase
 		FStaticParameterBase::AppendKeyString(KeyString);
 		KeyString += FString::FromInt(Value);
 	}
+
+	void GetValue(FMaterialParameterMetadata& OutResult) const
+	{
+		OutResult.Value = Value;
+#if WITH_EDITORONLY_DATA
+		OutResult.ExpressionGuid = ExpressionGUID;
+#endif
+	}
 };
 
 /**
@@ -207,6 +218,14 @@ struct FStaticComponentMaskParameter : public FStaticParameterBase
 		KeyString += FString::FromInt(B);
 		KeyString += FString::FromInt(A);
 	}
+
+	void GetValue(FMaterialParameterMetadata& OutResult) const
+	{
+		OutResult.Value = FMaterialParameterValue(R, G, B, A);
+#if WITH_EDITORONLY_DATA
+		OutResult.ExpressionGuid = ExpressionGUID;
+#endif
+	}
 };
 
 /**
@@ -413,6 +432,16 @@ struct FStaticParameterSet
 
 	bool Equivalent(const FStaticParameterSet& ReferenceSet) const;
 
+#if WITH_EDITORONLY_DATA
+	void SetParameterValue(const FMaterialParameterInfo& ParameterInfo, const FMaterialParameterMetadata& Meta, EMaterialSetParameterValueFlags Flags = EMaterialSetParameterValueFlags::None);
+	void AddParametersOfType(EMaterialParameterType Type, const TMap<FMaterialParameterInfo, FMaterialParameterMetadata>& Values);
+#endif // WITH_EDITORONLY_DATA
+
 private:
 	void SortForEquivalent();
+
+#if WITH_EDITORONLY_DATA
+	void SetStaticSwitchParameterValue(const FMaterialParameterInfo& ParameterInfo, const FGuid& ExpressionGuid, bool Value);
+	void SetStaticComponentMaskParameterValue(const FMaterialParameterInfo& ParameterInfo, const FGuid& ExpressionGuid, bool R, bool G, bool B, bool A);
+#endif // WITH_EDITORONLY_DATA
 };
