Fix for Solo Audio in First PIE Client for different net modes

diff --git a/Engine/Source/Editor/UnrealEd/Private/PlayLevel.cpp b/Engine/Source/Editor/UnrealEd/Private/PlayLevel.cpp
--- a/Engine/Source/Editor/UnrealEd/Private/PlayLevel.cpp
+++ b/Engine/Source/Editor/UnrealEd/Private/PlayLevel.cpp
@@ -2862,7 +2862,10 @@ UGameInstance* UEditorEngine::CreateInnerProcessPIEGameInstance(FRequestPlaySess
 				{
 					if (FAudioDeviceManager* DeviceManager = GEngine->GetAudioDeviceManager())
 					{
-						if (InPIEInstanceIndex == 0)
+						EPlayNetMode NetMode; 
+						InParams.EditorPlaySettings->GetPlayNetMode(NetMode);
+						if ((NetMode == PIE_Client && InPIEInstanceIndex == 1) || 
+							((NetMode == PIE_Standalone || NetMode == PIE_ListenServer) && InPIEInstanceIndex == 0))
 						{
 							DeviceManager->SetSoloDevice(GameInstanceAudioDevice->DeviceID);
 						}
