Cooker: Hybrid Iterative cooks
Add the concept of pre-exploration of dependencies - searching assetregistry or target domain dependencies of the set of requested packages up front, to find all packages that will be cooked and allow us to do asynchronous and batch operations on them.
Convert FRequestCluster into an FPackageData container, holding packages in a substate of EPackageState::Request. This allows us to create a new RequestClusters when we encounter a package that was not found in the initial discovery.
Add hybriditerative mode for iterative cooks; this uses targetdomain keys to check for packages that exist from the previous cook, rather tracing the dependency graph to find all packages that transitively reference packages with modified guids.
Modify AssetRegistryGenerator to use a TMap of iteratively loaded packages that can be edied as packages cook, rather than using a const previous state; this supports hybrid iterative where we do not calculate the list of iterative packages until we encounter them in a request cluster.
#rnx

diff --git a/Engine/Config/BaseEditor.ini b/Engine/Config/BaseEditor.ini
--- a/Engine/Config/BaseEditor.ini
+++ b/Engine/Config/BaseEditor.ini
@@ -327,7 +327,12 @@ MaxUsedMemory=16384
 
 [CookSettings]
 BulkDataRegistryEnabled=false
-ExploreSoftReferencesOnStart=false
+; HybridIterative uses the TargetDomain and BuildDependencies to decide when packages need to be recooked. If enabled,
+; every cook is iterative for packages with allowed classes (IterativeClassAllowList) unless -noiterative is passed.
+HybridIterativeEnabled=false
+; PreexploreDependenciesEnabled can cause editor-only packages to be cooked due to lack of robust editor-only reference detection.
+; Projects can disable it to prevent cooking those packages, at the cost of lower cook performance.
+PreexploreDependenciesEnabled=true
 PackagesPerGC=0
 IdleTimeToGC=20
 MemoryMaxUsedPhysical=16384
@@ -378,18 +383,19 @@ SaveUnversioned=false
 ClassBlockList=/Script/ControlRig.ControlRigBlueprintGeneratedClass
 ; +PackageBlockList=/Long/PackageName OR ../../../Standard/Filename/NoExtension OR D:\Absolute\Filename\NoExtension
 
-; Know bad files that are problematic if all of them come from EditorDomain. Blocking any one of these is sufficient to prevent the bug
+; Known bad files that are problematic if all of them come from EditorDomain. Blocking any one of these is sufficient to prevent the bug.
 +PackageBlockList=/Niagara/Functions/RandomRangeFloat
 +PackageBlockList=/Niagara/Modules/Update/Color/ScaleColor
 
 [TargetDomain]
-TargetDomainEnabled=false
-; IterativeClassAllowList/IterativeClassBlockList:
+; IterativeClassAllowListEnabled/IterativeClassAllowList/IterativeClassBlockList:
 ;    Used in HybridIterative only.
 ;    List of classes that may use iterative cooks. Allowed superclasses are also allowed.
 ;    Blocked subclasses are also blocked. BlockList overrides the AllowList.
+;    If IterativeClassAllowListEnabled is false, also classes are allowed unless blocked.
 ; +IterativeClassAllowList=/Long/PackageName.ClassName
 ; +IterativeClassBlockList=/Long/PackageName.ClassName
+IterativeClassAllowListEnabled=true
 IterativeClassAllowList=/Script/Engine.Texture2D
 +IterativeClassAllowList=/Script/CoreUObject.MetaData
 +IterativeClassAllowList=/Script/Engine.AssetImportData
diff --git a/Engine/Source/Editor/UnrealEd/Classes/CookOnTheSide/CookOnTheFlyServer.h b/Engine/Source/Editor/UnrealEd/Classes/CookOnTheSide/CookOnTheFlyServer.h
--- a/Engine/Source/Editor/UnrealEd/Classes/CookOnTheSide/CookOnTheFlyServer.h
+++ b/Engine/Source/Editor/UnrealEd/Classes/CookOnTheSide/CookOnTheFlyServer.h
@@ -209,6 +209,8 @@ private:
 	 * This allows us to have enough population in the LoadPrepareQueue to get benefit from the asynchronous work done on packages in the LoadPrepareQueue.
 	 */
 	uint32 DesiredLoadQueueLength;
+	/** A knob to tune performance - how many packages to pull off in each call to PumpRequests. */
+	uint32 RequestBatchSize;
 
 	ECookInitializationFlags CookFlags = ECookInitializationFlags::None;
 	TUniquePtr<class FSandboxPlatformFile> SandboxFile;
@@ -302,7 +304,6 @@ private:
 	enum class ECookAction
 	{
 		Done,			// The cook is complete; no requests remain in any non-idle state
-		RequestCluster,	// Process the RequestClusters
 		Request,		// Process the RequestQueue
 		Load,			// Process the LoadQueue
 		LoadLimited,	// Process the LoadQueue, stopping when loadqueuelength reaches the desired population level
@@ -315,14 +316,10 @@ private:
 	ECookAction DecideNextCookAction(UE::Cook::FTickStackData& StackData);
 	/** Execute any existing external callbacks and push any existing external cook requests into new RequestClusters. */
 	void PumpExternalRequests(const UE::Cook::FCookerTimer& CookerTimer);
-	/** Pull off each RequestCluster and find all transitive dependencies and push them into the RequestQueue. */
-	void PumpRequestClusters(UE::Cook::FTickStackData& StackData);
+	/** Send the PackageData back to request state to create a request cluster, if it has not yet been explored. */
+	bool TryCreateRequestCluster(UE::Cook::FPackageData& PackageData);
 	/** Inspect the next package in the RequestQueue and push it on to its next state. Report the number of PackageDatas that were pushed to load. */
 	void PumpRequests(UE::Cook::FTickStackData& StackData, int32& OutNumPushed);
-	/** Handle the requested PackageData that has been peeled off of the RequestQueue, e.g. by sending it on to the LoadQueue. Report the number of PackageDatas that were pushed to load (includes dependencies, so may be more than 1). */
-	void ProcessRequest(UE::Cook::FPackageData& PackageData, UE::Cook::FTickStackData& StackData, int32& OutNumPushed);
-	/** Get the list of unloaded Packages the load of the given PackageData is dependent upon, and add their PackageDatas to the load queue. Report the number of PackageDatass that were pushed to the load queue */
-	void AddDependenciesToLoadQueue(UE::Cook::FPackageData& PackageData, int32& OutNumPushed);
 
 	/** Load all packages in the LoadQueue until it's time to break. Report the number of loads that were pushed to save. */
 	void PumpLoads(UE::Cook::FTickStackData& StackData, uint32 DesiredQueueLength, int32& OutNumPushed, bool& bOutBusy);
@@ -341,9 +338,10 @@ private:
 	 * Inspect the given package and queue it for saving if necessary.
 	 *
 	 * @param Package				The package to be considered for saving.
+	 * @param bOutWasInProgress		Report whether the package was already in progress.
 	 * @return						Returns the PackageData for this queued package.
 	 */
-	UE::Cook::FPackageData* QueueDiscoveredPackage(UPackage* Package);
+	UE::Cook::FPackageData* QueueDiscoveredPackage(UPackage* Package, bool* bOutWasInProgress = nullptr);
 	/**
 	 * Inspect the given package and queue it for saving if necessary.
 	 *
@@ -467,11 +465,8 @@ public:
 	UE_DEPRECATED(4.26, "Unsolicited packages are now added directly to the savequeue and are not marked as unsolicited")
 	TArray<UPackage*> GetUnsolicitedPackages(const TArray<const ITargetPlatform*>& TargetPlatforms) const;
 
-	/**
-	* PostLoadPackageFixup
-	* after a package is loaded we might want to fix up some stuff before it gets saved
-	*/
-	void PostLoadPackageFixup(UPackage* Package);
+	/** Execute class-specific special case cook postloads and reference discovery on a given package. */
+	void PostLoadPackageFixup(UPackage* Package, TArray<FName>* OutDiscoveredPackageNames = nullptr);
 
 	/**
 	* Handles cook package requests until there are no more requests, then returns
@@ -512,8 +507,13 @@ public:
 	UE_DEPRECATED(4.25, "Use version that takes const ITargetPlatform* instead")
 	void ClearPlatformCookedData(const FString& PlatformName);
 
-	/** Clear all in-process flags that indicate packages have been cooked for the given platforms. */
-	void ResetCookResults(TConstArrayView<const ITargetPlatform*> TargetPlatforms);
+	/**
+	 * Clear platforms' explored flags for all PackageDatas and optionally clear the cookresult flags.
+	 *
+	 * @param TargetPlatforms List of pairs with targetplatforms to reset and bool bResetResults indicating whether
+	 *        the platform should clear bCookAttempted in addition to clearing bExplored.
+	 */
+	void ResetCook(TConstArrayView<TPair<const ITargetPlatform*, bool>> TargetPlatforms);
 
 	/**
 	* Recompile any global shader changes 
@@ -768,18 +768,9 @@ private:
 	//////////////////////////////////////////////////////////////////////////
 	// general functions
 
-	/**
-	 * Attempts to update the metadata for a package in an asset registry generator
-	 *
-	 * @param Generator The asset registry generator to update
-	 * @param Package The package to update info on
-	 * @param SavePackageResult The metadata to associate with the given package name
-	 */
-	void UpdateAssetRegistryPackageData(FAssetRegistryGenerator* Generator, const UPackage& Package, FSavePackageResultStruct& SavePackageResult);
-
 	/** Perform any special processing for freshly loaded packages 
 	 */
-	void ProcessUnsolicitedPackages();
+	void ProcessUnsolicitedPackages(TArray<FName>* OutDiscoveredPackageNames = nullptr);
 
 	/**
 	 * Loads a package and prepares it for cooking
@@ -799,6 +790,9 @@ private:
 	*/
 	void BeginCookSandbox(TConstArrayView<const ITargetPlatform*> TargetPlatforms);
 
+	/** Set parameters that rely on config and CookByTheBook settings. */
+	void SetBeginCookConfigSettings();
+
 	/**
 	* Finalize the package store
 	*/
@@ -1093,13 +1087,12 @@ private:
 	bool bSaveBusy = false;
 	/** If preloading is enabled, we call TryPreload until it returns true before sending the package to LoadReady, otherwise we skip TryPreload and it goes immediately. */
 	bool bPreloadingEnabled = false;
-	/** If TargetDomain is enabled, we load/save TargetDomainKey hashes and use those to test whether packages have already been cooked in hybrid-iterative builds. */
-	bool bTargetDomainEnabled = true;
-	/**
-	 * Following soft dependencies in request clusters can cause editor-only packages to be cooked due to lack of robust editor-only reference detection.
-	 * Projects can disable it to prevent cooking those packages, at the cost of lower cook performance.
-	 */
-	bool bExploreSoftReferencesOnStart = true;
+	/** If enabled, we load/save TargetDomainKey hashes and use those to test whether packages have already been cooked in hybrid-iterative builds. */
+	bool bHybridIterativeEnabled = true;
+	/** Whether to explore transitive dependencies up front, or discover them as we load packages. */
+	bool bPreexploreDependenciesEnabled = true;
+	/** Test mode for the debug of hybrid iterative dependencies. */
+	bool bHybridIterativeDebug = false;
 
 	/** Timers for tracking how long we have been busy, to manage retries and warnings of deadlock */
 	float SaveBusyTimeLastRetry = 0.f;
@@ -1113,7 +1106,6 @@ private:
 	TUniquePtr<UE::Cook::FPackageTracker> PackageTracker;
 	TUniquePtr<UE::Cook::FPackageDatas> PackageDatas;
 	TUniquePtr<UE::Cook::FExternalRequests> ExternalRequests;
-	TRingBuffer<UE::Cook::FRequestCluster> RequestClusters;
 
 	TArray<UE::Cook::FCookSavePackageContext*> SavePackageContexts;
 	/** Objects that were collected during the single-threaded PreGarbageCollect callback and that should be reported as referenced in CookerAddReferencedObjects. */
diff --git a/Engine/Source/Editor/UnrealEd/Private/Commandlets/AssetRegistryGenerator.cpp b/Engine/Source/Editor/UnrealEd/Private/Commandlets/AssetRegistryGenerator.cpp
--- a/Engine/Source/Editor/UnrealEd/Private/Commandlets/AssetRegistryGenerator.cpp
+++ b/Engine/Source/Editor/UnrealEd/Private/Commandlets/AssetRegistryGenerator.cpp
@@ -795,15 +795,52 @@ void FAssetRegistryGenerator::CleanManifestDirectories()
 	CleanTempPackagingDirectory(TargetPlatform->PlatformName());
 }
 
-void FAssetRegistryGenerator::SetPreviousAssetRegistry(TUniquePtr<FAssetRegistryState>&& InPreviousState)
+void FAssetRegistryGenerator::SetPreviousAssetRegistry(TUniquePtr<FAssetRegistryState>&& InPreviousState,
+	TMap<FName, FAssetPackageData*>&& PreviousAssetPackageData, TArray<FName>& OutTombstonePackages)
 {
+	OutTombstonePackages.Reset();
 	if (InPreviousState)
 	{
-		PreviousState = MoveTemp(*InPreviousState);
+		for (const TPair<FName, const FAssetPackageData*>& Pair : InPreviousState->GetAssetPackageDataMap())
+		{
+			FName PackageName = Pair.Key;
+			const FAssetPackageData& PreviousPackageData = *Pair.Value;
+			if (!State.GetAssetPackageData(PackageName))
+			{
+				// Tombstone means a package from the previous cook no longer exists in perforce
+				OutTombstonePackages.Add(PackageName);
+				continue;
+			}
+			TArrayView<FAssetData const * const> PreviousAssetDatas = InPreviousState->GetAssetsByPackageName(PackageName);
+			TPair<TArray<FAssetData>, FAssetPackageData>& UpdateData = PreviousPackagesToUpdate.FindOrAdd(PackageName);
+			if (PreviousAssetDatas.Num() > 0)
+			{
+				TArray<FAssetData>& UpdateAssetDatas = UpdateData.Get<0>();
+				UpdateAssetDatas.Reserve(PreviousAssetDatas.Num());
+				for (const FAssetData* AssetData : PreviousAssetDatas)
+				{
+					UpdateAssetDatas.Add(FAssetData(*AssetData));
+				}
+			}
+			UpdateData.Get<1>() = PreviousPackageData;
+		}
 	}
 	else
 	{
-		PreviousState.Reset();
+		for (TPair<FName, FAssetPackageData*>& Pair : PreviousAssetPackageData)
+		{
+			FName PackageName = Pair.Key;
+			FAssetPackageData& PreviousPackageData = *Pair.Value;
+			if (!State.GetAssetPackageData(PackageName))
+			{
+				// Tombstone means a package from the previous cook no longer exists in perforce
+				OutTombstonePackages.Add(PackageName);
+				continue;
+			}
+			TPair<TArray<FAssetData>, FAssetPackageData>& UpdateData = PreviousPackagesToUpdate.FindOrAdd(PackageName);
+			// No AssetDatas available in this case
+			UpdateData.Get<1>() = MoveTemp(PreviousPackageData);
+		}
 	}
 }
 
@@ -928,43 +965,18 @@ FAssetPackageData* FAssetRegistryGenerator::GetAssetPackageData(const FName& Pac
 	return State.CreateOrGetAssetPackageData(PackageName);
 }
 
-void FAssetRegistryGenerator::UpdateKeptPackagesDiskData(const TArray<FName>& InKeptPackages)
+void FAssetRegistryGenerator::UpdateKeptPackages()
 {
-	for (const FName& PackageName : InKeptPackages)
+	for (TPair<FName, TPair<TArray<FAssetData>, FAssetPackageData>>& Pair : PreviousPackagesToUpdate)
 	{
-		// Get mutable PackageData without creating it when it does not exist
-		FAssetPackageData* PackageData =
-			State.GetAssetPackageData(PackageName) ?
-			State.CreateOrGetAssetPackageData(PackageName) :
-			nullptr;
-		const FAssetPackageData* PreviousPackageData =
-			PreviousState.GetAssetPackageData(PackageName);
-
-		if (!PackageData || !PreviousPackageData)
-		{
-			continue;
-		}
-
-		PRAGMA_DISABLE_DEPRECATION_WARNINGS
-		if (PackageData->PackageGuid != PreviousPackageData->PackageGuid)
-		PRAGMA_ENABLE_DEPRECATION_WARNINGS
+		for (const FAssetData& PreviousAssetData : Pair.Value.Get<0>())
 		{
-			continue;
-		}
-
-		PackageData->CookedHash = PreviousPackageData->CookedHash;
-		PackageData->DiskSize = PreviousPackageData->DiskSize;
-	}
-}
-
-void FAssetRegistryGenerator::UpdateKeptPackagesAssetData()
-{
-	for (FName PackageName : KeptPackages)
-	{
-		for (const FAssetData* PreviousAssetData : PreviousState.GetAssetsByPackageName(PackageName))
-		{
-			State.UpdateAssetData(*PreviousAssetData);
+			State.UpdateAssetData(PreviousAssetData);
 		}
+		FAssetPackageData& PackageData = *State.CreateOrGetAssetPackageData(Pair.Key);
+		FAssetPackageData& PreviousPackageData = Pair.Value.Get<1>();
+		PackageData.CookedHash = PreviousPackageData.CookedHash;
+		PackageData.DiskSize = PreviousPackageData.DiskSize;
 	}
 }
 
@@ -1083,7 +1095,7 @@ void FAssetRegistryGenerator::ComputePackageDifferences(const FComputeDifference
 		else if (CurrentPackageData->PackageGuid == PreviousPackageData->PackageGuid)
 		PRAGMA_ENABLE_DEPRECATION_WARNINGS
 		{
-			if (!Options.bIgnoreDiskSize && PreviousPackageData->DiskSize < 0)
+			if (PreviousPackageData->DiskSize < 0)
 			{
 				OutDifference.IdenticalUncookedPackages.Add(PackageName);
 			}
@@ -1151,45 +1163,6 @@ void FAssetRegistryGenerator::ComputePackageDifferences(const FComputeDifference
 	}
 }
 
-void FAssetRegistryGenerator::ComputePackageDifferences(TSet<FName>& ModifiedPackages, TSet<FName>& NewPackages, TSet<FName>& RemovedPackages, TSet<FName>& IdenticalCookedPackages, TSet<FName>& IdenticalUncookedPackages, bool bRecurseModifications, bool bRecurseScriptModifications)
-{
-	FComputeDifferenceOptions Options;
-	Options.bRecurseModifications		= bRecurseModifications;
-	Options.bRecurseScriptModifications	= bRecurseScriptModifications;
-	
-	FAssetRegistryDifference Difference;
-	ComputePackageDifferences(Options, PreviousState.GetAssetPackageDataMap(), Difference);
-
-	ModifiedPackages			= MoveTemp(Difference.ModifiedPackages);
-	NewPackages					= MoveTemp(Difference.NewPackages);
-	RemovedPackages				= MoveTemp(Difference.RemovedPackages);
-	IdenticalCookedPackages		= MoveTemp(Difference.IdenticalCookedPackages);
-	IdenticalUncookedPackages	= MoveTemp(Difference.IdenticalUncookedPackages);
-}
-
-void FAssetRegistryGenerator::UpdateKeptPackages(const TArray<FName>& InKeptPackages)
-{
-	KeptPackages.Append(InKeptPackages);
-	// Update disk data right away, disk data is only updated when packages are saved, and kept packages are never saved
-	UpdateKeptPackagesDiskData(InKeptPackages);
-	// Delay update of AssetData with TagsAndValues, this data may be modified up until serialization in SaveAssetRegistry
-}
-
-void FAssetRegistryGenerator::UpdateKeptPackages(const TArray<FName>& InKeptPackages, TFunction<void(const FName&, FAssetPackageData*)>&& Update)
-{
-	KeptPackages.Append(InKeptPackages);
-
-	for (const FName& PackageName : InKeptPackages)
-	{
-		// Get mutable PackageData without creating it when it does not exist
-		FAssetPackageData* PackageData = State.GetAssetPackageData(PackageName)
-			? State.CreateOrGetAssetPackageData(PackageName)
-			: nullptr;
-
-		Update(PackageName, PackageData);
-	}
-}
-
 void FAssetRegistryGenerator::BuildChunkManifest(const TSet<FName>& InCookedPackages, const TSet<FName>& InDevelopmentOnlyPackages, FSandboxPlatformFile* InSandboxFile, bool bGenerateStreamingInstallManifest)
 {
 	// If we were asked to generate a streaming install manifest explicitly and we did not have bGenerateNoChunks set, we will generate chunks.
@@ -1416,8 +1389,8 @@ bool FAssetRegistryGenerator::SaveAssetRegistry(const FString& SandboxPath, bool
 	// First flush the asset registry and make sure the asset data is in sync, as it may have been updated during cook
 	AssetRegistry.Tick(-1.0f);
 	AssetRegistry.InitializeTemporaryAssetRegistryState(State, SaveOptions, true);
-	// Then possibly apply AssetData with TagsAndValues from a previous AssetRegistry for packages kept from a previous cook
-	UpdateKeptPackagesAssetData();
+	// Then possibly apply previous AssetData and AssetPackageData data for packages kept from a previous cook
+	UpdateKeptPackages();
 	UpdateCollectionAssetData();
 
 	if (DevelopmentSaveOptions.bSerializeAssetRegistry && bSerializeDevelopmentAssetRegistry)
@@ -2353,6 +2326,41 @@ const FAssetData* FAssetRegistryGenerator::CreateOrFindAssetData(UObject& Object
 	return AssetData;
 }
 
+void FAssetRegistryGenerator::UpdateAssetRegistryPackageData(const UPackage& Package, FSavePackageResultStruct& SavePackageResult)
+{
+	const FName PackageName = Package.GetFName();
+	PreviousPackagesToUpdate.Remove(PackageName);
+
+	bool bSaveSucceeded = SavePackageResult.IsSuccessful();
+	if (bSaveSucceeded)
+	{
+		// Ensure all assets in the package are recorded in the registry
+		CreateOrFindAssetDatas(Package);
+
+		FAssetPackageData* AssetPackageData = GetAssetPackageData(PackageName);
+		AssetPackageData->DiskSize = SavePackageResult.TotalFileSize;
+		// If there is no hash (e.g.: when SavePackageResult == ESavePackageResult::ReplaceCompletely), don't attempt to setup a continuation to update
+		// the AssetRegistry entry with it later.  Just leave the asset registry entry with a default constructed FMD5Hash which is marked as invalid.
+		if (SavePackageResult.CookedHash.IsValid())
+		{
+			SavePackageResult.CookedHash.Next([AssetPackageData](const FMD5Hash& CookedHash)
+				{
+					// Store the cooked hash in the Asset Registry when it is done computing in another thread.
+					// NOTE: For this to work, we rely on:
+					// 1) UPackage::WaitForAsyncFileWrites to have been called before any use of the CookedHash - it is called in CookByTheBookFinished before the registry does any work with the registries
+					// 2) AssetPackageData must continue to be a valid pointer - CreateOrFindAssetDatas allocates the FAssetPackageData individually and doesn't relocate or delete them until pruning, which happens after WaitForAsyncFileWrites
+					AssetPackageData->CookedHash = CookedHash;
+				});
+		}
+	}
+
+	// Set the package flags to zero to indicate that the package failed to save
+	const uint32 PackageFlags = bSaveSucceeded ? SavePackageResult.SerializedPackageFlags : 0;
+	const bool bUpdated = UpdateAssetPackageFlags(PackageName, PackageFlags);
+	UE_CLOG(!bUpdated && bSaveSucceeded, LogAssetRegistryGenerator, Warning,
+		TEXT("Trying to update asset package flags in package '%s' that does not exist"), *PackageName.ToString());
+}
+
 bool FAssetRegistryGenerator::UpdateAssetPackageFlags(const FName& PackageName, const uint32 PackageFlags)
 {
 	return State.UpdateAssetDataPackageFlags(PackageName, PackageFlags);
diff --git a/Engine/Source/Editor/UnrealEd/Private/CookOnTheFlyServer.cpp b/Engine/Source/Editor/UnrealEd/Private/CookOnTheFlyServer.cpp
--- a/Engine/Source/Editor/UnrealEd/Private/CookOnTheFlyServer.cpp
+++ b/Engine/Source/Editor/UnrealEd/Private/CookOnTheFlyServer.cpp
@@ -173,7 +173,6 @@ static FAutoConsoleVariableRef CVarCookDisplayWarnBusyTime(
 /// Cook on the fly server
 ///////////////////////////////////////////////////////////////
 
-constexpr uint32 ExpectedMaxNumPlatforms = 32;
 constexpr float TickCookableObjectsFrameTime = .100f;
 
 namespace UE
@@ -605,6 +604,7 @@ bool UCookOnTheFlyServer::StartCookOnTheFly(FCookOnTheFlyOptions InCookOnTheFlyO
 	}
 	PlatformManager->SetArePlatformsPrepopulated(true);
 
+	SetBeginCookConfigSettings();
 	CreateSandboxFile();
 	GenerateAssetRegistry();
 	for (FName NeverCookPackage : GetNeverCookPackageFileNames(TArrayView<const FString>()))
@@ -731,11 +731,8 @@ void UCookOnTheFlyServer::AddCookOnTheFlyPlatformFromGameThread(ITargetPlatform*
 
 	// When cooking on the fly the full registry is saved at the beginning
 	// in cook by the book asset registry is saved after the cook is finished
-	FAssetRegistryGenerator* Generator = PlatformData->RegistryGenerator.Get();
-	if (Generator)
-	{
-		Generator->SaveAssetRegistry(GetSandboxAssetRegistryFilename(), true);
-	}
+	FAssetRegistryGenerator& Generator = *PlatformData->RegistryGenerator;
+	Generator.SaveAssetRegistry(GetSandboxAssetRegistryFilename(), true);
 	check(PlatformData->bIsSandboxInitialized); // This should have been set by BeginCookSandbox, and it is what we use to determine whether a platform has been initialized
 }
 
@@ -1166,7 +1163,6 @@ FString UCookOnTheFlyServer::GetContentDirectoryForDLC() const
 	return GetBaseDirectoryForDLC() / TEXT("Content");
 }
 
-
 // allow for a command line to start async preloading a Development AssetRegistry if requested
 static FEventRef GPreloadAREvent(EEventMode::ManualReset);
 static FEventRef GPreloadARInfoEvent(EEventMode::ManualReset);
@@ -1310,9 +1306,6 @@ uint32 UCookOnTheFlyServer::TickCookOnTheSide(const float TimeSlice, uint32 &Coo
 			bool bBusy;
 			switch (CookAction)
 			{
-			case ECookAction::RequestCluster:
-				PumpRequestClusters(StackData);
-				break;
 			case ECookAction::Request:
 				PumpRequests(StackData, NumPushed);
 				if (NumPushed > 0)
@@ -1636,9 +1629,10 @@ UCookOnTheFlyServer::ECookAction UCookOnTheFlyServer::DecideNextCookAction(UE::C
 		return ECookAction::YieldTick;
 	}
 
-	if (RequestClusters.Num())
+	UE::Cook::FRequestQueue& RequestQueue = PackageDatas->GetRequestQueue();
+	if (RequestQueue.GetRequestClusters().Num() != 0 || RequestQueue.GetUnclusteredRequests().Num() != 0)
 	{
-		return ECookAction::RequestCluster;
+		return ECookAction::Request;
 	}
 
 	UE::Cook::FPackageDataMonitor& Monitor = PackageDatas->GetMonitor();
@@ -1680,9 +1674,7 @@ UCookOnTheFlyServer::ECookAction UCookOnTheFlyServer::DecideNextCookAction(UE::C
 		return ECookAction::LoadLimited;
 	}
 
-	int32 NumRequests = PackageDatas->GetRequestQueue().Num();
-	bool bRequestAvailable = NumRequests > 0;
-	if (bRequestAvailable)
+	if (!RequestQueue.IsReadyRequestsEmpty())
 	{
 		return ECookAction::Request;
 	}
@@ -1761,34 +1753,21 @@ void UCookOnTheFlyServer::PumpExternalRequests(const UE::Cook::FCookerTimer& Coo
 			}
 #endif
 			bool bRequestsAreUrgent = IsCookOnTheFlyMode();
-			UE::Cook::FRequestCluster::AddClusters(*this, MoveTemp(BuildRequests), bRequestsAreUrgent,
-				RequestClusters);
+			TRingBuffer<UE::Cook::FRequestCluster>& RequestClusters = PackageDatas->GetRequestQueue().GetRequestClusters();
+			UE::Cook::FRequestCluster::AddClusters(*this, MoveTemp(BuildRequests), bRequestsAreUrgent, RequestClusters);
 		}
 	}
 }
 
-void UCookOnTheFlyServer::PumpRequestClusters(UE::Cook::FTickStackData& StackData)
+bool UCookOnTheFlyServer::TryCreateRequestCluster(UE::Cook::FPackageData& PackageData)
 {
 	using namespace UE::Cook;
-
-	const FCookerTimer& CookerTimer = StackData.Timer;
-	while (RequestClusters.Num() > 0 && !CookerTimer.IsTimeUp())
+	if (!PackageData.AreAllRequestedPlatformsExplored())
 	{
-		FRequestCluster& RequestCluster = RequestClusters.First();
-		bool bComplete = false;
-		RequestCluster.Process(CookerTimer, bComplete);
-		if (bComplete)
-		{
-			for (FRequestCluster::FRequest& Request : RequestCluster.GetRequests())
-			{
-				FPackageData& PackageData(PackageDatas->FindOrAddPackageData(Request.PackageName, Request.FileName));
-				PackageData.UpdateRequestData(RequestCluster.GetPlatforms(), Request.bIsUrgent,
-					MoveTemp(Request.CompletionCallback));
-			}
-
-			RequestClusters.PopFront();
-		}
+		PackageData.SendToState(EPackageState::Request, ESendFlags::QueueAdd);
+		return true;
 	}
+	return false;
 }
 
 void UCookOnTheFlyServer::PumpRequests(UE::Cook::FTickStackData& StackData, int32& OutNumPushed)
@@ -1798,186 +1777,80 @@ void UCookOnTheFlyServer::PumpRequests(UE::Cook::FTickStackData& StackData, int3
 
 	OutNumPushed = 0;
 	FRequestQueue& RequestQueue = PackageDatas->GetRequestQueue();
-	COOK_STAT(DetailedCookStats::PeakRequestQueueSize = FMath::Max(DetailedCookStats::PeakRequestQueueSize, static_cast<int32>(RequestQueue.Num())));
-	if (!RequestQueue.IsEmpty())
-	{
-		FPackageData* PackageData = RequestQueue.PopRequest();
-		FPoppedPackageDataScope Scope(*PackageData);
-		int32 NumPushed;
-		ProcessRequest(*PackageData, StackData, NumPushed);
-		OutNumPushed += NumPushed;
-	}
-}
-
-void UCookOnTheFlyServer::ProcessRequest(UE::Cook::FPackageData& PackageData, UE::Cook::FTickStackData& StackData, int32& OutNumPushed)
-{
-	using namespace UE::Cook;
-	OutNumPushed = 0;
-
-	if (PackageData.AreAllRequestedPlatformsCooked(true /* bAllowFailedCooks */))
-	{
-#if DEBUG_COOKONTHEFLY
-		UE_LOG(LogCook, Display, TEXT("Package for platform already cooked %s, discarding request"), *PackageData.GetFileName().ToString());
-#endif
-		PackageData.SendToState(EPackageState::Idle, ESendFlags::QueueAdd);
-		return;
-	}
-
-	const FName& BuildFileName = PackageData.GetFileName();
-	FString BuildFileNameString = BuildFileName.ToString();
-	if (IsCookByTheBookMode() && CookByTheBookOptions->bErrorOnEngineContentUse)
+	FPackageDataSet& UnclusteredRequests = RequestQueue.GetUnclusteredRequests();
+	TRingBuffer<FRequestCluster>& RequestClusters = RequestQueue.GetRequestClusters();
+	const FCookerTimer& CookerTimer = StackData.Timer;
+	if (!UnclusteredRequests.IsEmpty())
 	{
-		check(IsCookingDLC());
-		FString DLCPath = FPaths::Combine(*GetBaseDirectoryForDLC(), TEXT("Content"));
-		FPaths::MakeStandardFilename(DLCPath);
-		if (BuildFileNameString.StartsWith(DLCPath) == false) // if we don't start with the dlc path then we shouldn't be cooking this data 
+		FRequestCluster::AddClusters(*this, UnclusteredRequests, RequestClusters, RequestQueue);
+		UnclusteredRequests.Empty();
+		if (CookerTimer.IsTimeUp())
 		{
-			UE_LOG(LogCook, Error, TEXT("Engine or Game content %s is being referenced by DLC!"), *BuildFileNameString);
-			RejectPackageToLoad(PackageData, TEXT("is base Game/Engine content and we are building DLC that is not allowed to refer to it"));
 			return;
 		}
 	}
 
-	if (PackageTracker->NeverCookPackageList.Contains(BuildFileName))
-	{
-#if DEBUG_COOKONTHEFLY
-		UE_LOG(LogCook, Display, TEXT("Package %s requested but is in the never cook package list, discarding request"), *BuildFileNameString);
-#else
-		UE_LOG(LogCook, Verbose, TEXT("Package %s requested but is in the never cook package list, discarding request"), *BuildFileNameString);
-#endif
-		RejectPackageToLoad(PackageData, TEXT("is in the never cook list"));
-		return;
-	}
-
-
-	if (!PackageData.GetIsUrgent() && (!IsCookByTheBookMode() || !CookByTheBookOptions->bSkipHardReferences))
-	{
-		// TODO: We need to AddDependenciesToLoadQueue only for requests that did not already do so
-		// because they were part of a CookRequestCluster
-		int32 NumPushed;
-		AddDependenciesToLoadQueue(PackageData, NumPushed);
-		OutNumPushed += NumPushed;
-	}
-	// AddDependenciesToLoadQueue is supposed to add the dependencies only and not add the passed-in packagedata, so it should still be in request
-	check(PackageData.GetState() == EPackageState::Request);
-	PackageData.SendToState(EPackageState::LoadPrepare, ESendFlags::QueueAdd);
-	++OutNumPushed;
-}
-
-void UCookOnTheFlyServer::AddDependenciesToLoadQueue(UE::Cook::FPackageData& PackageData, int32& OutNumPushed)
-{
-	using namespace UE::Cook;
-	OutNumPushed = 0;
-
-	struct FPackageAndDependencies
+	while (RequestClusters.Num() > 0)
 	{
-		FPackageData& PackageData;
-		TArray<FPackageData*> Dependencies;
-		int NextDependency = 0;
-
-		explicit FPackageAndDependencies(FPackageData& InPackageData, TArray<FName>& AssetDependenciesScratch, IAssetRegistry* AssetRegistry, FPackageDatas& PackageDatas)
-			: PackageData(InPackageData)
+		FRequestCluster& RequestCluster = RequestClusters.First();
+		bool bComplete = false;
+		RequestCluster.Process(CookerTimer, bComplete);
+		if (bComplete)
 		{
-			check(!PackageData.GetIsVisited());
-			PackageData.SetIsVisited(true);
-
-			AssetDependenciesScratch.Reset();
-			// TODO EditorOnly References: We only load Game dependencies, because if we explicitly load an EditorOnly dependency, that causes StaticLoadObjectInternal to SetLoadedByEditorPropertiesOnly(false), and we do not want to impact that value with our preloading of required packages 
-			if (AssetRegistry->GetDependencies(PackageData.GetPackageName(), AssetDependenciesScratch, UE::AssetRegistry::EDependencyCategory::Package, UE::AssetRegistry::EDependencyQuery::Hard | UE::AssetRegistry::EDependencyQuery::Game))
+			TArray<FPackageData*> RequestsToLoad;
+			TArray<FPackageData*> RequestsToDemote;
+			RequestCluster.ClearAndDetachOwnedPackageDatas(RequestsToLoad, RequestsToDemote);
+			for (FPackageData* PackageData : RequestsToLoad)
+			{
+				RequestQueue.AddReadyRequest(PackageData);
+			}
+			for (FPackageData* PackageData : RequestsToDemote)
+			{
+				PackageData->SendToState(EPackageState::Idle, ESendFlags::QueueAdd);
+			}
+			RequestClusters.PopFront();
+			if (IsCookByTheBookMode() && bHybridIterativeDebug)
 			{
-				Dependencies.Reserve(AssetDependenciesScratch.Num());
-				for (const FName& DependencyName : AssetDependenciesScratch)
+				for (FPackageData* PackageData : *PackageDatas)
 				{
-					TStringBuilder<256> NameBuffer;
-					DependencyName.ToString(NameBuffer);
-					if (FPackageName::IsScriptPackage(NameBuffer))
+					if (!PackageData->AreAllRequestedPlatformsExplored())
 					{
-						continue;
+						UE_LOG(LogCook, Warning, TEXT("Missing dependency: existing requested Package %s was not explored in the first cluster."),
+							*WriteToString<256>(PackageData->GetPackageName()));
 					}
-					FPackageData* DependencyData = PackageDatas.TryAddPackageDataByPackageName(DependencyName);
-					if (!DependencyData || DependencyData == &PackageData)
-					{
-						continue;
-					}
-					Dependencies.Add(DependencyData);
 				}
+				PackageDatas->SetLogDiscoveredPackages(true); // Turn this on for the rest of the cook after the initial cluster
 			}
 		}
-
-		FPackageAndDependencies(FPackageAndDependencies&& Other) = default;
-		
-		FPackageAndDependencies(const FPackageAndDependencies& Other) = delete;
-		FPackageAndDependencies& operator=(const FPackageAndDependencies& Other) = delete;
-		FPackageAndDependencies& operator=(FPackageAndDependencies&& Other) = delete;
-	};
-	TRingBuffer<FPackageAndDependencies> LoadStack;
-	TArray<FName> AssetDependenciesScratch;
-	const TArray<const ITargetPlatform*>& SessionPlatforms = PlatformManager->GetSessionPlatforms();
-	FRequestQueue& RequestQueue = PackageDatas->GetRequestQueue();
-
-	LoadStack.AddFront(FPackageAndDependencies(PackageData, AssetDependenciesScratch, AssetRegistry, *PackageDatas));
-
-	while (!LoadStack.IsEmpty())
-	{
-		FPackageAndDependencies& PackageAndDependencies(LoadStack.First());
-		FPackageData& CurrentPackageData(PackageAndDependencies.PackageData);
-		TArray<FPackageData*>& Dependencies(PackageAndDependencies.Dependencies);
-		int32& NextDependency(PackageAndDependencies.NextDependency);
-
-		// We search in DFS order so that we end up with a topological sort (or a mostly topological sort when there are cycles)
-		bool bAddedDependency = false;
-		while (NextDependency < Dependencies.Num())
+		if (CookerTimer.IsTimeUp())
 		{
-			FPackageData& DependencyData(*Dependencies[NextDependency++]);
-			if (DependencyData.GetState() >= EPackageState::LoadPrepare || DependencyData.GetIsVisited())
-			{
-				// If it's already been visited, or it's already loading or saving, don't add it again
-				// Skipping the add if it's already been visited is how we make sure we still terminate even with cycles in the dependency graph
-				continue;
-			}
-			if (FindObjectFast<UPackage>(nullptr, DependencyData.GetPackageName(), false /* ExactClass */, false /* AnyPackage */))
-			{
-				// If it's already loaded, no work to do for it
-				continue;
-			}
-
-			// Move the dependency into the request state and push it onto the DependencyStack, closer to front than the current PackageData that depends on it.
-			// Note that since we're moving it into the LoadQueue before we return, we do not need to spend time to add it to (or leave it in) the RequestQueue
-			bool bIsUrgent = false;
-			if (DependencyData.GetState() == EPackageState::Request)
-			{
-				RequestQueue.RemoveRequest(&DependencyData);
-				DependencyData.UpdateRequestData(SessionPlatforms, bIsUrgent, FCompletionCallback(), ESendFlags::QueueNone);
-			}
-			else
-			{
-				DependencyData.UpdateRequestData(SessionPlatforms, bIsUrgent, FCompletionCallback(), ESendFlags::QueueRemove);
-			}
-			LoadStack.AddFront(FPackageAndDependencies(DependencyData, AssetDependenciesScratch, AssetRegistry, *PackageDatas));
-			// PackageAndDependencies is now invalidated
-			bAddedDependency = true;
-			break;
+			return;
 		}
-		if (bAddedDependency)
+	}
+
+	COOK_STAT(DetailedCookStats::PeakRequestQueueSize = FMath::Max(DetailedCookStats::PeakRequestQueueSize, static_cast<int32>(RequestQueue.ReadyRequestsNum())));
+	uint32 NumInBatch = 0;
+	while (!RequestQueue.IsReadyRequestsEmpty() && NumInBatch < RequestBatchSize)
+	{
+		FPackageData* PackageData = RequestQueue.PopReadyRequest();
+		FPoppedPackageDataScope Scope(*PackageData);
+		// Exploring dependencies in ReadRequests is legacy behavior and is allowed even if bPreexploreDependenciesEnabled  is false
+		if (TryCreateRequestCluster(*PackageData))
 		{
-			// PackageAndDependencies is now invalidated
 			continue;
 		}
-		LoadStack.PopFront();
-		// PackageAndDependencies is now invalidated
-
-		check(CurrentPackageData.GetIsVisited());
-		CurrentPackageData.SetIsVisited(false);
-		if (&CurrentPackageData != &PackageData) // Caller is responsible for queueing the original PackageData; this function just queues dependencies
+		if (PackageData->AreAllRequestedPlatformsCooked(true /* bAllowFailedCooks */))
 		{
-			// We only record dependencies that were in idle or request, and we push idle to request when adding it to the dependency stack. If a package is no longer in request at this point, then there is a bug up above and we have pushed it into load twice
-			check(CurrentPackageData.GetState() == EPackageState::Request); 
-			// Send the package to the load queue
-			CurrentPackageData.SendToState(EPackageState::LoadPrepare, ESendFlags::QueueAdd);
-			++OutNumPushed;
-			COOK_STAT(++DetailedCookStats::NumPreloadedDependencies);
+#if DEBUG_COOKONTHEFLY
+			UE_LOG(LogCook, Display, TEXT("Package for platform already cooked %s, discarding request"), *PackageData.GetFileName().ToString());
+#endif
+			PackageData->SendToState(EPackageState::Idle, ESendFlags::QueueAdd);
+			continue;
 		}
+		PackageData->SendToState(EPackageState::LoadPrepare, ESendFlags::QueueAdd);
+		++NumInBatch;
 	}
+	OutNumPushed += NumInBatch;
 }
 
 void UCookOnTheFlyServer::PumpLoads(UE::Cook::FTickStackData& StackData, uint32 DesiredQueueLength, int32& OutNumPushed, bool& bOutBusy)
@@ -2018,6 +1891,11 @@ void UCookOnTheFlyServer::PumpLoads(UE::Cook::FTickStackData& StackData, uint32
 
 		FPackageData& PackageData(*LoadReadyQueue.PopFrontValue());
 		FPoppedPackageDataScope Scope(PackageData);
+		if (bPreexploreDependenciesEnabled && TryCreateRequestCluster(PackageData))
+		{
+			continue;
+		}
+
 		int32 NumPushed;
 		LoadPackageInQueue(PackageData, StackData.ResultFlags, NumPushed);
 		OutNumPushed += NumPushed;
@@ -2058,6 +1936,10 @@ void UCookOnTheFlyServer::PumpPreloadStarts()
 	while (!EntryQueue.IsEmpty() && Monitor.GetNumPreloadAllocated() < static_cast<int32>(MaxPreloadAllocated))
 	{
 		FPackageData* PackageData = EntryQueue.PopFrontValue();
+		if (bPreexploreDependenciesEnabled  && TryCreateRequestCluster(*PackageData))
+		{
+			continue;
+		}
 		if (bLocalPreloadingEnabled)
 		{
 			PackageData->TryPreload();
@@ -2245,7 +2127,7 @@ void UCookOnTheFlyServer::RejectPackageToLoad(UE::Cook::FPackageData& PackageDat
 
 //////////////////////////////////////////////////////////////////////////
 
-UE::Cook::FPackageData* UCookOnTheFlyServer::QueueDiscoveredPackage(UPackage* Package)
+UE::Cook::FPackageData* UCookOnTheFlyServer::QueueDiscoveredPackage(UPackage* Package, bool* bOutWasInProgress)
 {
 	check(Package != nullptr);
 
@@ -2255,6 +2137,10 @@ UE::Cook::FPackageData* UCookOnTheFlyServer::QueueDiscoveredPackage(UPackage* Pa
 		return nullptr;	// if we have name none that means we are in core packages or something...
 	}
 	UE::Cook::FPackageData& PackageData = PackageDatas->FindOrAddPackageData(Package->GetFName(), FileName);
+	if (bOutWasInProgress)
+	{
+		*bOutWasInProgress = PackageData.IsInProgress();
+	}
 	QueueDiscoveredPackageData(PackageData, true /* bIsLoadReady */);
 	return &PackageData;
 }
@@ -3084,7 +2970,7 @@ bool UCookOnTheFlyServer::LoadPackageForCooking(UE::Cook::FPackageData& PackageD
 	return bSuccess;
 }
 
-void UCookOnTheFlyServer::ProcessUnsolicitedPackages()
+void UCookOnTheFlyServer::ProcessUnsolicitedPackages(TArray<FName>* OutDiscoveredPackageNames)
 {
 	// Ensure sublevels are loaded by iterating all recently loaded packages and invoking
 	// PostLoadPackageFixup
@@ -3098,10 +2984,15 @@ void UCookOnTheFlyServer::ProcessUnsolicitedPackages()
 		{
 			if (!IsCookByTheBookMode() || !CookByTheBookOptions->bSkipSoftReferences)
 			{
-				PostLoadPackageFixup(Package);
+				PostLoadPackageFixup(Package, OutDiscoveredPackageNames);
 			}
 
-			QueueDiscoveredPackage(Package);
+			bool bWasInProgress;
+			UE::Cook::FPackageData* PackageData = QueueDiscoveredPackage(Package, &bWasInProgress);
+			if (OutDiscoveredPackageNames && PackageData && !bWasInProgress)
+			{
+				OutDiscoveredPackageNames->Add(PackageData->GetPackageName());
+			}
 		}
 	}
 }
@@ -3172,6 +3063,11 @@ void UCookOnTheFlyServer::PumpSaves(UE::Cook::FTickStackData& StackData, uint32
 	while (SaveQueue.Num() > static_cast<int32>(DesiredQueueLength))
 	{
 		FPackageData& PackageData(*SaveQueue.PopFrontValue());
+		if (bPreexploreDependenciesEnabled  && TryCreateRequestCluster(PackageData))
+		{
+			continue;
+		}
+
 		FPoppedPackageDataScope PoppedScope(PackageData);
 		UPackage* Package = PackageData.GetPackage();
 		
@@ -3358,43 +3254,7 @@ void UCookOnTheFlyServer::PumpSaves(UE::Cook::FTickStackData& StackData, uint32
 	}
 }
 
-void UCookOnTheFlyServer::UpdateAssetRegistryPackageData(FAssetRegistryGenerator* Generator, const UPackage& Package, FSavePackageResultStruct& SavePackageResult)
-{
-	if (!Generator)
-		return;
-
-	const bool bSaveSucceeded = SavePackageResult.IsSuccessful();
-
-	if (bSaveSucceeded)
-	{
-		// Ensure all assets in the package are recorded in the registry
-		Generator->CreateOrFindAssetDatas(Package);
-
-		const FName PackageName = Package.GetFName();
-		FAssetPackageData* AssetPackageData = Generator->GetAssetPackageData(PackageName);
-		AssetPackageData->DiskSize = SavePackageResult.TotalFileSize;
-		// If there is no hash (e.g.: when SavePackageResult == ESavePackageResult::ReplaceCompletely), don't attempt to setup a continuation to update
-		// the AssetRegistry entry with it later.  Just leave the asset registry entry with a default constructed FMD5Hash which is marked as invalid.
-		if (SavePackageResult.CookedHash.IsValid())
-		{
-			SavePackageResult.CookedHash.Next([AssetPackageData](const FMD5Hash& CookedHash)
-			{
-				// Store the cooked hash in the Asset Registry when it is done computing in another thread.
-				// NOTE: For this to work, we rely on:
-				// 1) UPackage::WaitForAsyncFileWrites to have been called before any use of the CookedHash - it is called in CookByTheBookFinished before the registry does any work with the registries
-				// 2) AssetPackageData must continue to be a valid pointer - the asset registry allocates the FAssetPackageData individually and doesn't relocate or delete them until pruning, which happens after WaitForAsyncFileWrites
-					AssetPackageData->CookedHash = CookedHash;
-			});
-		}
-	}
-
-	// Set the package flags to zero to indicate that the package failed to save
-	const uint32 PackageFlags = bSaveSucceeded ? SavePackageResult.SerializedPackageFlags : 0;
-	const bool bUpdated = Generator->UpdateAssetPackageFlags(Package.GetFName(), PackageFlags);
-	UE_CLOG(!bUpdated && bSaveSucceeded, LogCook, Warning, TEXT("Trying to update asset package flags in package '%s' that does not exist"), *Package.GetName());
-}
-
-void UCookOnTheFlyServer::PostLoadPackageFixup(UPackage* Package)
+void UCookOnTheFlyServer::PostLoadPackageFixup(UPackage* Package, TArray<FName>* OutDiscoveredPackageNames)
 {
 	if (Package->ContainsMap() == false)
 	{
@@ -3444,11 +3304,16 @@ void UCookOnTheFlyServer::PostLoadPackageFixup(UPackage* Package)
 
 	for (const FString& PackageName : NewPackagesToCook)
 	{
-		UE::Cook::FPackageData* NewPackageData = PackageDatas->TryAddPackageDataByPackageName(FName(*PackageName));
-		if (NewPackageData)
+		FName PackageFName(*PackageName);
+		UE::Cook::FPackageData* NewPackageData = PackageDatas->TryAddPackageDataByPackageName(PackageFName);
+		if (NewPackageData && !NewPackageData->IsInProgress())
 		{
 			bool bIsUrgent = false;
 			NewPackageData->UpdateRequestData(PlatformManager->GetSessionPlatforms(), bIsUrgent, UE::Cook::FCompletionCallback());
+			if (OutDiscoveredPackageNames)
+			{
+				OutDiscoveredPackageNames->Add(PackageFName);
+			}
 		}
 	}
 }
@@ -3976,11 +3841,15 @@ void UCookOnTheFlyServer::PreGarbageCollect()
 
 		FRequestQueue& RequestQueue = PackageDatas->GetRequestQueue();
 		TArray<FPackageData*> Requests;
-		Requests.Reserve(RequestQueue.Num());
-		while (!RequestQueue.IsEmpty())
+		Requests.Reserve(RequestQueue.ReadyRequestsNum());
+		while (!RequestQueue.IsReadyRequestsEmpty())
 		{
-			Requests.Add(RequestQueue.PopRequest());
+			Requests.Add(RequestQueue.PopReadyRequest());
 		}
+		// We are not looking at UnclusteredRequests or RequestClusters from the RequestQueue. These are supposed to 
+		// be quickly processed and should usually be empty, so failing to account for them will only rarely cause
+		// a performance issue (and will not cause a behavior issue)
+		
 		// Sort the cook requests by the packages which are loaded first
 		// then sort by the number of dependencies which are referenced by the package
 		// we want to process the packages with the highest dependencies so that they can
@@ -4607,18 +4476,15 @@ void FSaveCookedPackageContext::FinishPlatform()
 	bool bSuccessful = SavePackageResult.IsSuccessful();
 	bool bLocalReferencedOnlyByEditorOnlyData = SavePackageResult == ESavePackageResult::ReferencedOnlyByEditorOnlyData;
 
+	FAssetRegistryGenerator& Generator = *(COTFS.PlatformManager->GetPlatformData(TargetPlatform)->RegistryGenerator);
 	if (bPlatformSetupSuccessful)
 	{
-		FAssetRegistryGenerator* Generator = COTFS.PlatformManager->GetPlatformData(TargetPlatform)->RegistryGenerator.Get();
-		FAssetPackageData* AssetPackageData = Generator->GetAssetPackageData(Package->GetFName());
+		FAssetPackageData* AssetPackageData = Generator.GetAssetPackageData(Package->GetFName());
 		check(AssetPackageData);
 
 		FCbObject TargetDomainDependencies;
-		if (COTFS.bTargetDomainEnabled)
-		{
-			UE_SCOPED_HIERARCHICAL_COOKTIMER(TargetDomainDependencies);
-			TargetDomainDependencies = UE::TargetDomain::CollectDependenciesObject(Package, TargetPlatform, nullptr /* ErrorMessage */);
-		}
+		UE_SCOPED_HIERARCHICAL_COOKTIMER(TargetDomainDependencies);
+		TargetDomainDependencies = UE::TargetDomain::CollectDependenciesObject(Package, TargetPlatform, nullptr /* ErrorMessage */);
 
 		ICookedPackageWriter::FCommitPackageInfo Info;
 		Info.bSucceeded = bSuccessful;
@@ -4634,8 +4500,7 @@ void FSaveCookedPackageContext::FinishPlatform()
 	// Update asset registry
 	if (COTFS.CookByTheBookOptions)
 	{
-		FAssetRegistryGenerator* Generator = COTFS.PlatformManager->GetPlatformData(TargetPlatform)->RegistryGenerator.Get();
-		COTFS.UpdateAssetRegistryPackageData(Generator, *Package, SavePackageResult);
+		Generator.UpdateAssetRegistryPackageData(*Package, SavePackageResult);
 	}
 
 	// In success or failure we want to mark the package as cooked, unless the failure was due to being referenced only
@@ -4840,6 +4705,7 @@ void UCookOnTheFlyServer::Initialize( ECookMode::Type DesiredCookMode, ECookInit
 	MaxPreloadAllocated = 16;
 	DesiredSaveQueueLength = 8;
 	DesiredLoadQueueLength = 8;
+	RequestBatchSize = 16;
 
 	MinFreeUObjectIndicesBeforeGC = 100000;
 	GConfig->GetInt(TEXT("CookSettings"), TEXT("MinFreeUObjectIndicesBeforeGC"), MinFreeUObjectIndicesBeforeGC, GEditorIni);
@@ -4861,8 +4727,7 @@ void UCookOnTheFlyServer::Initialize( ECookMode::Type DesiredCookMode, ECookInit
 		bPreloadingEnabled = true;
 		FLinkerLoad::SetPreloadingEnabled(true);
 	}
-	GConfig->GetBool(TEXT("TargetDomain"), TEXT("TargetDomainEnabled"), bTargetDomainEnabled, GEditorIni);
-	GConfig->GetBool(TEXT("CookSettings"), TEXT("ExploreSoftReferencesOnStart"), bExploreSoftReferencesOnStart, GEditorIni);
+	GConfig->GetBool(TEXT("CookSettings"), TEXT("PreexploreDependenciesEnabled"), bPreexploreDependenciesEnabled, GEditorIni);
 
 	{
 		const FConfigSection* CacheSettings = GConfig->GetSectionPrivate(TEXT("CookPlatformDataCacheSettings"), false, true, GEditorIni);
@@ -4892,6 +4757,7 @@ void UCookOnTheFlyServer::Initialize( ECookMode::Type DesiredCookMode, ECookInit
 				UE_LOG(LogCook, Verbose, TEXT("Cooker startup package %s"), *It->GetName());
 			}
 		}
+		bHybridIterativeDebug = FParse::Param(FCommandLine::Get(), TEXT("hybriditerativedebug"));
 	}
 	
 	UE_LOG(LogCook, Display, TEXT("Mobile HDR setting %d"), IsMobileHDR());
@@ -5796,6 +5662,7 @@ FAsyncIODelete& UCookOnTheFlyServer::GetAsyncIODelete()
 
 void UCookOnTheFlyServer::PopulateCookedPackages(TArrayView<const ITargetPlatform* const> TargetPlatforms)
 {
+	using namespace UE::Cook;
 	TRACE_CPUPROFILER_EVENT_SCOPE(UCookOnTheFlyServer::PopulateCookedPackages);
 
 	struct FPreviousAssetPackageData
@@ -5819,109 +5686,124 @@ void UCookOnTheFlyServer::PopulateCookedPackages(TArrayView<const ITargetPlatfor
 				AssetPackageDataMap.Add(CookedPackageInfo.PackageName, &AssetPackageData);
 			}
 		}
+		TMap<FName, const FAssetPackageData*>& GetConstAssetPackageDataMap()
+		{
+			return reinterpret_cast<TMap<FName, const FAssetPackageData*>&>(AssetPackageDataMap);
+		}
 
 		TArray<FAssetPackageData> AssetPackageDatas;
-		TMap<FName, const FAssetPackageData*> AssetPackageDataMap;
+		TMap<FName, FAssetPackageData*> AssetPackageDataMap;
 	};
 
+	// TODO: NumPackagesIterativelySkipped is only counted for the first platform; to count all platforms we would
+	// have to check whether each one is already cooked.
+	bool bFirstPlatform = true;
+	COOK_STAT(DetailedCookStats::NumPackagesIterativelySkipped = 0);
 	for (const ITargetPlatform* TargetPlatform : TargetPlatforms)
 	{
-		UE::Cook::FPlatformData* PlatformData = PlatformManager->GetPlatformData(TargetPlatform);
-		FAssetRegistryGenerator* PlatformAssetRegistry = PlatformData->RegistryGenerator.Get();
-		UE::Cook::FCookSavePackageContext& CookSavePackageContext = FindOrCreateSaveContext(TargetPlatform);
+		FAssetRegistryGenerator& PlatformAssetRegistry = *(PlatformManager->GetPlatformData(TargetPlatform)->RegistryGenerator);
+		FCookSavePackageContext& CookSavePackageContext = FindOrCreateSaveContext(TargetPlatform);
 		ICookedPackageWriter& PackageWriter = *CookSavePackageContext.PackageWriter;
-		check(PlatformData);
-		check(PlatformAssetRegistry);
 		UE_LOG(LogCook, Display, TEXT("Populating cooked package(s) from %s package store on platform '%s'"),
 			*CookSavePackageContext.WriterDebugName, *TargetPlatform->PlatformName());
 
 		TArray<ICookedPackageWriter::FCookedPackageInfo> CookedPackages;
 		PackageWriter.GetCookedPackages(CookedPackages);
-		
+		// TODO: Implementation ReleasePreviousAssetRegistry for all PackageWriters and add an assertion here instead of handling null.
+		TUniquePtr<FAssetRegistryState> PreviousAssetRegistry(PackageWriter.ReleasePreviousAssetRegistry());
+
 		UE_LOG(LogCook, Display, TEXT("Found '%d' cooked package(s) in package store"), CookedPackages.Num());
 
 		FPreviousAssetPackageData PreviousAssetPackageData(CookedPackages);
 
-		FAssetRegistryGenerator::FComputeDifferenceOptions Options;
-		Options.bRecurseModifications		= true;
-		Options.bRecurseScriptModifications	= !IsCookFlagSet(ECookInitializationFlags::IgnoreScriptPackagesOutOfDate);
-		
-		FAssetRegistryGenerator::FAssetRegistryDifference Difference;
-		PlatformAssetRegistry->ComputePackageDifferences(Options, PreviousAssetPackageData.AssetPackageDataMap, Difference);
-
-		TArray<FName> PackagesToRemove;
-		PackagesToRemove.Reserve(Difference.ModifiedPackages.Num() + Difference.RemovedPackages.Num() + Difference.IdenticalUncookedPackages.Num());
-		for (FName PackageToRemove : Difference.ModifiedPackages)
-		{
-			PackagesToRemove.Add(PackageToRemove);
-		}
-		for (FName PackageToRemove : Difference.RemovedPackages)
+		if (bHybridIterativeEnabled)
 		{
-			PackagesToRemove.Add(PackageToRemove);
+			// HybridIterative does the equivalent operation of bRecurseModifications=true and bRecurseScriptModifications=true,
+			// but checks for out-of-datedness are done by FRequestCluster using the TargetDomainKey (which is built
+			// from dependencies), so we do not need to check for out-of-datedness here
 		}
-		for (FName PackageToRemove : Difference.IdenticalUncookedPackages)
-		{
-			PackagesToRemove.Add(PackageToRemove);
-		}
-		
-		UE_LOG(LogCook, Display, TEXT("Keeping '%d' and removing '%d' cooked package(s)"), Difference.IdenticalCookedPackages.Num(), PackagesToRemove.Num());
-		
-		PackageWriter.RemoveCookedPackages(PackagesToRemove);
-		
-		TArray<FName> PackagesToKeep;
-		for (const FName& IdenticalPackage : Difference.IdenticalCookedPackages)
+		else
 		{
-			// Mark this package as cooked so that we don't unnecessarily try to cook it again
-			UE::Cook::FPackageData* PackageData = PackageDatas->TryAddPackageDataByPackageName(IdenticalPackage);
-			if (PackageData)
+			// Without hybrid iterative, we use the AssetRegistry graph of dependencies to find out of date packages
+			// We also implement other legacy -iterate behaviors:
+			//  *) Remove modified packages from the PackageWriter in addition to the no-longer-exist packages
+			//  *) Skip packages that failed to cook on the previous cook
+			//  *) Cook all modified packages even if the requested cook packages don't reference them
+			FAssetRegistryGenerator::FComputeDifferenceOptions Options;
+			Options.bRecurseModifications = true;
+			Options.bRecurseScriptModifications = !IsCookFlagSet(ECookInitializationFlags::IgnoreScriptPackagesOutOfDate);
+			FAssetRegistryGenerator::FAssetRegistryDifference Difference;
+			PlatformAssetRegistry.ComputePackageDifferences(Options, PreviousAssetPackageData.GetConstAssetPackageDataMap(), Difference);
+
+			TArray<FName> PackagesToRemove;
+			PackagesToRemove.Reserve(Difference.ModifiedPackages.Num() + Difference.RemovedPackages.Num() + Difference.IdenticalUncookedPackages.Num());
+			for (FName PackageToRemove : Difference.ModifiedPackages)
 			{
-				PackageData->SetPlatformCooked(TargetPlatform, true /* bSucceeded */);
-				PackagesToKeep.Add(IdenticalPackage);
+				PackagesToRemove.Add(PackageToRemove);
 			}
-		}
-		for (FName UncookedPackage : Difference.IdenticalUncookedPackages)
-		{
-			// Mark this failed-to-cook package as cooked so that we don't unnecessarily try to cook it again
-			UE::Cook::FPackageData* PackageData = PackageDatas->TryAddPackageDataByPackageName(UncookedPackage);
-			if (PackageData)
+			for (FName PackageToRemove : Difference.RemovedPackages)
 			{
-				ensure(!PackageData->HasAnyCookedPlatforms({ TargetPlatform }, /* bIncludeFailed */ false));
-				PackageData->SetPlatformCooked(TargetPlatform, false /* bSucceeded */);
-				PackagesToKeep.Add(UncookedPackage);
+				PackagesToRemove.Add(PackageToRemove);
 			}
-		}
-		if (IsCookByTheBookMode())
-		{
-			for (const FName& RemovePackageName : Difference.ModifiedPackages)
+			for (FName PackageToRemove : Difference.IdenticalUncookedPackages)
+			{
+				PackagesToRemove.Add(PackageToRemove);
+			}
+
+			UE_LOG(LogCook, Display, TEXT("Keeping '%d' and removing '%d' cooked package(s)"), Difference.IdenticalCookedPackages.Num(), PackagesToRemove.Num());
+
+			PackageWriter.RemoveCookedPackages(PackagesToRemove);
+
+			if (bFirstPlatform)
 			{
-				// cook on the fly will requeue this package when it wants it, but for cook on the book we force cook the modified file
-				// so that the output set of packages is up to date (even if the user is currently cooking only a subset)
-				UE::Cook::FPackageData* PackageData = PackageDatas->TryAddPackageDataByPackageName(RemovePackageName);
+				COOK_STAT(DetailedCookStats::NumPackagesIterativelySkipped += Difference.IdenticalCookedPackages.Num());
+				bFirstPlatform = false;
+			}
+			for (const FName& IdenticalPackage : Difference.IdenticalCookedPackages)
+			{
+				// Mark this package as cooked so that we don't unnecessarily try to cook it again
+				FPackageData* PackageData = PackageDatas->TryAddPackageDataByPackageName(IdenticalPackage);
+				if (PackageData)
+				{
+					PackageData->SetPlatformCooked(TargetPlatform, true /* bSucceeded */);
+				}
+			}
+			for (FName UncookedPackage : Difference.IdenticalUncookedPackages)
+			{
+				// Mark this failed-to-cook package as cooked so that we don't unnecessarily try to cook it again
+				FPackageData* PackageData = PackageDatas->TryAddPackageDataByPackageName(UncookedPackage);
 				if (PackageData)
 				{
-					ExternalRequests->EnqueueUnique(UE::Cook::FFilePlatformRequest(PackageData->GetFileName(),
-						TConstArrayView<const ITargetPlatform*>{ TargetPlatform }));
+					ensure(!PackageData->HasAnyCookedPlatforms({ TargetPlatform }, /* bIncludeFailed */ false));
+					PackageData->SetPlatformCooked(TargetPlatform, false /* bSucceeded */);
+				}
+			}
+			if (IsCookByTheBookMode())
+			{
+				for (const FName& RemovePackageName : Difference.ModifiedPackages)
+				{
+					// cook on the fly will requeue this package when it wants it, but for cook on the book we force cook the modified file
+					// so that the output set of packages is up to date (even if the user is currently cooking only a subset)
+					FPackageData* PackageData = PackageDatas->TryAddPackageDataByPackageName(RemovePackageName);
+					if (PackageData)
+					{
+						ExternalRequests->EnqueueUnique(FFilePlatformRequest(PackageData->GetFileName(),
+							TConstArrayView<const ITargetPlatform*>{ TargetPlatform }));
+					}
 				}
 			}
 		}
 
-		TUniquePtr<FAssetRegistryState> PreviousAssetRegistry(PackageWriter.ReleasePreviousAssetRegistry());
-		if (PreviousAssetRegistry)
+		TArray<FName> TombstonePackages;
+		PlatformAssetRegistry.SetPreviousAssetRegistry(MoveTemp(PreviousAssetRegistry),
+			MoveTemp(PreviousAssetPackageData.AssetPackageDataMap), TombstonePackages);
+
+		if (bHybridIterativeEnabled)
 		{
-			// SetPreviousAssetRegistry is needed for UpdateKeptPackagesAssetData, which is called at the end
-			// of the cook for any packages that were kept from the previous cook and not cooked in this cook
-			PlatformAssetRegistry->SetPreviousAssetRegistry(MoveTemp(PreviousAssetRegistry));
-			PlatformAssetRegistry->UpdateKeptPackages(PackagesToKeep);
-		}
-		else
-		{
-			// TODO: Implementation ReleasePreviousAssetRegistry for all PackageWriters and add an assertion here instead of handling null.
-			PlatformAssetRegistry->UpdateKeptPackages(PackagesToKeep, [&PreviousAssetPackageData](const FName& PackageName, FAssetPackageData* Current)
-				{
-					const FAssetPackageData* Previous = PreviousAssetPackageData.AssetPackageDataMap.FindRef(PackageName);
-					Current->CookedHash = Previous->CookedHash;
-					Current->DiskSize = Previous->DiskSize;
-				});
+			// Remove packages that no longer exist in the WorkspaceDomain from the TargetDomain. We have to
+			// check this for all packages in the previous cook rather than just the currently referenced
+			// set because when packages are removed the references to them are usually also removed.
+			PackageWriter.RemoveCookedPackages(TombstonePackages);
 		}
 	}
 }
@@ -6967,7 +6849,7 @@ void UCookOnTheFlyServer::CookByTheBookFinished()
 			for (const ITargetPlatform* TargetPlatform : PlatformManager->GetSessionPlatforms())
 			{
 				UE::Cook::FPlatformData* PlatformData = PlatformManager->GetPlatformData(TargetPlatform);
-				FAssetRegistryGenerator& Generator = *PlatformData->RegistryGenerator.Get();
+				FAssetRegistryGenerator& Generator = *PlatformData->RegistryGenerator;
 				TArray<FName> CookedPackagesFilenames;
 				TArray<FName> IgnorePackageFilenames;
 
@@ -7280,9 +7162,32 @@ void UCookOnTheFlyServer::CancelAllQueues()
 		LoadPrepareQueue.PopFront()->SendToState(EPackageState::Idle, ESendFlags::QueueAdd);
 	}
 	FRequestQueue& RequestQueue = PackageDatas->GetRequestQueue();
-	while (!RequestQueue.IsEmpty())
+	FPackageDataSet& UnclusteredRequests = RequestQueue.GetUnclusteredRequests();
+	for (FPackageData* PackageData : UnclusteredRequests)
+	{
+		PackageData->SendToState(EPackageState::Idle, ESendFlags::QueueAdd);
+	}
+	UnclusteredRequests.Empty();
+	TRingBuffer<FRequestCluster>& RequestClusters = RequestQueue.GetRequestClusters();
+	for (FRequestCluster& RequestCluster : RequestClusters)
+	{
+		TArray<FPackageData*> RequestsToLoad;
+		TArray<FPackageData*> RequestsToDemote;
+		RequestCluster.ClearAndDetachOwnedPackageDatas(RequestsToLoad, RequestsToDemote);
+		for (FPackageData* PackageData : RequestsToLoad)
+		{
+			PackageData->SendToState(EPackageState::Idle, ESendFlags::QueueAdd);
+		}
+		for (FPackageData* PackageData : RequestsToDemote)
+		{
+			PackageData->SendToState(EPackageState::Idle, ESendFlags::QueueAdd);
+		}
+	}
+	RequestClusters.Empty();
+
+	while (!RequestQueue.IsReadyRequestsEmpty())
 	{
-		RequestQueue.PopRequest()->SendToState(EPackageState::Idle, ESendFlags::QueueAdd);
+		RequestQueue.PopReadyRequest()->SendToState(EPackageState::Idle, ESendFlags::QueueAdd);
 	}
 
 	SetLoadBusy(false);
@@ -7296,23 +7201,42 @@ void UCookOnTheFlyServer::ClearPlatformCookedData(const ITargetPlatform* TargetP
 	{
 		return;
 	}
-	ResetCookResults({ TargetPlatform });
+	ResetCook({ TPair<const ITargetPlatform*,bool>{TargetPlatform, true /* bResetResults */}});
 
 	FindOrCreatePackageWriter(TargetPlatform).RemoveCookedPackages();
 }
 
-void UCookOnTheFlyServer::ResetCookResults(TConstArrayView<const ITargetPlatform*> TargetPlatforms)
+void UCookOnTheFlyServer::ResetCook(TConstArrayView<TPair<const ITargetPlatform*, bool>> TargetPlatforms)
 {
 	for (UE::Cook::FPackageData* PackageData : *PackageDatas.Get())
 	{
-		PackageData->SetPlatformsNotCooked(TargetPlatforms);
+		for (const TPair<const ITargetPlatform*, bool>& Pair : TargetPlatforms)
+		{
+			const ITargetPlatform* TargetPlatform = Pair.Key;
+			UE::Cook::FPackageData::FPlatformData* PlatformData = PackageData->FindPlatformData(TargetPlatform);
+			if (PlatformData)
+			{
+				bool bResetResults = Pair.Value;
+				PlatformData->bExplored = false;
+				if (bResetResults)
+				{
+					PlatformData->bCookAttempted = false;
+					PlatformData->bCookSucceeded = false;
+				}
+			}
+		}
 	}
 
 	TArray<FName> PackageNames;
-	for (const ITargetPlatform* TargetPlatform : TargetPlatforms)
+	for (const TPair<const ITargetPlatform*, bool>& Pair : TargetPlatforms)
 	{
-		PackageNames.Reset();
-		PackageTracker->UnsolicitedCookedPackages.GetPackagesForPlatformAndRemove(TargetPlatform, PackageNames);
+		const ITargetPlatform* TargetPlatform = Pair.Key;
+		bool bResetResults = Pair.Value;
+		if (bResetResults)
+		{
+			PackageNames.Reset();
+			PackageTracker->UnsolicitedCookedPackages.GetPackagesForPlatformAndRemove(TargetPlatform, PackageNames);
+		}
 	}
 }
 
@@ -7364,6 +7288,15 @@ void UCookOnTheFlyServer::CreateSandboxFile()
 	SandboxFile->Initialize(&FPlatformFileManager::Get().GetPlatformFile(), *FString::Printf(TEXT("-sandbox=\"%s\""), *OutputDirectory));
 }
 
+void UCookOnTheFlyServer::SetBeginCookConfigSettings()
+{
+	GConfig->GetBool(TEXT("CookSettings"), TEXT("HybridIterativeEnabled"), bHybridIterativeEnabled, GEditorIni);
+	// TODO: HybridIterative is not yet implemented for DLC
+	bHybridIterativeEnabled &= !IsCookingDLC();
+	// HybridIterative uses TargetDomain storage of dependencies which is only implemented in ZenStore
+	bHybridIterativeEnabled &= IsUsingZenStore();
+}
+
 void UCookOnTheFlyServer::BeginCookSandbox(TConstArrayView<const ITargetPlatform*> TargetPlatforms)
 {
 #if OUTPUT_COOKTIMING
@@ -7378,9 +7311,9 @@ void UCookOnTheFlyServer::BeginCookSandbox(TConstArrayView<const ITargetPlatform
 		}
 
 		const bool bIsDiffOnly = FParse::Param(FCommandLine::Get(), TEXT("DIFFONLY"));
-		const bool bIsIterativeCook = IsCookFlagSet(ECookInitializationFlags::Iterative);
+		const bool bIterative = !FParse::Param(FCommandLine::Get(), TEXT("noiterate")) && (bHybridIterativeEnabled || IsCookFlagSet(ECookInitializationFlags::Iterative));
 		const bool bIsSharedIterativeCook = IsCookFlagSet(ECookInitializationFlags::IterateSharedBuild);
-		TArray<const ITargetPlatform*, TInlineAllocator<ExpectedMaxNumPlatforms>> ResetPlatforms;
+		TArray<TPair<const ITargetPlatform*,bool>, TInlineAllocator<ExpectedMaxNumPlatforms>> ResetPlatforms;
 		TArray<const ITargetPlatform*, TInlineAllocator<ExpectedMaxNumPlatforms>> PopulatePlatforms;
 
 		for (const ITargetPlatform* TargetPlatform : TargetPlatforms)
@@ -7388,7 +7321,7 @@ void UCookOnTheFlyServer::BeginCookSandbox(TConstArrayView<const ITargetPlatform
 			UE::Cook::FPlatformData* PlatformData = PlatformManager->GetPlatformData(TargetPlatform);
 			ICookedPackageWriter& PackageWriter = FindOrCreatePackageWriter(TargetPlatform);
 			bool bIterateSharedBuild = false;
-			if (bIsIterativeCook && bIsSharedIterativeCook && !PlatformData->bIsSandboxInitialized)
+			if (bIterative && bIsSharedIterativeCook && !PlatformData->bIsSandboxInitialized)
 			{
 				// see if the shared build is newer then the current cooked content in the local directory
 				FString SharedCookedAssetRegistry = FPaths::Combine(*FPaths::ProjectSavedDir(), TEXT("SharedIterativeBuild"),
@@ -7422,7 +7355,7 @@ void UCookOnTheFlyServer::BeginCookSandbox(TConstArrayView<const ITargetPlatform
 				// When looking for deterministic cooking differences in cooked packages, don't delete the packages on disk
 				bShouldClearCookedContent = false;
 			}
-			else if (bIsIterativeCook || PlatformData->bIsSandboxInitialized)
+			else if (bIterative || PlatformData->bIsSandboxInitialized)
 			{
 				if (!bIsIniSettingsOutOfDate)
 				{
@@ -7453,14 +7386,15 @@ void UCookOnTheFlyServer::BeginCookSandbox(TConstArrayView<const ITargetPlatform
 			CookInfo.CookMode = IsCookOnTheFlyMode() ? ICookedPackageWriter::FCookInfo::CookOnTheFlyMode : ICookedPackageWriter::FCookInfo::CookByTheBookMode;
 			CookInfo.bCleanBuild = false;
 			CookInfo.bIterateSharedBuild = bIterateSharedBuild && !bShouldClearCookedContent;
+			bool bResetResults = false;
 			if (bShouldClearCookedContent)
 			{
 				CookInfo.bCleanBuild = true;
-				ResetPlatforms.Add(TargetPlatform);
+				bResetResults = true;
 			}
 			else if (!PlatformData->bIsSandboxInitialized)
 			{
-				ResetPlatforms.Add(TargetPlatform);
+				bResetResults = true;
 				// Reset but don't populate if we are looking for deterministic cooking differences; start from an empty list of cooked packages
 				if (!bIsDiffOnly)
 				{
@@ -7468,13 +7402,12 @@ void UCookOnTheFlyServer::BeginCookSandbox(TConstArrayView<const ITargetPlatform
 				}
 			}
 			PackageWriter.BeginCook(CookInfo);
+			PlatformData->bCleanBuild = CookInfo.bCleanBuild;
 			PlatformData->bIsSandboxInitialized = true;
+			ResetPlatforms.Emplace(TargetPlatform, bResetResults);
 		}
 
-		if (ResetPlatforms.Num())
-		{
-			ResetCookResults(ResetPlatforms);
-		}
+		ResetCook(ResetPlatforms);
 		if (PopulatePlatforms.Num())
 		{
 			PopulateCookedPackages(PopulatePlatforms);
@@ -7685,6 +7618,7 @@ void UCookOnTheFlyServer::StartCookByTheBook( const FCookByTheBookStartupOptions
 		UE_LOG(LogCook, Warning, TEXT("Setting bSkipSoftReferences to true since bSkipHardReferences is true and skipping hard references requires skipping soft references."));
 		CookByTheBookOptions->bSkipSoftReferences = true;
 	}
+	SetBeginCookConfigSettings();
 
 	GenerateAssetRegistry();
 	if (!IsCookingInEditor())
@@ -7779,8 +7713,8 @@ void UCookOnTheFlyServer::StartCookByTheBook( const FCookByTheBookStartupOptions
 		{
 			for (const ITargetPlatform* TargetPlatform : TargetPlatforms)
 			{
-				FAssetRegistryGenerator* RegistryGenerator = PlatformManager->GetPlatformData(TargetPlatform)->RegistryGenerator.Get();
-				RegistryGenerator->RegisterChunkDataGenerator(MakeShared<FLocalizationChunkDataGenerator>(RegistryGenerator->GetPakchunkIndex(PackagingSettings->LocalizationTargetCatchAllChunkId), LocalizationTargetsToChunk, AllCulturesToCook));
+				FAssetRegistryGenerator& RegistryGenerator = *(PlatformManager->GetPlatformData(TargetPlatform)->RegistryGenerator);
+				RegistryGenerator.RegisterChunkDataGenerator(MakeShared<FLocalizationChunkDataGenerator>(RegistryGenerator.GetPakchunkIndex(PackagingSettings->LocalizationTargetCatchAllChunkId), LocalizationTargetsToChunk, AllCulturesToCook));
 			}
 		}
 	}
@@ -7812,7 +7746,11 @@ void UCookOnTheFlyServer::StartCookByTheBook( const FCookByTheBookStartupOptions
 	}
 	
 	FindOrCreateSaveContexts(TargetPlatforms);
-	BeginCookSandbox(TargetPlatforms); // This will either delete the sandbox or iteratively clean it
+	// TODO: Fix the elements of StartCookByTheBook that rely on BeginCookSandbox being called early, and remove this call to BeginCookSandbox
+	if (!bHybridIterativeEnabled)
+	{
+		BeginCookSandbox(TargetPlatforms); // This will either delete the sandbox or iteratively clean it
+	}
 	InitializeTargetPlatforms(TargetPlatforms);
 
 	if (CurrentCookMode == ECookMode::CookByTheBook && !IsCookFlagSet(ECookInitializationFlags::Iterative))
@@ -7956,8 +7894,8 @@ void UCookOnTheFlyServer::StartCookByTheBook( const FCookByTheBookStartupOptions
 	{
 		for (const ITargetPlatform* TargetPlatform : TargetPlatforms)
 		{
-			FAssetRegistryGenerator* RegistryGenerator = PlatformManager->GetPlatformData(TargetPlatform)->RegistryGenerator.Get();
-			RegistryGenerator->RegisterChunkDataGenerator(MakeShared<FShaderLibraryChunkDataGenerator>(TargetPlatform));
+			FAssetRegistryGenerator& RegistryGenerator = *(PlatformManager->GetPlatformData(TargetPlatform)->RegistryGenerator);
+			RegistryGenerator.RegisterChunkDataGenerator(MakeShared<FShaderLibraryChunkDataGenerator>(TargetPlatform));
 		}
 	}
 
@@ -8007,7 +7945,11 @@ void UCookOnTheFlyServer::StartCookByTheBook( const FCookByTheBookStartupOptions
 	
 	CollectFilesToCook(FilesInPath, CookMaps, CookDirectories, IniMapSections, CookOptions, TargetPlatforms, GameDefaultObjects);
 
-	// Add string asset packages after collecting files, to avoid accidentally activating the behavior to cook all maps if none are specified
+	// Add soft/hard startup references after collecting requested files and handling empty requests
+	if (bPreexploreDependenciesEnabled && !CookByTheBookOptions->bSkipHardReferences && !CookByTheBookOptions->bFullLoadAndSave)
+	{
+		ProcessUnsolicitedPackages(&FilesInPath);
+	}
 	for (FName SoftObjectPackage : StartupSoftObjectPackages)
 	{
 		TMap<FName, FName> RedirectedPaths;
@@ -8047,7 +7989,7 @@ void UCookOnTheFlyServer::StartCookByTheBook( const FCookByTheBookStartupOptions
 		UE_SCOPED_HIERARCHICAL_COOKTIMER(GenerateLongPackageName);
 		GenerateLongPackageNames(FilesInPath);
 	}
-	// add all the files for the requested platform to the cook list
+	// add all the files to the cook list for the requested platforms
 	for ( const FName& FileFName : FilesInPath )
 	{
 		if (FileFName == NAME_None)
@@ -8068,7 +8010,6 @@ void UCookOnTheFlyServer::StartCookByTheBook( const FCookByTheBookStartupOptions
 		}	
 	}
 
-
 	if (!IsCookingDLC())
 	{
 		// if we are not cooking dlc then basedOnRelease version just needs to make sure that we cook all the packages which are in the previous release (as well as the new ones)
@@ -8095,6 +8036,16 @@ void UCookOnTheFlyServer::StartCookByTheBook( const FCookByTheBookStartupOptions
 			}
 		}
 	}
+
+	if (bHybridIterativeEnabled)
+	{
+		BeginCookSandbox(TargetPlatforms);
+	}
+	if (bHybridIterativeDebug)
+	{
+		// Discoveries during the processing of the initial cluster are expected, so LogDiscoveredPackages must be off.
+		PackageDatas->SetLogDiscoveredPackages(false);
+	}
 }
 
 TArray<FName> UCookOnTheFlyServer::GetNeverCookPackageFileNames(TArrayView<const FString> ExtraNeverCookDirectories)
@@ -8846,8 +8797,9 @@ uint32 UCookOnTheFlyServer::FullLoadAndSave(uint32& CookedPackageCount)
 							}
 						}
 
-						FAssetRegistryGenerator* Generator = PlatformManager->GetPlatformData(Target)->RegistryGenerator.Get();
-						UpdateAssetRegistryPackageData(Generator, *Package, SaveResult);
+						// Update asset registry
+						FAssetRegistryGenerator& Generator = *(PlatformManager->GetPlatformData(Target)->RegistryGenerator);
+						Generator.UpdateAssetRegistryPackageData(*Package, SaveResult);
 
 						if (SaveResult.IsSuccessful())
 						{
diff --git a/Engine/Source/Editor/UnrealEd/Private/Cooker/CookPackageData.cpp b/Engine/Source/Editor/UnrealEd/Private/Cooker/CookPackageData.cpp
--- a/Engine/Source/Editor/UnrealEd/Private/Cooker/CookPackageData.cpp
+++ b/Engine/Source/Editor/UnrealEd/Private/Cooker/CookPackageData.cpp
@@ -5,8 +5,9 @@
 #include "Algo/AnyOf.h"
 #include "Algo/Count.h"
 #include "AssetCompilingManager.h"
+#include "Cooker/CookPlatformManager.h"
+#include "Cooker/CookRequestCluster.h"
 #include "CookOnTheSide/CookOnTheFlyServer.h"
-#include "CookPlatformManager.h"
 #include "Containers/StringView.h"
 #include "EditorDomain/EditorDomain.h"
 #include "HAL/FileManager.h"
@@ -30,8 +31,7 @@ namespace Cook
 	//////////////////////////////////////////////////////////////////////////
 	// FPackageData
 	FPackageData::FPlatformData::FPlatformData()
-		: bRequested(false), bCookAttempted(false), bCookSucceeded(false), bTargetDomainLookupAttempted(false)
-		, bExistsInTargetDomain(false)
+		: bRequested(false), bCookAttempted(false), bCookSucceeded(false), bExplored(false)
 	{
 	}
 
@@ -128,7 +128,41 @@ namespace Cook
 	{
 		for (const TPair<const ITargetPlatform*, FPlatformData>& Pair : PlatformDatas)
 		{
-			if (!Pair.Value.bCookAttempted || (!bAllowFailedCooks && !Pair.Value.bCookSucceeded))
+			if (Pair.Value.bRequested && (!Pair.Value.bCookAttempted || (!bAllowFailedCooks && !Pair.Value.bCookSucceeded)))
+			{
+				return false;
+			}
+		}
+		return true;
+	}
+
+	bool FPackageData::AreAllRequestedPlatformsExplored() const
+	{
+		for (const TPair<const ITargetPlatform*, FPlatformData>& Pair : PlatformDatas)
+		{
+			if (Pair.Value.bRequested && !Pair.Value.bExplored)
+			{
+				return false;
+			}
+		}
+		return true;
+	}
+
+	bool FPackageData::HasAllExploredPlatforms(const TArrayView<const ITargetPlatform* const>& Platforms) const
+	{
+		if (Platforms.Num() == 0)
+		{
+			return true;
+		}
+		if (PlatformDatas.Num() == 0)
+		{
+			return false;
+		}
+
+		for (const ITargetPlatform* QueryPlatform : Platforms)
+		{
+			const FPlatformData* PlatformData = FindPlatformData(QueryPlatform);
+			if (!PlatformData || !PlatformData->bExplored)
 			{
 				return false;
 			}
@@ -147,7 +181,7 @@ namespace Cook
 	}
 
 	void FPackageData::UpdateRequestData(const TConstArrayView<const ITargetPlatform*> InRequestedPlatforms,
-		bool bInIsUrgent, FCompletionCallback&& InCompletionCallback, ESendFlags SendFlags)
+		bool bInIsUrgent, FCompletionCallback&& InCompletionCallback)
 	{
 		if (IsInProgress())
 		{
@@ -165,20 +199,19 @@ namespace Cook
 				// Send back to the Request state (canceling any current operations) and then add the new platforms
 				if (GetState() != EPackageState::Request)
 				{
-					check(SendFlags == ESendFlags::QueueAddAndRemove);
 					SendToState(EPackageState::Request, ESendFlags::QueueAddAndRemove);
 				}
 				SetPlatformsRequested(InRequestedPlatforms, true);
 			}
-			else if (bUrgencyChanged && SendFlags == ESendFlags::QueueAddAndRemove)
+			else if (bUrgencyChanged)
 			{
-				SendToState(GetState(), SendFlags);
+				SendToState(GetState(), ESendFlags::QueueAddAndRemove);
 			}
 		}
 		else if (InRequestedPlatforms.Num() > 0)
 		{
 			SetRequestData(InRequestedPlatforms, bInIsUrgent, MoveTemp(InCompletionCallback));
-			SendToState(EPackageState::Request, SendFlags);
+			SendToState(EPackageState::Request, ESendFlags::QueueAddAndRemove);
 		}
 	}
 
@@ -225,10 +258,12 @@ namespace Cook
 	{
 		bool bHasAnyOthers = false;
 		bool bModified = false;
+		bool bExists = false;
 		for (TPair<const ITargetPlatform*, FPlatformData>& Pair : PlatformDatas)
 		{
 			if (Pair.Key == TargetPlatform)
 			{
+				bExists = true;
 				bModified = bModified | (Pair.Value.bCookAttempted == false);
 				Pair.Value.bCookAttempted = true;
 				Pair.Value.bCookSucceeded = bSucceeded;
@@ -238,6 +273,13 @@ namespace Cook
 				bHasAnyOthers = bHasAnyOthers | (Pair.Value.bCookAttempted != false);
 			}
 		}
+		if (!bExists)
+		{
+			FPlatformData& Value = PlatformDatas.FindOrAdd(TargetPlatform);
+			Value.bCookAttempted = true;
+			Value.bCookSucceeded = bSucceeded;
+			bModified = true;
+		}
 		if (bModified && !bHasAnyOthers)
 		{
 			PackageDatas.GetMonitor().OnFirstCookedPlatformAdded(*this);
@@ -295,6 +337,21 @@ namespace Cook
 		return PlatformDatas;
 	}
 
+	FPackageData::FPlatformData& FPackageData::FindOrAddPlatformData(const ITargetPlatform* TargetPlatform)
+	{
+		return PlatformDatas.FindOrAdd(TargetPlatform);
+	}
+
+	FPackageData::FPlatformData* FPackageData::FindPlatformData(const ITargetPlatform* TargetPlatform)
+	{
+		return PlatformDatas.Find(TargetPlatform);
+	}
+
+	const FPackageData::FPlatformData* FPackageData::FindPlatformData(const ITargetPlatform* TargetPlatform) const
+	{
+		return PlatformDatas.Find(TargetPlatform);
+	}
+
 	bool FPackageData::HasAnyCookedPlatform() const
 	{
 		return Algo::AnyOf(PlatformDatas,
@@ -604,6 +661,10 @@ namespace Cook
 
 	void FPackageData::OnExitIdle()
 	{
+		if (PackageDatas.GetLogDiscoveredPackages())
+		{
+			UE_LOG(LogCook, Warning, TEXT("Missing dependency: Package %s discovered after initial dependency search."), *WriteToString<256>(PackageName));
+		}
 	}
 
 	void FPackageData::OnEnterRequest()
@@ -1850,21 +1911,42 @@ namespace Cook
 
 	uint32 FRequestQueue::Num() const
 	{
-		return NormalRequests.Num() + UrgentRequests.Num();
+		uint32 Count = UnclusteredRequests.Num() + ReadyRequestsNum();
+		for (const FRequestCluster& RequestCluster : RequestClusters)
+		{
+			Count += RequestCluster.NumPackageDatas();
+		}
+		return Count;
 	}
 
 	bool FRequestQueue::Contains(const FPackageData* InPackageData) const
 	{
 		FPackageData* PackageData = const_cast<FPackageData*>(InPackageData);
-		return NormalRequests.Contains(PackageData) || UrgentRequests.Contains(PackageData);
+		if (UnclusteredRequests.Contains(PackageData) || NormalRequests.Contains(PackageData) || UrgentRequests.Contains(PackageData))
+		{
+			return true;
+		}
+		for (const FRequestCluster& RequestCluster : RequestClusters)
+		{
+			if (RequestCluster.Contains(PackageData))
+			{
+				return true;
+			}
+		}
+		return false;
 	}
 
 	uint32 FRequestQueue::RemoveRequest(FPackageData* PackageData)
 	{
-		uint32 OriginalNum = NormalRequests.Num() + UrgentRequests.Num();
+		uint32 OriginalNum = Num();
+		UnclusteredRequests.Remove(PackageData);
 		NormalRequests.Remove(PackageData);
 		UrgentRequests.Remove(PackageData);
-		uint32 Result = OriginalNum - (NormalRequests.Num() + UrgentRequests.Num());
+		for (FRequestCluster& RequestCluster : RequestClusters)
+		{
+			RequestCluster.RemovePackageData(PackageData);
+		}
+		uint32 Result = OriginalNum - Num();
 		check(Result == 0 || Result == 1);
 		return Result;
 	}
@@ -1874,7 +1956,17 @@ namespace Cook
 		return RemoveRequest(PackageData);
 	}
 
-	FPackageData* FRequestQueue::PopRequest()
+	bool FRequestQueue::IsReadyRequestsEmpty() const
+	{
+		return ReadyRequestsNum() == 0;
+	}
+
+	uint32 FRequestQueue::ReadyRequestsNum() const
+	{
+		return UrgentRequests.Num() + NormalRequests.Num();
+	}
+
+	FPackageData* FRequestQueue::PopReadyRequest()
 	{
 		for (FPackageData* PackageData : UrgentRequests)
 		{
@@ -1890,6 +1982,18 @@ namespace Cook
 	}
 
 	void FRequestQueue::AddRequest(FPackageData* PackageData, bool bForceUrgent)
+	{
+		if (!PackageData->AreAllRequestedPlatformsExplored())
+		{
+			UnclusteredRequests.Add(PackageData);
+		}
+		else
+		{
+			AddReadyRequest(PackageData, bForceUrgent);
+		}
+	}
+
+	void FRequestQueue::AddReadyRequest(FPackageData* PackageData, bool bForceUrgent)
 	{
 		if (bForceUrgent || PackageData->GetIsUrgent())
 		{
diff --git a/Engine/Source/Editor/UnrealEd/Private/Cooker/CookPackageData.h b/Engine/Source/Editor/UnrealEd/Private/Cooker/CookPackageData.h
--- a/Engine/Source/Editor/UnrealEd/Private/Cooker/CookPackageData.h
+++ b/Engine/Source/Editor/UnrealEd/Private/Cooker/CookPackageData.h
@@ -29,13 +29,12 @@ class UCookOnTheFlyServer;
 class UObject;
 class UPackage;
 
-namespace UE
-{
-namespace Cook
+namespace UE::Cook
 {
+	class FPackageDataQueue;
+	class FRequestCluster;
 	struct FGeneratorPackage;
 	struct FPackageDataMonitor;
-	class FPackageDataQueue;
 	struct FPendingCookedPlatformDataCancelManager;
 
 	extern const TCHAR* GeneratedPackageSubPath;
@@ -125,8 +124,7 @@ namespace Cook
 			bool bRequested : 1;
 			bool bCookAttempted : 1;
 			bool bCookSucceeded : 1;
-			bool bTargetDomainLookupAttempted : 1;
-			bool bExistsInTargetDomain : 1;
+			bool bExplored : 1;
 		};
 
 		FPackageData(FPackageDatas& PackageDatas, const FName& InPackageName, const FName& InFileName);
@@ -175,11 +173,19 @@ namespace Cook
 		 */
 		bool HasAllRequestedPlatforms(const TArrayView<const ITargetPlatform* const>& Platforms) const;
 
-		/** Return whether all Platforms that have been requested in this PackageData have also been cooked
+		/**
+		 * Return whether all Platforms requested have also been cooked.
 		 * 
 		 * @param bAllowFailedCooks Iff false, failed cooks cause function to return false.
 		 */
 		bool AreAllRequestedPlatformsCooked(bool bAllowFailedCooks) const;
+
+		/** Return whether all Platforms requested have also been explored for transitive references. */
+		bool AreAllRequestedPlatformsExplored() const;
+
+		/** Return true if and only if every element of Platforms has been explored. Returns true if Platforms is empty. */
+		bool HasAllExploredPlatforms(const TArrayView<const ITargetPlatform* const>& Platforms) const;
+
 		/**
 		 * Get the flag for whether this InProgress PackageData has been marked as an urgent request
 		 * (e.g. because it has been requested from the game client during cookonthefly.).
@@ -193,11 +199,9 @@ namespace Cook
 		 * If the PackageData is not in progress, just calls SetRequestData; skipped if InRequestedPlatforms is empty.
 		 * If the PackageData is in progress it is demoted back to an earlier state if necessary.
 		 * RequestData is set to the union of the previous and new values.
-		 
-		 * @param SendFlags If the PackageData needs to change state, specify whether to remove/add to state containers.
 		 */
 		void UpdateRequestData(const TConstArrayView<const ITargetPlatform*> InRequestedPlatforms, bool bInIsUrgent,
-			FCompletionCallback&& InCompletionCallback, ESendFlags SendFlags = ESendFlags::QueueAddAndRemove);
+			FCompletionCallback&& InCompletionCallback);
 		/*
 		 * Set the given data onto the existing request data.
 		 * It is invalid to call with empty InRequestedPlatforms.
@@ -230,6 +234,15 @@ namespace Cook
 		/** Access the information about platforms interacted with by *this. */
 		const TSortedMap<const ITargetPlatform*, FPlatformData>& GetPlatformDatas() const;
 
+		/** Add a platform if not already existing and return a writable pointer to its flags. */
+		FPlatformData& FindOrAddPlatformData(const ITargetPlatform* TargetPlatform);
+
+		/** Find a platform if it exists and return a writable pointer to its flags. */
+		FPlatformData* FindPlatformData(const ITargetPlatform* TargetPlatform);
+
+		/** Find a platform if it exists and return a writable pointer to its flags. */
+		const FPlatformData* FindPlatformData(const ITargetPlatform* TargetPlatform) const;
+
 		/** Return true if and only if at least one platform has been cooked. */
 		bool HasAnyCookedPlatform() const;
 		/**
@@ -252,6 +265,7 @@ namespace Cook
 		 * returns ECookResult::Unseen, otherwise returns Succeeded or Failed depending on its success flag.
 		 */
 		ECookResult GetCookResults(const ITargetPlatform* Platform) const;
+
 		/** Empties and then sets OutPlatforms to contain all platforms that have been requested and not cook-attempted. */
 		template <typename ArrayType>
 		void GetUncookedPlatforms(ArrayType& OutPlatforms)
@@ -718,11 +732,6 @@ namespace Cook
 		using TRingBuffer<FPackageData*>::TRingBuffer;
 	};
 
-	class FPackageDataSet : public TFastPointerSet<FPackageData*>
-	{
-		using TFastPointerSet<FPackageData*>::TFastPointerSet;
-	};
-
 	/**
 	 * A monitor class held by an FPackageDatas to provide reporting and decision making based on aggregated-data
 	 * across all InProgress or completed FPackageData.
@@ -774,7 +783,8 @@ namespace Cook
 
 	/**
 	 * A container for FPackageDatas in the Request state. This container needs to support fast find and remove,
-	 * and FIFO AddRequest/PopRequest that is overridden for urgent requests to push them to the front.
+	 * RequestClusters, staging for packages not yet in request clusters, and a FIFO for ready requests
+	 * using AddRequest/PopRequest that is overridden for urgent requests to push them to the front.
 	 */
 	class FRequestQueue
 	{
@@ -785,11 +795,19 @@ namespace Cook
 		bool Contains(const FPackageData* PackageData) const;
 		void Empty();
 
-		FPackageData* PopRequest();
 		void AddRequest(FPackageData* PackageData, bool bForceUrgent=false);
+
+		uint32 ReadyRequestsNum() const;
+		bool IsReadyRequestsEmpty() const;
+		FPackageData* PopReadyRequest();
+		void AddReadyRequest(FPackageData* PackageData, bool bForceUrgent = false);
 		uint32 RemoveRequest(FPackageData* PackageData);
 
+		FPackageDataSet& GetUnclusteredRequests() { return UnclusteredRequests; }
+		TRingBuffer<FRequestCluster>& GetRequestClusters() { return RequestClusters; }
 	private:
+		FPackageDataSet UnclusteredRequests;
+		TRingBuffer<FRequestCluster> RequestClusters;
 		FPackageDataSet UrgentRequests;
 		FPackageDataSet NormalRequests;
 	};
@@ -955,6 +973,9 @@ namespace Cook
 		/** Swap all ITargetPlatform* stored on this instance according to the mapping in @param Remap. */
 		void RemapTargetPlatforms(const TMap<ITargetPlatform*, ITargetPlatform*>& Remap);
 
+		/** Get/Set whether to log newly discovered packages to debug missing transitive references. */
+		void SetLogDiscoveredPackages(bool bLog) { bLogDiscoveredPackages = bLog; }
+		bool GetLogDiscoveredPackages() const { return bLogDiscoveredPackages; }
 	private:
 		/**
 		 * Construct a new FPackageData with the given PackageName and FileName and store references to it in the maps.
@@ -974,6 +995,7 @@ namespace Cook
 		FPackageDataQueue LoadReadyQueue;
 		FPackageDataQueue SaveQueue;
 		UCookOnTheFlyServer& CookOnTheFlyServer;
+		bool bLogDiscoveredPackages = false;
 	};
 
 	/**
@@ -990,5 +1012,4 @@ namespace Cook
 		FPackageData& PackageData;
 #endif
 	};
-}
-}
\ No newline at end of file
+} // namespace UE::Cook
\ No newline at end of file
diff --git a/Engine/Source/Editor/UnrealEd/Private/Cooker/CookPlatformManager.h b/Engine/Source/Editor/UnrealEd/Private/Cooker/CookPlatformManager.h
--- a/Engine/Source/Editor/UnrealEd/Private/Cooker/CookPlatformManager.h
+++ b/Engine/Source/Editor/UnrealEd/Private/Cooker/CookPlatformManager.h
@@ -47,6 +47,9 @@ namespace Cook
 		/* Whether BeginCookSandbox has been called for this platform.  If we have already initialized the sandbox we can take a faster refresh path on future sessions. */
 		bool bIsSandboxInitialized = false;
 
+		/* Whether BeginCookSandbox specified a clean build for this platform. */
+		bool bCleanBuild = false;
+
 		/*
 		 * The last FPlatformTime::GetSeconds() at which this platform was requested in a CookOnTheFly request.
 		 * If equal to 0, the platform was not requested in a CookOnTheFly since the last clear.
diff --git a/Engine/Source/Editor/UnrealEd/Private/Cooker/CookProfiling.cpp b/Engine/Source/Editor/UnrealEd/Private/Cooker/CookProfiling.cpp
--- a/Engine/Source/Editor/UnrealEd/Private/Cooker/CookProfiling.cpp
+++ b/Engine/Source/Editor/UnrealEd/Private/Cooker/CookProfiling.cpp
@@ -187,9 +187,11 @@ namespace DetailedCookStats
 	int32 PeakLoadQueueSize = 0;
 	int32 PeakSaveQueueSize = 0;
 	uint32 NumPreloadedDependencies = 0;
+	uint32 NumPackagesIterativelySkipped = 0;
 	FCookStatsManager::FAutoRegisterCallback RegisterCookOnTheFlyServerStats([](FCookStatsManager::AddStatFuncRef AddStat)
 		{
 			AddStat(TEXT("Package.Load"), FCookStatsManager::CreateKeyValueArray(TEXT("NumPreloadedDependencies"), NumPreloadedDependencies));
+			AddStat(TEXT("Package.Save"), FCookStatsManager::CreateKeyValueArray(TEXT("NumPackagesIterativelySkipped"), NumPackagesIterativelySkipped));
 			AddStat(TEXT("CookOnTheFlyServer"), FCookStatsManager::CreateKeyValueArray(TEXT("PeakRequestQueueSize"), PeakRequestQueueSize));
 			AddStat(TEXT("CookOnTheFlyServer"), FCookStatsManager::CreateKeyValueArray(TEXT("PeakLoadQueueSize"), PeakLoadQueueSize));
 			AddStat(TEXT("CookOnTheFlyServer"), FCookStatsManager::CreateKeyValueArray(TEXT("PeakSaveQueueSize"), PeakSaveQueueSize));
diff --git a/Engine/Source/Editor/UnrealEd/Private/Cooker/CookProfiling.h b/Engine/Source/Editor/UnrealEd/Private/Cooker/CookProfiling.h
--- a/Engine/Source/Editor/UnrealEd/Private/Cooker/CookProfiling.h
+++ b/Engine/Source/Editor/UnrealEd/Private/Cooker/CookProfiling.h
@@ -95,6 +95,7 @@ namespace DetailedCookStats
 
 	// Stats tracked through FAutoRegisterCallback
 	extern uint32 NumPreloadedDependencies;
+	extern uint32 NumPackagesIterativelySkipped;
 	extern int32 PeakRequestQueueSize;
 	extern int32 PeakLoadQueueSize;
 	extern int32 PeakSaveQueueSize;
diff --git a/Engine/Source/Editor/UnrealEd/Private/Cooker/CookRequestCluster.cpp b/Engine/Source/Editor/UnrealEd/Private/Cooker/CookRequestCluster.cpp
--- a/Engine/Source/Editor/UnrealEd/Private/Cooker/CookRequestCluster.cpp
+++ b/Engine/Source/Editor/UnrealEd/Private/Cooker/CookRequestCluster.cpp
@@ -2,11 +2,13 @@
 
 #include "Cooker/CookRequestCluster.h"
 
-#include "Algo/Unique.h"
 #include "Algo/Sort.h"
+#include "Algo/Unique.h"
 #include "Cooker/CookPackageData.h"
+#include "Cooker/CookPlatformManager.h"
 #include "Cooker/CookProfiling.h"
 #include "Cooker/CookRequests.h"
+#include "Cooker/CookTypes.h"
 #include "Cooker/PackageNameCache.h"
 #include "Cooker/PackageTracker.h"
 #include "CookOnTheSide/CookOnTheFlyServer.h"
@@ -27,6 +29,15 @@ FRequestCluster::FRequestCluster(UCookOnTheFlyServer& COTFS, TArray<const ITarge
 {
 }
 
+FRequestCluster::FRequestCluster(UCookOnTheFlyServer& COTFS, TConstArrayView<const ITargetPlatform*> InPlatforms)
+	: Platforms(InPlatforms)
+	, PackageDatas(*COTFS.PackageDatas)
+	, AssetRegistry(*IAssetRegistry::Get())
+	, PackageNameCache(COTFS.GetPackageNameCache())
+	, PackageTracker(*COTFS.PackageTracker)
+{
+}
+
 void FRequestCluster::AddClusters(UCookOnTheFlyServer& COTFS, TArray<FFilePlatformRequest>&& InRequests, bool bRequestsAreUrgent,
 	TRingBuffer<FRequestCluster>& OutClusters)
 {
@@ -51,9 +62,9 @@ void FRequestCluster::AddClusters(UCookOnTheFlyServer& COTFS, TArray<FFilePlatfo
 
 	UE::Cook::FRequestCluster* MRUCluster = FindOrAddClusterForPlatforms(MoveTemp(InRequests[0].GetPlatforms()));
 	// The usual case is all platforms are the same, so reserve the first Cluster's size assuming it will get all requests
-	MRUCluster->Requests.Reserve(InRequests.Num());
+	MRUCluster->InRequests.Reserve(InRequests.Num());
 	// Add the first Request to it. Since we've already taken away the Platforms from the first request, we have to handle it specially.
-	MRUCluster->Requests.Add(FRequest(MoveTemp(InRequests[0]), bRequestsAreUrgent));
+	MRUCluster->InRequests.Add(FFileNameRequest(MoveTemp(InRequests[0]), bRequestsAreUrgent));
 
 	for (FFilePlatformRequest& Request : TArrayView<FFilePlatformRequest>(InRequests).Slice(1, InRequests.Num() - 1))
 	{
@@ -67,7 +78,7 @@ void FRequestCluster::AddClusters(UCookOnTheFlyServer& COTFS, TArray<FFilePlatfo
 			Request.GetPlatforms().Reset();
 		}
 
-		MRUCluster->Requests.Add(FRequest(MoveTemp(Request), bRequestsAreUrgent));
+		MRUCluster->InRequests.Add(FFileNameRequest(MoveTemp(Request), bRequestsAreUrgent));
 	}
 
 	for (UE::Cook::FRequestCluster& AddedCluster : AddedClusters)
@@ -77,20 +88,91 @@ void FRequestCluster::AddClusters(UCookOnTheFlyServer& COTFS, TArray<FFilePlatfo
 	}
 }
 
-FRequestCluster::FRequest::FRequest(FFilePlatformRequest&& FileRequest, bool bInIsUrgent)
-	: PackageName(NAME_None)
-	, FileName(FileRequest.GetFilename())
-	, CompletionCallback(MoveTemp(FileRequest.GetCompletionCallback()))
-	, bIsUrgent(bInIsUrgent)
+void FRequestCluster::AddClusters(UCookOnTheFlyServer& COTFS, FPackageDataSet& UnclusteredRequests, TRingBuffer<FRequestCluster>& OutClusters,
+	FRequestQueue& QueueForReadyRequests)
 {
+	if (UnclusteredRequests.Num() == 0)
+	{
+		return;
+	}
+
+	TArray<FRequestCluster, TInlineAllocator<1>> AddedClusters;
+	TArray<const ITargetPlatform*, TInlineAllocator<ExpectedMaxNumPlatforms>> RequestedPlatforms;
+	auto FindOrAddClusterForPlatforms = [&AddedClusters, &COTFS, &RequestedPlatforms, &UnclusteredRequests]()
+	{
+		if (AddedClusters.Num() == 0)
+		{
+			FRequestCluster& Cluster = AddedClusters.Emplace_GetRef(COTFS, RequestedPlatforms);
+			// The usual case is all platforms are the same, so reserve the first Cluster's size assuming it will get all requests
+			Cluster.Requests.Reserve(UnclusteredRequests.Num());
+			return &Cluster;
+		}
+		for (FRequestCluster& Existing : AddedClusters)
+		{
+			if (Existing.GetPlatforms() == RequestedPlatforms)
+			{
+				return &Existing;
+			}
+		}
+		return &AddedClusters.Emplace_GetRef(COTFS, RequestedPlatforms);
+	};
+
+	bool bErrorOnEngineContentUse = false;
+	UCookOnTheFlyServer::FCookByTheBookOptions* Options = COTFS.CookByTheBookOptions;
+	FString DLCPath;
+	if (Options)
+	{
+		bErrorOnEngineContentUse = Options->bErrorOnEngineContentUse;
+		if (bErrorOnEngineContentUse)
+		{
+			DLCPath = FPaths::Combine(*COTFS.GetBaseDirectoryForDLC(), TEXT("Content"));
+			FPaths::MakeStandardFilename(DLCPath);
+		}
+	}
+
+	UE::Cook::FRequestCluster* MRUCluster = nullptr;
+	for (FPackageData* PackageData : UnclusteredRequests)
+	{
+		if (PackageData->AreAllRequestedPlatformsExplored())
+		{
+			QueueForReadyRequests.AddReadyRequest(PackageData);
+			continue;
+		}
+
+		// For non-cookable packages that we skipped in an earlier cluster but loaded because
+		// they are hard dependencies, avoid the work of creating a cluster just for them,
+		// by checking non-cookable and sending the package to idle instead of adding to
+		// a cluster
+		FName FileName = PackageData->GetFileName();
+		PackageData->GetRequestedPlatforms(RequestedPlatforms);
+		if (!IsRequestCookable(PackageData->GetPackageName(), FileName, PackageData,
+			COTFS.GetPackageNameCache(), *COTFS.PackageDatas, *COTFS.PackageTracker,
+			DLCPath, bErrorOnEngineContentUse, RequestedPlatforms))
+		{
+			PackageData->SendToState(EPackageState::Idle, ESendFlags::QueueAdd);
+			continue;
+		}
+
+		if (!MRUCluster || RequestedPlatforms != MRUCluster->GetPlatforms())
+		{
+			// MRUCluster points to data inside AddedClusters, so we have to recalculate MRUCluster whenever we add
+			MRUCluster = FindOrAddClusterForPlatforms();
+		}
+
+		MRUCluster->Requests.Add(PackageData);
+	}
+
+	for (UE::Cook::FRequestCluster& AddedCluster : AddedClusters)
+	{
+		AddedCluster.Initialize(COTFS);
+		OutClusters.Add(MoveTemp(AddedCluster));
+	}
 }
 
-FRequestCluster::FRequest::FRequest(FName InPackageName, FName InFileName, FCompletionCallback&& InCompletionCallback,
-	bool bInIsUrgent)
-	: PackageName(InPackageName)
-	, FileName(InFileName)
-	, CompletionCallback(MoveTemp(InCompletionCallback))
-	, bIsUrgent(bInIsUrgent)
+FRequestCluster::FFileNameRequest::FFileNameRequest(FFilePlatformRequest&& FileRequest, bool bInUrgent)
+	: FileName(FileRequest.GetFilename())
+	, CompletionCallback(MoveTemp(FileRequest.GetCompletionCallback()))
+	, bUrgent(bInUrgent)
 {
 }
 
@@ -100,15 +182,11 @@ void FRequestCluster::Initialize(UCookOnTheFlyServer& COTFS)
 	if (Options)
 	{
 		bAllowHardDependencies = !Options->bSkipHardReferences;
-		// We always skip soft dependencies in cluster exploration if the cook commandline is set to skip soft references.
-		// We also need to skip soft dependencies from assetregistry, and make other editor-only robustness fixes, 
-		// if the project has problems with editor-only robustness and has turned bExploreSoftReferencesOnStart off
-		bAllowSoftDependencies = !Options->bSkipSoftReferences && COTFS.bExploreSoftReferencesOnStart;
-		// Build dependency soft references are more robust, and can be enabled whenever the cook mode allows soft references.
-		bAllowSoftBuildDependencies = !Options->bSkipSoftReferences;
+		bAllowSoftDependencies = !Options->bSkipSoftReferences;
 		bErrorOnEngineContentUse = Options->bErrorOnEngineContentUse;
 	}
-	bTargetDomainEnabled = COTFS.bTargetDomainEnabled;
+	bHybridIterativeEnabled = COTFS.bHybridIterativeEnabled;
+	bPreexploreDependenciesEnabled = COTFS.bPreexploreDependenciesEnabled;
 	if (bErrorOnEngineContentUse)
 	{
 		DLCPath = FPaths::Combine(*COTFS.GetBaseDirectoryForDLC(), TEXT("Content"));
@@ -116,9 +194,41 @@ void FRequestCluster::Initialize(UCookOnTheFlyServer& COTFS)
 	}
 
 	PackageWriters.Reserve(Platforms.Num());
+	bCleanBuild = false;
+	bool bFirst = true;
 	for (const ITargetPlatform* TargetPlatform : Platforms)
 	{
 		PackageWriters.Add(&COTFS.FindOrCreatePackageWriter(TargetPlatform));
+		FPlatformData* PlatformData = COTFS.PlatformManager->GetPlatformData(TargetPlatform);
+		if (bFirst)
+		{
+			bCleanBuild = PlatformData->bCleanBuild;
+			bFirst = false;
+		}
+		else
+		{
+			if (PlatformData->bCleanBuild != bCleanBuild)
+			{
+				UE_LOG(LogCook, Warning, TEXT("Clean build is requested for some platforms but not others, but this is not supported. All platforms will be built clean."));
+				bCleanBuild = true;
+			}
+		}
+	}
+
+	bool bRemoveNulls = false;
+	for (FPackageData*& PackageData : Requests)
+	{
+		bool bAlreadyExists = false;
+		OwnedPackageDatas.Add(PackageData, &bAlreadyExists);
+		if (bAlreadyExists)
+		{
+			bRemoveNulls = true;
+			PackageData = nullptr;
+		}
+	}
+	if (bRemoveNulls)
+	{
+		Requests.Remove(nullptr);
 	}
 }
 
@@ -142,15 +252,19 @@ void FRequestCluster::FetchPackageNames(const FCookerTimer& CookerTimer, bool& b
 	}
 
 	constexpr int32 TimerCheckPeriod = 100; // Do not incur the cost of checking the timer on every package
-	int32 RequestsNum = Requests.Num();
-	while (this->NextRequest < RequestsNum)
+	int32 InRequestsNum = InRequests.Num();
+	if (this->NextRequest == 0)
+	{
+		Requests.Reserve(Requests.Num() + InRequestsNum);
+	}
+	for (;NextRequest < InRequestsNum; ++NextRequest)
 	{
 		if (NextRequest % TimerCheckPeriod == 0 && CookerTimer.IsTimeUp())
 		{
 			break;
 		}
 
-		FRequest& Request = Requests[NextRequest];
+		FFileNameRequest& Request = InRequests[NextRequest];
 		FName OriginalName = Request.FileName;
 #if DEBUG_COOKONTHEFLY
 		UE_LOG(LogCook, Display, TEXT("Processing request for package %s"), *OriginalName.ToString());
@@ -167,25 +281,78 @@ void FRequestCluster::FetchPackageNames(const FCookerTimer& CookerTimer, bool& b
 			{
 				CompletionCallback(nullptr);
 			}
-			Requests.RemoveAt(NextRequest); // Preserve the order of incoming requests
-			--RequestsNum;
 		}
 		else
 		{
-			Request.PackageName = *PackageName;
-			++this->NextRequest;
+			FPackageData& PackageData = PackageDatas.FindOrAddPackageData(*PackageName, Request.FileName);
+			if (TryTakeOwnership(PackageData, Request.bUrgent, MoveTemp(Request.CompletionCallback)))
+			{
+				bool bAlreadyExists;
+				OwnedPackageDatas.Add(&PackageData, &bAlreadyExists);
+				if (!bAlreadyExists)
+				{
+					Requests.Add(&PackageData);
+				}
+			}
 		}
 	}
-	if (NextRequest < RequestsNum)
+	if (NextRequest < InRequestsNum)
 	{
 		bOutComplete = false;
 		return;
 	}
 
+	InRequests.Empty();
 	NextRequest = 0;
 	bPackageNamesComplete = true;
 }
 
+bool FRequestCluster::TryTakeOwnership(FPackageData& PackageData, bool bUrgent, UE::Cook::FCompletionCallback && CompletionCallback)
+{
+	if (!PackageData.IsInProgress())
+	{
+		check(GetPlatforms().Num() != 0); // This is required for SetRequestData
+		if (PackageData.HasAllExploredPlatforms(GetPlatforms()))
+		{
+			// Leave it in idle - it's already been processed by a cluster
+			if (CompletionCallback)
+			{
+				CompletionCallback(&PackageData);
+			}
+			return false;
+		}
+		else
+		{
+			PackageData.SetRequestData(GetPlatforms(), bUrgent, MoveTemp(CompletionCallback));
+			PackageData.SendToState(EPackageState::Request, ESendFlags::QueueRemove);
+			return true;
+		}
+	}
+	else
+	{
+		if (PackageData.HasAllExploredPlatforms(GetPlatforms()))
+		{
+			// Leave it where it is - it's already been processed by a cluster - but update with our request data
+			// This might demote it back to Request, and add it to Normal or Urgent request, but that will not
+			// impact this RequestCluster
+			PackageData.UpdateRequestData(GetPlatforms(), bUrgent, MoveTemp(CompletionCallback));
+			return false;
+		}
+		else
+		{
+			if (!OwnedPackageDatas.Contains(&PackageData))
+			{
+				// Steal it from wherever it is and add it to this cluster
+				// This might steal it from another RequestCluster or from the UnclusteredRequests if it's in request
+				// Doing that steal is wasteful but okay; one of the RequestClusters will win it and keep it
+				PackageData.SendToState(EPackageState::Request, ESendFlags::QueueRemove);
+				PackageData.UpdateRequestData(GetPlatforms(), bUrgent, MoveTemp(CompletionCallback));
+			}
+			return true;
+		}
+	}
+}
+
 void FRequestCluster::FetchDependencies(const FCookerTimer& CookerTimer, bool& bOutComplete)
 {
 	if (bDependenciesComplete)
@@ -212,82 +379,11 @@ void FRequestCluster::FetchDependencies(const FCookerTimer& CookerTimer, bool& b
 	// When initial request assets are unrelated, we want to load those assets in the original order. Since we are
 	// (a) moving later requests before earlier requests but then (b) reversing the list, the reversals cancel out
 	// and we should iterate the initial requests from front to back.
-	TRingBuffer<FRequestClusterGraphData::FNamePair>& TransitiveRequests(DependencyGraphData.TransitiveRequests);
-	TRingBuffer<FRequestClusterGraphData::FNamePair>& Segment(DependencyGraphData.Segment);
-	TSet<FName>& Visited(DependencyGraphData.Visited);
-	TArray<FRequestClusterGraphData::FStackData>& StackStorage(DependencyGraphData.Stack);
-	int32& StackNum(DependencyGraphData.StackNum);
-	int32 NumRequests = Requests.Num();
-	int32 TimerCounter = 0;
-	auto IsStackEmpty = [&StackNum]()
-	{
-		return StackNum == 0;
-	};
-	auto TopStack = [&StackStorage, &StackNum]() -> FRequestClusterGraphData::FStackData&
-	{
-		return StackStorage[StackNum-1];
-	};
-	auto PushStack = [&StackStorage, &StackNum](FName PackageName) -> FRequestClusterGraphData::FStackData&
-	{
-		if (StackNum == StackStorage.Num())
-		{
-			StackStorage.Emplace();
-		}
-		FRequestClusterGraphData::FStackData& Top = StackStorage[StackNum++];
-		Top.Reset(PackageName);
-		return Top;
-	};
-	auto PopStack = [&StackNum]
+	if (NextRequest == 0)
 	{
-		--StackNum;
-	};
-
-	auto AddVertex = [this, &Segment, &Visited, &PushStack, &TimerCounter]
-		(FName PackageName, FName FileName, bool bInitial, bool& bOutPushedStack)
-	{
-		bool bAlreadyInSet;
-		Visited.Add(PackageName, &bAlreadyInSet);
-		if (bAlreadyInSet)
-		{
-			bOutPushedStack = false;
-			return;
-		}
-		++TimerCounter;
-
-		// Add the package to the transitive output if its cookable, and always add the initial vertices
-		bool bCookable = IsRequestCookable(PackageName, FileName);
-		if (bInitial || bCookable)
-		{
-			bOutPushedStack = true;
-			Segment.Add(FRequestClusterGraphData::FNamePair{ PackageName, FileName });
-			FRequestClusterGraphData::FStackData& StackData = PushStack(PackageName);
-
-			bool bExploreDependencies = bCookable;
-			if (bCookable && !bAllowSoftDependencies)
-			{
-				// TODO: Editor-only packages may be loaded at editor startup, which makes them skip the request state
-				// and go straight to load and save, where they are culled and returned to idle.
-				// If we add dependencies from these startup packages, in some projects we add transitive dependency
-				// packages that are not otherwise cooked. To prevent breaking projects we are temporarily skipping the
-				// dependency search from these packages.
-				bool bInProgress = false;
-				FPackageData* PackageData = PackageDatas.FindPackageDataByPackageName(PackageName);
-				if (PackageData && (PackageData->IsInProgress() || FindPackage(nullptr, *PackageName.ToString())))
-				{
-					bExploreDependencies = false;
-				}
-			}
-			if (bExploreDependencies)
-			{
-				GetDependencies(PackageName, StackData.Dependencies);
-			}
-		}
-		else
-		{
-			bOutPushedStack = false;
-		}
-	};
-
+		check(Requests.Num() == OwnedPackageDatas.Num());
+	}
+	int32 TimerCounter = 0;
 	constexpr int32 TimerCheckPeriod = 10; // Do not incur the cost of checking the timer on every package
 	for (;;)
 	{
@@ -303,11 +399,11 @@ void FRequestCluster::FetchDependencies(const FCookerTimer& CookerTimer, bool& b
 				TimerCounter = 0;
 			}
 
-			FRequestClusterGraphData::FStackData& CurrentTop(TopStack());
+			FStackData& CurrentTop(TopStack());
 			bool bPushedStack = false;
 			while (CurrentTop.NextDependency < CurrentTop.Dependencies.Num())
 			{
-				AddVertex(CurrentTop.Dependencies[CurrentTop.NextDependency++], NAME_None, false /* bInitial */, bPushedStack);
+				AddVertex(CurrentTop.Dependencies[CurrentTop.NextDependency++], nullptr, bPushedStack, TimerCounter);
 				if (bPushedStack)
 				{
 					// CurrentTop is now invalid
@@ -329,49 +425,41 @@ void FRequestCluster::FetchDependencies(const FCookerTimer& CookerTimer, bool& b
 			TransitiveRequests.AddFront(Segment.PopValue());
 		}
 
-		while (NextRequest < NumRequests)
+		while (NextRequest < Requests.Num())
 		{
-			FRequest& Request = Requests[NextRequest++];
-			FRequestClusterGraphData::FNamePair NamePair{ Request.PackageName, Request.FileName };
+			FPackageData* PackageData = Requests[NextRequest++];
 			bool bPushedStack = false;
-			AddVertex(Request.PackageName, Request.FileName, true /* bInitial */, bPushedStack);
+			FName PackageName = PackageData->GetPackageName();
+			EVisitStatus& Status = AddVertex(PackageName, PackageData, bPushedStack, TimerCounter);
 			if (bPushedStack)
 			{
+				check(Status == EVisitStatus::Visited);
 				break;
 			}
+			else if (Status != EVisitStatus::Visited)
+			{
+				// If an initial request was found from transitive dependencies of another initial request, and was not
+				// cookable, it was skipped. We need to reprocess it with the information that it is an initial vertex.
+				check(Status == EVisitStatus::Skipped);
+				Visited.Remove(PackageData->GetPackageName());
+				Status = AddVertex(PackageName, PackageData, bPushedStack, TimerCounter);
+				check(Status == EVisitStatus::Visited);
+			}
 		}
 		if (IsStackEmpty())
 		{
-			check(NextRequest == NumRequests);
+			check(NextRequest == Requests.Num());
 			break;
 		}
 	}
 
-	// Map PackageNames to Request for the initial requests, so we can preserve their data (e.g. bIsUrgent) in the output
-	TMap<FName, FRequest*> InitialRequests;
-	for (FRequest& Request : Requests)
-	{
-		InitialRequests.Add(Request.PackageName, &Request);
-	}
-
-	// Convert all NamePairs found to an FRequest and push them all into the output in LeafToRoot order
-	TArray<FRequest> OutRequests;
-	OutRequests.Reserve(TransitiveRequests.Num());
-	for (int32 LeafToRootIndex = TransitiveRequests.Num() - 1; LeafToRootIndex >= 0; --LeafToRootIndex)
+	// Reverse the TransitiveRequests array to put it in LeafToRoot order, and push it into the output Requests
+	COOK_STAT(DetailedCookStats::NumPreloadedDependencies += TransitiveRequests.Num() - Requests.Num());
+	Requests.Reset(TransitiveRequests.Num());
+	for (int32 Index = TransitiveRequests.Num() - 1; Index >= 0; --Index)
 	{
-		FRequestClusterGraphData::FNamePair& NamePair = TransitiveRequests[LeafToRootIndex];
-		FRequest** InitialRequest = InitialRequests.Find(NamePair.PackageName);
-		if (InitialRequest)
-		{
-			OutRequests.Add(MoveTemp(**InitialRequest));
-		}
-		else
-		{
-			OutRequests.Add(FRequest{ NamePair.PackageName, NamePair.FileName, FCompletionCallback(), false /* bIsUrgent */ });
-			COOK_STAT(++DetailedCookStats::NumPreloadedDependencies);
-		}
+		Requests.Add(TransitiveRequests[Index]);
 	}
-	Swap(OutRequests, Requests);
 
 	TransitiveRequests.Empty();
 	StackStorage.Empty();
@@ -382,49 +470,237 @@ void FRequestCluster::FetchDependencies(const FCookerTimer& CookerTimer, bool& b
 	bDependenciesComplete = true;
 }
 
-void FRequestCluster::GetDependencies(FName PackageName, TArray<FName>& OutDependencies)
+int32 FRequestCluster::NumPackageDatas() const
+{
+	return OwnedPackageDatas.Num();
+}
+
+void FRequestCluster::RemovePackageData(FPackageData* PackageData)
+{
+	if (OwnedPackageDatas.Remove(PackageData) == 0)
+	{
+		return;
+	}
+
+	Requests.Remove(PackageData);
+	TransitiveRequests.Remove(PackageData);
+	Segment.Remove(PackageData);
+}
+
+bool FRequestCluster::Contains(FPackageData* PackageData) const
+{
+	return OwnedPackageDatas.Contains(PackageData);
+}
+
+void FRequestCluster::ClearAndDetachOwnedPackageDatas(TArray<FPackageData*>& OutRequestsToLoad, TArray<FPackageData*>& OutRequestsToDemote)
+{
+	if (bDependenciesComplete)
+	{
+		check(TransitiveRequests.Num() == 0 && Segment.Num() == 0);
+		OutRequestsToLoad = MoveTemp(Requests);
+		OutRequestsToDemote = MoveTemp(RequestsToDemote);
+		check(OutRequestsToLoad.Num() + OutRequestsToDemote.Num() == OwnedPackageDatas.Num())
+	}
+	else
+	{
+		OutRequestsToLoad = OwnedPackageDatas.Array();
+		OutRequestsToDemote.Reset();
+	}
+	InRequests.Empty();
+	Requests.Empty();
+	RequestsToDemote.Empty();
+	OwnedPackageDatas.Empty();
+	TransitiveRequests.Empty();
+	Segment.Empty();
+	Visited.Empty();
+	StackStorage.Empty();
+	StackNum = 0;
+	NextRequest = 0;
+}
+
+FRequestCluster::EVisitStatus& FRequestCluster::AddVertex(FName PackageName, FPackageData* PackageData,
+	bool& bOutPushedStack, int32& TimerCounter)
+{
+	EVisitStatus& Status = Visited.FindOrAdd(PackageName, EVisitStatus::New);
+	if (Status != EVisitStatus::New)
+	{
+		bOutPushedStack = false;
+		return Status;
+	}
+	++TimerCounter;
+
+	// If it was a direct request we already own it and we can't skip exploring it, but we still check whether
+	// it is cookable to see if it should add dependencies and move to load or ignore dependencies and return to idle.
+	bool bInitialRequest = false;
+	FName FileName = NAME_None;
+	if (PackageData)
+	{
+		bInitialRequest = true;
+		FileName = PackageData->GetFileName();
+	}
+	bool bCookable = IsRequestCookable(PackageName, FileName, PackageData);
+	// If not an intial request, take ownership of the packagedata if cookable and available
+	if (!bInitialRequest)
+	{
+		if (!bCookable)
+		{
+			Status = EVisitStatus::Skipped;
+			bOutPushedStack = false;
+			return Status;
+		}
+		if (!PackageData) // IsRequestCookable may have looked up the PackageData
+		{
+			PackageData = &PackageDatas.FindOrAddPackageData(PackageName, FileName);
+		}
+		if (!TryTakeOwnership(*PackageData, false /* bUrgent */, FCompletionCallback()))
+		{
+			Status = EVisitStatus::Skipped;
+			bOutPushedStack = false;
+			return Status;
+		}
+		// Note it may already be in OwnedPackageDatas because it might be present later in the Requests array.
+		// That's okay, we still need to add it to Segment now.
+		OwnedPackageDatas.Add(PackageData);
+	}
+
+	bool bExploreDependencies = bCookable;
+	if (bCookable && !bPreexploreDependenciesEnabled)
+	{
+		// TODO: Editor-only packages may be loaded at editor startup, which makes them skip the request state
+		// and go straight to load and save, where they are culled and returned to idle.
+		// If we add dependencies from these startup packages, in some projects we add transitive dependency
+		// packages that are not otherwise cooked. To prevent breaking projects we are temporarily skipping the
+		// dependency search from these packages.
+		bool bInProgress = false;
+		if (FindPackage(nullptr, *PackageName.ToString()))
+		{
+			bExploreDependencies = false;
+		}
+	}
+	FStackData* StackData = nullptr;
+	if (bExploreDependencies)
+	{
+		bOutPushedStack = true;
+		StackData = &PushStack(PackageName);
+	}
+
+	bool bAlreadyCooked;
+	VisitPackageData(PackageData, StackData ? &StackData->Dependencies : nullptr, bAlreadyCooked);
+
+	if (bCookable && !bAlreadyCooked)
+	{
+		Segment.Add(PackageData);
+	}
+	else
+	{
+		RequestsToDemote.Add(PackageData);
+	}
+	Status = EVisitStatus::Visited;
+	bOutPushedStack = StackData != nullptr;
+	return Status;
+}
+
+bool FRequestCluster::IsStackEmpty() const
+{
+	return StackNum == 0;
+};
+
+FRequestCluster::FStackData& FRequestCluster::TopStack()
 {
-	TArray<FName>& BuildDependencies(OutDependencies);
-	TArray<FName>& RuntimeDependencies(this->RuntimeScratch);
-	BuildDependencies.Reset();
-	RuntimeDependencies.Reset();
+	return StackStorage[StackNum - 1];
+};
 
-	// TODO EditorOnly References: We only fetch Game dependencies, because the cooker explicitly loads all of
-	// the dependencies that we report. And if we explicitly load an EditorOnly dependency, that causes
-	// StaticLoadObjectInternal to SetLoadedByEditorPropertiesOnly(false), which then treats the editor-only package
-	// as needed in-game.
-	UE::AssetRegistry::EDependencyQuery DependencyQuery = UE::AssetRegistry::EDependencyQuery::Game;
-	if (!bAllowSoftDependencies)
+FRequestCluster::FStackData& FRequestCluster::PushStack(FName PackageName)
+{
+	if (StackNum == StackStorage.Num())
 	{
-		DependencyQuery |= UE::AssetRegistry::EDependencyQuery::Hard;
+		StackStorage.Emplace();
 	}
-	AssetRegistry.GetDependencies(PackageName, BuildDependencies,
-		UE::AssetRegistry::EDependencyCategory::Package, DependencyQuery);
-	if (bTargetDomainEnabled)
+	FStackData& Top = StackStorage[StackNum++];
+	Top.Reset(PackageName);
+	return Top;
+};
+
+void FRequestCluster::PopStack()
+{
+	--StackNum;
+};
+
+void FRequestCluster::VisitPackageData(FPackageData* PackageData, TArray<FName>* OutDependencies, bool& bAlreadyCooked)
+{
+	if (OutDependencies)
 	{
-		for (int32 PlatIndex = 0; PlatIndex < Platforms.Num(); ++PlatIndex)
+		TArray<FName>& BuildDependencies(*OutDependencies);
+		TArray<FName>& RuntimeDependencies(this->RuntimeScratch);
+		BuildDependencies.Reset();
+		RuntimeDependencies.Reset();
+		FName PackageName = PackageData->GetPackageName();
+
+		// TODO EditorOnly References: We only fetch Game dependencies, because the cooker explicitly loads all of
+		// the dependencies that we report. And if we explicitly load an EditorOnly dependency, that causes
+		// StaticLoadObjectInternal to SetLoadedByEditorPropertiesOnly(false), which then treats the editor-only package
+		// as needed in-game.
+		UE::AssetRegistry::EDependencyQuery DependencyQuery = UE::AssetRegistry::EDependencyQuery::Game;
+		// We always skip assetregistry soft dependencies if the cook commandline is set to skip soft references.
+		// We also need to skip them if the project has problems with editor-only robustness and has turned
+		// ExploreDependencies off
+		if (!bAllowSoftDependencies || !bPreexploreDependenciesEnabled)
+		{
+			DependencyQuery |= UE::AssetRegistry::EDependencyQuery::Hard;
+		}
+		AssetRegistry.GetDependencies(PackageName, BuildDependencies,
+			UE::AssetRegistry::EDependencyCategory::Package, DependencyQuery);
+		if (bHybridIterativeEnabled && !bCleanBuild)
 		{
-			const ITargetPlatform* TargetPlatform = Platforms[PlatIndex];
-			ICookedPackageWriter* PackageWriter = PackageWriters[PlatIndex];
-			if (!PackageWriter)
+			bool bExists = false;
+			for (int32 PlatIndex = 0; PlatIndex < Platforms.Num(); ++PlatIndex)
 			{
-				continue;
+				const ITargetPlatform* TargetPlatform = Platforms[PlatIndex];
+				ICookedPackageWriter* PackageWriter = PackageWriters[PlatIndex];
+				if (!PackageWriter)
+				{
+					continue;
+				}
+				bool bExistsInTargetDomain = UE::TargetDomain::TryFetchKeyAndDependencies(PackageWriter, PackageName,
+					TargetPlatform, nullptr /* OutHash */, &BuildDependencies, &RuntimeDependencies,
+					nullptr /* OutErrorMessage */);
+				if (bExistsInTargetDomain && UE::TargetDomain::IsIterativeEnabled(PackageName))
+				{
+					bExists = true;
+					PackageData->SetPlatformCooked(TargetPlatform, true);
+				}
+			}
+			if (bExists)
+			{
+				COOK_STAT(++DetailedCookStats::NumPackagesIterativelySkipped);
 			}
-			UE::TargetDomain::TryFetchKeyAndDependencies(PackageWriter, PackageName, TargetPlatform, nullptr /* OutHash */, &BuildDependencies,
-				&RuntimeDependencies, nullptr /* OutErrorMessage */);
 		}
+		if (bAllowSoftDependencies)
+		{
+			BuildDependencies.Append(RuntimeDependencies);
+		}
+
+		// Sort the list of BuildDependencies to check for uniqueness and make them deterministic
+		Algo::Sort(BuildDependencies, FNameLexicalLess());
+		BuildDependencies.SetNum(Algo::Unique(BuildDependencies));
 	}
-	if (bAllowSoftBuildDependencies)
+
+	for (const ITargetPlatform* TargetPlatform : Platforms)
 	{
-		BuildDependencies.Append(RuntimeDependencies);
+		PackageData->FindOrAddPlatformData(TargetPlatform).bExplored = true;
 	}
+	bAlreadyCooked = PackageData->AreAllRequestedPlatformsCooked(true /* bAllowFailedCooks */);
+}
 
-	// Sort the list of BuildDependencies to check for uniqueness and make them deterministic
-	Algo::Sort(BuildDependencies, FNameLexicalLess());
-	BuildDependencies.SetNum(Algo::Unique(BuildDependencies));
+bool FRequestCluster::IsRequestCookable(FName PackageName, FName& InOutFileName, FPackageData*& PackageData)
+{
+	return IsRequestCookable(PackageName, InOutFileName, PackageData, PackageNameCache, PackageDatas, PackageTracker,
+		DLCPath, bErrorOnEngineContentUse, GetPlatforms());
 }
 
-bool FRequestCluster::IsRequestCookable(FName PackageName, FName& InOutFileName)
+bool FRequestCluster::IsRequestCookable(FName PackageName, FName& InOutFileName, FPackageData*& PackageData,
+	const FPackageNameCache& InPackageNameCache, FPackageDatas& InPackageDatas, FPackageTracker& InPackageTracker,
+	FStringView InDLCPath, bool bInErrorOnEngineContentUse, TConstArrayView<const ITargetPlatform*> RequestPlatforms)
 {
 	TStringBuilder<256> NameBuffer;
 	// We need to reject packagenames from adding themselves or their transitive dependencies using all the same rules that
@@ -437,7 +713,7 @@ bool FRequestCluster::IsRequestCookable(FName PackageName, FName& InOutFileName)
 
 	if (InOutFileName.IsNone())
 	{
-		InOutFileName = PackageNameCache.GetCachedStandardFileName(PackageName);
+		InOutFileName = InPackageNameCache.GetCachedStandardFileName(PackageName);
 		if (InOutFileName.IsNone())
 		{
 			// Package does not exist on disk
@@ -445,17 +721,26 @@ bool FRequestCluster::IsRequestCookable(FName PackageName, FName& InOutFileName)
 		}
 	}
 
-	if (bErrorOnEngineContentUse && !DLCPath.IsEmpty())
+	if (bInErrorOnEngineContentUse && !InDLCPath.IsEmpty())
 	{
 		InOutFileName.ToString(NameBuffer);
-		if (!FStringView(NameBuffer).StartsWith(DLCPath))
+		if (!FStringView(NameBuffer).StartsWith(InDLCPath))
 		{
+			if (!PackageData)
+			{
+				PackageData = &InPackageDatas.FindOrAddPackageData(PackageName, InOutFileName);
+			}
+			if (!PackageData->HasAllCookedPlatforms(RequestPlatforms, true /* bIncludeFailed */))
+			{
+				UE_LOG(LogCook, Error, TEXT("Uncooked Engine or Game content %s is being referenced by DLC!"), *NameBuffer);
+			}
 			return false;
 		}
 	}
 
-	if (PackageTracker.NeverCookPackageList.Contains(InOutFileName))
+	if (InPackageTracker.NeverCookPackageList.Contains(InOutFileName))
 	{
+		UE_LOG(LogCook, Verbose, TEXT("Package %s is referenced but is in the never cook package list, discarding request"), *NameBuffer);
 		return false;
 	}
 
diff --git a/Engine/Source/Editor/UnrealEd/Private/Cooker/CookRequestCluster.h b/Engine/Source/Editor/UnrealEd/Private/Cooker/CookRequestCluster.h
--- a/Engine/Source/Editor/UnrealEd/Private/Cooker/CookRequestCluster.h
+++ b/Engine/Source/Editor/UnrealEd/Private/Cooker/CookRequestCluster.h
@@ -13,44 +13,13 @@ class IAssetRegistry;
 class ITargetPlatform;
 class UCookOnTheFlyServer;
 struct FPackageNameCache;
+namespace UE::Cook { class FRequestQueue; }
 namespace UE::Cook { struct FFilePlatformRequest; }
 namespace UE::Cook { struct FPackageTracker; }
 
 namespace UE::Cook
 {
 
-/** Data used during the reentrant dependencies and topological sort operation of FRequestCluster. */
-struct FRequestClusterGraphData
-{
-	struct FNamePair
-	{
-		FName PackageName;
-		FName FileName;
-	};
-	/** All transitive PackageNames found from from all initial requests. Root to leaf order. */
-	TRingBuffer<FNamePair> TransitiveRequests;
-	/** Transitive PackageNames from the current initial request not already in TransitiveRequests. Root to leaf order. */
-	TRingBuffer<FNamePair> Segment;
-	/** Which PackageNames have already been reached (possibly still on the stack). */
-	TSet<FName> Visited;
-	/** Requests that are currently iterating over their dependencies. */
-	struct FStackData
-	{
-		void Reset(FName InPackageName)
-		{
-			PackageName = InPackageName;
-			Dependencies.Reset();
-			NextDependency = 0;
-		}
-		TArray<FName> Dependencies;
-		FName PackageName;
-		int32 NextDependency;
-	};
-	TArray<FStackData> Stack;
-	/** We do not remove FStackData from the stack, because we want to avoid reallocating Dependencies arrays. */
-	int32 StackNum = 0;
-};
-
 /**
  * A group of external requests sent to CookOnTheFlyServer's tick loop. Transitive dependencies are found and all of the
  * requested or dependent packagenames are added as requests together to the cooking state machine.
@@ -59,17 +28,16 @@ class FRequestCluster
 {
 public:
 	/** The cluster's data about a package from the external requests or discovered dependency. */
-	struct FRequest
+	struct FFileNameRequest
 	{
-		FName PackageName;
 		FName FileName;
 		FCompletionCallback CompletionCallback;
-		bool bIsUrgent;
+		bool bUrgent;
 
-		FRequest(FName InPackageName, FName InFileName, FCompletionCallback&& InCompletionCallback, bool bInIsUrgent);
-		FRequest(FFilePlatformRequest&& FileRequest, bool bInIsUrgent);
+		FFileNameRequest(FFilePlatformRequest&& FileRequest, bool bInUrgent);
 	};
 	explicit FRequestCluster(UCookOnTheFlyServer& COTFS, TArray<const ITargetPlatform*>&& Platforms);
+	explicit FRequestCluster(UCookOnTheFlyServer& COTFS, TConstArrayView<const ITargetPlatform*> Platforms);
 	FRequestCluster(FRequestCluster&&) = default;
 
 	/**
@@ -79,27 +47,73 @@ public:
 	void Process(const FCookerTimer& CookerTimer, bool& bOutComplete);
 
 	TConstArrayView<const ITargetPlatform*> GetPlatforms() const { return Platforms; }
-	TArray<FRequest>& GetRequests() { return Requests; }
+	/** PackageData container interface: return the number of PackageDatas owned by this container. */
+	int32 NumPackageDatas() const;
+	/** PackageData container interface: remove the PackageData from this container. */
+	void RemovePackageData(FPackageData* PackageData);
+	/** PackageData container interface: whether the PackageData is owned by this container. */
+	bool Contains(FPackageData* PackageData) const;
+	/**
+	 * Remove all PackageDatas owned by this container and return them.
+	 * OutRequestsToLoad is the set of PackageDatas sorted by leaf to root load order.
+	 * OutRequestToDemote is the set of Packages that are uncookable or have already been cooked.
+	 * If called before Process sets bOutComplete=true, all packages are put in OutRequestToLoad and are unsorted.
+	 */
+	void ClearAndDetachOwnedPackageDatas(TArray<FPackageData*>& OutRequestsToLoad, TArray<FPackageData*>& OutRequestsToDemote);
 
 	/**
-	 * Create clusters(s) for all the given requests and append them to OutClusters.
+	 * Create clusters(s) for all the given name or packagedata requests and append them to OutClusters.
 	 * Multiple clusters are necessary if the list of platforms differs between requests.
 	 */
-	static void AddClusters(UCookOnTheFlyServer& COTFS, TArray<FFilePlatformRequest>&& InRequests, bool bRequestsAreUrgent,
-		TRingBuffer<FRequestCluster>& OutClusters);
+	static void AddClusters(UCookOnTheFlyServer& COTFS, TArray<FFilePlatformRequest>&& InRequests,
+		bool bRequestsAreUrgent, TRingBuffer<FRequestCluster>& OutClusters);
+	static void AddClusters(UCookOnTheFlyServer& COTFS, FPackageDataSet& UnclusteredRequests,
+		TRingBuffer<FRequestCluster>& OutClusters, FRequestQueue& QueueForReadyRequests);
 
 private:
+	struct FStackData
+	{
+		void Reset(FName InPackageName)
+		{
+			PackageName = InPackageName;
+			Dependencies.Reset();
+			NextDependency = 0;
+		}
+		TArray<FName> Dependencies;
+		FName PackageName;
+		int32 NextDependency;
+	};
+	enum class EVisitStatus
+	{
+		New,
+		Visited,
+		Skipped,
+	};
+
 	void Initialize(UCookOnTheFlyServer& COTFS);
 	void FetchPackageNames(const FCookerTimer& CookerTimer, bool& bOutComplete);
 	void FetchDependencies(const FCookerTimer& CookerTimer, bool& bOutComplete);
-	void GetDependencies(FName PackageName, TArray<FName>& OutDependencies);
-	bool IsRequestCookable(FName PackageName, FName& InOutFileName);
+	bool TryTakeOwnership(FPackageData& PackageData, bool bUrgent, FCompletionCallback&& CompletionCallback);
+	void VisitPackageData(FPackageData* PackageData, TArray<FName>* OutDependencies, bool& bAlreadyCooked);
+	bool IsRequestCookable(FName PackageName, FName& InOutFileName, FPackageData*& PackageData);
+	static bool IsRequestCookable(FName PackageName, FName& InOutFileName, FPackageData*& PackageData,
+		const FPackageNameCache& InPackageNameCache, FPackageDatas& InPackageDatas, FPackageTracker& InPackageTracker,
+		FStringView InDLCPath, bool bInErrorOnEngineContentUse, TConstArrayView<const ITargetPlatform*> RequestPlatforms);
+
+	// Graph Search functions
+	EVisitStatus& AddVertex(FName PackageName, FPackageData* PackageData, bool& bOutPushedStack, int32& TimerCounter);
+	bool IsStackEmpty() const;
+	FStackData& TopStack();
+	FStackData& PushStack(FName PackageName);
+	void PopStack();
 
-	TArray<FRequest> Requests;
+	TArray<FFileNameRequest> InRequests;
+	TArray<FPackageData*> Requests;
+	TArray<FPackageData*> RequestsToDemote;
 	TArray<FName> RuntimeScratch;
 	TArray<const ITargetPlatform*> Platforms;
 	TArray<ICookedPackageWriter*> PackageWriters;
-	FRequestClusterGraphData DependencyGraphData;
+	FPackageDataSet OwnedPackageDatas;
 	FString DLCPath;
 	FPackageDatas& PackageDatas;
 	IAssetRegistry& AssetRegistry;
@@ -108,11 +122,25 @@ private:
 	int32 NextRequest = 0;
 	bool bAllowHardDependencies = true;
 	bool bAllowSoftDependencies = true;
-	bool bAllowSoftBuildDependencies = true;
-	bool bTargetDomainEnabled = true;
+	bool bPreexploreDependenciesEnabled = true;
+	bool bHybridIterativeEnabled = true;
 	bool bErrorOnEngineContentUse = false;
 	bool bPackageNamesComplete = false;
 	bool bDependenciesComplete = false;
+	bool bCleanBuild = false;
+
+
+	/** Data used during the reentrant dependencies and topological sort operation of FRequestCluster. */
+	/** All transitive PackageNames found from from all initial requests. Root to leaf order. */
+	TRingBuffer<FPackageData*> TransitiveRequests;
+	/** Transitive PackageNames from the current initial request not already in TransitiveRequests. Root to leaf order. */
+	TRingBuffer<FPackageData*> Segment;
+	/** Which PackageNames have already been reached (possibly still on the stack). */
+	TMap<FName, EVisitStatus> Visited;
+	/** Storage for requests that are currently iterating over their dependencies. */
+	TArray<FStackData> StackStorage;
+	/** Num elements in StackStorage. We avoid TArray::Add/Remove to avoid reallocating FStackData internals. */
+	int32 StackNum = 0;
 };
 
 }
diff --git a/Engine/Source/Editor/UnrealEd/Private/Cooker/CookTypes.h b/Engine/Source/Editor/UnrealEd/Private/Cooker/CookTypes.h
--- a/Engine/Source/Editor/UnrealEd/Private/Cooker/CookTypes.h
+++ b/Engine/Source/Editor/UnrealEd/Private/Cooker/CookTypes.h
@@ -61,6 +61,11 @@ namespace UE::Cook
 	/** A function that is called when a requested package finishes cooking (when successful, failed, or skipped) */
 	typedef TUniqueFunction<void(FPackageData*)> FCompletionCallback;
 
+	class FPackageDataSet : public TFastPointerSet<FPackageData*>
+	{
+		using TFastPointerSet<FPackageData*>::TFastPointerSet;
+	};
+
 	/** External Requests to the cooker can either by cook requests for a specific file, or arbitrary callbacks that need to execute within the Scheduler's lock. */
 	enum class EExternalRequestType
 	{
@@ -230,4 +235,5 @@ public:
 
 void LogCookerMessage(const FString& MessageText, EMessageSeverity::Type Severity);
 const FString& GetDevelopmentAssetRegistryFilename();
+constexpr uint32 ExpectedMaxNumPlatforms = 32;
 #define REMAPPED_PLUGINS TEXT("RemappedPlugins")
diff --git a/Engine/Source/Editor/UnrealEd/Private/EditorDomain/EditorDomain.cpp b/Engine/Source/Editor/UnrealEd/Private/EditorDomain/EditorDomain.cpp
--- a/Engine/Source/Editor/UnrealEd/Private/EditorDomain/EditorDomain.cpp
+++ b/Engine/Source/Editor/UnrealEd/Private/EditorDomain/EditorDomain.cpp
@@ -48,12 +48,7 @@ public:
 			{
 				UE_LOG(LogEditorDomain, Error, TEXT("Editor.ini:[CookSettings]:EditorDomainEnabled is deprecated, use Editor.ini:[EditorDomain]:EditorDomainEnabled instead."));
 			}
-			bool bTargetDomainEnabled = true;
-			GConfig->GetBool(TEXT("TargetDomain"), TEXT("TargetDomainEnabled"), bTargetDomainEnabled, GEditorIni);
-			if (bEditorDomainEnabled || bTargetDomainEnabled)
-			{
-				UE::EditorDomain::UtilsInitialize();
-			}
+			UE::EditorDomain::UtilsInitialize();
 			if (bEditorDomainEnabled)
 			{
 				// Set values for config settings EditorDomain depends on
diff --git a/Engine/Source/Editor/UnrealEd/Private/EditorDomain/EditorDomainUtils.cpp b/Engine/Source/Editor/UnrealEd/Private/EditorDomain/EditorDomainUtils.cpp
--- a/Engine/Source/Editor/UnrealEd/Private/EditorDomain/EditorDomainUtils.cpp
+++ b/Engine/Source/Editor/UnrealEd/Private/EditorDomain/EditorDomainUtils.cpp
@@ -39,16 +39,10 @@ FClassDigestMap& GetClassDigests()
 }
 
 TSet<FName> GEditorClassBlockList;
-const TSet<FName>& GetEditorClassBlockList()
-{
-	return GEditorClassBlockList;
-}
-
 TSet<FName> GEditorPackageBlockList;
-const TSet<FName>& GetEditorPackageBlockList()
-{
-	return GEditorPackageBlockList;
-}
+TSet<FName> GTargetDomainClassBlockList;
+bool GTargetDomainClassUseAllowList = true;
+bool GTargetDomainClassEmptyAllowList = false;
 
 // Change to a new guid when EditorDomain needs to be invalidated
 const TCHAR* EditorDomainVersion = TEXT("86E449B024FD412BB80DB1397EFAE34D");
@@ -206,10 +200,18 @@ void PrecacheClassDigests(TConstArrayView<FName> ClassNames, TMap<FName, FClassD
 		
 		FClassData& ClassData = ClassDatas.Emplace_GetRef();
 		ClassData.Name = ClassName;
-		ClassData.DigestData.bEditorDomainEnabled = !GetEditorClassBlockList().Contains(ClassName);
+		ClassData.DigestData.bEditorDomainEnabled = !GEditorClassBlockList.Contains(ClassName);
 		if (LookupName != ClassName)
 		{
-			ClassData.DigestData.bEditorDomainEnabled &= !GetEditorClassBlockList().Contains(LookupName);
+			ClassData.DigestData.bEditorDomainEnabled &= !GEditorClassBlockList.Contains(LookupName);
+		}
+		if (!GTargetDomainClassUseAllowList)
+		{
+			ClassData.DigestData.bTargetIterativeEnabled = !GTargetDomainClassBlockList.Contains(ClassName);
+			if (LookupName != ClassName)
+			{
+				ClassData.DigestData.bTargetIterativeEnabled &= !GTargetDomainClassBlockList.Contains(LookupName);
+			}
 		}
 
 		if (Struct)
@@ -273,6 +275,10 @@ void PrecacheClassDigests(TConstArrayView<FName> ClassNames, TMap<FName, FClassD
 				if (ParentDigest)
 				{
 					ClassData.DigestData.bEditorDomainEnabled &= ParentDigest->bEditorDomainEnabled;
+					if (!GTargetDomainClassUseAllowList)
+					{
+						ClassData.DigestData.bTargetIterativeEnabled &= ParentDigest->bTargetIterativeEnabled;
+					}
 				}
 				else
 				{
@@ -356,6 +362,10 @@ void PrecacheClassDigests(TConstArrayView<FName> ClassNames, TMap<FName, FClassD
 		if (ParentDigest)
 		{
 			ClassData.DigestData.bEditorDomainEnabled &= ParentDigest->bEditorDomainEnabled;
+			if (!GTargetDomainClassUseAllowList)
+			{
+				ClassData.DigestData.bTargetIterativeEnabled &= ParentDigest->bTargetIterativeEnabled;
+			}
 		}
 	};
 	for (TPair<FName, FClassData>& RemainingPair : RemainingBatch)
@@ -511,6 +521,30 @@ void UtilsInitialize()
 	GEditorClassBlockList = ConstructClassBlockList();
 	GEditorPackageBlockList = ConstructPackageNameBlockList();
 
+	bool bTargetDomainClassUseBlockList = true;
+	if (FParse::Param(FCommandLine::Get(), TEXT("noiterate")))
+	{
+		// Allow list is marked as used, but is initialized empty
+		bTargetDomainClassUseBlockList = false;
+		GTargetDomainClassUseAllowList = true;
+		GTargetDomainClassEmptyAllowList = true;
+	}
+	else if (FParse::Param(FCommandLine::Get(), TEXT("iterate")))
+	{
+		bTargetDomainClassUseBlockList = false;
+		GTargetDomainClassUseAllowList = false;
+	}
+	else
+	{
+		GConfig->GetBool(TEXT("TargetDomain"), TEXT("IterativeClassAllowListEnabled"), GTargetDomainClassUseAllowList, GEditorIni);
+		GTargetDomainClassEmptyAllowList = false;
+	}
+
+	if (!GTargetDomainClassUseAllowList && bTargetDomainClassUseBlockList)
+	{
+		GTargetDomainClassBlockList = ConstructTargetIterativeClassBlockList();
+	}
+
 	// Constructing allowlists requires use of UStructs, and the early SetPackageResourceManager
 	// where UtilsInitialize is called is too early; trying to call UStruct->GetSchemaHash at that
 	// time will break the UClass. Defer the construction of allowlist-based data until OnPostEngineInit
@@ -523,7 +557,10 @@ void UtilsPostEngineInit()
 	GUtilsPostInitDelegate.Reset();
 
 	// Note that constructing AllowLists depends on all BlockLists having been parsed already
-	ConstructTargetIterativeClassAllowList();
+	if (GTargetDomainClassUseAllowList && !GTargetDomainClassEmptyAllowList)
+	{
+		ConstructTargetIterativeClassAllowList();
+	}
 }
 
 EPackageDigestResult GetPackageDigest(IAssetRegistry& AssetRegistry, FName PackageName,
@@ -551,7 +588,7 @@ EPackageDigestResult AppendPackageDigest(IAssetRegistry& AssetRegistry, FName Pa
 		return EPackageDigestResult::FileDoesNotExist;
 	}
 	EPackageDigestResult Result = AppendPackageDigest(Builder, bOutEditorDomainEnabled, OutErrorMessage, *PackageData, PackageName);
-	bOutEditorDomainEnabled &= !GetEditorPackageBlockList().Contains(PackageName);
+	bOutEditorDomainEnabled &= !GEditorPackageBlockList.Contains(PackageName);
 	return Result;
 }
 
diff --git a/Engine/Source/Editor/UnrealEd/Private/TargetDomain/TargetDomainUtils.cpp b/Engine/Source/Editor/UnrealEd/Private/TargetDomain/TargetDomainUtils.cpp
--- a/Engine/Source/Editor/UnrealEd/Private/TargetDomain/TargetDomainUtils.cpp
+++ b/Engine/Source/Editor/UnrealEd/Private/TargetDomain/TargetDomainUtils.cpp
@@ -104,9 +104,20 @@ bool TryCollectKeyAndDependencies(UPackage* Package, const ITargetPlatform* Targ
 
 	TArray<FName> SortedBuild;
 	SortedBuild = BuildDependencies.Array();
+	TStringBuilder<256> StringBuffer;
+	FName TransientPackageName = GetTransientPackage()->GetFName();
+	auto IsTransientPackageName = [&StringBuffer, TransientPackageName](FName PackageName)
+	{
+		PackageName.ToString(StringBuffer);
+		return PackageName == TransientPackageName ||
+			FPackageName::IsMemoryPackage(StringBuffer) ||
+			FPackageName::IsScriptPackage(StringBuffer);
+	};
+	SortedBuild.RemoveAllSwap(IsTransientPackageName, false /* bAllowShrinking */);
 	SortedBuild.Sort(FNameLexicalLess());
 	TArray<FName> SortedRuntimeOnly;
 	SortedRuntimeOnly = RuntimeOnlyDependencies.Array();
+	SortedRuntimeOnly.RemoveAllSwap(IsTransientPackageName, false /* bAllowShrinking */);
 	SortedRuntimeOnly.Sort(FNameLexicalLess());
 
 	if (!TryCreateKey(PackageName, SortedBuild, OutHash, OutErrorMessage))
diff --git a/Engine/Source/Editor/UnrealEd/Public/Commandlets/AssetRegistryGenerator.h b/Engine/Source/Editor/UnrealEd/Public/Commandlets/AssetRegistryGenerator.h
--- a/Engine/Source/Editor/UnrealEd/Public/Commandlets/AssetRegistryGenerator.h
+++ b/Engine/Source/Editor/UnrealEd/Public/Commandlets/AssetRegistryGenerator.h
@@ -44,16 +44,18 @@ public:
 	/** 
 	 * Sets asset registry from a previous run that is used for iterative or DLC cooking
 	 */
-	void SetPreviousAssetRegistry(TUniquePtr<FAssetRegistryState>&& PreviousState);
-	
+	void SetPreviousAssetRegistry(TUniquePtr<FAssetRegistryState>&& PreviousState,
+		TMap<FName, FAssetPackageData*>&& PreviousAssetPackageData, TArray<FName>& OutTombstonePackages);
+
 	/**
 	 * Options when computing the differences between current and previous state.
 	 */
 	struct FComputeDifferenceOptions
 	{
+		/** if true, modified packages are recursed to X in X->Y->Z chains. Otherwise, only Y and Z are seen as modified */
 		bool bRecurseModifications;
+		/** if true, modified script / c++ packages are recursed, if false only asset references are recursed */
 		bool bRecurseScriptModifications;
-		bool bIgnoreDiskSize = false;
 	};
 
 	/**
@@ -61,10 +63,15 @@ public:
 	 */
 	struct FAssetRegistryDifference
 	{
+		/** ModifiedPackages list of packages which existed beforeand now, but need to be recooked */
 		TSet<FName> ModifiedPackages;
+		/** NewPackages list of packages that did not exist before, but exist now */
 		TSet<FName> NewPackages;
+		/** RemovedPackages list of packages that existed before, but do not any more */
 		TSet<FName> RemovedPackages;
+		/** IdenticalCookedPackages list of cooked packages that have not changed */
 		TSet<FName> IdenticalCookedPackages;
+		/** IdenticalUncookedPackages list of uncooked packages that have not changed.These were filtered out by platform or editor only */
 		TSet<FName> IdenticalUncookedPackages;
 	};
 
@@ -77,26 +84,6 @@ public:
 	 */
 	void ComputePackageDifferences(const FComputeDifferenceOptions& Options, const TMap<FName, const FAssetPackageData*>& PreviousAssetPackageDataMap, FAssetRegistryDifference& OutDifference);
 
-	/**
-	 * Computes differences between the previous asset registry and the current one 
-	 *
-	 * @param ModifiedPackages list of packages which existed before and now, but need to be recooked
-	 * @param NewPackages list of packages that did not exist before, but exist now
-	 * @param RemovedPackages list of packages that existed before, but do not any more
-	 * @param IdenticalCookedPackages list of cooked packages that have not changed
-	 * @param IdenticalUncookedPackages list of uncooked packages that have not changed. These were filtered out by platform or editor only
-	 * @param bRecurseModifications if true, modified packages are recursed to X in X->Y->Z chains. Otherwise, only Y and Z are seen as modified
-	* @param bRecurseModifications if true, modified script/c++ packages are recursed, if false only asset references are recursed
-	 */
-	void ComputePackageDifferences(TSet<FName>& ModifiedPackages, TSet<FName>& NewPackages, TSet<FName>& RemovedPackages, TSet<FName>& IdenticalCookedPackages, TSet<FName>& IdenticalUncookedPackages, bool bRecurseModifications, bool bRecurseScriptModifications);
-
-	/**
-	 * Tracks packages that were kept from a previous cook.
-	 * Updates the current asset registry from the previous one for all kept packages.
-	 */
-	void UpdateKeptPackages(const TArray<FName>& InKeptPackages);
-	void UpdateKeptPackages(const TArray<FName>& InKeptPackages, TFunction<void(const FName&, FAssetPackageData*)>&& Update);
-
 	/**
 	 * GenerateChunkManifest 
 	 * generate chunk manifest for the packages passed in using the asset registry to determine dependencies
@@ -205,6 +192,15 @@ public:
 	 */
 	void GetChunkAssignments(TArray<TSet<FName>>& OutAssignments) const;
 
+	/**
+	 * Attempts to update the metadata for a package in an asset registry generator
+	 *
+	 * @param Package The package to update info on
+	 * @param SavePackageResult The metadata to associate with the given package name
+	 */
+	void UpdateAssetRegistryPackageData(const UPackage& Package, FSavePackageResultStruct& SavePackageResult);
+
+private:
 	/**
 	 * Ensures all assets in the input package are present in the registry
 	 * @param Package - Package to process
@@ -222,14 +218,16 @@ public:
 	 */
 	bool UpdateAssetPackageFlags(const FName& PackageName, const uint32 PackageFlags);
 
-private:
+	/**
+	 * Updates AssetData with previous TagsAndValues, and updates PackageData values with previous PackageData,
+	 * for all packages kept from a previous cook.
+	 */
+	void UpdateKeptPackages();
 
 	/** State of the asset registry that is being built for this platform */
 	FAssetRegistryState State;
 
-	/** Base state, which is either a release build or an iterative cook */
-	FAssetRegistryState PreviousState;
-
+	TMap<FName, TPair<TArray<FAssetData>, FAssetPackageData>> PreviousPackagesToUpdate;
 	/** List of packages that were loaded at startup */
 	TSet<FName> StartupPackages;
 	/** List of packages that were successfully cooked */
@@ -295,18 +293,6 @@ private:
 		uint32		ParentNodeIndex;
 	};
 
-	/**
-	 * Updates disk data with CookedHash and DiskSize from previous asset registry
-	 * for all packages kept from a previous cook.
-	 */
-	void UpdateKeptPackagesDiskData(const TArray<FName>& InKeptPackages);
-
-	/**
-	 * Updates AssetData with TagsAndValues from previous asset registry
-	 * for all packages kept from a previous cook.
-	 */
-	void UpdateKeptPackagesAssetData();
-
 	/**
 	 * Updates AssetData with TagsAndValues corresponding to any collections 
 	 * flagged for inclusion as asset registry tags.
diff --git a/Engine/Source/Runtime/CoreUObject/Private/UObject/SavePackage2.cpp b/Engine/Source/Runtime/CoreUObject/Private/UObject/SavePackage2.cpp
--- a/Engine/Source/Runtime/CoreUObject/Private/UObject/SavePackage2.cpp
+++ b/Engine/Source/Runtime/CoreUObject/Private/UObject/SavePackage2.cpp
@@ -2250,8 +2250,8 @@ FText GetSlowTaskStatusMessage(const FSaveContext& SaveContext)
 
 FSavePackageResultStruct UPackage::Save2(UPackage* InPackage, UObject* InAsset, const TCHAR* InFilename, FSavePackageArgs& SaveArgs)
 {
-	//COOK_STAT(FScopedDurationTimer FuncSaveTimer(FSavePackageStats::SavePackageTimeSec));
-	//COOK_STAT(FSavePackageStats::NumPackagesSaved++);
+	COOK_STAT(FScopedDurationTimer FuncSaveTimer(FSavePackageStats::SavePackageTimeSec));
+	COOK_STAT(FSavePackageStats::NumPackagesSaved++);
 	SCOPED_SAVETIMER(UPackage_Save2);
 	FSaveContext SaveContext(InPackage, InAsset, InFilename, SaveArgs);
 
