Updated Play and Debug buttons placement and style in toolbars.

diff --git a/Engine/Source/Editor/EditorStyle/Private/StarshipStyle.cpp b/Engine/Source/Editor/EditorStyle/Private/StarshipStyle.cpp
--- a/Engine/Source/Editor/EditorStyle/Private/StarshipStyle.cpp
+++ b/Engine/Source/Editor/EditorStyle/Private/StarshipStyle.cpp
@@ -5384,24 +5384,113 @@ void FStarshipEditorStyle::FStyle::SetupPersonaStyle()
 
 	// Play in editor / play in world
 	{
+		// Leftmost button for backplate style toolbar buttons
+		FToolBarStyle MainToolbarLeftButton = FStarshipCoreStyle::GetCoreStyle().GetWidgetStyle<FToolBarStyle>("AssetEditorToolbar");
+
+		const FButtonStyle LeftToolbarButton = FButtonStyle(MainToolbarLeftButton.ButtonStyle)
+			.SetNormal(FSlateRoundedBoxBrush(FStyleColors::Dropdown, FVector4(4.0f, 0.0f, 0.0f, 4.0f)))
+			.SetHovered(FSlateRoundedBoxBrush(FStyleColors::Dropdown, FVector4(4.0f, 0.0f, 0.0f, 4.0f)))
+			.SetPressed(FSlateRoundedBoxBrush(FStyleColors::Dropdown, FVector4(4.0f, 0.0f, 0.0f, 4.0f)))
+			.SetDisabled(FSlateRoundedBoxBrush(FStyleColors::Dropdown, FVector4(4.0f, 0.0f, 0.0f, 4.0f)))
+			.SetNormalPadding(FMargin(8.f, 4.f, 6.f, 4.f))
+			.SetPressedPadding(FMargin(8.f, 4.f, 6.f, 4.f));
+
+
+		MainToolbarLeftButton.SetButtonStyle(LeftToolbarButton);
+		MainToolbarLeftButton.SetButtonPadding(FMargin(8.f, 0.0f, 0.0f, 0.0f));
+		MainToolbarLeftButton.SetSeparatorPadding(FMargin(0.f, 0.f, 8.f, 0.f));
+
+		Set("Toolbar.BackplateLeft", MainToolbarLeftButton);
+
+		// Specialized Play Button (Left button with green color)
 		FLinearColor GreenHSV = FStyleColors::AccentGreen.GetSpecifiedColor().LinearRGBToHSV();
-		
 		FLinearColor GreenHover = FLinearColor(GreenHSV.R, GreenHSV.G * .5, GreenHSV.B, GreenHSV.A).HSVToLinearRGB();
 		FLinearColor GreenPress = FLinearColor(GreenHSV.R, GreenHSV.G, GreenHSV.B*.5, GreenHSV.A).HSVToLinearRGB(); 
 
-		FToolBarStyle PlayToolbar = FStarshipCoreStyle::GetCoreStyle().GetWidgetStyle<FToolBarStyle>("AssetEditorToolbar");
+		FToolBarStyle MainToolbarPlayButton = MainToolbarLeftButton;
 
-		const FButtonStyle PlayToolbarButton = FButtonStyle(PlayToolbar.SettingsButtonStyle)
+		const FButtonStyle PlayToolbarButton = FButtonStyle(MainToolbarPlayButton.ButtonStyle)
 			.SetNormalForeground(FStyleColors::AccentGreen)
 			.SetPressedForeground(GreenPress)
 			.SetHoveredForeground(GreenHover);
 
+		MainToolbarPlayButton.SetButtonStyle(PlayToolbarButton);
+
+		Set("Toolbar.BackplateLeftPlay", MainToolbarPlayButton);
+
+		// Center Buttons for backplate style toolbar buttons
+		FToolBarStyle MainToolbarCenterButton = MainToolbarLeftButton;
+
+		const FButtonStyle CenterToolbarButton = FButtonStyle(MainToolbarCenterButton.ButtonStyle)
+			.SetNormal(FSlateRoundedBoxBrush(FStyleColors::Dropdown, FVector4(0.0f, 0.0f, 0.0f, 0.0f)))
+			.SetHovered(FSlateRoundedBoxBrush(FStyleColors::Dropdown, FVector4(0.0f, 0.0f, 0.0f, 0.0f)))
+			.SetPressed(FSlateRoundedBoxBrush(FStyleColors::Dropdown, FVector4(0.0f, 0.0f, 0.0f, 0.0f)))
+			.SetDisabled(FSlateRoundedBoxBrush(FStyleColors::Dropdown, FVector4(0.0f, 0.0f, 0.0f, 0.0f)))
+			.SetNormalPadding(FMargin(2.f, 4.f, 6.f, 4.f))
+			.SetPressedPadding(FMargin(2.f, 4.f, 6.f, 4.f));
+
+		MainToolbarCenterButton.SetButtonPadding(0.0f);
+		MainToolbarCenterButton.SetButtonStyle(CenterToolbarButton);
+
+		Set("Toolbar.BackplateCenter", MainToolbarCenterButton);
+
+		// Specialized Stop Button (Center button + Red color)
+
+		FLinearColor RedHSV = FStyleColors::AccentRed.GetSpecifiedColor().LinearRGBToHSV();
+
+		FLinearColor RedHover = FLinearColor(RedHSV.R, RedHSV.G * .5, RedHSV.B, RedHSV.A).HSVToLinearRGB();
+		FLinearColor RedPress = FLinearColor(RedHSV.R, RedHSV.G, RedHSV.B * .5, RedHSV.A).HSVToLinearRGB();
+
+		FToolBarStyle MainToolbarStopButton = MainToolbarCenterButton;
+
+		const FButtonStyle StopToolbarButton = FButtonStyle(MainToolbarStopButton.ButtonStyle)
+			.SetNormalForeground(FStyleColors::AccentRed)
+			.SetPressedForeground(RedPress)
+			.SetHoveredForeground(RedHover);
+
+		MainToolbarStopButton.SetButtonStyle(StopToolbarButton);
+
+		Set("Toolbar.BackplateCenterStop", MainToolbarStopButton);
+
+		// Rightmost button for backplate style toolbar buttons
+		FToolBarStyle MainToolbarRightButton = MainToolbarLeftButton;
+
+		const FButtonStyle RightToolbarButton = FButtonStyle(MainToolbarRightButton.ButtonStyle)
+			.SetNormal(FSlateRoundedBoxBrush(FStyleColors::Dropdown, FVector4(0.0f, 4.0f, 4.0f, 0.0f)))
+			.SetHovered(FSlateRoundedBoxBrush(FStyleColors::Dropdown, FVector4(0.0f, 4.0f, 4.0f, 0.0f)))
+			.SetPressed(FSlateRoundedBoxBrush(FStyleColors::Dropdown, FVector4(0.0f, 4.0f, 4.0f, 0.0f)))
+			.SetDisabled(FSlateRoundedBoxBrush(FStyleColors::Dropdown, FVector4(0.0f, 4.0f, 4.0f, 0.0f)))
+			.SetNormalPadding(FMargin(2.f, 4.f, 8.f, 4.f))
+			.SetPressedPadding(FMargin(2.f, 4.f, 8.f, 4.f));
+
+		MainToolbarRightButton.SetButtonStyle(RightToolbarButton);
+		MainToolbarRightButton.SetButtonPadding(FMargin(0.0f, 0.0f, 8.0f, 0.0f));
+		MainToolbarRightButton.SetSeparatorPadding(FMargin(8.f, 0.f, 8.f, 0.f));
+
+		Set("Toolbar.BackplateRight", MainToolbarRightButton);
+
+		// Rightmost button for backplate style toolbar buttons as a combo button
+		FToolBarStyle MainToolbarRightComboButton = MainToolbarLeftButton;
+
+		const FButtonStyle RightToolbarComboButton = FButtonStyle(MainToolbarRightComboButton.ButtonStyle)
+			.SetNormal(FSlateRoundedBoxBrush(FStyleColors::Dropdown, FVector4(0.0f, 4.0f, 4.0f, 0.0f)))
+			.SetHovered(FSlateRoundedBoxBrush(FStyleColors::Dropdown, FVector4(0.0f, 4.0f, 4.0f, 0.0f)))
+			.SetPressed(FSlateRoundedBoxBrush(FStyleColors::Dropdown, FVector4(0.0f, 4.0f, 4.0f, 0.0f)))
+			.SetDisabled(FSlateRoundedBoxBrush(FStyleColors::Dropdown, FVector4(0.0f, 4.0f, 4.0f, 0.0f)))
+			.SetNormalPadding(FMargin(7.f, 4.f, 6.f, 4.f))
+			.SetPressedPadding(FMargin(7.f, 4.f, 6.f, 4.f));
 
-		PlayToolbar.SetSettingsButtonStyle(PlayToolbarButton);
+		FComboButtonStyle PlayToolbarComboButton = FComboButtonStyle(FStarshipCoreStyle::GetCoreStyle().GetWidgetStyle<FComboButtonStyle>("ComboButton"))
+			.SetDownArrowPadding(FMargin(-19.f, 0.f, 2.f, 0.f))
+			.SetDownArrowImage(CORE_IMAGE_BRUSH_SVG("Starship/Common/ellipsis-vertical-narrow", FVector2D(6, 24)));
+		PlayToolbarComboButton.ButtonStyle = RightToolbarComboButton;
 
-		PlayToolbar.SetShowLabels(true);
+		MainToolbarRightComboButton.SetButtonStyle(RightToolbarComboButton);
+		MainToolbarRightComboButton.SetComboButtonStyle(PlayToolbarComboButton);
+		MainToolbarRightComboButton.SetSeparatorPadding(FMargin(8.f, 0.f, 8.f, 0.f));
+		MainToolbarRightComboButton.SetComboButtonPadding(FMargin(1.0f, 0.0f, 8.0f, 0.0f));
 
-		Set("PlayToolBar", PlayToolbar);
+		Set("Toolbar.BackplateRightCombo", MainToolbarRightComboButton);
 
 		Set("PlayWorld.Simulate", new IMAGE_BRUSH_SVG("Starship/MainToolbar/simulate", Icon20x20));
 		Set( "PlayWorld.RepeatLastPlay", new IMAGE_BRUSH_SVG("Starship/MainToolbar/play", Icon20x20));
@@ -5434,14 +5523,14 @@ void FStarshipEditorStyle::FStyle::SetupPersonaStyle()
 		Set( "PlayWorld.PossessPlayer", new IMAGE_BRUSH_SVG("Starship/AssetIcons/PlayerController_16", Icon20x20));
 		Set( "PlayWorld.EjectFromPlayer", new IMAGE_BRUSH_SVG("Starship/MainToolbar/eject", Icon40x40));
 
-		Set( "PlayWorld.ShowCurrentStatement", new IMAGE_BRUSH( "Icons/icon_findnode_40x", Icon40x40 ) );
-		Set( "PlayWorld.ShowCurrentStatement.Small", new IMAGE_BRUSH( "Icons/icon_findnode_40x", Icon20x20 ) );
-		Set( "PlayWorld.StepOut", new IMAGE_BRUSH("Icons/icon_DebugStepOut_40x", Icon40x40));
-		Set( "PlayWorld.StepOut.Small", new IMAGE_BRUSH("Icons/icon_DebugStepOut_40x", Icon20x20));
-		Set( "PlayWorld.StepInto", new IMAGE_BRUSH( "Icons/icon_DebugStepIn_40x", Icon40x40 ) );
-		Set( "PlayWorld.StepInto.Small", new IMAGE_BRUSH( "Icons/icon_DebugStepIn_40x", Icon20x20 ) );
-		Set( "PlayWorld.StepOver", new IMAGE_BRUSH( "Icons/icon_DebugStepOver_40x", Icon40x40 ) );
-		Set( "PlayWorld.StepOver.Small", new IMAGE_BRUSH("Icons/icon_DebugStepOver_40x", Icon20x20));
+		Set( "PlayWorld.ShowCurrentStatement", new IMAGE_BRUSH_SVG( "Starship/MainToolbar/DebugFindNode", Icon40x40 ) );
+		Set( "PlayWorld.ShowCurrentStatement.Small", new IMAGE_BRUSH_SVG( "Starship/MainToolbar/DebugFindNode", Icon20x20 ) );
+		Set( "PlayWorld.StepOut", new IMAGE_BRUSH_SVG("Starship/MainToolbar/DebugStepOut", Icon40x40));
+		Set( "PlayWorld.StepOut.Small", new IMAGE_BRUSH_SVG("Starship/MainToolbar/DebugStepOut", Icon20x20));
+		Set( "PlayWorld.StepInto", new IMAGE_BRUSH_SVG( "Starship/MainToolbar/DebugStepInto", Icon40x40 ) );
+		Set( "PlayWorld.StepInto.Small", new IMAGE_BRUSH_SVG( "Starship/MainToolbar/DebugStepInto", Icon20x20 ) );
+		Set( "PlayWorld.StepOver", new IMAGE_BRUSH_SVG( "Starship/MainToolbar/DebugStepOver", Icon40x40 ) );
+		Set( "PlayWorld.StepOver.Small", new IMAGE_BRUSH_SVG("Starship/MainToolbar/DebugStepOver", Icon20x20));
 	}
 
 
diff --git a/Engine/Source/Editor/Kismet/Private/BlueprintEditor.cpp b/Engine/Source/Editor/Kismet/Private/BlueprintEditor.cpp
--- a/Engine/Source/Editor/Kismet/Private/BlueprintEditor.cpp
+++ b/Engine/Source/Editor/Kismet/Private/BlueprintEditor.cpp
@@ -3031,7 +3031,7 @@ void FBlueprintEditor::CreateDefaultCommands()
 		FExecuteAction::CreateSP(this, &FBlueprintEditor::ToggleHideUnrelatedNodes),
 		FCanExecuteAction(),
 		FIsActionChecked::CreateSP(this, &FBlueprintEditor::IsToggleHideUnrelatedNodesChecked),
-		FIsActionButtonVisible::CreateSP(this, &FBlueprintEditor::ShouldShowToggleHideUnrelatedNodes, true)
+		FIsActionButtonVisible()
 	);
 }
 
@@ -7948,23 +7948,6 @@ TSharedRef<SWidget> FBlueprintEditor::MakeHideUnrelatedNodesOptionsMenu()
 
 	MenuBuilder.AddWidget(OptionsHeading, FText::GetEmpty(), true);
 
-	TSharedPtr<FUICommandInfo> ToggleCmd = FBlueprintEditorCommands::Get().ToggleHideUnrelatedNodes;
-
-	// Add a menu version of toggle, when we can't show the full one
-	MenuBuilder.AddMenuEntry
-	(
-		ToggleCmd->GetLabel(),
-		ToggleCmd->GetDescription(),
-		FSlateIcon(),
-		FUIAction(
-			FExecuteAction::CreateSP(this, &FBlueprintEditor::ToggleHideUnrelatedNodes),
-			FCanExecuteAction(),
-			FIsActionChecked::CreateSP(this, &FBlueprintEditor::IsToggleHideUnrelatedNodesChecked),
-			FIsActionButtonVisible::CreateSP(this, &FBlueprintEditor::ShouldShowToggleHideUnrelatedNodes, false)),
-		NAME_None,
-		EUserInterfaceActionType::ToggleButton
-	);
-
 	MenuBuilder.AddMenuEntry(FUIAction(), LockNodeStateCheckBox);
 
 	return MenuBuilder.MakeWidget();
diff --git a/Engine/Source/Editor/Kismet/Private/SBlueprintEditorSelectedDebugObjectWidget.cpp b/Engine/Source/Editor/Kismet/Private/SBlueprintEditorSelectedDebugObjectWidget.cpp
--- a/Engine/Source/Editor/Kismet/Private/SBlueprintEditorSelectedDebugObjectWidget.cpp
+++ b/Engine/Source/Editor/Kismet/Private/SBlueprintEditorSelectedDebugObjectWidget.cpp
@@ -122,6 +122,7 @@ TSharedRef<SWidget> SBlueprintEditorSelectedDebugObjectWidget::OnGetActiveDetail
 		SNew(SHorizontalBox)
 		+ SHorizontalBox::Slot()
 		.AutoWidth()
+		.Padding(FMargin(8.0f, 0.0f, 0.0f, 0.0f))
 		[
 			DebugObjectsComboBox.ToSharedRef()
 		]
@@ -129,7 +130,7 @@ TSharedRef<SWidget> SBlueprintEditorSelectedDebugObjectWidget::OnGetActiveDetail
 		.AutoWidth()
 		.HAlign(HAlign_Right)
 		.VAlign(VAlign_Center)
-		.Padding(2.0f)
+		.Padding(4.0f)
 		[
 			BrowseButton
 		];
diff --git a/Engine/Source/Editor/Kismet/Public/BlueprintEditor.h b/Engine/Source/Editor/Kismet/Public/BlueprintEditor.h
--- a/Engine/Source/Editor/Kismet/Public/BlueprintEditor.h
+++ b/Engine/Source/Editor/Kismet/Public/BlueprintEditor.h
@@ -1415,6 +1415,8 @@ public:
 	/** Make nodes which are unrelated to the selected nodes fade out */
 	void ToggleHideUnrelatedNodes();
 	bool IsToggleHideUnrelatedNodesChecked() const;
+	
+	UE_DEPRECATED(5.0, "The Toggle Hide Unrelated Nodes button is always shown now")
 	bool ShouldShowToggleHideUnrelatedNodes(bool bIsToolbar) const;
 
 	/** Make a drop down menu to control the opacity of unrelated nodes */
diff --git a/Engine/Source/Editor/UnrealEd/Private/Kismet2/DebuggerCommands.cpp b/Engine/Source/Editor/UnrealEd/Private/Kismet2/DebuggerCommands.cpp
--- a/Engine/Source/Editor/UnrealEd/Private/Kismet2/DebuggerCommands.cpp
+++ b/Engine/Source/Editor/UnrealEd/Private/Kismet2/DebuggerCommands.cpp
@@ -471,7 +471,7 @@ void FPlayWorldCommands::BindGlobalPlayWorldCommands()
 		FExecuteAction::CreateStatic(&FInternalPlayWorldCommandCallbacks::StopPlaySession_Clicked),
 		FCanExecuteAction::CreateStatic(&FInternalPlayWorldCommandCallbacks::HasPlayWorld),
 		FIsActionChecked(),
-		FIsActionButtonVisible::CreateStatic(&FInternalPlayWorldCommandCallbacks::HasPlayWorld)
+		FIsActionButtonVisible()
 	);
 
 	// Late join session
@@ -501,7 +501,7 @@ void FPlayWorldCommands::BindGlobalPlayWorldCommands()
 		FExecuteAction::CreateStatic(&FInternalPlayWorldCommandCallbacks::SingleFrameAdvance_Clicked),
 		FCanExecuteAction::CreateStatic(&FPlayWorldCommandCallbacks::HasPlayWorldAndPaused),
 		FIsActionChecked(),
-		FIsActionChecked::CreateStatic(&FPlayWorldCommandCallbacks::HasPlayWorldAndPaused)
+		FIsActionButtonVisible()
 	);
 
 	ActionList.MapAction(Commands.TogglePlayPauseOfPlaySession,
@@ -524,7 +524,7 @@ void FPlayWorldCommands::BindGlobalPlayWorldCommands()
 		FExecuteAction::CreateStatic(&FInternalPlayWorldCommandCallbacks::PossessEjectPlayer_Clicked),
 		FCanExecuteAction::CreateStatic(&FInternalPlayWorldCommandCallbacks::CanPossessEjectPlayer),
 		FIsActionChecked(),
-		FIsActionButtonVisible::CreateStatic(&FInternalPlayWorldCommandCallbacks::CanPossessEjectPlayer)
+		FIsActionButtonVisible()
 	);
 
 	// Breakpoint-only commands
@@ -584,59 +584,96 @@ void FPlayWorldCommands::BuildToolbar(FToolMenuSection& InSection, bool bInclude
 	FToolMenuEntry PlayMenuEntry =
 		FToolMenuEntry::InitToolBarButton(
 			FPlayWorldCommands::Get().RepeatLastPlay,
-			LOCTEXT("RepeatLastPlay", "Play"),
+			TAttribute<FText>(),
 			TAttribute< FText >::Create(TAttribute< FText >::FGetter::CreateStatic(&FInternalPlayWorldCommandCallbacks::GetRepeatLastPlayToolTip)),
 			TAttribute< FSlateIcon >::Create(TAttribute< FSlateIcon >::FGetter::CreateStatic(&FInternalPlayWorldCommandCallbacks::GetRepeatLastPlayIcon)),
 			FName(TEXT("LevelToolbarPlay")));
-	PlayMenuEntry.StyleNameOverride = FName("PlayToolbar");
-
-	// Play combo box
-	FUIAction SpecialPIEOptionsMenuAction;
-	SpecialPIEOptionsMenuAction.IsActionVisibleDelegate = FIsActionButtonVisible::CreateStatic(&FInternalPlayWorldCommandCallbacks::CanShowNonPlayWorldOnlyActions);
-
-	PlayMenuEntry.AddOptionsDropdown(
-		SpecialPIEOptionsMenuAction,
-		FOnGetContent::CreateStatic(&GeneratePlayMenuContent, GlobalPlayWorldActions.ToSharedRef()),
-		LOCTEXT("PIEComboToolTip", "Change Play Mode and Play Settings")
-	);
+	PlayMenuEntry.StyleNameOverride = FName("Toolbar.BackplateLeftPlay");
 
 	// Play
 	InSection.AddEntry(PlayMenuEntry);
 
-	if (bIncludeLaunchButtonAndOptions)
-	{
-		ITurnkeySupportModule::Get().MakeTurnkeyMenu(InSection);
-	}
-
 	// Resume/pause toggle (only one will be visible, and only in PIE/SIE)
-	InSection.AddEntry(FToolMenuEntry::InitToolBarButton(FPlayWorldCommands::Get().ResumePlaySession, TAttribute<FText>(),
+	FToolMenuEntry ResumeEntry = FToolMenuEntry::InitToolBarButton(FPlayWorldCommands::Get().ResumePlaySession, TAttribute<FText>(),
 		TAttribute<FText>::Create(TAttribute<FText>::FGetter::CreateStatic(&FInternalPlayWorldCommandCallbacks::GetResumePlaySessionToolTip)),
 		TAttribute<FSlateIcon>::Create(TAttribute<FSlateIcon>::FGetter::CreateStatic(&FInternalPlayWorldCommandCallbacks::GetResumePlaySessionImage)),
-		FName(TEXT("ResumePlaySession"))
-	));
+		FName(TEXT("ResumePlaySession")));
+	ResumeEntry.StyleNameOverride = FName("Toolbar.BackplateLeftPlay");
+
+	FToolMenuEntry PauseEntry = FToolMenuEntry::InitToolBarButton(FPlayWorldCommands::Get().PausePlaySession, TAttribute<FText>(), TAttribute<FText>(), TAttribute<FSlateIcon>(), FName(TEXT("PausePlaySession")));
+	PauseEntry.StyleNameOverride = FName("Toolbar.BackplateLeft");
 
-	InSection.AddEntry(FToolMenuEntry::InitToolBarButton(FPlayWorldCommands::Get().PausePlaySession, TAttribute<FText>(), TAttribute<FText>(), TAttribute<FSlateIcon>(), FName(TEXT("PausePlaySession"))));
-	InSection.AddEntry(FToolMenuEntry::InitToolBarButton(FPlayWorldCommands::Get().SingleFrameAdvance, TAttribute<FText>(), TAttribute<FText>(), TAttribute<FSlateIcon>(), FName(TEXT("SingleFrameAdvance"))));
+	InSection.AddEntry(ResumeEntry);
+	InSection.AddEntry(PauseEntry);
+	
+	FToolMenuEntry SingleFrameAdvanceEntry = FToolMenuEntry::InitToolBarButton(FPlayWorldCommands::Get().SingleFrameAdvance, TAttribute<FText>(), TAttribute<FText>(), TAttribute<FSlateIcon>(), FName(TEXT("SingleFrameAdvance")));
+	SingleFrameAdvanceEntry.StyleNameOverride = FName("Toolbar.BackplateCenter");
+	InSection.AddEntry(SingleFrameAdvanceEntry);
 
 	// Stop
-	InSection.AddEntry(FToolMenuEntry::InitToolBarButton(FPlayWorldCommands::Get().StopPlaySession, TAttribute<FText>(), TAttribute<FText>(), TAttribute<FSlateIcon>(), FName(TEXT("StopPlaySession"))));
+	FToolMenuEntry StopEntry = FToolMenuEntry::InitToolBarButton(FPlayWorldCommands::Get().StopPlaySession, TAttribute<FText>(), TAttribute<FText>(), TAttribute<FSlateIcon>(), FName(TEXT("StopPlaySession")));
+	StopEntry.StyleNameOverride = FName("Toolbar.BackplateCenterStop");
+
+	InSection.AddEntry(StopEntry);
 
 	// Late Join
 	InSection.AddEntry(FToolMenuEntry::InitToolBarButton(FPlayWorldCommands::Get().LateJoinSession, TAttribute<FText>(), TAttribute<FText>(), TAttribute<FSlateIcon>(), FName(TEXT("LateJoinSession"))));
 
 	// Eject/possess toggle
-	InSection.AddEntry(FToolMenuEntry::InitToolBarButton(FPlayWorldCommands::Get().PossessEjectPlayer,
+
+	FToolMenuEntry EjectMenuEntry =
+	FToolMenuEntry::InitToolBarButton(FPlayWorldCommands::Get().PossessEjectPlayer,
 		TAttribute<FText>::Create(TAttribute<FText>::FGetter::CreateStatic(&FInternalPlayWorldCommandCallbacks::GetPossessEjectLabel)),
 		TAttribute<FText>::Create(TAttribute<FText>::FGetter::CreateStatic(&FInternalPlayWorldCommandCallbacks::GetPossessEjectTooltip)),
 		TAttribute<FSlateIcon>::Create(TAttribute<FSlateIcon>::FGetter::CreateStatic(&FInternalPlayWorldCommandCallbacks::GetPossessEjectImage)),
-		FName(TEXT("PossessEjectPlayer"))
-	));
+		FName(TEXT("PossessEjectPlayer")));
+	EjectMenuEntry.StyleNameOverride = FName("Toolbar.BackplateCenter");
+
+	// Play combo box
+	/*EjectMenuEntry.AddOptionsDropdown(
+		SpecialPIEOptionsMenuAction,
+		FOnGetContent::CreateStatic(&GeneratePlayMenuContent, GlobalPlayWorldActions.ToSharedRef()),
+		LOCTEXT("PIEComboToolTip", "Change Play Mode and Play Settings")
+	);*/
+
+	InSection.AddEntry(EjectMenuEntry);
+
+	FUIAction SpecialPIEOptionsMenuAction;
+	SpecialPIEOptionsMenuAction.CanExecuteAction = FCanExecuteAction::CreateStatic(&FInternalPlayWorldCommandCallbacks::CanShowNonPlayWorldOnlyActions);
+
+	FToolMenuEntry PIEComboEntry = FToolMenuEntry::InitComboButton("PIECombo", SpecialPIEOptionsMenuAction, FOnGetContent::CreateStatic(&GeneratePlayMenuContent, GlobalPlayWorldActions.ToSharedRef()), FText(), LOCTEXT("PIEComboToolTip", "Change Play Mode and Play Settings"));
+	PIEComboEntry.StyleNameOverride = FName("Toolbar.BackplateRightCombo");
+
+	InSection.AddEntry(PIEComboEntry);
 
 	// Single-stepping only buttons
-	InSection.AddEntry(FToolMenuEntry::InitToolBarButton(FPlayWorldCommands::Get().ShowCurrentStatement, TAttribute<FText>(), TAttribute<FText>(), TAttribute<FSlateIcon>(), FName(TEXT("ShowCurrentStatement"))));
-	InSection.AddEntry(FToolMenuEntry::InitToolBarButton(FPlayWorldCommands::Get().StepInto, TAttribute<FText>(), TAttribute<FText>(), TAttribute<FSlateIcon>(), FName(TEXT("StepInto"))));
-	InSection.AddEntry(FToolMenuEntry::InitToolBarButton(FPlayWorldCommands::Get().StepOver, TAttribute<FText>(), TAttribute<FText>(), TAttribute<FSlateIcon>(), FName(TEXT("StepOver"))));
-	InSection.AddEntry(FToolMenuEntry::InitToolBarButton(FPlayWorldCommands::Get().StepOut, TAttribute<FText>(), TAttribute<FText>(), TAttribute<FSlateIcon>(), FName(TEXT("StepOut"))));
+
+	InSection.AddSeparator(NAME_None).StyleNameOverride = FName("Toolbar.BackplateRight");
+	
+	FToolMenuEntry ShowCurrentStatementEntry = FToolMenuEntry::InitToolBarButton(FPlayWorldCommands::Get().ShowCurrentStatement, TAttribute<FText>(), TAttribute<FText>(), TAttribute<FSlateIcon>(), FName(TEXT("ShowCurrentStatement")));
+	ShowCurrentStatementEntry.StyleNameOverride = FName("Toolbar.BackplateLeft");
+
+	FToolMenuEntry StepIntoEntry = FToolMenuEntry::InitToolBarButton(FPlayWorldCommands::Get().StepInto, TAttribute<FText>(), TAttribute<FText>(), TAttribute<FSlateIcon>(), FName(TEXT("StepInto")));
+	StepIntoEntry.StyleNameOverride = FName("Toolbar.BackplateCenter");
+
+	FToolMenuEntry StepOverEntry = FToolMenuEntry::InitToolBarButton(FPlayWorldCommands::Get().StepOver, TAttribute<FText>(), TAttribute<FText>(), TAttribute<FSlateIcon>(), FName(TEXT("StepOver")));
+	StepOverEntry.StyleNameOverride = FName("Toolbar.BackplateCenter");
+
+	FToolMenuEntry StepOutEntry = FToolMenuEntry::InitToolBarButton(FPlayWorldCommands::Get().StepOut, TAttribute<FText>(), TAttribute<FText>(), TAttribute<FSlateIcon>(), FName(TEXT("StepOut")));
+	StepOutEntry.StyleNameOverride = FName("Toolbar.BackplateRight");
+
+	InSection.AddEntry(ShowCurrentStatementEntry);
+
+	InSection.AddEntry(StepIntoEntry);
+	
+	InSection.AddEntry(StepOverEntry);
+	
+	InSection.AddEntry(StepOutEntry);
+
+	if (bIncludeLaunchButtonAndOptions)
+	{
+		ITurnkeySupportModule::Get().MakeTurnkeyMenu(InSection);
+	}
 }
 
 static void MakePreviewDeviceMenu(FMenuBuilder& MenuBuilder)
@@ -1035,17 +1072,13 @@ FText FInternalPlayWorldCommandCallbacks::GetPossessEjectTooltip()
 
 FSlateIcon FInternalPlayWorldCommandCallbacks::GetPossessEjectImage()
 {
-	if (IsInPIE())
-	{
-		return FSlateIcon(FEditorStyle::GetStyleSetName(), "PlayWorld.EjectFromPlayer");
-	}
-	else if (IsInSIE())
+	if (IsInSIE())
 	{
 		return FSlateIcon(FEditorStyle::GetStyleSetName(), "PlayWorld.PossessPlayer");
 	}
 	else
 	{
-		return FSlateIcon();
+		return FSlateIcon(FEditorStyle::GetStyleSetName(), "PlayWorld.EjectFromPlayer");
 	}
 }
 
