Change cleanbuild -> fullbuild to match our public-facing terminology for clean aka full aka noniterative builds.
#rnx

diff --git a/Engine/Source/Developer/IoStoreUtilities/Private/IoStoreUtilities.cpp b/Engine/Source/Developer/IoStoreUtilities/Private/IoStoreUtilities.cpp
--- a/Engine/Source/Developer/IoStoreUtilities/Private/IoStoreUtilities.cpp
+++ b/Engine/Source/Developer/IoStoreUtilities/Private/IoStoreUtilities.cpp
@@ -4046,7 +4046,7 @@ int32 Staged2Zen(const FString& BuildPath, const FKeyChain& KeyChain, const ITar
 	TUniquePtr<FZenStoreWriter> ZenStoreWriter = MakeUnique<FZenStoreWriter>(CookedOutputPath, MetaDataOutputPath, TargetPlatform);
 	
 	ICookedPackageWriter::FCookInfo CookInfo;
-	CookInfo.bCleanBuild = true;
+	CookInfo.bFullBuild = true;
 	ZenStoreWriter->BeginCook(CookInfo);
 	int32 LocalPackageIndex = 0;
 	TArray<FPackageInfo> PackagesArray;
diff --git a/Engine/Source/Developer/IoStoreUtilities/Private/ZenStoreHttpClient.cpp b/Engine/Source/Developer/IoStoreUtilities/Private/ZenStoreHttpClient.cpp
--- a/Engine/Source/Developer/IoStoreUtilities/Private/ZenStoreHttpClient.cpp
+++ b/Engine/Source/Developer/IoStoreUtilities/Private/ZenStoreHttpClient.cpp
@@ -112,11 +112,11 @@ FZenStoreHttpClient::Initialize(FStringView InProjectId,
 	bAllowEdit = true;
 }
 
-void FZenStoreHttpClient::EstablishWritableOpLog(FStringView InProjectId, FStringView InOplogId, bool bCleanBuild)
+void FZenStoreHttpClient::EstablishWritableOpLog(FStringView InProjectId, FStringView InOplogId, bool bFullBuild)
 {
 	UE::Zen::FZenScopedRequestPtr Request(RequestPool.Get());
 
-	if (bCleanBuild)
+	if (bFullBuild)
 	{
 		UE_LOG(LogZenStore, Display, TEXT("Deleting oplog '%s'/'%s' if it exists"), *FString(InProjectId), *FString(InOplogId));
 		Request->PerformBlockingDelete(OplogPath);
@@ -562,7 +562,7 @@ void FZenStoreHttpClient::Initialize(
 {
 }
 
-void FZenStoreHttpClient::EstablishWritableOpLog(FStringView InProjectId, FStringView InOplogId, bool bCleanBuild)
+void FZenStoreHttpClient::EstablishWritableOpLog(FStringView InProjectId, FStringView InOplogId, bool bFullBuild)
 {
 }
 
diff --git a/Engine/Source/Developer/IoStoreUtilities/Private/ZenStoreHttpClient.h b/Engine/Source/Developer/IoStoreUtilities/Private/ZenStoreHttpClient.h
--- a/Engine/Source/Developer/IoStoreUtilities/Private/ZenStoreHttpClient.h
+++ b/Engine/Source/Developer/IoStoreUtilities/Private/ZenStoreHttpClient.h
@@ -29,7 +29,7 @@ public:
 					FStringView ServerRoot, 
 					FStringView EngineRoot, 
 					FStringView ProjectRoot);
-	void EstablishWritableOpLog(FStringView InProjectId, FStringView InOplogId, bool bCleanBuild);
+	void EstablishWritableOpLog(FStringView InProjectId, FStringView InOplogId, bool bFullBuild);
 
 	void InitializeReadOnly(FStringView InProjectId, FStringView InOplogId);
 
diff --git a/Engine/Source/Developer/IoStoreUtilities/Private/ZenStoreWriter.cpp b/Engine/Source/Developer/IoStoreUtilities/Private/ZenStoreWriter.cpp
--- a/Engine/Source/Developer/IoStoreUtilities/Private/ZenStoreWriter.cpp
+++ b/Engine/Source/Developer/IoStoreUtilities/Private/ZenStoreWriter.cpp
@@ -423,9 +423,9 @@ void FZenStoreWriter::BeginCook(const FCookInfo& Info)
 	{
 		FString ProjectId = FApp::GetProjectName();
 		FString OplogId = TargetPlatform.PlatformName();
-		HttpClient->EstablishWritableOpLog(ProjectId, OplogId, Info.bCleanBuild);
+		HttpClient->EstablishWritableOpLog(ProjectId, OplogId, Info.bFullBuild);
 
-		if (!Info.bCleanBuild)
+		if (!Info.bFullBuild)
 		{
 			UE_LOG(LogZenStoreWriter, Display, TEXT("Fetching oplog..."), *ProjectId, *OplogId);
 
@@ -505,7 +505,7 @@ void FZenStoreWriter::BeginCook(const FCookInfo& Info)
 	}
 	else
 	{
-		if (Info.bCleanBuild)
+		if (Info.bFullBuild)
 		{
 			RemoveCookedPackages();
 		}
diff --git a/Engine/Source/Editor/UnrealEd/Private/CookOnTheFlyServer.cpp b/Engine/Source/Editor/UnrealEd/Private/CookOnTheFlyServer.cpp
--- a/Engine/Source/Editor/UnrealEd/Private/CookOnTheFlyServer.cpp
+++ b/Engine/Source/Editor/UnrealEd/Private/CookOnTheFlyServer.cpp
@@ -7311,7 +7311,7 @@ void UCookOnTheFlyServer::BeginCookSandbox(TConstArrayView<const ITargetPlatform
 		}
 
 		const bool bIsDiffOnly = FParse::Param(FCommandLine::Get(), TEXT("DIFFONLY"));
-		const bool bIterative = !FParse::Param(FCommandLine::Get(), TEXT("noiterate")) && (bHybridIterativeEnabled || IsCookFlagSet(ECookInitializationFlags::Iterative));
+		const bool bIterative = !FParse::Param(FCommandLine::Get(), TEXT("fullcook")) && (bHybridIterativeEnabled || IsCookFlagSet(ECookInitializationFlags::Iterative));
 		const bool bIsSharedIterativeCook = IsCookFlagSet(ECookInitializationFlags::IterateSharedBuild);
 		TArray<TPair<const ITargetPlatform*,bool>, TInlineAllocator<ExpectedMaxNumPlatforms>> ResetPlatforms;
 		TArray<const ITargetPlatform*, TInlineAllocator<ExpectedMaxNumPlatforms>> PopulatePlatforms;
@@ -7384,12 +7384,12 @@ void UCookOnTheFlyServer::BeginCookSandbox(TConstArrayView<const ITargetPlatform
 
 			ICookedPackageWriter::FCookInfo CookInfo;
 			CookInfo.CookMode = IsCookOnTheFlyMode() ? ICookedPackageWriter::FCookInfo::CookOnTheFlyMode : ICookedPackageWriter::FCookInfo::CookByTheBookMode;
-			CookInfo.bCleanBuild = false;
+			CookInfo.bFullBuild = false;
 			CookInfo.bIterateSharedBuild = bIterateSharedBuild && !bShouldClearCookedContent;
 			bool bResetResults = false;
 			if (bShouldClearCookedContent)
 			{
-				CookInfo.bCleanBuild = true;
+				CookInfo.bFullBuild = true;
 				bResetResults = true;
 			}
 			else if (!PlatformData->bIsSandboxInitialized)
@@ -7402,7 +7402,7 @@ void UCookOnTheFlyServer::BeginCookSandbox(TConstArrayView<const ITargetPlatform
 				}
 			}
 			PackageWriter.BeginCook(CookInfo);
-			PlatformData->bCleanBuild = CookInfo.bCleanBuild;
+			PlatformData->bFullBuild = CookInfo.bFullBuild;
 			PlatformData->bIsSandboxInitialized = true;
 			ResetPlatforms.Emplace(TargetPlatform, bResetResults);
 		}
diff --git a/Engine/Source/Editor/UnrealEd/Private/Cooker/CookPlatformManager.h b/Engine/Source/Editor/UnrealEd/Private/Cooker/CookPlatformManager.h
--- a/Engine/Source/Editor/UnrealEd/Private/Cooker/CookPlatformManager.h
+++ b/Engine/Source/Editor/UnrealEd/Private/Cooker/CookPlatformManager.h
@@ -47,8 +47,8 @@ namespace Cook
 		/* Whether BeginCookSandbox has been called for this platform.  If we have already initialized the sandbox we can take a faster refresh path on future sessions. */
 		bool bIsSandboxInitialized = false;
 
-		/* Whether BeginCookSandbox specified a clean build for this platform. */
-		bool bCleanBuild = false;
+		/* Whether BeginCookSandbox specified a full, non-iterative build for this platform. */
+		bool bFullBuild = false;
 
 		/*
 		 * The last FPlatformTime::GetSeconds() at which this platform was requested in a CookOnTheFly request.
diff --git a/Engine/Source/Editor/UnrealEd/Private/Cooker/CookRequestCluster.cpp b/Engine/Source/Editor/UnrealEd/Private/Cooker/CookRequestCluster.cpp
--- a/Engine/Source/Editor/UnrealEd/Private/Cooker/CookRequestCluster.cpp
+++ b/Engine/Source/Editor/UnrealEd/Private/Cooker/CookRequestCluster.cpp
@@ -194,7 +194,7 @@ void FRequestCluster::Initialize(UCookOnTheFlyServer& COTFS)
 	}
 
 	PackageWriters.Reserve(Platforms.Num());
-	bCleanBuild = false;
+	bFullBuild = false;
 	bool bFirst = true;
 	for (const ITargetPlatform* TargetPlatform : Platforms)
 	{
@@ -202,15 +202,15 @@ void FRequestCluster::Initialize(UCookOnTheFlyServer& COTFS)
 		FPlatformData* PlatformData = COTFS.PlatformManager->GetPlatformData(TargetPlatform);
 		if (bFirst)
 		{
-			bCleanBuild = PlatformData->bCleanBuild;
+			bFullBuild = PlatformData->bFullBuild;
 			bFirst = false;
 		}
 		else
 		{
-			if (PlatformData->bCleanBuild != bCleanBuild)
+			if (PlatformData->bFullBuild != bFullBuild)
 			{
-				UE_LOG(LogCook, Warning, TEXT("Clean build is requested for some platforms but not others, but this is not supported. All platforms will be built clean."));
-				bCleanBuild = true;
+				UE_LOG(LogCook, Warning, TEXT("Full build is requested for some platforms but not others, but this is not supported. All platforms will be built full."));
+				bFullBuild = true;
 			}
 		}
 	}
@@ -650,7 +650,7 @@ void FRequestCluster::VisitPackageData(FPackageData* PackageData, TArray<FName>*
 		}
 		AssetRegistry.GetDependencies(PackageName, BuildDependencies,
 			UE::AssetRegistry::EDependencyCategory::Package, DependencyQuery);
-		if (bHybridIterativeEnabled && !bCleanBuild)
+		if (bHybridIterativeEnabled && !bFullBuild)
 		{
 			bool bExists = false;
 			for (int32 PlatIndex = 0; PlatIndex < Platforms.Num(); ++PlatIndex)
diff --git a/Engine/Source/Editor/UnrealEd/Private/Cooker/CookRequestCluster.h b/Engine/Source/Editor/UnrealEd/Private/Cooker/CookRequestCluster.h
--- a/Engine/Source/Editor/UnrealEd/Private/Cooker/CookRequestCluster.h
+++ b/Engine/Source/Editor/UnrealEd/Private/Cooker/CookRequestCluster.h
@@ -127,7 +127,7 @@ private:
 	bool bErrorOnEngineContentUse = false;
 	bool bPackageNamesComplete = false;
 	bool bDependenciesComplete = false;
-	bool bCleanBuild = false;
+	bool bFullBuild = false;
 
 
 	/** Data used during the reentrant dependencies and topological sort operation of FRequestCluster. */
diff --git a/Engine/Source/Editor/UnrealEd/Private/Cooker/LooseCookedPackageWriter.cpp b/Engine/Source/Editor/UnrealEd/Private/Cooker/LooseCookedPackageWriter.cpp
--- a/Engine/Source/Editor/UnrealEd/Private/Cooker/LooseCookedPackageWriter.cpp
+++ b/Engine/Source/Editor/UnrealEd/Private/Cooker/LooseCookedPackageWriter.cpp
@@ -83,7 +83,7 @@ FDateTime FLooseCookedPackageWriter::GetPreviousCookTime() const
 void FLooseCookedPackageWriter::BeginCook(const FCookInfo& Info)
 {
 	bIterateSharedBuild = Info.bIterateSharedBuild;
-	if (Info.bCleanBuild)
+	if (Info.bFullBuild)
 	{
 		DeleteSandboxDirectory();
 	}
diff --git a/Engine/Source/Editor/UnrealEd/Private/EditorDomain/EditorDomainUtils.cpp b/Engine/Source/Editor/UnrealEd/Private/EditorDomain/EditorDomainUtils.cpp
--- a/Engine/Source/Editor/UnrealEd/Private/EditorDomain/EditorDomainUtils.cpp
+++ b/Engine/Source/Editor/UnrealEd/Private/EditorDomain/EditorDomainUtils.cpp
@@ -522,7 +522,7 @@ void UtilsInitialize()
 	GEditorPackageBlockList = ConstructPackageNameBlockList();
 
 	bool bTargetDomainClassUseBlockList = true;
-	if (FParse::Param(FCommandLine::Get(), TEXT("noiterate")))
+	if (FParse::Param(FCommandLine::Get(), TEXT("fullcook")))
 	{
 		// Allow list is marked as used, but is initialized empty
 		bTargetDomainClassUseBlockList = false;
diff --git a/Engine/Source/Runtime/Core/Public/Serialization/PackageWriter.h b/Engine/Source/Runtime/Core/Public/Serialization/PackageWriter.h
--- a/Engine/Source/Runtime/Core/Public/Serialization/PackageWriter.h
+++ b/Engine/Source/Runtime/Core/Public/Serialization/PackageWriter.h
@@ -157,7 +157,7 @@ public:
 		};
 
 		ECookMode CookMode = CookByTheBookMode;
-		bool bCleanBuild = true;
+		bool bFullBuild = true;
 		bool bIterateSharedBuild = false;
 	};
 
