Add profiling data to the Mirage local caching operation.

#rnx

- The large white space change in the cpp, is to fix a formatting issue introduced in Rev2. I am not sure how the file got so messed up which seems to have happened inbetween the review stage and the submit.
- The debug config option 'ValidateAfterPushOperation' now works when caching data locally inadditon to pushing data to persistent storage.
- Collect profiling data when caching payloads to local storage.
- Add caching statistics to the final log output and the Mirage editor widget although this widget is currently slated to be merged with the DDC widget in the near future.

diff --git a/Engine/Source/Editor/UnrealEd/Private/Virtualization/SVirtualizationStatusIndicator.cpp b/Engine/Source/Editor/UnrealEd/Private/Virtualization/SVirtualizationStatusIndicator.cpp
--- a/Engine/Source/Editor/UnrealEd/Private/Virtualization/SVirtualizationStatusIndicator.cpp
+++ b/Engine/Source/Editor/UnrealEd/Private/Virtualization/SVirtualizationStatusIndicator.cpp
@@ -16,177 +16,182 @@
 namespace UE::Virtualization
 {
 
-	/** The amount of time (s) that the indication arrows will take to fade in/out */
-	static constexpr float ArrowFadeTime = 0.5f;
-
-	/** 
-	* The tool tip that will show when the mouse is hovered over SVirtualizationStatusIndicator. 
-	* 
-	* Currently this will display the overall amount of data that content virtualization has 
-	* transfered to and from the backends.
-	*/
-	class SVirtualizationToolTip : public SCompoundWidget
-	{
-		SLATE_BEGIN_ARGS(SVirtualizationToolTip) {}
-		SLATE_END_ARGS()
-
-			void Construct(const FArguments& InArgs);
-
-	private:
-		FText GetToolTipText() const;
-	};
-
-	/** 
-	* A widget used to show the status of content virtualization.
-	* 
-	* A downwards pointing arrow will be displayed when ever a payload is pulled from a backend.
-	* An upwards pointing arrow will be displayed when ever a payload is pushed to a backend.
-	* This gives the user immediately visual feedback as to when the system is in action.
-	* Additional data can be found when hovering the mouse of the widget.
-	*/
-	class SVirtualizationStatusIndicator : public SCompoundWidget
-	{
-		SLATE_BEGIN_ARGS(SVirtualizationStatusIndicator) {}
-		SLATE_END_ARGS()
-
-			void Construct(const FArguments& InArgs);
+/** The amount of time (s) that the indication arrows will take to fade in/out */
+static constexpr float ArrowFadeTime = 0.5f;
+
+/**
+* The tool tip that will show when the mouse is hovered over SVirtualizationStatusIndicator.
+*
+* Currently this will display the overall amount of data that content virtualization has
+* transfered to and from the backends.
+*/
+class SVirtualizationToolTip : public SCompoundWidget
+{
+	SLATE_BEGIN_ARGS(SVirtualizationToolTip) {}
+	SLATE_END_ARGS()
+
+	void Construct(const FArguments& InArgs);
+
+private:
+	FText GetToolTipText() const;
+};
+
+/**
+* A widget used to show the status of content virtualization.
+*
+* A downwards pointing arrow will be displayed when ever a payload is pulled from a backend.
+* An upwards pointing arrow will be displayed when ever a payload is pushed to a backend.
+* This gives the user immediately visual feedback as to when the system is in action.
+* Additional data can be found when hovering the mouse of the widget.
+*/
+class SVirtualizationStatusIndicator : public SCompoundWidget
+{
+	SLATE_BEGIN_ARGS(SVirtualizationStatusIndicator) {}
+	SLATE_END_ARGS()
 
-	private:
+	void Construct(const FArguments& InArgs);
 
-		EActiveTimerReturnType UpdateLastDataAccess(double InCurrentTime, float InDeltaTime);
+private:
 
-		int64 PayloadsPulled = 0;
-		int64 PayloadsPushed = 0;
+	EActiveTimerReturnType UpdateLastDataAccess(double InCurrentTime, float InDeltaTime);
+	
+	int64 PayloadsPulled = 0;
+	int64 PayloadsUploaded = 0; // Covers both pushing to persistent backends and caching to local ones	
 
-		FCurveSequence FadePullArrow;
-		FCurveSequence FadePushArrow;
-	};
+	FCurveSequence FadePullArrow;
+	FCurveSequence FadeUploadArrow;
+};
 
-	void SVirtualizationToolTip::Construct(const FArguments& InArgs)
-	{
-		this->ChildSlot
-			[
-				SNew(SVerticalBox)
+void SVirtualizationToolTip::Construct(const FArguments& InArgs)
+{
+	this->ChildSlot
+	[
+		SNew(SVerticalBox)
+
+		+ SVerticalBox::Slot()
+		.AutoHeight()
+		[
+			SNew(STextBlock)
+			.Text(this, &SVirtualizationToolTip::GetToolTipText)
+		]
+	];
+}
+
+FText SVirtualizationToolTip::GetToolTipText() const
+{
+	using namespace UE::Virtualization;
 
-				+ SVerticalBox::Slot()
-			.AutoHeight()
-			[
-				SNew(STextBlock)
-				.Text(this, &SVirtualizationToolTip::GetToolTipText)
-			]
-			];
-	}
+	FPayloadActivityInfo Info = FVirtualizationManager::Get().GetPayloadActivityInfo();
+	TStringBuilder<1024> Builder;
 
-	FText SVirtualizationToolTip::GetToolTipText() const
-	{
-		using namespace UE::Virtualization;
+	Builder << TEXT("Cached: ") << Info.Cache.TotalBytes / (int64)(1024 * 1024) << TEXT(" MB in ");
+	Builder.Appendf(TEXT("%.3f Seconds\n"), Info.Cache.CyclesSpent * FPlatformTime::GetSecondsPerCycle());
 
-		FPayloadActivityInfo Info = FVirtualizationManager::Get().GetPayloadActivityInfo();
-		TStringBuilder<1024> Builder;
+	Builder << TEXT("Pushed: ") << Info.Push.TotalBytes / (int64)(1024 * 1024) << TEXT(" MB in ");
+	Builder.Appendf(TEXT("%.3f Seconds\n"), Info.Push.CyclesSpent * FPlatformTime::GetSecondsPerCycle());
 
-		Builder << TEXT("Pushed: ") << Info.Push.TotalBytes / (int64)(1024 * 1024) << TEXT(" MB in ");
-		Builder.Appendf(TEXT("%.3f Seconds\n"), Info.Push.CyclesSpent * FPlatformTime::GetSecondsPerCycle());
+	Builder << TEXT("Pulled: ") << Info.Pull.TotalBytes / (int64)(1024 * 1024) << TEXT(" MB in ");
+	Builder.Appendf(TEXT("%.3f Seconds"), Info.Pull.CyclesSpent * FPlatformTime::GetSecondsPerCycle());
 
-		Builder << TEXT("Pulled: ") << Info.Pull.TotalBytes / (int64)(1024 * 1024) << TEXT(" MB in ");
-		Builder.Appendf(TEXT("%.3f Seconds"), Info.Pull.CyclesSpent  * FPlatformTime::GetSecondsPerCycle());
+	return FText::FromString(Builder.ToString());
+}
 
-		return FText::FromString(Builder.ToString());
-	}
+void SVirtualizationStatusIndicator::Construct(const FArguments& InArgs)
+{
+	using namespace UE::Virtualization;
 
-	void SVirtualizationStatusIndicator::Construct(const FArguments& InArgs)
-	{
-		using namespace UE::Virtualization;
+	FPayloadActivityInfo Info = FVirtualizationManager::Get().GetPayloadActivityInfo();
+	PayloadsPulled = Info.Pull.PayloadCount;
+	PayloadsUploaded = Info.Push.PayloadCount + Info.Cache.PayloadCount;
 
-		FPayloadActivityInfo Info = FVirtualizationManager::Get().GetPayloadActivityInfo();
-		PayloadsPulled = Info.Pull.PayloadCount;
-		PayloadsPushed = Info.Push.PayloadCount;
+	FadeUploadArrow = FCurveSequence(0.0f, ArrowFadeTime, ECurveEaseFunction::Linear);
+	FadePullArrow = FCurveSequence(0.0f, ArrowFadeTime, ECurveEaseFunction::Linear);
 
-		FadePushArrow = FCurveSequence(0.0f, ArrowFadeTime, ECurveEaseFunction::Linear);
-		FadePullArrow = FCurveSequence(0.0f, ArrowFadeTime, ECurveEaseFunction::Linear);
+	const FText MirageLabel = FText::FromString(TEXT("Mirage"));
 
-		const FText MirageLabel = FText::FromString(TEXT("Mirage"));
+	this->ChildSlot
+	[
+		SNew(SHorizontalBox)
+		.ToolTip(SNew(SToolTip)[SNew(SVirtualizationToolTip)])
 
-		this->ChildSlot
-			[
-				SNew(SHorizontalBox)
-				.ToolTip(SNew(SToolTip)[SNew(SVirtualizationToolTip)])
-
-			+ SHorizontalBox::Slot()
-			.AutoWidth()
-			.VAlign(VAlign_Center)
-			.Padding(0, 0, 3, 0)
-			[
-				SNew(SOverlay)
+		+ SHorizontalBox::Slot()
+		.AutoWidth()
+		.VAlign(VAlign_Center)
+		.Padding(0, 0, 3, 0)
+		[
+			SNew(SOverlay)
 
-				+ SOverlay::Slot()
+			+ SOverlay::Slot()
 			.HAlign(HAlign_Center)
 			.VAlign(VAlign_Top)
 			.Padding(0, 0, 4, 4)
 			[
 				SNew(SImage)
 				.Image(FAppStyle::Get().GetBrush("Icons.ArrowUp"))
-			.ColorAndOpacity_Lambda([this] { return FLinearColor::Green.CopyWithNewOpacity(FadePushArrow.GetLerp()); })
+				.ColorAndOpacity_Lambda([this] { return FLinearColor::Green.CopyWithNewOpacity(FadeUploadArrow.GetLerp()); })
 			]
 
-		+ SOverlay::Slot()
+			+ SOverlay::Slot()
 			.HAlign(HAlign_Center)
 			.VAlign(VAlign_Bottom)
 			.Padding(4, 4, 0, 0)
 			[
 				SNew(SImage)
 				.Image(FAppStyle::Get().GetBrush("Icons.ArrowDown"))
-			.ColorAndOpacity_Lambda([this] { return FLinearColor::Green.CopyWithNewOpacity(FadePullArrow.GetLerp()); })
-			]
+				.ColorAndOpacity_Lambda([this] { return FLinearColor::Green.CopyWithNewOpacity(FadePullArrow.GetLerp()); })
 			]
+		]
 
 		+ SHorizontalBox::Slot()
-			.AutoWidth()
-			.VAlign(VAlign_Center)
-			.Padding(0, 0, 10, 0)
-			[
-				SNew(STextBlock)
-				.Text(MirageLabel)
-			]
-			];
+		.AutoWidth()
+		.VAlign(VAlign_Center)
+		.Padding(0, 0, 10, 0)
+		[
+			SNew(STextBlock)
+			.Text(MirageLabel)
+		]
+	];
+
+	RegisterActiveTimer(ArrowFadeTime, FWidgetActiveTimerDelegate::CreateSP(this, &SVirtualizationStatusIndicator::UpdateLastDataAccess));
+}
+
+EActiveTimerReturnType SVirtualizationStatusIndicator::UpdateLastDataAccess(double InCurrentTime, float InDeltaTime)
+{
+	using namespace UE::Virtualization;
 
-		RegisterActiveTimer(ArrowFadeTime, FWidgetActiveTimerDelegate::CreateSP(this, &SVirtualizationStatusIndicator::UpdateLastDataAccess));
-	}
+	FPayloadActivityInfo Info = FVirtualizationManager::Get().GetPayloadActivityInfo();
 
-	EActiveTimerReturnType SVirtualizationStatusIndicator::UpdateLastDataAccess(double InCurrentTime, float InDeltaTime)
-	{
-		using namespace UE::Virtualization;
+	// Don't change the fade if it is currently playing, the tick duration of this method is not that accurate
+	// and we don't want to cut out the previous fade too soon or the user might not be able to notice it.
 
-		FPayloadActivityInfo Info = FVirtualizationManager::Get().GetPayloadActivityInfo();
+	if (!FadePullArrow.IsPlaying())
+	{
+		FadePullArrow.PlayRelative(this->AsShared(), PayloadsPulled != Info.Pull.PayloadCount);
+		PayloadsPulled = Info.Pull.PayloadCount;
+	}
 
-		// Don't change the fade if it is currently playing, the tick duration of this method is not that accurate
-		// and we don't want to cut out the previous fade too soon or the user might not be able to notice it.
+	if (!FadeUploadArrow.IsPlaying())
+	{
+		int64 NewPayloadsUploaded = Info.Push.PayloadCount + Info.Cache.PayloadCount;
 
-		if (!FadePullArrow.IsPlaying())
-		{
-			FadePullArrow.PlayRelative(this->AsShared(), PayloadsPulled != Info.Pull.PayloadCount);
-			PayloadsPulled = Info.Pull.PayloadCount;
-		}
+		FadeUploadArrow.PlayRelative(this->AsShared(), PayloadsUploaded != NewPayloadsUploaded);
+		PayloadsUploaded = NewPayloadsUploaded;
+	}
 
-		if (!FadePushArrow.IsPlaying())
-		{
-			FadePushArrow.PlayRelative(this->AsShared(), PayloadsPushed != Info.Push.PayloadCount);
-			PayloadsPushed = Info.Push.PayloadCount;
-		}
+	return EActiveTimerReturnType::Continue;
+}
 
-		return EActiveTimerReturnType::Continue;
+TSharedPtr<SWidget> GetVirtualizationStatusIndicator()
+{
+	if (FVirtualizationManager::Get().IsEnabled())
+	{
+		return SNew(SVirtualizationStatusIndicator);
 	}
-
-	TSharedPtr<SWidget> GetVirtualizationStatusIndicator()
+	else
 	{
-		if (FVirtualizationManager::Get().IsEnabled())
-		{
-			return SNew(SVirtualizationStatusIndicator);
-		}
-		else
-		{
-			return nullptr;
-		}
+		return nullptr;
 	}
+}
 
 } // namespace UE::Virtualization
 
diff --git a/Engine/Source/Runtime/CoreUObject/Private/Virtualization/VirtualizationManager.cpp b/Engine/Source/Runtime/CoreUObject/Private/Virtualization/VirtualizationManager.cpp
--- a/Engine/Source/Runtime/CoreUObject/Private/Virtualization/VirtualizationManager.cpp
+++ b/Engine/Source/Runtime/CoreUObject/Private/Virtualization/VirtualizationManager.cpp
@@ -107,15 +107,22 @@ TArray<FString> ParseEntries(const FString& Data)
 namespace Profiling
 {
 #if ENABLE_COOK_STATS
+	TMap<FString, FCookStats::CallStats> CacheStats;
 	TMap<FString, FCookStats::CallStats> PushStats;
 	TMap<FString, FCookStats::CallStats> PullStats;
 
 	void CreateStats(const IVirtualizationBackend& Backend)
 	{
+		CacheStats.Add(Backend.GetDebugString());	
 		PushStats.Add(Backend.GetDebugString());
 		PullStats.Add(Backend.GetDebugString());
 	}
 
+	FCookStats::CallStats& GetCacheStats(const IVirtualizationBackend& Backend)
+	{
+		return *CacheStats.Find(Backend.GetDebugString());
+	}
+	
 	FCookStats::CallStats& GetPushStats(const IVirtualizationBackend& Backend)
 	{
 		return *PushStats.Find(Backend.GetDebugString());
@@ -147,7 +154,7 @@ namespace Profiling
 			return false;
 		};
 
-		return HasAccumulatedData(PushStats) || HasAccumulatedData(PullStats);
+		return HasAccumulatedData(CacheStats) || HasAccumulatedData(PushStats) || HasAccumulatedData(PullStats);
 	}
 
 	void LogStats()
@@ -161,6 +168,27 @@ namespace Profiling
 		UE_LOG(LogVirtualization, Log, TEXT("Virtualization ProfileData"));
 		UE_LOG(LogVirtualization, Log, TEXT("======================================================================================="));
 
+		if (CacheStats.Num() > 0)
+		{
+			UE_LOG(LogVirtualization, Log, TEXT("%-40s|%17s|%12s|%14s|"), TEXT("Caching Data"), TEXT("TotalSize (MB)"), TEXT("TotalTime(s)"), TEXT("DataRate(MB/S)"));
+			UE_LOG(LogVirtualization, Log, TEXT("----------------------------------------|-----------------|------------|--------------|"));
+
+			for (const auto& Iterator : CacheStats)
+			{
+				const double Time = Iterator.Value.GetAccumulatedValueAnyThread(FCookStats::CallStats::EHitOrMiss::Hit, FCookStats::CallStats::EStatType::Cycles) * FPlatformTime::GetSecondsPerCycle();
+				const int64 DataSizeMB = Iterator.Value.GetAccumulatedValueAnyThread(FCookStats::CallStats::EHitOrMiss::Hit, FCookStats::CallStats::EStatType::Bytes) / (1024 * 1024);
+				const double MBps = Time != 0.0 ? (DataSizeMB / Time) : 0.0;
+
+				UE_LOG(LogVirtualization, Log, TEXT("%-40.40s|%17" UINT64_FMT "|%12.3f|%14.3f|"),
+					*Iterator.Key,
+					DataSizeMB,
+					Time,
+					MBps);
+			}
+
+			UE_LOG(LogVirtualization, Log, TEXT("======================================================================================="));
+		}
+
 		if (PushStats.Num() > 0)
 		{
 			UE_LOG(LogVirtualization, Log, TEXT("%-40s|%17s|%12s|%14s|"), TEXT("Pushing Data"), TEXT("TotalSize (MB)"), TEXT("TotalTime(s)"), TEXT("DataRate(MB/S)"));
@@ -330,7 +358,7 @@ bool FVirtualizationManager::PushData(const FPayloadId& Id, const FCompressedBuf
 		}
 	}
 
-	UE_CLOG(!bWasPayloadPushed, LogVirtualization, Fatal, TEXT("Payload '%s' failed to be pushed to any backend'"), *Id.ToString());
+	UE_CLOG(!bWasPayloadPushed, LogVirtualization, Error, TEXT("Payload '%s' failed to be pushed to any backend'"), *Id.ToString());
 
 	return bWasPayloadPushed;
 }
@@ -389,6 +417,13 @@ FPayloadActivityInfo FVirtualizationManager::GetPayloadActivityInfo() const
 	FPayloadActivityInfo Info;
 
 #if ENABLE_COOK_STATS
+	for (const auto& Iterator : Profiling::CacheStats)
+	{
+		Info.Cache.PayloadCount += Iterator.Value.GetAccumulatedValueAnyThread(FCookStats::CallStats::EHitOrMiss::Hit, FCookStats::CallStats::EStatType::Counter);
+		Info.Cache.TotalBytes += Iterator.Value.GetAccumulatedValueAnyThread(FCookStats::CallStats::EHitOrMiss::Hit, FCookStats::CallStats::EStatType::Bytes);
+		Info.Cache.CyclesSpent += Iterator.Value.GetAccumulatedValueAnyThread(FCookStats::CallStats::EHitOrMiss::Hit, FCookStats::CallStats::EStatType::Cycles);
+	}
+
 	for (const auto& Iterator : Profiling::PushStats)
 	{
 		Info.Push.PayloadCount	+= Iterator.Value.GetAccumulatedValueAnyThread(FCookStats::CallStats::EHitOrMiss::Hit, FCookStats::CallStats::EStatType::Counter);
@@ -640,12 +675,34 @@ void FVirtualizationManager::CachePayload(const FPayloadId& Id, const FCompresse
 			return; // No point going past BackendSource
 		}
 
-		EPushResult Result = BackendToCache->PushData(Id, Payload);
-		UE_CLOG(Result == EPushResult::Failed, LogVirtualization, Warning,
-			TEXT("Failed to cache payload '%s' to backend '%s'"),
-			*Id.ToString(),
-			*BackendToCache->GetDebugString());
+		bool bResult = TryCacheDataToBackend(*BackendToCache, Id, Payload);
+		UE_CLOG(	!bResult, LogVirtualization, Warning,
+					TEXT("Failed to cache payload '%s' to backend '%s'"),
+					*Id.ToString(),
+					*BackendToCache->GetDebugString());
+
+		// Debugging operation where we immediately try to pull the payload after each push (when possible) and assert 
+		// that the pulled payload is the same as the original
+		if (bValidateAfterPushOperation && bResult && BackendToCache->SupportsPullOperations())
+		{
+			FCompressedBuffer PulledPayload = PullDataFromBackend(*BackendToCache, Id);
+			checkf(Payload.GetRawHash() == PulledPayload.GetRawHash(), TEXT("[%s] Failed to pull payload '%s' after it was cached to backend"),
+				*BackendToCache->GetDebugString(), *Id.ToString());
+		}
+	}
+}
+
+bool FVirtualizationManager::TryCacheDataToBackend(IVirtualizationBackend& Backend, const FPayloadId& Id, const FCompressedBuffer& Payload)
+{
+	COOK_STAT(FCookStats::FScopedStatsCounter Timer(Profiling::GetCacheStats(Backend)));
+	const EPushResult Result = Backend.PushData(Id, Payload);
+
+	if (Result == EPushResult::Success)
+	{
+		COOK_STAT(Timer.AddHit(Payload.GetCompressedSize()));
 	}
+
+	return Result != EPushResult::Failed;
 }
 
 bool FVirtualizationManager::TryPushDataToBackend(IVirtualizationBackend& Backend, const FPayloadId& Id, const FCompressedBuffer& Payload)
diff --git a/Engine/Source/Runtime/CoreUObject/Public/Virtualization/VirtualizationManager.h b/Engine/Source/Runtime/CoreUObject/Public/Virtualization/VirtualizationManager.h
--- a/Engine/Source/Runtime/CoreUObject/Public/Virtualization/VirtualizationManager.h
+++ b/Engine/Source/Runtime/CoreUObject/Public/Virtualization/VirtualizationManager.h
@@ -64,6 +64,7 @@ struct FPayloadActivityInfo
 
 	FActivity Pull;
 	FActivity Push;
+	FActivity Cache;
 };
 
 /** Describes the type of storage to use for a given action */
@@ -135,6 +136,7 @@ private:
 
 	void CachePayload(const FPayloadId& Id, const FCompressedBuffer& Payload, const IVirtualizationBackend* BackendSource);
 
+	bool TryCacheDataToBackend(IVirtualizationBackend& Backend, const FPayloadId& Id, const FCompressedBuffer& Payload);
 	bool TryPushDataToBackend(IVirtualizationBackend& Backend, const FPayloadId& Id, const FCompressedBuffer& Payload);
 	FCompressedBuffer PullDataFromBackend(IVirtualizationBackend& Backend, const FPayloadId& Id);
 
