UEditorEngine::ReplaceActors: Unregisters old actor components before creating new actor (fixes duplicate registration of LevelInstances)

diff --git a/Engine/Source/Editor/UnrealEd/Private/EditorEngine.cpp b/Engine/Source/Editor/UnrealEd/Private/EditorEngine.cpp
--- a/Engine/Source/Editor/UnrealEd/Private/EditorEngine.cpp
+++ b/Engine/Source/Editor/UnrealEd/Private/EditorEngine.cpp
@@ -4994,6 +4994,10 @@ void UEditorEngine::ReplaceActors(UActorFactory* Factory, const FAssetData& Asse
 		ULevel* Level = OldActor->GetLevel();
 		AActor* NewActor = NULL;
 
+		// Unregister this actors components because we are effectively replacing it with an actor sharing the same ActorGuid.
+		// This allows it to be unregistered before a new actor with the same guid gets registered avoiding conflicts.
+		OldActor->UnregisterAllComponents();
+
 		const FName OldActorName = OldActor->GetFName();
 		const FName OldActorReplacedNamed = MakeUniqueObjectName(OldActor->GetOuter(), OldActor->GetClass(), *FString::Printf(TEXT("%s_REPLACED"), *OldActorName.ToString()));
 		
@@ -5120,6 +5124,7 @@ void UEditorEngine::ReplaceActors(UActorFactory* Factory, const FAssetData& Asse
 		{
 			// If creating the new Actor failed, put the old Actor's name back
 			OldActor->UObject::Rename(*OldActorName.ToString(), OldActor->GetOuter(), REN_DoNotDirty | REN_DontCreateRedirectors | REN_ForceNoResetLoaders);
+			OldActor->RegisterAllComponents();
 		}
 	}
 
diff --git a/Engine/Source/Runtime/Engine/Private/LevelInstance/LevelInstanceSubsystem.cpp b/Engine/Source/Runtime/Engine/Private/LevelInstance/LevelInstanceSubsystem.cpp
--- a/Engine/Source/Runtime/Engine/Private/LevelInstance/LevelInstanceSubsystem.cpp
+++ b/Engine/Source/Runtime/Engine/Private/LevelInstance/LevelInstanceSubsystem.cpp
@@ -1172,6 +1172,12 @@ void ULevelInstanceSubsystem::OnActorDeleted(AActor* Actor)
 			return;
 		}
 
+		// Unregistered Level Instance Actor nothing to do.
+		if (!LevelInstanceActor->HasValidLevelInstanceID())
+		{
+			return;
+		}
+
 		const bool bIsEditingLevelInstance = IsEditingLevelInstance(LevelInstanceActor);
 		if (!bIsEditingLevelInstance && LevelInstanceActor->IsA<APackedLevelInstance>())
 		{
