Cooker: Fix softlock in the cooker when cooking with no hard references, i.e. -cooksinglepackagenoref.
#rnx

diff --git a/Engine/Source/Editor/UnrealEd/Private/Cooker/CookRequestCluster.cpp b/Engine/Source/Editor/UnrealEd/Private/Cooker/CookRequestCluster.cpp
--- a/Engine/Source/Editor/UnrealEd/Private/Cooker/CookRequestCluster.cpp
+++ b/Engine/Source/Editor/UnrealEd/Private/Cooker/CookRequestCluster.cpp
@@ -362,6 +362,13 @@ void FRequestCluster::FetchDependencies(const FCookerTimer& CookerTimer, bool& b
 
 	if (!bAllowHardDependencies)
 	{
+		// FetchDependencies is responsible for marking all InitialRequests as explored. If we're skipping
+		// the dependencies search, just handle that responsibility and return.
+		for (FPackageData* PackageData : Requests)
+		{
+			bool bAlreadyCooked;
+			VisitPackageData(PackageData, nullptr /* Dependencies */, bAlreadyCooked);
+		}
 		bDependenciesComplete = true;
 		return;
 	}
