Blueprints: Re-ordered the message log tokens in BP runtime errors so the most relevant token (the node causing the problem) is first instead of last

diff --git a/Engine/Source/Editor/UnrealEd/Private/Kismet2/KismetDebugUtilities.cpp b/Engine/Source/Editor/UnrealEd/Private/Kismet2/KismetDebugUtilities.cpp
--- a/Engine/Source/Editor/UnrealEd/Private/Kismet2/KismetDebugUtilities.cpp
+++ b/Engine/Source/Editor/UnrealEd/Private/Kismet2/KismetDebugUtilities.cpp
@@ -271,17 +271,6 @@ void FKismetDebugUtilities::OnScriptException(const UObject* ActiveObject, const
 				TSharedRef<FTokenizedMessage> Message = FTokenizedMessage::Create(EMessageSeverity::Error);
 				Message->AddToken(FTextToken::Create(FText::Format(LOCTEXT("RuntimeErrorMessageFmt", "Blueprint Runtime Error: \"{0}\"."), Info.GetDescription())));
 
-				Message->AddToken(FTextToken::Create(LOCTEXT("RuntimeErrorBlueprintObjectLabel", "Blueprint: ")));
-				Message->AddToken(FUObjectToken::Create(BlueprintObj, FText::FromString(BlueprintObj->GetName()))
-					->OnMessageTokenActivated(FOnMessageTokenActivated::CreateStatic(&Local::OnMessageLogLinkActivated))
-				);
-
-				// NOTE: StackFrame.Node is not a blueprint node like you may think ("Node" has some legacy meaning)
-				Message->AddToken(FTextToken::Create(LOCTEXT("RuntimeErrorBlueprintFunctionLabel", "Function: ")));
-				Message->AddToken(FUObjectToken::Create(StackFrame.Node, StackFrame.Node->GetDisplayNameText())
-					->OnMessageTokenActivated(FOnMessageTokenActivated::CreateStatic(&Local::OnMessageLogLinkActivated))
-				);
-
 #if WITH_EDITORONLY_DATA // to protect access to GeneratedClass->DebugData
 				UBlueprintGeneratedClass* GeneratedClass = Cast<UBlueprintGeneratedClass>(ClassContainingCode);
 				if ((GeneratedClass != nullptr) && GeneratedClass->DebugData.IsValid())
@@ -290,18 +279,30 @@ void FKismetDebugUtilities::OnScriptException(const UObject* ActiveObject, const
 					// if instead, there is a node we can point to...
 					if (BlueprintNode != nullptr)
 					{
-						Message->AddToken(FTextToken::Create(LOCTEXT("RuntimeErrorBlueprintGraphLabel", "Graph: ")));
-						Message->AddToken(FUObjectToken::Create(BlueprintNode->GetGraph(), FText::FromString(GetNameSafe(BlueprintNode->GetGraph())))
+						Message->AddToken(FTextToken::Create(LOCTEXT("RuntimeErrorBlueprintNodeLabel", "Node: ")));
+						Message->AddToken(FUObjectToken::Create(BlueprintNode, BlueprintNode->GetNodeTitle(ENodeTitleType::ListView))
 							->OnMessageTokenActivated(FOnMessageTokenActivated::CreateStatic(&Local::OnMessageLogLinkActivated))
 						);
 
-						Message->AddToken(FTextToken::Create(LOCTEXT("RuntimeErrorBlueprintNodeLabel", "Node: ")));
-						Message->AddToken(FUObjectToken::Create(BlueprintNode, BlueprintNode->GetNodeTitle(ENodeTitleType::ListView))
+						Message->AddToken(FTextToken::Create(LOCTEXT("RuntimeErrorBlueprintGraphLabel", "Graph: ")));
+						Message->AddToken(FUObjectToken::Create(BlueprintNode->GetGraph(), FText::FromString(GetNameSafe(BlueprintNode->GetGraph())))
 							->OnMessageTokenActivated(FOnMessageTokenActivated::CreateStatic(&Local::OnMessageLogLinkActivated))
 						);
 					}
 				}
 #endif // WITH_EDITORONLY_DATA
+
+				// NOTE: StackFrame.Node is not a blueprint node like you may think ("Node" has some legacy meaning)
+				Message->AddToken(FTextToken::Create(LOCTEXT("RuntimeErrorBlueprintFunctionLabel", "Function: ")));
+				Message->AddToken(FUObjectToken::Create(StackFrame.Node, StackFrame.Node->GetDisplayNameText())
+					->OnMessageTokenActivated(FOnMessageTokenActivated::CreateStatic(&Local::OnMessageLogLinkActivated))
+				);
+
+				Message->AddToken(FTextToken::Create(LOCTEXT("RuntimeErrorBlueprintObjectLabel", "Blueprint: ")));
+				Message->AddToken(FUObjectToken::Create(BlueprintObj, FText::FromString(BlueprintObj->GetName()))
+					->OnMessageTokenActivated(FOnMessageTokenActivated::CreateStatic(&Local::OnMessageLogLinkActivated))
+				);
+
 				FMessageLog("PIE").AddMessage(Message);
 			}
 			bForceToCurrentObject = true;
