WorldPartition: Cleanup
- Remove unused callbacks
- Remove bEnableWorldPartiton property
- Add UWorld::IsPartitionedWorld() instead of checking for GetWorldPartition() != nullptr

diff --git a/Engine/Source/Editor/LevelEditor/Private/SLevelViewport.cpp b/Engine/Source/Editor/LevelEditor/Private/SLevelViewport.cpp
--- a/Engine/Source/Editor/LevelEditor/Private/SLevelViewport.cpp
+++ b/Engine/Source/Editor/LevelEditor/Private/SLevelViewport.cpp
@@ -3930,7 +3930,7 @@ EVisibility SLevelViewport::GetCurrentLevelTextVisibility() const
 	return (&GetLevelViewportClient() == GCurrentLevelEditingViewportClient) 
 		&& !IsPlayInEditorViewportActive() 
 		&& GetWorld() && GetWorld()->GetCurrentLevel()->OwningWorld->GetLevels().Num() > 1
-		&& !GetWorld()->GetWorldPartition()
+		&& !GetWorld()->IsPartitionedWorld()
 		?  ContentVisibility : EVisibility::Collapsed;
 }
 
diff --git a/Engine/Source/Editor/PropertyEditor/Private/UserInterface/PropertyEditor/SPropertyEditorAsset.cpp b/Engine/Source/Editor/PropertyEditor/Private/UserInterface/PropertyEditor/SPropertyEditorAsset.cpp
--- a/Engine/Source/Editor/PropertyEditor/Private/UserInterface/PropertyEditor/SPropertyEditorAsset.cpp
+++ b/Engine/Source/Editor/PropertyEditor/Private/UserInterface/PropertyEditor/SPropertyEditorAsset.cpp
@@ -611,7 +611,7 @@ void SPropertyEditorAsset::Construct(const FArguments& InArgs, const TSharedPtr<
 			}
 		}
 
-		if (World && World->GetWorldPartition())
+		if (World && World->IsPartitionedWorld())
 		{
 			bDisplayBrowse = false;
 		}
@@ -717,7 +717,7 @@ SPropertyEditorAsset::EActorReferenceState SPropertyEditorAsset::GetActorReferen
 				UWorld* World = Cast<UWorld>(MapObject);
 
 				// In a partitioned world, the world object will exist but the actor itself can be unloaded
-				if (World && World->GetWorldPartition())
+				if (World && World->IsPartitionedWorld())
 				{
 					UObject* Object = nullptr;
 					if (World->ResolveSubobject(*Value.ObjectPath.GetSubPathString(), Object, /*bLoadIfExists*/false))
diff --git a/Engine/Source/Editor/SceneOutliner/Private/ActorBrowsingMode.cpp b/Engine/Source/Editor/SceneOutliner/Private/ActorBrowsingMode.cpp
--- a/Engine/Source/Editor/SceneOutliner/Private/ActorBrowsingMode.cpp
+++ b/Engine/Source/Editor/SceneOutliner/Private/ActorBrowsingMode.cpp
@@ -205,7 +205,7 @@ void FActorBrowsingMode::Rebuild()
 	ApplicableUnloadedActors.Empty();
 	ApplicableActors.Empty();
 
-	bRepresentingWorldPartitionedWorld = RepresentingWorld.IsValid() && RepresentingWorld->GetWorldPartition() != nullptr;
+	bRepresentingWorldPartitionedWorld = RepresentingWorld.IsValid() && RepresentingWorld->IsPartitionedWorld();
 
 	if (bRepresentingWorldPartitionedWorld)
 	{
diff --git a/Engine/Source/Editor/UnrealEd/Private/Commandlets/ContentCommandlets.cpp b/Engine/Source/Editor/UnrealEd/Private/Commandlets/ContentCommandlets.cpp
--- a/Engine/Source/Editor/UnrealEd/Private/Commandlets/ContentCommandlets.cpp
+++ b/Engine/Source/Editor/UnrealEd/Private/Commandlets/ContentCommandlets.cpp
@@ -2526,7 +2526,7 @@ int32 UWrangleContentCommandlet::Main( const FString& Params )
 				FLinkerLoad* Linker = LoadPackageLinker(nullptr, PackagePath, LOAD_Quiet | LOAD_NoWarn | LOAD_NoVerify);
 
 				UWorld* World = UWorld::FindWorldInPackage(Package);
-				if (World && World->GetWorldPartition())
+				if (World && World->IsPartitionedWorld())
 				{
 					TArray<FString> ExternalActorPackages = World->PersistentLevel->GetOnDiskExternalActorPackages();
 					if (ExternalActorPackages.Num())
diff --git a/Engine/Source/Editor/UnrealEd/Private/Commandlets/WorldPartitionConvertCommandlet.cpp b/Engine/Source/Editor/UnrealEd/Private/Commandlets/WorldPartitionConvertCommandlet.cpp
--- a/Engine/Source/Editor/UnrealEd/Private/Commandlets/WorldPartitionConvertCommandlet.cpp
+++ b/Engine/Source/Editor/UnrealEd/Private/Commandlets/WorldPartitionConvertCommandlet.cpp
@@ -831,7 +831,7 @@ int32 UWorldPartitionConvertCommandlet::Main(const FString& Params)
 
 	// Make sure the world isn't already partitioned
 	AWorldSettings* MainWorldSettings = MainWorld->GetWorldSettings();
-	if (MainWorldSettings->GetWorldPartition())
+	if (MainWorldSettings->IsPartitionedWorld())
 	{
 		UE_LOG(LogWorldPartitionConvertCommandlet, Error, TEXT("Level '%s' is already partitionned"), *Tokens[0]);
 		return 1;
diff --git a/Engine/Source/Editor/UnrealEd/Private/PlayLevelViaLauncher.cpp b/Engine/Source/Editor/UnrealEd/Private/PlayLevelViaLauncher.cpp
--- a/Engine/Source/Editor/UnrealEd/Private/PlayLevelViaLauncher.cpp
+++ b/Engine/Source/Editor/UnrealEd/Private/PlayLevelViaLauncher.cpp
@@ -251,7 +251,7 @@ void UEditorEngine::StartPlayUsingLauncherSession(FRequestPlaySessionParams& InR
 	{
 		// For now World Partition doesn't these cook modes
 		FWorldContext& EditorContext = GetEditorWorldContext();
-		if (EditorContext.World()->GetWorldPartition())
+		if (EditorContext.World()->IsPartitionedWorld())
 		{
 			FString ErrorMsg = FString::Printf(TEXT("Error launching map %s : Quick launch %s doesn't yet support partitioned worlds."), *EditorContext.World()->GetOutermost()->GetName(), (CurrentLauncherCookMode == ELauncherProfileCookModes::ByTheBookInEditor) ? TEXT("cook by the book") : TEXT("cook on the fly"));
 			UE_LOG(LogPlayLevel, Error, TEXT("%s"), *ErrorMsg);
diff --git a/Engine/Source/Editor/UnrealEd/Private/UnrealEdMisc.cpp b/Engine/Source/Editor/UnrealEd/Private/UnrealEdMisc.cpp
--- a/Engine/Source/Editor/UnrealEd/Private/UnrealEdMisc.cpp
+++ b/Engine/Source/Editor/UnrealEd/Private/UnrealEdMisc.cpp
@@ -530,7 +530,6 @@ void FUnrealEdMisc::OnInit()
 
 	// Handles "Enable World Composition" option in WorldSettings
 	UWorldComposition::EnableWorldCompositionEvent.BindRaw(this, &FUnrealEdMisc::EnableWorldComposition);
-	UWorldPartition::EnableWorldPartitionEvent.BindRaw(this, &FUnrealEdMisc::EnableWorldPartition);
 
 	const double TotalEditorStartupTime = (FStudioAnalytics::GetAnalyticSeconds() - GStartTime);
 	UE_LOG(LogUnrealEdMisc, Log, TEXT("Total Editor Startup Time, took %.3f"), TotalEditorStartupTime);
@@ -936,33 +935,6 @@ bool FUnrealEdMisc::EnableWorldComposition(UWorld* InWorld, bool bEnable)
 	return true;
 }
 
-bool FUnrealEdMisc::EnableWorldPartition(UWorld* InWorld, bool bEnable)
-{
-	if (InWorld == nullptr || InWorld->WorldType != EWorldType::Editor)
-	{
-		return false;
-	}
-			
-	if (!bEnable)
-	{
-		UWorldPartition::WorldPartitionChangedEvent.Broadcast(InWorld);
-		return false;
-	}
-	
-	FString RootPackageName = InWorld->GetOutermost()->GetName();
-
-	// Map should be saved to disk
-	if (!FPackageName::DoesPackageExist(RootPackageName))
-	{
-		FMessageDialog::Open(EAppMsgType::Ok, LOCTEXT("EnableWorldPartitionNotSaved_Message", "Please save your level to disk before enabling World Partition"));
-		return false;
-	}
-			
-	UWorldPartition::WorldPartitionChangedEvent.Broadcast(InWorld);
-	
-	return true;
-}
-
 /** Build and return the path to the current project (used for relaunching the editor.)	 */
 FString CreateProjectPath()
 {
@@ -1078,7 +1050,6 @@ void FUnrealEdMisc::OnExit()
 	ISourceControlModule::Get().GetProvider().Close();
 
 	UWorldComposition::EnableWorldCompositionEvent.Unbind();
-	UWorldPartition::EnableWorldPartitionEvent.Unbind();
 	
 	//Fbx dll is currently compile with a different windows platform sdk so we should use this memory bypass later when unreal ed will be using windows 10 platform sdk
 	//UnloadFBxLibraries();
diff --git a/Engine/Source/Editor/UnrealEd/Public/UnrealEdMisc.h b/Engine/Source/Editor/UnrealEd/Public/UnrealEdMisc.h
--- a/Engine/Source/Editor/UnrealEd/Public/UnrealEdMisc.h
+++ b/Engine/Source/Editor/UnrealEd/Public/UnrealEdMisc.h
@@ -304,8 +304,6 @@ public:
 	/** Handles "Enable World Composition" option in WorldSettings */
 	bool EnableWorldComposition(UWorld* InWorld, bool bEnable);
 
-	bool EnableWorldPartition(UWorld* InWorld, bool bEnable);
-
 	/** Get the path to the executable that runs the editor */
 	FString GetProjectEditorBinaryPath();
 
diff --git a/Engine/Source/Editor/WorldPartitionEditor/Private/WorldPartitionEditorModule.cpp b/Engine/Source/Editor/WorldPartitionEditor/Private/WorldPartitionEditorModule.cpp
--- a/Engine/Source/Editor/WorldPartitionEditor/Private/WorldPartitionEditorModule.cpp
+++ b/Engine/Source/Editor/WorldPartitionEditor/Private/WorldPartitionEditorModule.cpp
@@ -100,8 +100,6 @@ void FWorldPartitionEditorModule::StartupModule()
 
 void FWorldPartitionEditorModule::ShutdownModule()
 {
-	UWorldPartition::WorldPartitionChangedEvent.RemoveAll(this);
-	
 	if (!IsRunningGame())
 	{
 		if (FLevelEditorModule* LevelEditorModule = FModuleManager::Get().GetModulePtr<FLevelEditorModule>("LevelEditor"))
diff --git a/Engine/Source/Runtime/Engine/Classes/Engine/World.h b/Engine/Source/Runtime/Engine/Classes/Engine/World.h
--- a/Engine/Source/Runtime/Engine/Classes/Engine/World.h
+++ b/Engine/Source/Runtime/Engine/Classes/Engine/World.h
@@ -2516,6 +2516,11 @@ public:
 	 */
 	UWorldPartition* GetWorldPartition() const;
 
+	/**
+	* Returns true if world contains an associated UWorldPartition object.
+	*/
+	bool IsPartitionedWorld() const { return GetWorldPartition() != nullptr; }
+
 	/**
 	 * Returns the current levels BSP model.
 	 *
diff --git a/Engine/Source/Runtime/Engine/Classes/GameFramework/WorldSettings.h b/Engine/Source/Runtime/Engine/Classes/GameFramework/WorldSettings.h
--- a/Engine/Source/Runtime/Engine/Classes/GameFramework/WorldSettings.h
+++ b/Engine/Source/Runtime/Engine/Classes/GameFramework/WorldSettings.h
@@ -469,16 +469,9 @@ public:
 	 * Enables tools for composing a tiled world. 
 	 * Level has to be saved and all sub-levels removed before enabling this option.
 	 */
-	UPROPERTY(EditAnywhere, BlueprintReadOnly, Category=World)
+	UPROPERTY(EditAnywhere, BlueprintReadOnly, Category=World, meta=(EditConditionHides, EditCondition="WorldPartition == nullptr"))
 	uint8 bEnableWorldComposition:1;
 
-	/** 
-	 * Enables tools for creating a partitioned world. 
-	 * Level has to be saved and all sub-levels removed before enabling this option.
-	 */
-	UPROPERTY(BlueprintReadOnly, Category=WorldPartitionSetup)
-	uint8 bEnableWorldPartition:1;
-
 	/** Additional transform applied when applying LevelStreaming Transform to LevelInstance */
 	UPROPERTY()
 	FVector LevelInstancePivotOffset;
@@ -540,7 +533,7 @@ protected:
 	UPROPERTY(Transient)
 	TObjectPtr<UNavigationSystemConfig> NavigationSystemConfigOverride;
 
-	UPROPERTY(EditAnywhere, Category=WorldPartitionSetup, NoClear, meta=(NoResetToDefault, HideEditConditionToggle, EditCondition = "bEnableWorldPartition"), Instanced)
+	UPROPERTY(VisibleAnywhere, Category=WorldPartitionSetup, Instanced)
 	TObjectPtr<UWorldPartition> WorldPartition;
 
 public:
@@ -766,6 +759,7 @@ public:
 	//~ End AActor Interface.
 
 	UWorldPartition* GetWorldPartition() const;
+	bool IsPartitionedWorld() const { return GetWorldPartition() != nullptr; }
 	void SetWorldPartition(UWorldPartition* InWorldPartition);
 
 	/**
diff --git a/Engine/Source/Runtime/Engine/Private/Level.cpp b/Engine/Source/Runtime/Engine/Private/Level.cpp
--- a/Engine/Source/Runtime/Engine/Private/Level.cpp
+++ b/Engine/Source/Runtime/Engine/Private/Level.cpp
@@ -2757,7 +2757,7 @@ bool ULevel::CanEditChange(const FProperty* PropertyThatWillChange) const
 		}
 
 		// Can't set a partition if using world composition or partition
-		if (WorldSettings && (WorldSettings->bEnableWorldComposition || WorldSettings->bEnableWorldPartition))
+		if (WorldSettings && (WorldSettings->bEnableWorldComposition || WorldSettings->IsPartitionedWorld()))
 		{
 			return false;
 		}
diff --git a/Engine/Source/Runtime/Engine/Private/ReplayHelper.cpp b/Engine/Source/Runtime/Engine/Private/ReplayHelper.cpp
--- a/Engine/Source/Runtime/Engine/Private/ReplayHelper.cpp
+++ b/Engine/Source/Runtime/Engine/Private/ReplayHelper.cpp
@@ -139,7 +139,7 @@ void FReplayHelper::StartRecording(UNetConnection* Connection)
 	World = Connection ? Connection->GetWorld() : nullptr;
 
 	// World Partition needs bHasLevelStreamingFixes to be true
-	bHasLevelStreamingFixes = !!CVarWithLevelStreamingFixes.GetValueOnAnyThread() || (World->GetWorldPartition() != nullptr);
+	bHasLevelStreamingFixes = !!CVarWithLevelStreamingFixes.GetValueOnAnyThread() || World->IsPartitionedWorld();
 	bHasDeltaCheckpoints = !!CVarWithDeltaCheckpoints.GetValueOnAnyThread() && ReplayStreamer->IsCheckpointTypeSupported(EReplayCheckpointType::Delta);
 	bHasGameSpecificFrameData = !!CVarWithGameSpecificFrameData.GetValueOnAnyThread();
 
diff --git a/Engine/Source/Runtime/Engine/Private/UnrealEngine.cpp b/Engine/Source/Runtime/Engine/Private/UnrealEngine.cpp
--- a/Engine/Source/Runtime/Engine/Private/UnrealEngine.cpp
+++ b/Engine/Source/Runtime/Engine/Private/UnrealEngine.cpp
@@ -15906,7 +15906,7 @@ int32 UEngine::RenderStatLevels(UWorld* World, FViewport* Viewport, FCanvas* Can
 		FString LevelPackageName = LevelStatus.PackageName.ToString();
 		FString DisplayName = LevelPackageName;
 
-		if (World->GetWorldPartition())
+		if (World->IsPartitionedWorld())
 		{
 			int32 CutPos = LevelPackageName.Find(*MapShortName);
 			if (CutPos != INDEX_NONE)
diff --git a/Engine/Source/Runtime/Engine/Private/WorldPartition/DataLayer/DataLayerSubsystem.cpp b/Engine/Source/Runtime/Engine/Private/WorldPartition/DataLayer/DataLayerSubsystem.cpp
--- a/Engine/Source/Runtime/Engine/Private/WorldPartition/DataLayer/DataLayerSubsystem.cpp
+++ b/Engine/Source/Runtime/Engine/Private/WorldPartition/DataLayer/DataLayerSubsystem.cpp
@@ -26,7 +26,7 @@ bool UDataLayerSubsystem::ShouldCreateSubsystem(UObject* Outer) const
 
 	if (UWorld* WorldOuter = Cast<UWorld>(Outer))
 	{
-		return WorldOuter->GetWorldPartition() != nullptr;
+		return WorldOuter->IsPartitionedWorld();
 	}
 
 	return false;
diff --git a/Engine/Source/Runtime/Engine/Private/WorldPartition/HLOD/HLODSubsystem.cpp b/Engine/Source/Runtime/Engine/Private/WorldPartition/HLOD/HLODSubsystem.cpp
--- a/Engine/Source/Runtime/Engine/Private/WorldPartition/HLOD/HLODSubsystem.cpp
+++ b/Engine/Source/Runtime/Engine/Private/WorldPartition/HLOD/HLODSubsystem.cpp
@@ -69,7 +69,7 @@ bool UHLODSubsystem::ShouldCreateSubsystem(UObject* Outer) const
 
 	if (UWorld* WorldOuter = Cast<UWorld>(Outer))
 	{
-		return WorldOuter->GetWorldPartition() != nullptr;
+		return WorldOuter->IsPartitionedWorld();
 	}
 	return false;
 }
diff --git a/Engine/Source/Runtime/Engine/Private/WorldPartition/WorldPartition.cpp b/Engine/Source/Runtime/Engine/Private/WorldPartition/WorldPartition.cpp
--- a/Engine/Source/Runtime/Engine/Private/WorldPartition/WorldPartition.cpp
+++ b/Engine/Source/Runtime/Engine/Private/WorldPartition/WorldPartition.cpp
@@ -223,8 +223,6 @@ struct FWorldPartionCellUpdateContext
 
 int32 FWorldPartionCellUpdateContext::UpdatesInProgress = 0;
 
-UWorldPartition::FEnableWorldPartitionEvent UWorldPartition::EnableWorldPartitionEvent;
-UWorldPartition::FWorldPartitionChangedEvent UWorldPartition::WorldPartitionChangedEvent;
 #endif
 
 UWorldPartition::UWorldPartition(const FObjectInitializer& ObjectInitializer)
diff --git a/Engine/Source/Runtime/Engine/Private/WorldPartition/WorldPartitionCookPackageSplitter.cpp b/Engine/Source/Runtime/Engine/Private/WorldPartition/WorldPartitionCookPackageSplitter.cpp
--- a/Engine/Source/Runtime/Engine/Private/WorldPartition/WorldPartitionCookPackageSplitter.cpp
+++ b/Engine/Source/Runtime/Engine/Private/WorldPartition/WorldPartitionCookPackageSplitter.cpp
@@ -18,7 +18,7 @@ REGISTER_COOKPACKAGE_SPLITTER(FWorldPartitionCookPackageSplitter, UWorld);
 bool FWorldPartitionCookPackageSplitter::ShouldSplit(UObject* SplitData)
 {
 	UWorld* World = Cast<UWorld>(SplitData);
-	return World && !!World->GetWorldPartition();
+	return World && World->IsPartitionedWorld();
 }
 
 bool FWorldPartitionCookPackageSplitter::UseDeferredPopulate()
@@ -71,7 +71,7 @@ UWorld* FWorldPartitionCookPackageSplitter::ValidateDataObject(UObject* SplitDat
 	UWorld* PartitionedWorld = CastChecked<UWorld>(SplitData);
 	check(PartitionedWorld);
 	check(PartitionedWorld->PersistentLevel);
-	check(PartitionedWorld->PersistentLevel->GetWorldPartition());
+	check(PartitionedWorld->IsPartitionedWorld());
 	return PartitionedWorld;
 }
 
diff --git a/Engine/Source/Runtime/Engine/Private/WorldPartition/WorldPartitionMiniMapHelper.cpp b/Engine/Source/Runtime/Engine/Private/WorldPartition/WorldPartitionMiniMapHelper.cpp
--- a/Engine/Source/Runtime/Engine/Private/WorldPartition/WorldPartitionMiniMapHelper.cpp
+++ b/Engine/Source/Runtime/Engine/Private/WorldPartition/WorldPartitionMiniMapHelper.cpp
@@ -23,7 +23,7 @@ DEFINE_LOG_CATEGORY_STATIC(LogWorldPartitionMiniMapHelper, All, All);
 
 AWorldPartitionMiniMap* FWorldPartitionMiniMapHelper::GetWorldPartitionMiniMap(UWorld* World, bool bCreateNewMiniMap)
 {
-	if (!World->GetWorldPartition())
+	if (!World->IsPartitionedWorld())
 	{
 		UE_LOG(LogWorldPartitionMiniMapHelper, Error, TEXT("No WorldPartition Found. WorldPartition must exist to get WorldPartitionMiniMap"));
 		return nullptr;
diff --git a/Engine/Source/Runtime/Engine/Private/WorldPartition/WorldPartitionReplay.cpp b/Engine/Source/Runtime/Engine/Private/WorldPartition/WorldPartitionReplay.cpp
--- a/Engine/Source/Runtime/Engine/Private/WorldPartition/WorldPartitionReplay.cpp
+++ b/Engine/Source/Runtime/Engine/Private/WorldPartition/WorldPartitionReplay.cpp
@@ -84,9 +84,8 @@ void AWorldPartitionReplay::Initialize(UWorld* World)
 #if !UE_BUILD_SHIPPING
 	if (World->IsGameWorld())
 	{
-		check(World->GetWorldPartition());
 		UWorldPartition* WorldPartition = World->GetWorldPartition();
-
+		check(WorldPartition);
 		check(!WorldPartition->Replay);
 		FActorSpawnParameters SpawnParams;
 		SpawnParams.Name = AWorldPartitionReplay::StaticClass()->GetFName();
@@ -99,7 +98,7 @@ bool AWorldPartitionReplay::IsEnabled(UWorld* World)
 {
 #if !UE_BUILD_SHIPPING
 	check(World->IsGameWorld());
-	check(World->GetWorldPartition());
+	check(World->IsPartitionedWorld());
 	UWorldPartition* WorldPartition = World->GetWorldPartition();
 	return WorldPartition->Replay && WorldPartition->Replay->IsEnabled();
 #else
@@ -111,7 +110,7 @@ void AWorldPartitionReplay::BeginPlay()
 {
 	UWorld* World = GetWorld();
 	check(World->IsGameWorld());
-	check(World->GetWorldPartition());
+	check(World->IsPartitionedWorld());
 	ENetMode NetMode = World->GetNetMode();
 	
 	// Recording will happen only in NM_Standalone/NM_Client. Replay can happen on any net mode.
@@ -130,8 +129,7 @@ void AWorldPartitionReplay::PreReplication(IRepChangedPropertyTracker& ChangedPr
 	check(GetWorld()->IsRecordingReplay());
 	Super::PreReplication(ChangedPropertyTracker);
 	
-	UWorldPartition* WorldPartition = GetWorld()->GetWorldPartition();
-	check(WorldPartition);
+	check(GetWorld()->IsPartitionedWorld());
 
 	FWorldPartitionReplaySample Replay(this);
 	if (Replay.StreamingSources.Num())
diff --git a/Engine/Source/Runtime/Engine/Private/WorldPartition/WorldPartitionSubsystem.cpp b/Engine/Source/Runtime/Engine/Private/WorldPartition/WorldPartitionSubsystem.cpp
--- a/Engine/Source/Runtime/Engine/Private/WorldPartition/WorldPartitionSubsystem.cpp
+++ b/Engine/Source/Runtime/Engine/Private/WorldPartition/WorldPartitionSubsystem.cpp
@@ -77,7 +77,7 @@ bool UWorldPartitionSubsystem::ShouldCreateSubsystem(UObject* Outer) const
 
 	if (UWorld* WorldOuter = Cast<UWorld>(Outer))
 	{
-		return WorldOuter->GetWorldPartition() != nullptr;
+		return WorldOuter->IsPartitionedWorld();
 	}
 
 	return false;
diff --git a/Engine/Source/Runtime/Engine/Private/WorldSettings.cpp b/Engine/Source/Runtime/Engine/Private/WorldSettings.cpp
--- a/Engine/Source/Runtime/Engine/Private/WorldSettings.cpp
+++ b/Engine/Source/Runtime/Engine/Private/WorldSettings.cpp
@@ -209,7 +209,7 @@ void AWorldSettings::PostRegisterAllComponents()
 
 UWorldPartition* AWorldSettings::GetWorldPartition() const
 {
-	return bEnableWorldPartition ? WorldPartition : nullptr;
+	return WorldPartition;
 }
 
 void AWorldSettings::SetWorldPartition(UWorldPartition* InWorldPartition)
@@ -217,7 +217,6 @@ void AWorldSettings::SetWorldPartition(UWorldPartition* InWorldPartition)
 	check(!WorldPartition);
 	check(InWorldPartition);
 	WorldPartition = InWorldPartition;
-	bEnableWorldPartition = true;
 	bEnableWorldComposition = false;
 }
 
@@ -601,7 +600,7 @@ bool AWorldSettings::CanEditChange(const FProperty* InProperty) const
 		}
 		else if (PropertyName == GET_MEMBER_NAME_STRING_CHECKED(AWorldSettings, bEnableWorldComposition))
 		{
-			return !bEnableWorldPartition;
+			return !IsPartitionedWorld();
 		}
 		if (PropertyName == GET_MEMBER_NAME_STRING_CHECKED(AWorldSettings, InstancedFoliageGridSize))
 		{
diff --git a/Engine/Source/Runtime/Engine/Public/WorldPartition/WorldPartition.h b/Engine/Source/Runtime/Engine/Public/WorldPartition/WorldPartition.h
--- a/Engine/Source/Runtime/Engine/Public/WorldPartition/WorldPartition.h
+++ b/Engine/Source/Runtime/Engine/Public/WorldPartition/WorldPartition.h
@@ -87,11 +87,6 @@ public:
 	static UWorldPartition* CreateWorldPartition(AWorldSettings* WorldSettings, TSubclassOf<UWorldPartitionEditorHash> EditorHashClass = nullptr, TSubclassOf<UWorldPartitionRuntimeHash> RuntimeHashClass = nullptr);
 
 private:
-	DECLARE_DELEGATE_RetVal_TwoParams(bool, FEnableWorldPartitionEvent, UWorld*, bool);
-	static FEnableWorldPartitionEvent EnableWorldPartitionEvent;
-
-	DECLARE_MULTICAST_DELEGATE_OneParam(FWorldPartitionChangedEvent, UWorld*);
-	static FWorldPartitionChangedEvent WorldPartitionChangedEvent;
 
 	void SavePerUserSettings();
 		
