Silence false positive static analysis warning C6011: Dereferencing NULL pointer 'Variable'
#rnx

diff --git a/Engine/Source/Editor/UnrealEd/Private/Cooker/CookRequestCluster.cpp b/Engine/Source/Editor/UnrealEd/Private/Cooker/CookRequestCluster.cpp
--- a/Engine/Source/Editor/UnrealEd/Private/Cooker/CookRequestCluster.cpp
+++ b/Engine/Source/Editor/UnrealEd/Private/Cooker/CookRequestCluster.cpp
@@ -158,6 +158,7 @@ void FRequestCluster::AddClusters(UCookOnTheFlyServer& COTFS, FPackageDataSet& U
 			// MRUCluster points to data inside AddedClusters, so we have to recalculate MRUCluster whenever we add
 			MRUCluster = FindOrAddClusterForPlatforms();
 		}
+		CA_ASSUME(MRUCluster);
 
 		MRUCluster->Requests.Add(PackageData);
 	}
