Silence false positive static analysis warning V772: Calling a 'delete' operator for a void pointer will cause undefined behavior.
#rnx

diff --git a/Engine/Source/Editor/UnrealEd/Private/Cooker/CookPlatformManager.cpp b/Engine/Source/Editor/UnrealEd/Private/Cooker/CookPlatformManager.cpp
--- a/Engine/Source/Editor/UnrealEd/Private/Cooker/CookPlatformManager.cpp
+++ b/Engine/Source/Editor/UnrealEd/Private/Cooker/CookPlatformManager.cpp
@@ -86,7 +86,8 @@ namespace Cook
 			FRWScopeLock PlatformDatasScopeLock(PlatformDatasLock, FRWScopeLockType::SLT_Write);
 			for (TPair<const ITargetPlatform*, FPlatformData*>& kvpair : PlatformDatas)
 			{
-				delete kvpair.Value;
+				FPlatformData* Value = kvpair.Value;
+				delete Value;
 			}
 			PlatformDatas.Empty();
 			PlatformDatasByName.Empty();
