Added a 'Preview Profiles' combo box to the advances preview scene toolbar to quickly change preview profiles.
  - The profile combo is only visible when the user has more than one profile
  - This was added to the Material Editor, the Static Mesh Editor and Niagara Editor.

diff --git a/Engine/Plugins/FX/Niagara/Source/NiagaraEditor/Private/Widgets/SNiagaraSystemViewportToolBar.cpp b/Engine/Plugins/FX/Niagara/Source/NiagaraEditor/Private/Widgets/SNiagaraSystemViewportToolBar.cpp
--- a/Engine/Plugins/FX/Niagara/Source/NiagaraEditor/Private/Widgets/SNiagaraSystemViewportToolBar.cpp
+++ b/Engine/Plugins/FX/Niagara/Source/NiagaraEditor/Private/Widgets/SNiagaraSystemViewportToolBar.cpp
@@ -7,6 +7,7 @@
 #include "EditorViewportCommands.h"
 #include "NiagaraEditorSettings.h"
 #include "SEditorViewportToolBarButton.h"
+#include "PreviewProfileController.h"
 
 #define LOCTEXT_NAMESPACE "NiagaraSystemViewportToolBar"
 
@@ -17,7 +18,7 @@
 void SNiagaraSystemViewportToolBar::Construct(const FArguments& InArgs, TSharedPtr<class SNiagaraSystemViewport> InViewport)
 {
 	// we don't want a realtime button here as we create a custom one by extending the left menu
-	SCommonEditorViewportToolbarBase::Construct(SCommonEditorViewportToolbarBase::FArguments().AddRealtimeButton(true), InViewport);
+	SCommonEditorViewportToolbarBase::Construct(SCommonEditorViewportToolbarBase::FArguments().AddRealtimeButton(true).PreviewProfileController(MakeShared<FPreviewProfileController>()), InViewport);
 	Sequencer = InArgs._Sequencer;
 }
 
diff --git a/Engine/Source/Editor/AdvancedPreviewScene/Private/AdvancedPreviewScene.cpp b/Engine/Source/Editor/AdvancedPreviewScene/Private/AdvancedPreviewScene.cpp
--- a/Engine/Source/Editor/AdvancedPreviewScene/Private/AdvancedPreviewScene.cpp
+++ b/Engine/Source/Editor/AdvancedPreviewScene/Private/AdvancedPreviewScene.cpp
@@ -434,7 +434,14 @@ void FAdvancedPreviewScene::HandleTogglePostProcessing()
 
 void FAdvancedPreviewScene::OnAssetViewerSettingsRefresh(const FName& InPropertyName)
 {
-	if (DefaultSettings->Profiles.IsValidIndex(CurrentProfileIndex))
+	// If the profile was changed, update the current index and the scene.
+	if (InPropertyName == GET_MEMBER_NAME_CHECKED(FPreviewSceneProfile, ProfileName))
+	{
+		CurrentProfileIndex = GetDefault<UEditorPerProjectUserSettings>()->AssetViewerProfileIndex;
+		CurrentProfileIndex = DefaultSettings->Profiles.IsValidIndex(CurrentProfileIndex) ? CurrentProfileIndex : 0;
+		UpdateScene(DefaultSettings->Profiles[CurrentProfileIndex]);
+	}
+	else if (DefaultSettings->Profiles.IsValidIndex(CurrentProfileIndex))
 	{
 		const bool bNameNone = InPropertyName == NAME_None;
 
diff --git a/Engine/Source/Editor/AdvancedPreviewScene/Private/PreviewProfileController.cpp b/Engine/Source/Editor/AdvancedPreviewScene/Private/PreviewProfileController.cpp
new file mode 100644
--- /dev/null
+++ b/Engine/Source/Editor/AdvancedPreviewScene/Private/PreviewProfileController.cpp
@@ -0,0 +1,111 @@
+// Copyright Epic Games, Inc. All Rights Reserved.
+
+#include "PreviewProfileController.h"
+
+#include "Editor/EditorPerProjectUserSettings.h"
+#include "AssetViewerSettings.h"
+
+FPreviewProfileController::FPreviewProfileController()
+{
+	PerProjectSettings = GetMutableDefault<UEditorPerProjectUserSettings>();
+	if (PerProjectSettings)
+	{
+		AssetViewerSettings = UAssetViewerSettings::Get();
+		if (AssetViewerSettings)
+		{
+			AssetViewerSettingsProfileAddRemoveHandle = AssetViewerSettings->OnAssetViewerProfileAddRemoved().AddLambda([this]()
+			{
+				UpdateAssetViewerProfiles();
+				OnPreviewProfileListChanged().Broadcast();
+			});
+
+			AssetViewerSettingsChangedHandle = AssetViewerSettings->OnAssetViewerSettingsChanged().AddLambda([this](const FName& InPropertyName)
+			{
+				FString CurrProfileName = AssetViewerProfileNames[CurrentProfileIndex];
+				UpdateAssetViewerProfiles();
+				if (CurrProfileName != AssetViewerProfileNames[CurrentProfileIndex])
+				{
+					OnPreviewProfileChanged().Broadcast();
+				}
+			});
+			
+			UpdateAssetViewerProfiles();
+		}
+	}
+}
+
+FPreviewProfileController::~FPreviewProfileController()
+{
+	if (AssetViewerSettings)
+	{
+		AssetViewerSettings->OnAssetViewerProfileAddRemoved().Remove(AssetViewerSettingsProfileAddRemoveHandle);
+		AssetViewerSettings->OnAssetViewerSettingsChanged().Remove(AssetViewerSettingsChangedHandle);
+	}
+}
+
+void FPreviewProfileController::UpdateAssetViewerProfiles()
+{
+	AssetViewerProfileNames.Empty();
+
+	if (AssetViewerSettings && PerProjectSettings)
+	{
+		// Rebuild the profile list.
+		for (const FPreviewSceneProfile& Profile : AssetViewerSettings->Profiles)
+		{
+			AssetViewerProfileNames.Add(Profile.ProfileName + (Profile.bSharedProfile ? TEXT(" (Shared)") : TEXT("")));
+		}
+
+		CurrentProfileIndex = PerProjectSettings->AssetViewerProfileIndex;
+
+		EnsureProfilesStateCoherence();
+	}
+}
+
+TArray<FString> FPreviewProfileController::GetPreviewProfiles(int32& OutCurrentProfileIndex) const
+{
+	if (PerProjectSettings && AssetViewerSettings)
+	{
+		EnsureProfilesStateCoherence();
+		OutCurrentProfileIndex = CurrentProfileIndex;
+	}
+	return AssetViewerProfileNames;
+}
+
+bool FPreviewProfileController::SetActiveProfile(const FString& ProfileName)
+{
+	if (PerProjectSettings && AssetViewerSettings)
+	{
+		EnsureProfilesStateCoherence();
+
+		int32 Index = AssetViewerProfileNames.IndexOfByKey(ProfileName);
+		if (Index != INDEX_NONE && Index != PerProjectSettings->AssetViewerProfileIndex)
+		{
+			// Store the settings.
+			PerProjectSettings->AssetViewerProfileIndex = Index;
+			CurrentProfileIndex = Index;
+
+			// Notify the observer about the change.
+			AssetViewerSettings->OnAssetViewerSettingsChanged().Broadcast(GET_MEMBER_NAME_CHECKED(FPreviewSceneProfile, ProfileName));
+			return true;
+		}
+	}
+	return false;
+}
+
+FString FPreviewProfileController::GetActiveProfile() const
+{
+	if (PerProjectSettings && AssetViewerSettings)
+	{
+		EnsureProfilesStateCoherence();
+		return AssetViewerProfileNames[PerProjectSettings->AssetViewerProfileIndex];
+	}
+	return FString();
+}
+
+void FPreviewProfileController::EnsureProfilesStateCoherence() const
+{
+	ensureMsgf(AssetViewerProfileNames.Num() == AssetViewerSettings->Profiles.Num(), TEXT("List of profiles is out of sync with the list of corresponding profile names."));
+	ensureMsgf(AssetViewerProfileNames.Num() > 0, TEXT("The list of profiles is expected to always have at least one default profile"));
+	ensureMsgf(AssetViewerSettings->Profiles.IsValidIndex(PerProjectSettings->AssetViewerProfileIndex), TEXT("The active asset viewer profile is invalid."));
+	ensureMsgf(CurrentProfileIndex == PerProjectSettings->AssetViewerProfileIndex, TEXT("The cached selected profile index is out of sync."));
+}
diff --git a/Engine/Source/Editor/AdvancedPreviewScene/Private/SAdvancedPreviewDetailsTab.cpp b/Engine/Source/Editor/AdvancedPreviewScene/Private/SAdvancedPreviewDetailsTab.cpp
--- a/Engine/Source/Editor/AdvancedPreviewScene/Private/SAdvancedPreviewDetailsTab.cpp
+++ b/Engine/Source/Editor/AdvancedPreviewScene/Private/SAdvancedPreviewDetailsTab.cpp
@@ -232,7 +232,7 @@ void SAdvancedPreviewDetailsTab::OnAssetViewerSettingsRefresh(const FName& InPro
 {
 	if (InPropertyName == GET_MEMBER_NAME_CHECKED(FPreviewSceneProfile, ProfileName) || InPropertyName == GET_MEMBER_NAME_CHECKED(FPreviewSceneProfile, bSharedProfile))
 	{
-		UpdateProfileNames();
+		Refresh();
 	}
 }
 
diff --git a/Engine/Source/Editor/AdvancedPreviewScene/Public/PreviewProfileController.h b/Engine/Source/Editor/AdvancedPreviewScene/Public/PreviewProfileController.h
new file mode 100644
--- /dev/null
+++ b/Engine/Source/Editor/AdvancedPreviewScene/Public/PreviewProfileController.h
@@ -0,0 +1,61 @@
+// Copyright Epic Games, Inc. All Rights Reserved.
+
+#pragma once
+
+#include "CoreMinimal.h"
+#include "IPreviewProfileController.h"
+
+class UAssetViewerSettings;
+class UEditorPerProjectUserSettings;
+
+/**
+ * Controls the access to the preview profiles. It serves as a bridge between the AdvancedPreviewScene
+ * and the UnrealEd modules and enable UnrealEd to change or observe change to the active profile.
+ * 
+ * @note This class was created to decouple UnrealEd from AdvancedPreviewScene and prevent
+ *       circular dependencies between the modules.
+ */
+class ADVANCEDPREVIEWSCENE_API FPreviewProfileController : public IPreviewProfileController
+{
+public:
+	FPreviewProfileController();
+	virtual ~FPreviewProfileController();
+
+	/** Returns the list of available preview profiles names. */
+	virtual TArray<FString> GetPreviewProfiles(int32& OutCurrentProfileIndex) const override;
+
+	/** Set the specified preview profiles as the active one. */
+	virtual bool SetActiveProfile(const FString& ProfileName) override;
+
+	/** Returns the preview profiles currently active. */
+	virtual FString GetActiveProfile() const override;
+
+	/** Invoked after the list of available profiles has changed. */
+	virtual FOnPreviewProfileListChanged& OnPreviewProfileListChanged() override { return OnPreviewProfileListChangedDelegate; }
+
+	/** Invoked after the active preview profile changed. */
+	virtual FOnPreviewProfileChanged& OnPreviewProfileChanged() override { return OnPreviewSettingChangedDelegate; }
+
+private:
+	void UpdateAssetViewerProfiles();
+	void EnsureProfilesStateCoherence() const;
+
+private:
+	/** Provides the list of available preview profiles along with the profile change delegates (to update the profile combo box selection) */
+	UAssetViewerSettings* AssetViewerSettings = nullptr;
+
+	/** Holds the asset viewer profile currently used by the project. */
+	UEditorPerProjectUserSettings* PerProjectSettings = nullptr;
+
+	/** The list of available profiles. */
+	TArray<FString> AssetViewerProfileNames;
+
+	/** Holds the current profile index in the list. This is kept consistent with the cache list of names. */
+	int32 CurrentProfileIndex = 0;
+
+	FOnPreviewProfileListChanged OnPreviewProfileListChangedDelegate;
+	FOnPreviewProfileChanged OnPreviewSettingChangedDelegate;
+
+	FDelegateHandle AssetViewerSettingsProfileAddRemoveHandle;
+	FDelegateHandle AssetViewerSettingsChangedHandle;
+};
diff --git a/Engine/Source/Editor/MaterialEditor/Private/SMaterialEditorViewportToolBar.cpp b/Engine/Source/Editor/MaterialEditor/Private/SMaterialEditorViewportToolBar.cpp
--- a/Engine/Source/Editor/MaterialEditor/Private/SMaterialEditorViewportToolBar.cpp
+++ b/Engine/Source/Editor/MaterialEditor/Private/SMaterialEditorViewportToolBar.cpp
@@ -4,6 +4,7 @@
 #include "Widgets/Layout/SBorder.h"
 #include "EditorStyleSet.h"
 #include "MaterialEditorActions.h"
+#include "PreviewProfileController.h"
 
 #define LOCTEXT_NAMESPACE "MaterialEditorViewportToolBar"
 
@@ -52,7 +53,7 @@ void SMaterialEditorViewportPreviewShapeToolBar::Construct(const FArguments& InA
 
 void SMaterialEditorViewportToolBar::Construct(const FArguments& InArgs, TSharedPtr<class SMaterialEditor3DPreviewViewport> InViewport)
 {
-	SCommonEditorViewportToolbarBase::Construct(SCommonEditorViewportToolbarBase::FArguments(), InViewport);
+	SCommonEditorViewportToolbarBase::Construct(SCommonEditorViewportToolbarBase::FArguments().PreviewProfileController(MakeShared<FPreviewProfileController>()), InViewport);
 }
 
 TSharedRef<SWidget> SMaterialEditorViewportToolBar::GenerateShowMenu() const
diff --git a/Engine/Source/Editor/StaticMeshEditor/Private/SStaticMeshEditorViewportToolBar.cpp b/Engine/Source/Editor/StaticMeshEditor/Private/SStaticMeshEditorViewportToolBar.cpp
--- a/Engine/Source/Editor/StaticMeshEditor/Private/SStaticMeshEditorViewportToolBar.cpp
+++ b/Engine/Source/Editor/StaticMeshEditor/Private/SStaticMeshEditorViewportToolBar.cpp
@@ -19,6 +19,7 @@
 #include "Engine/StaticMeshSocket.h"
 #include "Editor/UnrealEd/Public/SEditorViewportToolBarMenu.h"
 #include "StaticMeshViewportLODCommands.h"
+#include "PreviewProfileController.h"
 
 #define LOCTEXT_NAMESPACE "StaticMeshEditorViewportToolbar"
 
@@ -28,7 +29,7 @@
 
 void SStaticMeshEditorViewportToolbar::Construct(const FArguments& InArgs, TSharedPtr<class ICommonEditorViewportToolbarInfoProvider> InInfoProvider)
 {
-	SCommonEditorViewportToolbarBase::Construct(SCommonEditorViewportToolbarBase::FArguments(), InInfoProvider);
+	SCommonEditorViewportToolbarBase::Construct(SCommonEditorViewportToolbarBase::FArguments().PreviewProfileController(MakeShared<FPreviewProfileController>()), InInfoProvider);
 }
 
 // SCommonEditorViewportToolbarBase interface
diff --git a/Engine/Source/Editor/UnrealEd/Private/SCommonEditorViewportToolbarBase.cpp b/Engine/Source/Editor/UnrealEd/Private/SCommonEditorViewportToolbarBase.cpp
--- a/Engine/Source/Editor/UnrealEd/Private/SCommonEditorViewportToolbarBase.cpp
+++ b/Engine/Source/Editor/UnrealEd/Private/SCommonEditorViewportToolbarBase.cpp
@@ -5,9 +5,9 @@
 #include "Widgets/Layout/SBorder.h"
 #include "Widgets/Layout/SBox.h"
 #include "Widgets/Input/SSpinBox.h"
+#include "Widgets/Input/STextComboBox.h"
 #include "EditorStyleSet.h"
 
-
 #include "STransformViewportToolbar.h"
 #include "EditorShowFlags.h"
 #include "SEditorViewport.h"
@@ -26,12 +26,31 @@
 //////////////////////////////////////////////////////////////////////////
 // SCommonEditorViewportToolbarBase
 
+SCommonEditorViewportToolbarBase::~SCommonEditorViewportToolbarBase()
+{
+	if (PreviewProfileController)
+	{
+		PreviewProfileController->OnPreviewProfileListChanged().RemoveAll(this);
+		PreviewProfileController->OnPreviewProfileChanged().RemoveAll(this);
+	}
+}
+
+
 void SCommonEditorViewportToolbarBase::Construct(const FArguments& InArgs, TSharedPtr<class ICommonEditorViewportToolbarInfoProvider> InInfoProvider)
 {
 	InfoProviderPtr = InInfoProvider;
- 	TSharedRef<SEditorViewport> ViewportRef = GetInfoProvider().GetViewportWidget();
+	PreviewProfileController = InArgs._PreviewProfileController;
+
+	TSharedRef<SEditorViewport> ViewportRef = GetInfoProvider().GetViewportWidget();
 	TSharedPtr<SHorizontalBox> MainBoxPtr;
 
+	if (PreviewProfileController)
+	{
+		PreviewProfileController->OnPreviewProfileListChanged().AddRaw(this, &SCommonEditorViewportToolbarBase::UpdateAssetViewerProfileList);
+		PreviewProfileController->OnPreviewProfileChanged().AddRaw(this, &SCommonEditorViewportToolbarBase::UpdateAssetViewerProfileSelection);
+		UpdateAssetViewerProfileList();
+	}
+
 	const FMargin ToolbarSlotPadding(4.0f, 1.0f);
 	const FMargin ToolbarButtonPadding(4.0f, 0.0f);
 
@@ -89,6 +108,17 @@ void SCommonEditorViewportToolbarBase::Construct(const FArguments& InArgs, TShar
 			.OnGetMenuContent(this, &SCommonEditorViewportToolbarBase::GenerateShowMenu)
 		];
 
+	// Profile menu (Controls the Preview Scene Settings)
+	if (InArgs._PreviewProfileController)
+	{
+		MainBoxPtr->AddSlot()
+			.AutoWidth()
+			.Padding(ToolbarSlotPadding)
+			[
+				MakeAssetViewerProfileComboBox()
+			];
+	}
+
 	// Realtime button
 	if (InArgs._AddRealtimeButton)
 	{
@@ -156,6 +186,79 @@ void SCommonEditorViewportToolbarBase::Construct(const FArguments& InArgs, TShar
 	SViewportToolBar::Construct(SViewportToolBar::FArguments());
 }
 
+void SCommonEditorViewportToolbarBase::UpdateAssetViewerProfileList()
+{
+	if (PreviewProfileController)
+	{
+		// Pull the latest profile list.
+		int32 CurrProfileIndex = 0;
+		TArray<FString> ProfileNames = PreviewProfileController->GetPreviewProfiles(CurrProfileIndex);
+
+		// Rebuild the combo box list.
+		AssetViewerProfileNames.Empty();
+		for (const FString& Profile : ProfileNames)
+		{
+			AssetViewerProfileNames.Add(MakeShared<FString>(Profile));
+		}
+
+		// Select the current profile item.
+		if (AssetViewerProfileComboBox)
+		{
+			AssetViewerProfileComboBox->RefreshOptions();
+			AssetViewerProfileComboBox->SetSelectedItem(AssetViewerProfileNames[CurrProfileIndex]);
+		}
+	}
+}
+
+void SCommonEditorViewportToolbarBase::UpdateAssetViewerProfileSelection()
+{
+	if (PreviewProfileController)
+	{
+		FString ActiveProfileName = PreviewProfileController->GetActiveProfile();
+		if (TSharedPtr<FString>* Match = AssetViewerProfileNames.FindByPredicate(
+			[&ActiveProfileName](const TSharedPtr<FString>& Candidate) { return *Candidate == ActiveProfileName; }))
+		{
+			AssetViewerProfileComboBox->SetSelectedItem(*Match);
+		}
+		else // The profile was likely renamed.
+		{
+			UpdateAssetViewerProfileList();
+		}
+	}
+}
+
+void SCommonEditorViewportToolbarBase::OnAssetViewerProfileComboBoxSelectionChanged(TSharedPtr<FString> NewSelection, ESelectInfo::Type SelectInfo)
+{
+	int32 NewSelectionIndex;
+	if (AssetViewerProfileNames.Find(NewSelection, NewSelectionIndex))
+	{
+		// If that's the user changing the combo box, not an update coming from code to reflect a change that already occurred.
+		if (SelectInfo != ESelectInfo::Direct)
+		{
+			PreviewProfileController->SetActiveProfile(*NewSelection);
+		}
+	}
+}
+
+TSharedRef<SWidget> SCommonEditorViewportToolbarBase::MakeAssetViewerProfileComboBox()
+{
+	AssetViewerProfileComboBox = SNew(STextComboBox)
+		.OptionsSource(&AssetViewerProfileNames)
+		.ButtonStyle(&FAppStyle::Get().GetWidgetStyle<FButtonStyle>("EditorViewportToolBar.Button"))
+		.ContentPadding(FMargin(2, 0))
+		.ToolTipText(LOCTEXT("AssetViewerProfile_ToolTip", "Changes the asset viewer profile"))
+		.OnSelectionChanged(this, &SCommonEditorViewportToolbarBase::OnAssetViewerProfileComboBoxSelectionChanged)
+		.Visibility_Lambda([this]() { return AssetViewerProfileNames.Num() > 1 ? EVisibility::Visible : EVisibility::Collapsed; });
+
+	AssetViewerProfileComboBox->RefreshOptions();
+	if (!AssetViewerProfileNames.IsEmpty())
+	{
+		AssetViewerProfileComboBox->SetSelectedItem(AssetViewerProfileNames[0]);
+	}
+
+	return AssetViewerProfileComboBox.ToSharedRef();
+}
+
 FText SCommonEditorViewportToolbarBase::GetCameraMenuLabel() const
 {
 	return GetCameraMenuLabelFromViewportType( GetViewportClient().GetViewportType() );
diff --git a/Engine/Source/Editor/UnrealEd/Public/IPreviewProfileController.h b/Engine/Source/Editor/UnrealEd/Public/IPreviewProfileController.h
new file mode 100644
--- /dev/null
+++ b/Engine/Source/Editor/UnrealEd/Public/IPreviewProfileController.h
@@ -0,0 +1,34 @@
+// Copyright Epic Games, Inc. All Rights Reserved.
+
+#pragma once
+
+#include "CoreMinimal.h"
+
+/**
+ * Sets or gets the available preview profiles. The module(s) offering a preview (like AdvancedPreviewScene) should
+ * implement this interface to enable this module (UnrealEd) to generically display or change the current preview profile.
+ * 
+ * @note This interface was added to prevent circular dependencies between AdvancedPreviewScene and UnrealEd modules.
+ */
+class IPreviewProfileController
+{
+public:
+	virtual ~IPreviewProfileController() = default;
+
+	/** Returns the list of available preview profiles names. */
+	virtual TArray<FString> GetPreviewProfiles(int32& OutCurrentProfileIndex) const = 0;
+
+	/** Set the specified preview profiles as the active one. */
+	virtual bool SetActiveProfile(const FString& ProfileName) = 0;
+
+	/** Returns the preview profiles currently active. */
+	virtual FString GetActiveProfile() const = 0;
+
+	/** Invoked after the list of available profiles has changed. */
+	DECLARE_EVENT(IPreviewProfileController, FOnPreviewProfileListChanged);
+	virtual FOnPreviewProfileListChanged& OnPreviewProfileListChanged() = 0;
+
+	/** Invoked after the active preview profile changed. */
+	DECLARE_EVENT(IPreviewProfileController, FOnPreviewProfileChanged);
+	virtual FOnPreviewProfileChanged& OnPreviewProfileChanged() = 0;
+};
diff --git a/Engine/Source/Editor/UnrealEd/Public/SCommonEditorViewportToolbarBase.h b/Engine/Source/Editor/UnrealEd/Public/SCommonEditorViewportToolbarBase.h
--- a/Engine/Source/Editor/UnrealEd/Public/SCommonEditorViewportToolbarBase.h
+++ b/Engine/Source/Editor/UnrealEd/Public/SCommonEditorViewportToolbarBase.h
@@ -9,6 +9,7 @@
 #include "Framework/MultiBox/MultiBoxExtender.h"
 #include "Framework/MultiBox/MultiBoxBuilder.h"
 #include "SViewportToolBar.h"
+#include "IPreviewProfileController.h"
 
 // This is the interface that the host of a SCommonEditorViewportToolbarBase must implement
 class ICommonEditorViewportToolbarInfoProvider
@@ -67,11 +68,15 @@ class UNREALED_API SCommonEditorViewportToolbarBase : public SViewportToolBar
 public:
 	SLATE_BEGIN_ARGS(SCommonEditorViewportToolbarBase)
 		: _AddRealtimeButton(false)
+		, _PreviewProfileController(nullptr)
 		{}
 
 		SLATE_ARGUMENT(bool, AddRealtimeButton)
+		SLATE_ARGUMENT(TSharedPtr<IPreviewProfileController>, PreviewProfileController) // Should be null if the Preview doesn't require profile.
 	SLATE_END_ARGS()
 
+	virtual ~SCommonEditorViewportToolbarBase();
+
 	void Construct(const FArguments& InArgs, TSharedPtr<class ICommonEditorViewportToolbarInfoProvider> InInfoProvider);
 
 private:
@@ -206,8 +211,27 @@ protected:
 	/** Called when the ScreenPercentage slider is adjusted in the viewport */
 	void OnScreenPercentageValueChanged(int32 NewValue);
 
+	/** Update the list of asset viewer profiles displayed by the combo box. */
+	void UpdateAssetViewerProfileList();
+	void UpdateAssetViewerProfileSelection();
+
+	/** Invoked when the asset viewer profile combo box selection changes. */
+	void OnAssetViewerProfileComboBoxSelectionChanged(TSharedPtr<FString> NewSelection, ESelectInfo::Type /*SelectInfo*/);
+
+	/** Creates and returns the asset viewr profile combo box.*/
+	TSharedRef<SWidget> MakeAssetViewerProfileComboBox();
+
 private:
 	/** The viewport that we are in */
 	TWeakPtr<class ICommonEditorViewportToolbarInfoProvider> InfoProviderPtr;
+
+	/** Interface to set/get/list the preview profiles. */
+	TSharedPtr<IPreviewProfileController> PreviewProfileController;
+
+	/** List of advanced preview profiles to fill up the Profiles combo box. */
+	TArray<TSharedPtr<FString>> AssetViewerProfileNames;
+
+	/** Displays/Selects the active advanced viewer profile. */
+	TSharedPtr<STextComboBox> AssetViewerProfileComboBox;
 };
 
