EditorDomain: Allow EditorDomain packages to propagate upstream. Bulk download EditorDomain packages from remote. Make EditorDomain archive opens a blocking get.
#rnx

diff --git a/Engine/Source/Editor/UnrealEd/Private/Cooker/CookRequestCluster.cpp b/Engine/Source/Editor/UnrealEd/Private/Cooker/CookRequestCluster.cpp
--- a/Engine/Source/Editor/UnrealEd/Private/Cooker/CookRequestCluster.cpp
+++ b/Engine/Source/Editor/UnrealEd/Private/Cooker/CookRequestCluster.cpp
@@ -4,6 +4,7 @@
 
 #include "Algo/Sort.h"
 #include "Algo/Unique.h"
+#include "AssetRegistry/IAssetRegistry.h"
 #include "Cooker/CookPackageData.h"
 #include "Cooker/CookPlatformManager.h"
 #include "Cooker/CookProfiling.h"
@@ -12,6 +13,8 @@
 #include "Cooker/PackageNameCache.h"
 #include "Cooker/PackageTracker.h"
 #include "CookOnTheSide/CookOnTheFlyServer.h"
+#include "DerivedDataCache.h"
+#include "DerivedDataRequestOwner.h"
 #include "EditorDomain/EditorDomainUtils.h"
 #include "Interfaces/ITargetPlatform.h"
 #include "Misc/StringBuilder.h"
@@ -243,6 +246,11 @@ void FRequestCluster::Process(const FCookerTimer& CookerTimer, bool& bOutComplet
 		return;
 	}
 	FetchDependencies(CookerTimer, bOutComplete);
+	if (!bOutComplete)
+	{
+		return;
+	}
+	StartAsync(CookerTimer, bOutComplete);
 }
 
 void FRequestCluster::FetchPackageNames(const FCookerTimer& CookerTimer, bool& bOutComplete)
@@ -478,6 +486,51 @@ void FRequestCluster::FetchDependencies(const FCookerTimer& CookerTimer, bool& b
 	bDependenciesComplete = true;
 }
 
+void FRequestCluster::StartAsync(const FCookerTimer& CookerTimer, bool& bOutComplete)
+{
+	if (bStartAsyncComplete)
+	{
+		return;
+	}
+
+	if (FEditorDomain::Get())
+	{
+		// Disable BatchDownload until cache storage implements fetch-head requests in response to a Get with SkipData
+		bool bBatchDownloadEnabled = false;
+		GConfig->GetBool(TEXT("EditorDomain"), TEXT("BatchDownloadEnabled"), bBatchDownloadEnabled, GEditorIni);
+		if (bBatchDownloadEnabled)
+		{
+			// If the EditorDomain is active, then batch-download all packages to cook from remote cache into local
+			TArray<UE::DerivedData::FCacheKey> CacheKeys;
+			FString ErrorMessage;
+			CacheKeys.Reserve(Requests.Num());
+			for (FPackageData* PackageData : Requests)
+			{
+				UE::EditorDomain::FPackageDigest PackageDigest;
+				bool bEditorDomainEnabled;
+				UE::EditorDomain::EPackageDigestResult Result = UE::EditorDomain::GetPackageDigest(AssetRegistry,
+					PackageData->GetPackageName(), PackageDigest, bEditorDomainEnabled, ErrorMessage);
+				if (Result == UE::EditorDomain::EPackageDigestResult::Success && bEditorDomainEnabled)
+				{
+					CacheKeys.Add(UE::EditorDomain::GetEditorDomainPackageKey(PackageDigest));
+				}
+			}
+
+			if (CacheKeys.Num() > 0)
+			{
+				UE::DerivedData::ICache& Cache = UE::DerivedData::GetCache();
+				UE::DerivedData::ECachePolicy CachePolicy = UE::DerivedData::ECachePolicy::Default | UE::DerivedData::ECachePolicy::SkipData;
+				UE::DerivedData::FRequestOwner Owner(UE::DerivedData::EPriority::Highest);
+				Owner.KeepAlive();
+				Cache.Get(CacheKeys, TEXT("RequestClusterEditorDomainDownload"_SV), CachePolicy, Owner,
+					[](UE::DerivedData::FCacheGetCompleteParams&& Params) {});
+			}
+		}
+	}
+
+	bStartAsyncComplete = true;
+}
+
 int32 FRequestCluster::NumPackageDatas() const
 {
 	return OwnedPackageDatas.Num();
@@ -502,7 +555,7 @@ bool FRequestCluster::Contains(FPackageData* PackageData) const
 
 void FRequestCluster::ClearAndDetachOwnedPackageDatas(TArray<FPackageData*>& OutRequestsToLoad, TArray<FPackageData*>& OutRequestsToDemote)
 {
-	if (bDependenciesComplete)
+	if (bStartAsyncComplete)
 	{
 		check(TransitiveRequests.Num() == 0 && Segment.Num() == 0);
 		OutRequestsToLoad = MoveTemp(Requests);
diff --git a/Engine/Source/Editor/UnrealEd/Private/Cooker/CookRequestCluster.h b/Engine/Source/Editor/UnrealEd/Private/Cooker/CookRequestCluster.h
--- a/Engine/Source/Editor/UnrealEd/Private/Cooker/CookRequestCluster.h
+++ b/Engine/Source/Editor/UnrealEd/Private/Cooker/CookRequestCluster.h
@@ -93,6 +93,7 @@ private:
 	void Initialize(UCookOnTheFlyServer& COTFS);
 	void FetchPackageNames(const FCookerTimer& CookerTimer, bool& bOutComplete);
 	void FetchDependencies(const FCookerTimer& CookerTimer, bool& bOutComplete);
+	void StartAsync(const FCookerTimer& CookerTimer, bool& bOutComplete);
 	bool TryTakeOwnership(FPackageData& PackageData, bool bUrgent, FCompletionCallback&& CompletionCallback);
 	void VisitPackageData(FPackageData* PackageData, TArray<FName>* OutDependencies, bool& bAlreadyCooked);
 	bool IsRequestCookable(FName PackageName, FName& InOutFileName, FPackageData*& PackageData);
@@ -127,6 +128,7 @@ private:
 	bool bErrorOnEngineContentUse = false;
 	bool bPackageNamesComplete = false;
 	bool bDependenciesComplete = false;
+	bool bStartAsyncComplete = false;
 	bool bFullBuild = false;
 
 
diff --git a/Engine/Source/Editor/UnrealEd/Private/EditorDomain/EditorDomain.cpp b/Engine/Source/Editor/UnrealEd/Private/EditorDomain/EditorDomain.cpp
--- a/Engine/Source/Editor/UnrealEd/Private/EditorDomain/EditorDomain.cpp
+++ b/Engine/Source/Editor/UnrealEd/Private/EditorDomain/EditorDomain.cpp
@@ -306,7 +306,10 @@ FOpenPackageResult FEditorDomain::OpenReadPackage(const FPackagePath& PackagePat
 		return Workspace->OpenReadPackage(PackagePath, PackageSegment, OutUpdatedPath);
 	}
 
-	FEditorDomainReadArchive* Result = new FEditorDomainReadArchive(Locks, PackagePath, PackageSource);
+	// TODO: Change priority to High instead of Blocking once we have removed the GetPackageFormat below
+	// and don't need to block on the result before exiting this function
+	EPriority Priority = EPriority::Blocking;
+	FEditorDomainReadArchive* Result = new FEditorDomainReadArchive(Locks, PackagePath, PackageSource, Priority);
 	const FPackageDigest PackageSourceDigest = PackageSource->Digest;
 	const bool bHasEditorSource = (PackageSource->Source == EPackageSource::Editor);
 
diff --git a/Engine/Source/Editor/UnrealEd/Private/EditorDomain/EditorDomainArchive.cpp b/Engine/Source/Editor/UnrealEd/Private/EditorDomain/EditorDomainArchive.cpp
--- a/Engine/Source/Editor/UnrealEd/Private/EditorDomain/EditorDomainArchive.cpp
+++ b/Engine/Source/Editor/UnrealEd/Private/EditorDomain/EditorDomainArchive.cpp
@@ -17,9 +17,9 @@
 #include "Serialization/CompactBinary.h"
 
 FEditorDomainPackageSegments::FEditorDomainPackageSegments(const TRefCountPtr<FEditorDomain::FLocks>& InLocks,
-	const FPackagePath& InPackagePath, const TRefCountPtr<FEditorDomain::FPackageSource>& InPackageSource)
+	const FPackagePath& InPackagePath, const TRefCountPtr<FEditorDomain::FPackageSource>& InPackageSource, UE::DerivedData::EPriority Priority)
 	: EditorDomainLocks(InLocks)
-	, RequestOwner(UE::DerivedData::EPriority::Normal)
+	, RequestOwner(Priority)
 	, PackagePath(InPackagePath)
 	, PackageSource(InPackageSource)
 	, PackageDigest(InPackageSource->Digest)
@@ -434,8 +434,9 @@ void FEditorDomainPackageSegments::SendSegmentRequest(FSegment& Segment)
 }
 
 FEditorDomainReadArchive::FEditorDomainReadArchive(const TRefCountPtr<FEditorDomain::FLocks>& InLocks,
-	const FPackagePath& InPackagePath, const TRefCountPtr<FEditorDomain::FPackageSource>& InPackageSource)
-	: Segments(InLocks, InPackagePath, InPackageSource)
+	const FPackagePath& InPackagePath, const TRefCountPtr<FEditorDomain::FPackageSource>& InPackageSource,
+	UE::DerivedData::EPriority Priority)
+	: Segments(InLocks, InPackagePath, InPackageSource, Priority)
 {
 	this->SetIsLoading(true);
 	this->SetIsPersistent(true);
@@ -729,7 +730,7 @@ protected:
 
 FEditorDomainAsyncReadFileHandle::FEditorDomainAsyncReadFileHandle(const TRefCountPtr<FEditorDomain::FLocks>& InLocks,
 	const FPackagePath& InPackagePath, const TRefCountPtr<FEditorDomain::FPackageSource>& InPackageSource)
-	: Segments(InLocks, InPackagePath, InPackageSource)
+	: Segments(InLocks, InPackagePath, InPackageSource, UE::DerivedData::EPriority::Normal)
 {
 }
 
diff --git a/Engine/Source/Editor/UnrealEd/Private/EditorDomain/EditorDomainArchive.h b/Engine/Source/Editor/UnrealEd/Private/EditorDomain/EditorDomainArchive.h
--- a/Engine/Source/Editor/UnrealEd/Private/EditorDomain/EditorDomainArchive.h
+++ b/Engine/Source/Editor/UnrealEd/Private/EditorDomain/EditorDomainArchive.h
@@ -40,7 +40,8 @@ public:
 	};
 
 	FEditorDomainPackageSegments(const TRefCountPtr<FEditorDomain::FLocks>& InLocks,
-		const FPackagePath& InPackagePath, const TRefCountPtr<FEditorDomain::FPackageSource>& InPackageSource);
+		const FPackagePath& InPackagePath, const TRefCountPtr<FEditorDomain::FPackageSource>& InPackageSource,
+		UE::DerivedData::EPriority Priority);
 	~FEditorDomainPackageSegments();
 
 	/** Get the request owner for requests that will feed this archive. */
@@ -197,7 +198,8 @@ class FEditorDomainReadArchive : public FArchive
 {
 public:
 	FEditorDomainReadArchive(const TRefCountPtr<FEditorDomain::FLocks>& InLocks,
-		const FPackagePath& InPackagePath, const TRefCountPtr<FEditorDomain::FPackageSource>& InPackageSource);
+		const FPackagePath& InPackagePath, const TRefCountPtr<FEditorDomain::FPackageSource>& InPackageSource,
+		UE::DerivedData::EPriority Priority);
 
 	/** Get the request owner for requests that will feed this archive. */
 	UE::DerivedData::IRequestOwner& GetRequestOwner() { return Segments.GetRequestOwner(); }
diff --git a/Engine/Source/Editor/UnrealEd/Private/EditorDomain/EditorDomainUtils.cpp b/Engine/Source/Editor/UnrealEd/Private/EditorDomain/EditorDomainUtils.cpp
--- a/Engine/Source/Editor/UnrealEd/Private/EditorDomain/EditorDomainUtils.cpp
+++ b/Engine/Source/Editor/UnrealEd/Private/EditorDomain/EditorDomainUtils.cpp
@@ -614,12 +614,19 @@ void RequestEditorDomainPackage(const FPackagePath& PackagePath,
 	const FPackageDigest& PackageDigest, UE::DerivedData::ECachePolicy SkipFlags, UE::DerivedData::IRequestOwner& Owner,
 	UE::DerivedData::FOnCacheGetComplete&& Callback)
 {
-	UE::DerivedData::ICache& Cache = UE::DerivedData::GetCache();
-	checkf((SkipFlags & (~UE::DerivedData::ECachePolicy::SkipData)) == UE::DerivedData::ECachePolicy::None,
+	using namespace UE::DerivedData;
+
+	ICache& Cache = GetCache();
+	checkf((SkipFlags & (~ECachePolicy::SkipData)) == ECachePolicy::None,
 		TEXT("SkipFlags should only contain ECachePolicy::Skip* flags"));
-	Cache.Get({ GetEditorDomainPackageKey(PackageDigest) },
-		PackagePath.GetDebugName(),
-		SkipFlags | UE::DerivedData::ECachePolicy::QueryLocal, Owner, MoveTemp(Callback));
+
+	// Set the CachePolicy to only query from local; we do not want to wait for download from remote.
+	// Downloading from remote is done in batch  see FRequestCluster::StartAsync.
+	// But set the CachePolicy to store into remote. This will cause the CacheStore to push
+	// any existing local value into upstream storage and refresh the last-used time in the upstream.
+	ECachePolicy CachePolicy = SkipFlags | ECachePolicy::Local | ECachePolicy::StoreRemote;
+	Cache.Get({ GetEditorDomainPackageKey(PackageDigest) }, PackagePath.GetDebugName(),
+		CachePolicy, Owner, MoveTemp(Callback));
 }
 
 /** Stores data from SavePackage in accessible fields */
@@ -841,7 +848,7 @@ bool TrySavePackage(UPackage* Package)
 
 	RecordBuilder.SetMeta(MetaData.Save().AsObject());
 	FRequestOwner Owner(EPriority::Normal);
-	Cache.Put({RecordBuilder.Build()}, Package->GetName(), ECachePolicy::Local, Owner);
+	Cache.Put({RecordBuilder.Build()}, Package->GetName(), ECachePolicy::Default, Owner);
 	Owner.KeepAlive();
 	return true;
 }
