Added HLOD generation options
* Added option to skip submit (will currently have to fetch build artifacts manually)
* Added option to only build a given HLOD level
* -AllowCommandletRendering not required anymore for the "-SetupHLODs" step
* General cleanup in UWorldPartitionHLODsBuilder

diff --git a/Engine/Source/Editor/UnrealEd/Private/WorldPartition/WorldPartitionHLODsBuilder.cpp b/Engine/Source/Editor/UnrealEd/Private/WorldPartition/WorldPartitionHLODsBuilder.cpp
--- a/Engine/Source/Editor/UnrealEd/Private/WorldPartition/WorldPartitionHLODsBuilder.cpp
+++ b/Engine/Source/Editor/UnrealEd/Private/WorldPartition/WorldPartitionHLODsBuilder.cpp
@@ -156,18 +156,22 @@ UWorldPartitionHLODsBuilder::UWorldPartitionHLODsBuilder(const FObjectInitialize
 	: Super(ObjectInitializer)
 	, BuilderIdx(INDEX_NONE)
 	, BuilderCount(INDEX_NONE)
+	, HLODLevelToBuild(INDEX_NONE)
 	, DistributedBuildWorkingDir(FPaths::RootDir() / DistributedBuildWorkingDirName)
 	, DistributedBuildManifest(DistributedBuildWorkingDir / DistributedBuildManifestName)
 {
-	bSetupHLODs = FParse::Param(FCommandLine::Get(), TEXT("SetupHLODs"));
-	bBuildHLODs = FParse::Param(FCommandLine::Get(), TEXT("BuildHLODs"));
-	bDeleteHLODs = FParse::Param(FCommandLine::Get(), TEXT("DeleteHLODs"));
-	bSubmitHLODs = FParse::Param(FCommandLine::Get(), TEXT("SubmitHLODs"));
-	bDumpStats = FParse::Param(FCommandLine::Get(), TEXT("DumpStats"));
-
-	bSingleBuildStep = !bSetupHLODs && !bBuildHLODs && !bSubmitHLODs && !bDeleteHLODs && !bDumpStats;
+	BuildOptions = FParse::Param(FCommandLine::Get(), TEXT("SetupHLODs")) ? EHLODBuildStep::HLOD_Setup : EHLODBuildStep::None;
+	BuildOptions |= FParse::Param(FCommandLine::Get(), TEXT("BuildHLODs")) ? EHLODBuildStep::HLOD_Build : EHLODBuildStep::None;
+	BuildOptions |= FParse::Param(FCommandLine::Get(), TEXT("DeleteHLODs")) ? EHLODBuildStep::HLOD_Delete : EHLODBuildStep::None;
+	BuildOptions |= FParse::Param(FCommandLine::Get(), TEXT("SubmitHLODs")) ? EHLODBuildStep::HLOD_Submit : EHLODBuildStep::None;
+	BuildOptions |= FParse::Param(FCommandLine::Get(), TEXT("DumpStats")) ? EHLODBuildStep::HLOD_Stats : EHLODBuildStep::None;
+
+	// Default behavior without any option is to setup + build
+	if (BuildOptions == EHLODBuildStep::None)
+	{
+		BuildOptions = EHLODBuildStep::HLOD_Setup | EHLODBuildStep::HLOD_Build;
+	}
 
-	bAutoSubmit = FParse::Param(FCommandLine::Get(), TEXT("AutoSubmit"));
 	bResumeBuild = FParse::Value(FCommandLine::Get(), TEXT("ResumeBuild="), ResumeBuildIndex);
 
 	bDistributedBuild = FParse::Param(FCommandLine::Get(), TEXT("DistributedBuild"));
@@ -176,6 +180,8 @@ UWorldPartitionHLODsBuilder::UWorldPartitionHLODsBuilder(const FObjectInitialize
 	FParse::Value(FCommandLine::Get(), TEXT("BuilderIdx="), BuilderIdx);
 	FParse::Value(FCommandLine::Get(), TEXT("BuilderCount="), BuilderCount);
 
+	FParse::Value(FCommandLine::Get(), TEXT("HLODLevel="), HLODLevelToBuild);
+
 	if (bDistributedBuild)
 	{
 		if (!BuildManifest.IsEmpty())
@@ -191,13 +197,17 @@ bool UWorldPartitionHLODsBuilder::RequiresCommandletRendering() const
 {
 	// Commandlet requires rendering only for building HLODs
 	// Building will occur either if -BuildHLODs is provided or no explicit step arguments are provided
-	// Currently required with -SetupHLODs as we want to build texture streaming data for HLODs on save.
-	return bSetupHLODs || bBuildHLODs || bSingleBuildStep;
+	return EnumHasAnyFlags(BuildOptions, EHLODBuildStep::HLOD_Build);
+}
+
+bool UWorldPartitionHLODsBuilder::ShouldRunStep(const EHLODBuildStep BuildStep) const
+{
+	return (BuildOptions & BuildStep) == BuildStep;
 }
 
 bool UWorldPartitionHLODsBuilder::ValidateParams() const
 {
-	if (bSetupHLODs && IsUsingBuildManifest())
+	if (ShouldRunStep(EHLODBuildStep::HLOD_Setup) && IsUsingBuildManifest())
 	{
 		if (BuilderCount <= 0)
 		{
@@ -206,7 +216,7 @@ bool UWorldPartitionHLODsBuilder::ValidateParams() const
 		}
 	}
 
-	if (bBuildHLODs && IsUsingBuildManifest())
+	if (ShouldRunStep(EHLODBuildStep::HLOD_Build) && IsUsingBuildManifest())
 	{
 		if (BuilderIdx < 0)
 		{
@@ -241,21 +251,9 @@ bool UWorldPartitionHLODsBuilder::ValidateParams() const
 		}
 	}
 
-	if (bSubmitHLODs && !IsDistributedBuild())
+	if (ShouldRunStep(EHLODBuildStep::HLOD_Submit) && !ISourceControlModule::Get().GetProvider().IsEnabled())
 	{
-		UE_LOG(LogWorldPartitionHLODsBuilder, Error, TEXT("-SubmitHLODs argument only valid for distributed builds, exiting..."), *BuildManifest);
-		return false;
-	}
-
-	if (bAutoSubmit && !bSingleBuildStep)
-	{
-		UE_LOG(LogWorldPartitionHLODsBuilder, Error, TEXT("-AutoSubmit argument only valid when building HLODs in a single step, exiting..."), *BuildManifest);
-		return false;
-	}
-
-	if (IsDistributedBuild() && bSubmitHLODs && !ISourceControlModule::Get().GetProvider().IsEnabled())
-	{
-		UE_LOG(LogWorldPartitionHLODsBuilder, Error, TEXT("Distributed builds submit step requires that a valid source control provider is enabled, exiting..."), *BuildManifest);
+		UE_LOG(LogWorldPartitionHLODsBuilder, Error, TEXT("Submit step requires that a valid source control provider is enabled, exiting..."), *BuildManifest);
 		return false;
 	}
 
@@ -272,9 +270,9 @@ bool UWorldPartitionHLODsBuilder::PreWorldInitialization(FPackageSourceControlHe
 	bool bRet = true;
 
 	// When running a distributed build, retrieve relevant build products from the previous steps
-	if (IsDistributedBuild() && (bBuildHLODs || bSubmitHLODs))
+	if (IsDistributedBuild() && (ShouldRunStep(EHLODBuildStep::HLOD_Build) || ShouldRunStep(EHLODBuildStep::HLOD_Submit)))
 	{
-		FString WorkingDirFolder = bBuildHLODs ? GetHLODBuilderFolderName(BuilderIdx) : GetToSubmitFolderName();
+		FString WorkingDirFolder = ShouldRunStep(EHLODBuildStep::HLOD_Build) ? GetHLODBuilderFolderName(BuilderIdx) : GetToSubmitFolderName();
 		bRet = CopyFilesFromWorkingDir(WorkingDirFolder);
 	}
 
@@ -290,32 +288,27 @@ bool UWorldPartitionHLODsBuilder::RunInternal(UWorld* World, const FBox& Bounds,
 
 	bool bRet = true;
 
-	if (bSingleBuildStep)
+	if (bRet && ShouldRunStep(EHLODBuildStep::HLOD_Setup))
 	{
-		bRet = SetupHLODActors(false);
+		bRet = SetupHLODActors();
 	}
-	else
+
+	if (bRet && ShouldRunStep(EHLODBuildStep::HLOD_Build))
 	{
-		if (bSetupHLODs)
-		{
-			bRet = SetupHLODActors(true);
-		}
-		else if (bBuildHLODs)
-		{
-			bRet = BuildHLODActors();
-		}
-		else if (bDeleteHLODs)
-		{
-			bRet = DeleteHLODActors();
-		}
+		bRet = BuildHLODActors();
+	}
+
+	if (bRet && ShouldRunStep(EHLODBuildStep::HLOD_Delete))
+	{
+		bRet = DeleteHLODActors();
 	}
 
-	if (bRet && (bSubmitHLODs || bAutoSubmit))
+	if (bRet && ShouldRunStep(EHLODBuildStep::HLOD_Submit))
 	{
 		bRet = SubmitHLODActors();
 	}
 
-	if (bRet && bDumpStats)
+	if (bRet && ShouldRunStep(EHLODBuildStep::HLOD_Stats))
 	{
 		bRet = DumpStats();
 	}
@@ -329,100 +322,98 @@ bool UWorldPartitionHLODsBuilder::RunInternal(UWorld* World, const FBox& Bounds,
 	return bRet;
 }
 
-bool UWorldPartitionHLODsBuilder::SetupHLODActors(bool bCreateOnly)
+bool UWorldPartitionHLODsBuilder::SetupHLODActors()
 {
-	WorldPartition->GenerateHLOD(SourceControlHelper, bCreateOnly);
+	const bool bCreateActorsOnly = true;
+	WorldPartition->GenerateHLOD(SourceControlHelper, bCreateActorsOnly);
 
-	if (bCreateOnly)
+	// When performing a distributed build, ensure our work folder is empty
+	if (IsDistributedBuild())
 	{
-		// When performing a distributed build, ensure our work folder is empty
-		if (IsDistributedBuild())
-		{
-			IFileManager::Get().DeleteDirectory(*DistributedBuildWorkingDir, false, true);
-		}
+		IFileManager::Get().DeleteDirectory(*DistributedBuildWorkingDir, false, true);
+	}
 
-		UE_LOG(LogWorldPartitionHLODsBuilder, Display, TEXT("#### World HLOD actors ####"));
+	UE_LOG(LogWorldPartitionHLODsBuilder, Display, TEXT("#### World HLOD actors ####"));
 
-		int32 NumActors = 0;
-		for (FActorDescList::TIterator<AWorldPartitionHLOD> HLODIterator(WorldPartition); HLODIterator; ++HLODIterator)
-		{
-			FWorldPartitionActorDesc* HLODActorDesc = *HLODIterator;
-			FString PackageName = HLODActorDesc->GetActorPackage().ToString();
+	int32 NumActors = 0;
+	for (FActorDescList::TIterator<AWorldPartitionHLOD> HLODIterator(WorldPartition); HLODIterator; ++HLODIterator)
+	{
+		FWorldPartitionActorDesc* HLODActorDesc = *HLODIterator;
+		FString PackageName = HLODActorDesc->GetActorPackage().ToString();
 
-			UE_LOG(LogWorldPartitionHLODsBuilder, Display, TEXT("    [%d] %s"), NumActors, *PackageName);
+		UE_LOG(LogWorldPartitionHLODsBuilder, Display, TEXT("    [%d] %s"), NumActors, *PackageName);
 
-			NumActors++;
-		}
+		NumActors++;
+	}
 
-		UE_LOG(LogWorldPartitionHLODsBuilder, Display, TEXT("#### World contains %d HLOD actors ####"), NumActors);
+	UE_LOG(LogWorldPartitionHLODsBuilder, Display, TEXT("#### World contains %d HLOD actors ####"), NumActors);
 
-		if (IsUsingBuildManifest())
+	if (IsUsingBuildManifest())
+	{
+		TMap<FString, int32> FilesToBuilderMap;
+		bool bGenerated = GenerateBuildManifest(FilesToBuilderMap);
+		if (!bGenerated)
 		{
-			TMap<FString, int32> FilesToBuilderMap;
-			bool bGenerated = GenerateBuildManifest(FilesToBuilderMap);
-			if (!bGenerated)
-			{
-				return false;
-			}
+			return false;
+		}
 
-			// When performing a distributed build, move modified files to the temporary working dir, to be submitted later in the last "submit" step
-			if (IsDistributedBuild())
-			{
-				// Ensure we don't hold on to packages of always loaded actors
-				// When running distributed builds, we wanna leave the machine clean, so added files are deleted, check'd out files are reverted
-				// and deleted files are restored.
-				WorldPartition->Uninitialize();
-				FWorldPartitionHelpers::DoCollectGarbage();
+		// When performing a distributed build, move modified files to the temporary working dir, to be submitted later in the last "submit" step
+		if (IsDistributedBuild())
+		{
+			// Ensure we don't hold on to packages of always loaded actors
+			// When running distributed builds, we wanna leave the machine clean, so added files are deleted, check'd out files are reverted
+			// and deleted files are restored.
+			WorldPartition->Uninitialize();
+			FWorldPartitionHelpers::DoCollectGarbage();
 
-				ModifiedFiles.Append(SourceControlHelper->GetModifiedFiles());
+			ModifiedFiles.Append(SourceControlHelper->GetModifiedFiles());
 
-				TArray<FHLODModifiedFiles> BuildersFiles;
-				BuildersFiles.SetNum(BuilderCount);
+			TArray<FHLODModifiedFiles> BuildersFiles;
+			BuildersFiles.SetNum(BuilderCount);
 
-				for (int32 i = 0; i < FHLODModifiedFiles::EFileOperation::NumFileOperations; i++)
+			for (int32 i = 0; i < FHLODModifiedFiles::EFileOperation::NumFileOperations; i++)
+			{
+				FHLODModifiedFiles::EFileOperation FileOp = (FHLODModifiedFiles::EFileOperation)i;
+				for (const FString& ModifiedFile : ModifiedFiles.Get(FileOp))
 				{
-					FHLODModifiedFiles::EFileOperation FileOp = (FHLODModifiedFiles::EFileOperation)i;
-					for (const FString& ModifiedFile : ModifiedFiles.Get(FileOp))
+					int32* Idx = FilesToBuilderMap.Find(ModifiedFile);
+					if (Idx)
 					{
-						int32* Idx = FilesToBuilderMap.Find(ModifiedFile);
-						if (Idx)
-						{
-							BuildersFiles[*Idx].Add(FileOp, ModifiedFile);
-						}
-						else
-						{
-							// Add general files to the last builder
-							BuildersFiles.Last().Add(FileOp, ModifiedFile);
-						}
+						BuildersFiles[*Idx].Add(FileOp, ModifiedFile);
 					}
-				}
-
-				// Gather build product to ensure intermediary files are copied between the different HLOD generation steps
-				TArray<FString> BuildProducts;
-
-				// Copy files that will be handled by the different builders
-				for (int32 Idx = 0; Idx < BuilderCount; Idx++)
-				{
-					if (!CopyFilesToWorkingDir(GetHLODBuilderFolderName(Idx), BuildersFiles[Idx], BuildProducts))
+					else
 					{
-						return false;
+						// Add general files to the last builder
+						BuildersFiles.Last().Add(FileOp, ModifiedFile);
 					}
 				}
+			}
 
-				// The build manifest must also be included as a build product to be available in the next steps
-				BuildProducts.Add(BuildManifest);
+			// Gather build product to ensure intermediary files are copied between the different HLOD generation steps
+			TArray<FString> BuildProducts;
 
-				// Write build products to a file
-				FString BuildProductsFile = DistributedBuildWorkingDir / BuildProductsFileName;
-				bool bRet = FFileHelper::SaveStringArrayToFile(BuildProducts, *BuildProductsFile);
-				if (!bRet)
+			// Copy files that will be handled by the different builders
+			for (int32 Idx = 0; Idx < BuilderCount; Idx++)
+			{
+				if (!CopyFilesToWorkingDir(GetHLODBuilderFolderName(Idx), BuildersFiles[Idx], BuildProducts))
 				{
-					UE_LOG(LogWorldPartitionHLODsBuilder, Error, TEXT("Error writing build product file %s"), *BuildProductsFile);
 					return false;
 				}
+			}
+
+			// The build manifest must also be included as a build product to be available in the next steps
+			BuildProducts.Add(BuildManifest);
 
-				ModifiedFiles.Empty();
+			// Write build products to a file
+			FString BuildProductsFile = DistributedBuildWorkingDir / BuildProductsFileName;
+			bool bRet = FFileHelper::SaveStringArrayToFile(BuildProducts, *BuildProductsFile);
+			if (!bRet)
+			{
+				UE_LOG(LogWorldPartitionHLODsBuilder, Error, TEXT("Error writing build product file %s"), *BuildProductsFile);
+				return false;
 			}
+
+			ModifiedFiles.Empty();
 		}
 	}
 
@@ -459,7 +450,10 @@ bool UWorldPartitionHLODsBuilder::BuildHLODActors()
 
 		UE_LOG(LogWorldPartitionHLODsBuilder, Display, TEXT("    [%d/%d] Building HLOD actor %s..."), CurrentActor, HLODActorsToBuild.Num(), *HLODActor->GetActorLabel());
 
-		HLODActor->BuildHLOD();
+		if (HLODLevelToBuild == INDEX_NONE || HLODActor->GetLODLevel() == HLODLevelToBuild)
+		{
+			HLODActor->BuildHLOD();
+		}
 
 		UPackage* ActorPackage = HLODActor->GetPackage();
 		if (ActorPackage->IsDirty())
diff --git a/Engine/Source/Editor/UnrealEd/Public/WorldPartition/WorldPartitionHLODsBuilder.h b/Engine/Source/Editor/UnrealEd/Public/WorldPartition/WorldPartitionHLODsBuilder.h
--- a/Engine/Source/Editor/UnrealEd/Public/WorldPartition/WorldPartitionHLODsBuilder.h
+++ b/Engine/Source/Editor/UnrealEd/Public/WorldPartition/WorldPartitionHLODsBuilder.h
@@ -2,6 +2,7 @@
 
 #pragma once
 
+#include "Misc/EnumClassFlags.h"
 #include "WorldPartition/WorldPartitionBuilder.h"
 #include "WorldPartitionHLODsBuilder.generated.h"
 
@@ -57,6 +58,17 @@ private:
 	TSet<FString> Files[NumFileOperations];
 };
 
+enum class EHLODBuildStep : uint8
+{
+	None		= 0,
+	HLOD_Setup	= 1 << 0,
+	HLOD_Build	= 1 << 1,
+	HLOD_Submit = 1 << 2,
+	HLOD_Delete = 1 << 3,
+	HLOD_Stats	= 1 << 4
+};
+ENUM_CLASS_FLAGS(EHLODBuildStep);
+
 UCLASS()
 class UWorldPartitionHLODsBuilder : public UWorldPartitionBuilder
 {
@@ -74,7 +86,7 @@ protected:
 	bool IsUsingBuildManifest() const { return !BuildManifest.IsEmpty(); }
 	bool ValidateParams() const;
 
-	bool SetupHLODActors(bool bCreateOnly);
+	bool SetupHLODActors();
 	bool BuildHLODActors();
 	bool DeleteHLODActors();
 	bool SubmitHLODActors();
@@ -89,24 +101,22 @@ protected:
 	bool CopyFilesToWorkingDir(const FString& TargetDir, const FHLODModifiedFiles& ModifiedFiles, TArray<FString>& BuildProducts);
 	bool CopyFilesFromWorkingDir(const FString& SourceDir);
 
+	bool ShouldRunStep(const EHLODBuildStep BuildStep) const;
+
 private:
 	class UWorldPartition* WorldPartition;
 	class FSourceControlHelper* SourceControlHelper;
 
 	// Options --
-	bool bSetupHLODs;
-	bool bBuildHLODs;
-	bool bDeleteHLODs;
-	bool bSubmitHLODs;
-	bool bDumpStats;
-	bool bSingleBuildStep;
-	bool bAutoSubmit;
+	EHLODBuildStep BuildOptions;
+
 	bool bDistributedBuild;
 	FString BuildManifest;
 	int32 BuilderIdx;
 	int32 BuilderCount;
 	bool bResumeBuild;
 	int32 ResumeBuildIndex;
+	int32 HLODLevelToBuild;
 
 
 	const FString DistributedBuildWorkingDir;
