Add per vertex factory type histogram.

Results in output that looks like the following:

VFType, Count, Percent Total
FLocalVertexFactory, 147816, 24.49
FInstancedStaticMeshVertexFactory, 66850, 11.08
FGPUSkinPassthroughVertexFactory, 63316, 10.49
TGPUSkinVertexFactoryDefault, 53532, 8.87
TGPUSkinVertexFactoryUnlimited, 53532, 8.87
TGPUSkinMorphVertexFactoryUnlimited, 23737, 3.93
TGPUSkinMorphVertexFactoryDefault, 23737, 3.93
FGeometryCacheVertexVertexFactory, 16274, 2.70
FNiagaraMeshVertexFactory, 15952, 2.64
FNiagaraMeshVertexFactoryEx, 15952, 2.64

diff --git a/Engine/Source/Editor/UnrealEd/Private/Commandlets/GatherShaderStatsFromMaterialsCommandlet.cpp b/Engine/Source/Editor/UnrealEd/Private/Commandlets/GatherShaderStatsFromMaterialsCommandlet.cpp
--- a/Engine/Source/Editor/UnrealEd/Private/Commandlets/GatherShaderStatsFromMaterialsCommandlet.cpp
+++ b/Engine/Source/Editor/UnrealEd/Private/Commandlets/GatherShaderStatsFromMaterialsCommandlet.cpp
@@ -50,18 +50,28 @@ public:
 		{
 			FullShaderTypeHistogram.FindOrAdd(AbsoluteShaderName, 1);
 		}
+
+		FString VFTypeName(VertexFactoryName);
+		if (int32* Existing = VertexFactoryTypeHistogram.Find(VFTypeName))
+		{
+			++(*Existing);
+		}
+		else
+		{
+			VertexFactoryTypeHistogram.FindOrAdd(VFTypeName, 1);
+		}
 	}
 
-	void PrintHistogram()
+	void PrintHistogram(int TotalShaders)
 	{
 		if (ShaderTypeHistogram.Num() > 0)
 		{
 			ShaderTypeHistogram.ValueSort(TGreater<int32>());
-			const char ShaderTypeHeader[] = "\nShaderType, Count\n";
+			const char ShaderTypeHeader[] = "\nShaderType, Count, Percent Total\n";
 			DebugWriter->Serialize(const_cast<char*>(ShaderTypeHeader), sizeof(ShaderTypeHeader) - 1);
 			for (TPair<FString, int32> ShaderUsage : ShaderTypeHistogram)
 			{
-				FString OuputLine = FString::Printf(TEXT("%s, %d\n"), *ShaderUsage.Key, ShaderUsage.Value);
+				FString OuputLine = FString::Printf(TEXT("%s, %d, %.2f\n"), *ShaderUsage.Key, ShaderUsage.Value, (ShaderUsage.Value / (float)TotalShaders) * 100.0f);
 				DebugWriter->Serialize(const_cast<ANSICHAR*>(StringCast<ANSICHAR>(*OuputLine).Get()), OuputLine.Len());
 			}
 		}
@@ -69,11 +79,23 @@ public:
 		if (FullShaderTypeHistogram.Num() > 0)
 		{
 			FullShaderTypeHistogram.ValueSort(TGreater<int32>());
-			const char FullShaderTypeHeader[] = "\nFullShaderType, Count\n";
+			const char FullShaderTypeHeader[] = "\nFullShaderType, Count, Percent Total\n";
 			DebugWriter->Serialize(const_cast<char*>(FullShaderTypeHeader), sizeof(FullShaderTypeHeader) - 1);
 			for (TPair<FString, int32> ShaderUsage : FullShaderTypeHistogram)
 			{
-				FString OuputLine = FString::Printf(TEXT("%s, %d\n"), *ShaderUsage.Key, ShaderUsage.Value);
+				FString OuputLine = FString::Printf(TEXT("%s, %d, %.2f\n"), *ShaderUsage.Key, ShaderUsage.Value, (ShaderUsage.Value / (float)TotalShaders) * 100.0f);
+				DebugWriter->Serialize(const_cast<ANSICHAR*>(StringCast<ANSICHAR>(*OuputLine).Get()), OuputLine.Len());
+			}
+		}
+
+		if (VertexFactoryTypeHistogram.Num() > 0)
+		{
+			VertexFactoryTypeHistogram.ValueSort(TGreater<int32>());
+			const char FullVFTypeHeader[] = "\nVFType, Count, Percent Total\n";
+			DebugWriter->Serialize(const_cast<char*>(FullVFTypeHeader), sizeof(FullVFTypeHeader) - 1);
+			for (TPair<FString, int32> VFTypeUsage : VertexFactoryTypeHistogram)
+			{
+				FString OuputLine = FString::Printf(TEXT("%s, %d, %.2f\n"), *VFTypeUsage.Key, VFTypeUsage.Value, (VFTypeUsage.Value / (float)TotalShaders) * 100.0f);
 				DebugWriter->Serialize(const_cast<ANSICHAR*>(StringCast<ANSICHAR>(*OuputLine).Get()), OuputLine.Len());
 			}
 		}
@@ -106,6 +128,19 @@ public:
 				DebugWriter->Serialize(const_cast<ANSICHAR*>(StringCast<ANSICHAR>(*OuputLine).Get()), OuputLine.Len());
 			}
 		}
+
+		if (VertexFactoryTypeHistogram.Num() > 0)
+		{
+			VertexFactoryTypeHistogram.KeySort(TGreater<FString>());
+			const char FullVFTypeAlphabeticHeader[] = "\nVertexFactoryType only\n";
+			DebugWriter->Serialize(const_cast<char*>(FullVFTypeAlphabeticHeader), sizeof(FullVFTypeAlphabeticHeader) - 1);
+			for (TPair<FString, int32> VFTypeUsage : VertexFactoryTypeHistogram)
+			{
+				// do not print numbers here as it complicates the diff
+				FString OuputLine = FString::Printf(TEXT("%s\n"), *VFTypeUsage.Key);
+				DebugWriter->Serialize(const_cast<ANSICHAR*>(StringCast<ANSICHAR>(*OuputLine).Get()), OuputLine.Len());
+			}
+		}
 	}
 
 	void Log(const FString& OutString)
@@ -123,6 +158,9 @@ private:
 
 	/** Map of full shader display names to their counts. */
 	TMap<FString, int32>	FullShaderTypeHistogram;
+
+	/** Map of vertex factory display names to their counts. */
+	TMap<FString, int32>	VertexFactoryTypeHistogram;
 };
 
 UGatherShaderStatsFromMaterialsCommandlet::UGatherShaderStatsFromMaterialsCommandlet(const FObjectInitializer& ObjectInitializer)
@@ -375,7 +413,7 @@ int32 UGatherShaderStatsFromMaterialsCommandlet::Main(const FString& Params)
 	Output.Log(FString::Printf(TEXT("Total Assets: %d"), TotalAssets));
 	Output.Log(FString::Printf(TEXT("Total Shaders: %d"), TotalShaders));
 	Output.Log(FString::Printf(TEXT("Histogram:")));
-	Output.PrintHistogram();
+	Output.PrintHistogram(TotalShaders);
 	Output.Log(FString::Printf(TEXT("Alphabetic:")));
 	Output.PrintAlphabeticList();
 
