Fix mismatched combo dropdown and toggle button on viewport toolbars

diff --git a/Engine/Source/Editor/UnrealEd/Private/SViewportToolBarComboMenu.cpp b/Engine/Source/Editor/UnrealEd/Private/SViewportToolBarComboMenu.cpp
--- a/Engine/Source/Editor/UnrealEd/Private/SViewportToolBarComboMenu.cpp
+++ b/Engine/Source/Editor/UnrealEd/Private/SViewportToolBarComboMenu.cpp
@@ -73,7 +73,6 @@ void SViewportToolBarComboMenu::Construct( const FArguments& InArgs )
 	[
 		SNew(SHorizontalBox)
 		+SHorizontalBox::Slot()
-		.VAlign(VAlign_Center)
 		.AutoWidth()
 		[
 			ToggleControl
