Fix for a regression in which we could fail to resolve object paths for certain Blueprint assets during BP graph node pin reference initialization at load time.

Change summary:
- Reverted FEdGraphPinReference back to using TWeakObjectPtr (fixes the regression by restoring previous behavior).
- Added FBlueprintWatchedPin as a new type intended for FBlueprintEditorSettings config data; this internally uses TSoftObjectPtr to match the existing behavior that was being used to export pin objects as soft rather than weak. It also matches how breakpoints are being stored, and allows for future extensions to watched pin context without affecting the underlying pin reference type used by the rest of the engine code.
- Modified the FKismetDebugUtilities APIs for watched pins to utilize the new type in place of FEdGraphPinReference.

diff --git a/Engine/Source/Editor/BlueprintGraph/Private/BlueprintEditorSettings.cpp b/Engine/Source/Editor/BlueprintGraph/Private/BlueprintEditorSettings.cpp
--- a/Engine/Source/Editor/BlueprintGraph/Private/BlueprintEditorSettings.cpp
+++ b/Engine/Source/Editor/BlueprintGraph/Private/BlueprintEditorSettings.cpp
@@ -94,6 +94,7 @@ void UBlueprintEditorSettings::OnAssetRemoved(UObject* Object)
 	if(UBlueprint* Blueprint = Cast<UBlueprint>(Object))
 	{
 		FKismetDebugUtilities::ClearBreakpoints(Blueprint);
+		FKismetDebugUtilities::ClearPinWatches(Blueprint);
 	}
 }
 
@@ -143,4 +144,4 @@ void UBlueprintEditorSettings::PostEditChangeProperty(FPropertyChangedEvent& Pro
 	Super::PostEditChangeProperty(PropertyChangedEvent);
 
 	OnSettingsChange.Broadcast(this, PropertyChangedEvent.ChangeType);
-}
+}
\ No newline at end of file
diff --git a/Engine/Source/Editor/BlueprintGraph/Public/BlueprintEditorSettings.h b/Engine/Source/Editor/BlueprintGraph/Public/BlueprintEditorSettings.h
--- a/Engine/Source/Editor/BlueprintGraph/Public/BlueprintEditorSettings.h
+++ b/Engine/Source/Editor/BlueprintGraph/Public/BlueprintEditorSettings.h
@@ -7,6 +7,7 @@
 #include "UObject/Object.h"
 #include "Engine/Blueprint.h"
 #include "Kismet2/Breakpoint.h"
+#include "Kismet2/WatchedPin.h"
 #include "Engine/DeveloperSettings.h"
 #include "Kismet2/KismetDebugUtilities.h"
 #include "BlueprintEditorSettings.generated.h"
@@ -32,7 +33,7 @@ struct BLUEPRINTGRAPH_API FPerBlueprintSettings
 	TArray<FBlueprintBreakpoint> Breakpoints;
 
 	UPROPERTY()
-	TArray<FEdGraphPinReference> WatchedPins;
+	TArray<FBlueprintWatchedPin> WatchedPins;
 
 	bool operator==(const FPerBlueprintSettings& Other) const
 	{
diff --git a/Engine/Source/Editor/UnrealEd/Private/Kismet2/KismetDebugUtilities.cpp b/Engine/Source/Editor/UnrealEd/Private/Kismet2/KismetDebugUtilities.cpp
--- a/Engine/Source/Editor/UnrealEd/Private/Kismet2/KismetDebugUtilities.cpp
+++ b/Engine/Source/Editor/UnrealEd/Private/Kismet2/KismetDebugUtilities.cpp
@@ -24,6 +24,7 @@
 #include "Animation/AnimBlueprintGeneratedClass.h"
 #include "UnrealEdGlobals.h"
 #include "Kismet2/Breakpoint.h"
+#include "Kismet2/WatchedPin.h"
 #include "ActorEditorUtils.h"
 #include "EdGraphSchema_K2.h"
 #include "K2Node.h"
@@ -802,7 +803,7 @@ TArray<FBlueprintBreakpoint>* FKismetDebugUtilities::GetBreakpoints(const UBluep
 		&Settings->Breakpoints;
 }
 
-TArray<FEdGraphPinReference>* FKismetDebugUtilities::GetWatchedPins(const UBlueprint* Blueprint)
+TArray<FBlueprintWatchedPin>* FKismetDebugUtilities::GetWatchedPins(const UBlueprint* Blueprint)
 {
 	FPerBlueprintSettings* Settings = GetPerBlueprintSettings(Blueprint);
 
@@ -1308,9 +1309,9 @@ bool FKismetDebugUtilities::CanWatchPin(const UBlueprint* Blueprint, const UEdGr
 
 bool FKismetDebugUtilities::IsPinBeingWatched(const UBlueprint* Blueprint, const UEdGraphPin* Pin)
 {
-	if(TArray<FEdGraphPinReference>* WatchedPins = GetWatchedPins(Blueprint))
+	if(TArray<FBlueprintWatchedPin>* WatchedPins = GetWatchedPins(Blueprint))
 	{
-		return WatchedPins->Contains(const_cast<UEdGraphPin*>(Pin));
+		return WatchedPins->Contains(Pin);
 	}
 	return false;
 }
@@ -1332,13 +1333,8 @@ void FKismetDebugUtilities::AddPinWatch(const UBlueprint* Blueprint, const UEdGr
 	check(Settings);
 	FPerBlueprintSettings &BlueprintSettings = Settings->PerBlueprintSettings.FindOrAdd(Blueprint->GetPathName());
 
-	// ensure that this node doesn't already contain a breakpoint
-	checkSlow(!BlueprintSettings.WatchedPins.ContainsByPredicate(
-		[Pin](const FEdGraphPinReference& Other)
-		{
-			return Other == Pin;
-		}
-	));
+	// ensure that this pin isn't already being watched
+	checkSlow(!IsPinBeingWatched(Blueprint, Pin));
 	
 	BlueprintSettings.WatchedPins.Add(Pin);
 	SaveBlueprintEditorSettings();
@@ -1348,14 +1344,7 @@ void FKismetDebugUtilities::AddPinWatch(const UBlueprint* Blueprint, const UEdGr
 
 void FKismetDebugUtilities::TogglePinWatch(const UBlueprint* Blueprint, const UEdGraphPin* Pin)
 {
-	int32 ExistingWatchIndex = INDEX_NONE;
-	
-	if(TArray<FEdGraphPinReference>* WatchedPins = GetWatchedPins(Blueprint))
-	{
-		ExistingWatchIndex = WatchedPins->Find(const_cast<UEdGraphPin*>(Pin));
-	}
-
-	if (ExistingWatchIndex != INDEX_NONE)
+	if (IsPinBeingWatched(Blueprint, Pin))
 	{
 		RemovePinWatch(Blueprint, Pin);
 	}
@@ -1388,11 +1377,11 @@ bool FKismetDebugUtilities::BlueprintHasPinWatches(const UBlueprint* Blueprint)
 
 void FKismetDebugUtilities::ForeachPinWatch(const UBlueprint* Blueprint, TFunctionRef<void(UEdGraphPin*)> Task)
 {
-	if(TArray<FEdGraphPinReference>* Pins = GetWatchedPins(Blueprint))
+	if(TArray<FBlueprintWatchedPin>* WatchedPins = GetWatchedPins(Blueprint))
 	{
-		for(FEdGraphPinReference& PinRef : *Pins)
+		for(FBlueprintWatchedPin& WatchedPin : *WatchedPins)
 		{
-			if(UEdGraphPin* Pin = PinRef.Get())
+			if (UEdGraphPin* Pin = WatchedPin.Get())
 			{
 				Task(Pin);
 			}
@@ -1403,17 +1392,17 @@ void FKismetDebugUtilities::ForeachPinWatch(const UBlueprint* Blueprint, TFuncti
 bool FKismetDebugUtilities::RemovePinWatchesByPredicate(const UBlueprint* Blueprint,
 	const TFunctionRef<bool(const UEdGraphPin*)> Predicate)
 {
-	auto ModifiedPedicate = [Predicate](FEdGraphPinReference& PinRef)
+	auto ModifiedPedicate = [Predicate](FBlueprintWatchedPin& WatchedPin)
 	{
-		UEdGraphPin* Pin = PinRef.Get();
+		const UEdGraphPin* Pin = WatchedPin.Get();
 		return Pin && Predicate(Pin);
 	};
 	
-	if(TArray<FEdGraphPinReference>* Pins = GetWatchedPins(Blueprint))
+	if(TArray<FBlueprintWatchedPin>* WatchedPins = GetWatchedPins(Blueprint))
 	{
-		if(Pins->RemoveAllSwap(ModifiedPedicate, false))
+		if(WatchedPins->RemoveAllSwap(ModifiedPedicate, false))
 		{
-			if(Pins->IsEmpty())
+			if(WatchedPins->IsEmpty())
 			{
 				// keeps the ini file clean by removing empty arrays
 				ClearPinWatches(Blueprint);
@@ -1429,11 +1418,11 @@ bool FKismetDebugUtilities::RemovePinWatchesByPredicate(const UBlueprint* Bluepr
 UEdGraphPin* FKismetDebugUtilities::FindPinWatchByPredicate(const UBlueprint* Blueprint,
 	const TFunctionRef<bool(const UEdGraphPin*)> Predicate)
 {
-	if(TArray<FEdGraphPinReference>* Pins = GetWatchedPins(Blueprint))
+	if(TArray<FBlueprintWatchedPin>* WatchedPins = GetWatchedPins(Blueprint))
 	{
-		for(FEdGraphPinReference& PinRef : *Pins)
+		for(FBlueprintWatchedPin& WatchedPin : *WatchedPins)
 		{
-			UEdGraphPin* Pin = PinRef.Get();
+			UEdGraphPin* Pin = WatchedPin.Get();
 			if(Pin && Predicate(Pin))
 			{
 				return Pin;
diff --git a/Engine/Source/Editor/UnrealEd/Private/Kismet2/WatchedPin.cpp b/Engine/Source/Editor/UnrealEd/Private/Kismet2/WatchedPin.cpp
new file mode 100644
--- /dev/null
+++ b/Engine/Source/Editor/UnrealEd/Private/Kismet2/WatchedPin.cpp
@@ -0,0 +1,48 @@
+// Copyright Epic Games, Inc. All Rights Reserved.
+
+#include "Kismet2/WatchedPin.h"
+#include "EdGraph/EdGraphPin.h"
+
+FBlueprintWatchedPin::FBlueprintWatchedPin()
+{
+}
+
+FBlueprintWatchedPin::FBlueprintWatchedPin(const UEdGraphPin* Pin)
+{
+	SetFromPin(Pin);
+}
+
+UEdGraphPin* FBlueprintWatchedPin::Get() const
+{
+	UEdGraphPin* FoundPin = CachedPinRef.Get();
+
+	if (!FoundPin && PinId.IsValid() && OwningNode.IsValid())
+	{
+		const UEdGraphNode* Node = OwningNode.Get();
+		check(Node);
+
+		UEdGraphPin* const* FoundPinPtr = Node->Pins.FindByPredicate([PinId = this->PinId](const UEdGraphPin* Pin)
+		{
+			return PinId == Pin->PinId;
+		});
+
+		if (FoundPinPtr)
+		{
+			FoundPin = *FoundPinPtr;
+			CachedPinRef = FoundPin;
+		}
+	}
+
+	return FoundPin;
+}
+
+void FBlueprintWatchedPin::SetFromPin(const UEdGraphPin* Pin)
+{
+	if (Pin)
+	{
+		PinId = Pin->PinId;
+		OwningNode = Pin->GetOwningNodeUnchecked();
+	}
+
+	CachedPinRef = Pin;
+}
\ No newline at end of file
diff --git a/Engine/Source/Editor/UnrealEd/Public/Kismet2/KismetDebugUtilities.h b/Engine/Source/Editor/UnrealEd/Public/Kismet2/KismetDebugUtilities.h
--- a/Engine/Source/Editor/UnrealEd/Public/Kismet2/KismetDebugUtilities.h
+++ b/Engine/Source/Editor/UnrealEd/Public/Kismet2/KismetDebugUtilities.h
@@ -12,6 +12,7 @@ static_assert(DO_BLUEPRINT_GUARD, "KismetDebugUtilities assumes BP exception tra
 
 class UBlueprint;
 struct FBlueprintBreakpoint;
+struct FBlueprintWatchedPin;
 template<typename ElementType> class TSimpleRingBuffer;
 
 DECLARE_LOG_CATEGORY_EXTERN(LogBlueprintDebug, Log, All);
@@ -296,23 +297,23 @@ public:
 	static bool BlueprintHasPinWatches(const UBlueprint* Blueprint);
 
 	/**
-	* Performs a task on every breakpoint in the provided blueprint
-	* @param Blueprint The owning blueprint of the breakpoints to iterate
+	* Performs a task on every watched pin in the provided blueprint
+	* @param Blueprint The owning blueprint of the watched pins to iterate
 	* @param Task function to be called on every element
 	*/
 	static void ForeachPinWatch(const UBlueprint* Blueprint, TFunctionRef<void(UEdGraphPin*)> Task);
 
 	/**
-	* Removes any breakpoint that matches the provided predicate
-	* @param Blueprint The owning blueprint of the breakpoints to iterate
-	* @param Predicate function that returns true if a breakpoint should be removed
+	* Removes any watched pin that matches the provided predicate
+	* @param Blueprint The owning blueprint of the watched pins to iterate
+	* @param Predicate function that returns true if a watched pin should be removed
 	*/
 	static bool RemovePinWatchesByPredicate(const UBlueprint* Blueprint, const TFunctionRef<bool(const UEdGraphPin*)> Predicate);
 
 	/**
-	* Returns the first breakpoint that matches the provided predicate or nullptr if nothing matched
-	* @param Blueprint The owning blueprint of the breakpoints to iterate
-	* @param Predicate function that returns true for the found breakpoint
+	* Returns the first watched pin that matches the provided predicate or nullptr if nothing matched
+	* @param Blueprint The owning blueprint of the watched pins to iterate
+	* @param Predicate function that returns true for the found watched pin
 	*/
 	static UEdGraphPin* FindPinWatchByPredicate(const UBlueprint* Blueprint, const TFunctionRef<bool(const UEdGraphPin*)> Predicate);
 
@@ -374,9 +375,9 @@ protected:
 	*	returns null if there are no breakpoints associated with this blueprint */
 	static TArray<FBlueprintBreakpoint>* GetBreakpoints(const UBlueprint* Blueprint);
 
-	/**	Retrieve the Array of breakpoints associated with a blueprint.
-	*	returns null if there are no breakpoints associated with this blueprint */
-	static TArray<FEdGraphPinReference>* GetWatchedPins(const UBlueprint* Blueprint);
+	/**	Retrieve the Array of watched pins associated with a blueprint.
+	*	returns null if there are no watched pins associated with this blueprint */
+	static TArray<FBlueprintWatchedPin>* GetWatchedPins(const UBlueprint* Blueprint);
 	
 	/** Save any modifications made to breakpoints */
 	static void SaveBlueprintEditorSettings();
diff --git a/Engine/Source/Editor/UnrealEd/Public/Kismet2/WatchedPin.h b/Engine/Source/Editor/UnrealEd/Public/Kismet2/WatchedPin.h
new file mode 100644
--- /dev/null
+++ b/Engine/Source/Editor/UnrealEd/Public/Kismet2/WatchedPin.h
@@ -0,0 +1,48 @@
+// Copyright Epic Games, Inc. All Rights Reserved.
+
+#pragma once
+
+#include "CoreMinimal.h"
+#include "WatchedPin.generated.h"
+
+class UEdGraphPin;
+
+/** Contains information about a watched pin in a Blueprint graph for local settings data.
+ */
+USTRUCT()
+struct UNREALED_API FBlueprintWatchedPin
+{
+	GENERATED_BODY()
+
+	FBlueprintWatchedPin();
+	FBlueprintWatchedPin(const UEdGraphPin* Pin);
+
+	/** Returns a reference to the underlying graph pin */
+	UEdGraphPin* Get() const;
+
+	/** Resets the pin watch to the given graph pin */
+	void SetFromPin(const UEdGraphPin* Pin);
+
+	bool operator==(const FBlueprintWatchedPin& Other) const
+	{
+		return PinId == Other.PinId && OwningNode == Other.OwningNode;
+	}
+
+private:
+	/** Node that owns the pin that the watch is placed on */
+	UPROPERTY()
+	TSoftObjectPtr<UEdGraphNode> OwningNode;
+
+	/** Unique ID of the pin that the watch is placed on */
+	UPROPERTY()
+	FGuid PinId;
+
+	/** Holds a cached reference to the underlying pin object. We don't save this directly to settings data,
+	 *  because it internally maintains a weak object reference to the owning node that it will then try to
+	 *  load after parsing the underlying value from the user's local settings file. To avoid issues and
+	 *	overhead of trying to load referenced assets when reading the config file at editor startup, we
+	 *  maintain our own soft object reference for the settings data instead. Additionally, we can add more
+	 *  context this way without affecting other parts of the engine that rely on the pin reference type.
+	 */
+	mutable FEdGraphPinReference CachedPinRef;
+};
\ No newline at end of file
diff --git a/Engine/Source/Runtime/Engine/Classes/EdGraph/EdGraphPin.h b/Engine/Source/Runtime/Engine/Classes/EdGraph/EdGraphPin.h
--- a/Engine/Source/Runtime/Engine/Classes/EdGraph/EdGraphPin.h
+++ b/Engine/Source/Runtime/Engine/Classes/EdGraph/EdGraphPin.h
@@ -278,7 +278,7 @@ struct ENGINE_API FEdGraphPinReference
 private:
 	/** The node that owns the pin referred to by this struct. Updated at Set and Save time. */
 	UPROPERTY()
-	TSoftObjectPtr<UEdGraphNode> OwningNode;
+	TWeakObjectPtr<UEdGraphNode> OwningNode;
 
 	/** The pin's unique ID. Updated at Set and Save time. */
 	UPROPERTY()
