Whitespace change: Unindent functions in namespace.
#rnx

diff --git a/Engine/Source/Editor/UnrealEd/Private/Cooker/CookTypes.cpp b/Engine/Source/Editor/UnrealEd/Private/Cooker/CookTypes.cpp
--- a/Engine/Source/Editor/UnrealEd/Private/Cooker/CookTypes.cpp
+++ b/Engine/Source/Editor/UnrealEd/Private/Cooker/CookTypes.cpp
@@ -8,91 +8,91 @@
 namespace UE::Cook
 {
 
-	FCookerTimer::FCookerTimer(const float& InTimeSlice, bool bInIsRealtimeMode, int InMaxNumPackagesToSave)
-		: bIsRealtimeMode(bInIsRealtimeMode), StartTime(FPlatformTime::Seconds()), TimeSlice(InTimeSlice)
-		, MaxNumPackagesToSave(InMaxNumPackagesToSave), NumPackagesSaved(0)
-	{
-	}
+FCookerTimer::FCookerTimer(const float& InTimeSlice, bool bInIsRealtimeMode, int InMaxNumPackagesToSave)
+	: bIsRealtimeMode(bInIsRealtimeMode), StartTime(FPlatformTime::Seconds()), TimeSlice(InTimeSlice)
+	, MaxNumPackagesToSave(InMaxNumPackagesToSave), NumPackagesSaved(0)
+{
+}
 
-	FCookerTimer::FCookerTimer(EForever)
-		:bIsRealtimeMode(false), StartTime(FPlatformTime::Seconds()), TimeSlice(ForeverTimeSlice), MaxNumPackagesToSave(TNumericLimits<int>::Max()), NumPackagesSaved(0)
-	{
-	}
+FCookerTimer::FCookerTimer(EForever)
+	:bIsRealtimeMode(false), StartTime(FPlatformTime::Seconds()), TimeSlice(ForeverTimeSlice), MaxNumPackagesToSave(TNumericLimits<int>::Max()), NumPackagesSaved(0)
+{
+}
 
-	FCookerTimer::FCookerTimer(ENoWait)
-		:bIsRealtimeMode(false), StartTime(FPlatformTime::Seconds()), TimeSlice(ZeroTimeSlice), MaxNumPackagesToSave(TNumericLimits<int>::Max()), NumPackagesSaved(0)
-	{
-	}
+FCookerTimer::FCookerTimer(ENoWait)
+	:bIsRealtimeMode(false), StartTime(FPlatformTime::Seconds()), TimeSlice(ZeroTimeSlice), MaxNumPackagesToSave(TNumericLimits<int>::Max()), NumPackagesSaved(0)
+{
+}
 
-	double FCookerTimer::GetTimeTillNow() const
-	{
-		return FPlatformTime::Seconds() - StartTime;
-	}
+double FCookerTimer::GetTimeTillNow() const
+{
+	return FPlatformTime::Seconds() - StartTime;
+}
 	
-	bool FCookerTimer::IsTimeUp() const
+bool FCookerTimer::IsTimeUp() const
+{
+	if (bIsRealtimeMode)
 	{
-		if (bIsRealtimeMode)
-		{
-			if ((FPlatformTime::Seconds() - StartTime) > TimeSlice)
-			{
-				return true;
-			}
-		}
-		if (NumPackagesSaved >= MaxNumPackagesToSave)
+		if ((FPlatformTime::Seconds() - StartTime) > TimeSlice)
 		{
 			return true;
 		}
-		return false;
 	}
-
-	void FCookerTimer::SavedPackage()
+	if (NumPackagesSaved >= MaxNumPackagesToSave)
 	{
-		++NumPackagesSaved;
+		return true;
 	}
+	return false;
+}
 
-	double FCookerTimer::GetTimeRemain() const
-	{
-		return TimeSlice - (FPlatformTime::Seconds() - StartTime);
-	}
+void FCookerTimer::SavedPackage()
+{
+	++NumPackagesSaved;
+}
 
-	float FCookerTimer::ZeroTimeSlice = 0.0f;
-	float FCookerTimer::ForeverTimeSlice = TNumericLimits<float>::Max();
+double FCookerTimer::GetTimeRemain() const
+{
+	return TimeSlice - (FPlatformTime::Seconds() - StartTime);
+}
 
-	static uint32 SchedulerThreadTlsSlot = 0;
-	void InitializeTls()
-	{
-		if (SchedulerThreadTlsSlot == 0)
-		{
-			SchedulerThreadTlsSlot = FPlatformTLS::AllocTlsSlot();
-			SetIsSchedulerThread(true);
-		}
-	}
+float FCookerTimer::ZeroTimeSlice = 0.0f;
+float FCookerTimer::ForeverTimeSlice = TNumericLimits<float>::Max();
 
-	bool IsSchedulerThread()
+static uint32 SchedulerThreadTlsSlot = 0;
+void InitializeTls()
+{
+	if (SchedulerThreadTlsSlot == 0)
 	{
-		return FPlatformTLS::GetTlsValue(SchedulerThreadTlsSlot) != 0;
+		SchedulerThreadTlsSlot = FPlatformTLS::AllocTlsSlot();
+		SetIsSchedulerThread(true);
 	}
+}
 
-	void SetIsSchedulerThread(bool bValue)
-	{
-		FPlatformTLS::SetTlsValue(SchedulerThreadTlsSlot, bValue ? (void*)0x1 : (void*)0x0);
-	}
+bool IsSchedulerThread()
+{
+	return FPlatformTLS::GetTlsValue(SchedulerThreadTlsSlot) != 0;
+}
 
-	FCookSavePackageContext::FCookSavePackageContext(const ITargetPlatform* InTargetPlatform,
-		ICookedPackageWriter* InPackageWriter, bool bInForceLegacyOffsets, FStringView InWriterDebugName)
-		: SaveContext(InTargetPlatform, InPackageWriter->GetCookCapabilities().bSavePackageSupported ? InPackageWriter : nullptr, bInForceLegacyOffsets)
-		, WriterDebugName(InWriterDebugName)
-		, PackageWriter(InPackageWriter)
-	{
-		PackageWriterCapabilities = InPackageWriter->GetCookCapabilities();
-	}
+void SetIsSchedulerThread(bool bValue)
+{
+	FPlatformTLS::SetTlsValue(SchedulerThreadTlsSlot, bValue ? (void*)0x1 : (void*)0x0);
+}
+
+FCookSavePackageContext::FCookSavePackageContext(const ITargetPlatform* InTargetPlatform,
+	ICookedPackageWriter* InPackageWriter, bool bInForceLegacyOffsets, FStringView InWriterDebugName)
+	: SaveContext(InTargetPlatform, InPackageWriter->GetCookCapabilities().bSavePackageSupported ? InPackageWriter : nullptr, bInForceLegacyOffsets)
+	, WriterDebugName(InWriterDebugName)
+	, PackageWriter(InPackageWriter)
+{
+	PackageWriterCapabilities = InPackageWriter->GetCookCapabilities();
+}
 
-	FCookSavePackageContext::~FCookSavePackageContext()
+FCookSavePackageContext::~FCookSavePackageContext()
+{
+	// SaveContext destructor deletes the PackageWriter, so if we passed our writer into SaveContext, we do not delete it
+	if (!SaveContext.PackageWriter)
 	{
-		// SaveContext destructor deletes the PackageWriter, so if we passed our writer into SaveContext, we do not delete it
-		if (!SaveContext.PackageWriter)
-		{
-			delete PackageWriter;
-		}
+		delete PackageWriter;
 	}
 }
+}
