Horde HLOD build -
* Added option to shelve the build result to the workspace of a specific user (skipping the submit)
* Created a generic WorldPartitionBuilder BuildGraph command that will allow this kind of operation on any WP builder
* Added a P4.UpdateChange() overload that takes a client (workspace) argument
* Renamed the HLOD_Submit step to HLOD_Finalize, as it can either: gather files from multiple builders locally without submitting / submit the result / shelve the result

diff --git a/Engine/Source/Editor/UnrealEd/Private/WorldPartition/WorldPartitionBuilder.cpp b/Engine/Source/Editor/UnrealEd/Private/WorldPartition/WorldPartitionBuilder.cpp
--- a/Engine/Source/Editor/UnrealEd/Private/WorldPartition/WorldPartitionBuilder.cpp
+++ b/Engine/Source/Editor/UnrealEd/Private/WorldPartition/WorldPartitionBuilder.cpp
@@ -21,7 +21,9 @@ DEFINE_LOG_CATEGORY_STATIC(LogWorldPartitionBuilder, Log, All);
 
 UWorldPartitionBuilder::UWorldPartitionBuilder(const FObjectInitializer& ObjectInitializer)
 	: Super(ObjectInitializer)
-{}
+{
+	bSubmit = FParse::Param(FCommandLine::Get(), TEXT("Submit"));
+}
 
 bool UWorldPartitionBuilder::RunBuilder(TSubclassOf<UWorldPartitionBuilder> BuilderClass, UWorld* World)
 {
diff --git a/Engine/Source/Editor/UnrealEd/Private/WorldPartition/WorldPartitionHLODsBuilder.cpp b/Engine/Source/Editor/UnrealEd/Private/WorldPartition/WorldPartitionHLODsBuilder.cpp
--- a/Engine/Source/Editor/UnrealEd/Private/WorldPartition/WorldPartitionHLODsBuilder.cpp
+++ b/Engine/Source/Editor/UnrealEd/Private/WorldPartition/WorldPartitionHLODsBuilder.cpp
@@ -163,7 +163,7 @@ UWorldPartitionHLODsBuilder::UWorldPartitionHLODsBuilder(const FObjectInitialize
 	BuildOptions = FParse::Param(FCommandLine::Get(), TEXT("SetupHLODs")) ? EHLODBuildStep::HLOD_Setup : EHLODBuildStep::None;
 	BuildOptions |= FParse::Param(FCommandLine::Get(), TEXT("BuildHLODs")) ? EHLODBuildStep::HLOD_Build : EHLODBuildStep::None;
 	BuildOptions |= FParse::Param(FCommandLine::Get(), TEXT("DeleteHLODs")) ? EHLODBuildStep::HLOD_Delete : EHLODBuildStep::None;
-	BuildOptions |= FParse::Param(FCommandLine::Get(), TEXT("SubmitHLODs")) ? EHLODBuildStep::HLOD_Submit : EHLODBuildStep::None;
+	BuildOptions |= FParse::Param(FCommandLine::Get(), TEXT("FinalizeHLODs")) ? EHLODBuildStep::HLOD_Finalize : EHLODBuildStep::None;
 	BuildOptions |= FParse::Param(FCommandLine::Get(), TEXT("DumpStats")) ? EHLODBuildStep::HLOD_Stats : EHLODBuildStep::None;
 
 	// Default behavior without any option is to setup + build
@@ -251,9 +251,9 @@ bool UWorldPartitionHLODsBuilder::ValidateParams() const
 		}
 	}
 
-	if (ShouldRunStep(EHLODBuildStep::HLOD_Submit) && !ISourceControlModule::Get().GetProvider().IsEnabled())
+	if (ShouldRunStep(EHLODBuildStep::HLOD_Finalize) && bSubmit && !ISourceControlModule::Get().GetProvider().IsEnabled())
 	{
-		UE_LOG(LogWorldPartitionHLODsBuilder, Error, TEXT("Submit step requires that a valid source control provider is enabled, exiting..."), *BuildManifest);
+		UE_LOG(LogWorldPartitionHLODsBuilder, Error, TEXT("Submit requires that a valid source control provider is enabled, exiting..."), *BuildManifest);
 		return false;
 	}
 
@@ -270,7 +270,7 @@ bool UWorldPartitionHLODsBuilder::PreWorldInitialization(FPackageSourceControlHe
 	bool bRet = true;
 
 	// When running a distributed build, retrieve relevant build products from the previous steps
-	if (IsDistributedBuild() && (ShouldRunStep(EHLODBuildStep::HLOD_Build) || ShouldRunStep(EHLODBuildStep::HLOD_Submit)))
+	if (IsDistributedBuild() && (ShouldRunStep(EHLODBuildStep::HLOD_Build) || ShouldRunStep(EHLODBuildStep::HLOD_Finalize)))
 	{
 		FString WorkingDirFolder = ShouldRunStep(EHLODBuildStep::HLOD_Build) ? GetHLODBuilderFolderName(BuilderIdx) : GetToSubmitFolderName();
 		bRet = CopyFilesFromWorkingDir(WorkingDirFolder);
@@ -303,7 +303,7 @@ bool UWorldPartitionHLODsBuilder::RunInternal(UWorld* World, const FBox& Bounds,
 		bRet = DeleteHLODActors();
 	}
 
-	if (bRet && ShouldRunStep(EHLODBuildStep::HLOD_Submit))
+	if (bRet && ShouldRunStep(EHLODBuildStep::HLOD_Finalize) && bSubmit)
 	{
 		bRet = SubmitHLODActors();
 	}
@@ -316,9 +316,6 @@ bool UWorldPartitionHLODsBuilder::RunInternal(UWorld* World, const FBox& Bounds,
 	WorldPartition = nullptr;
 	delete SourceControlHelper;
 
-	// TODO: DDC shutdown crash workaround - Wait for any DDC writes to complete
-	GetDerivedDataCacheRef().WaitForQuiescence(true);
-	
 	return bRet;
 }
 
diff --git a/Engine/Source/Editor/UnrealEd/Public/WorldPartition/WorldPartitionBuilder.h b/Engine/Source/Editor/UnrealEd/Public/WorldPartition/WorldPartitionBuilder.h
--- a/Engine/Source/Editor/UnrealEd/Public/WorldPartition/WorldPartitionBuilder.h
+++ b/Engine/Source/Editor/UnrealEd/Public/WorldPartition/WorldPartitionBuilder.h
@@ -36,4 +36,6 @@ protected:
 	TSet<FName> DataLayerLabels;
 	bool bLoadNonDynamicDataLayers = true;
 	bool bLoadInitiallyActiveDataLayers = true;
+
+	bool bSubmit = true;
 };
\ No newline at end of file
diff --git a/Engine/Source/Editor/UnrealEd/Public/WorldPartition/WorldPartitionHLODsBuilder.h b/Engine/Source/Editor/UnrealEd/Public/WorldPartition/WorldPartitionHLODsBuilder.h
--- a/Engine/Source/Editor/UnrealEd/Public/WorldPartition/WorldPartitionHLODsBuilder.h
+++ b/Engine/Source/Editor/UnrealEd/Public/WorldPartition/WorldPartitionHLODsBuilder.h
@@ -61,11 +61,11 @@ private:
 enum class EHLODBuildStep : uint8
 {
 	None		= 0,
-	HLOD_Setup	= 1 << 0,
-	HLOD_Build	= 1 << 1,
-	HLOD_Submit = 1 << 2,
-	HLOD_Delete = 1 << 3,
-	HLOD_Stats	= 1 << 4
+	HLOD_Setup	= 1 << 0,		// Create/delete HLOD actors to populate the world.
+	HLOD_Build	= 1 << 1,		// Create components/merged meshes/etc - can run on multiple machines if this step is distributed.
+	HLOD_Finalize = 1 << 2,		// When performing a distributed build, this step will gather the result generated from the different machines and, optionnaly, will submit it to source control.
+	HLOD_Delete = 1 << 3,		// Delete all HLOD actors from the given world.
+	HLOD_Stats	= 1 << 4		// Print stats on all the HLOD actors.
 };
 ENUM_CLASS_FLAGS(EHLODBuildStep);
 
diff --git a/Engine/Source/Programs/AutomationTool/AutomationUtils/P4Utils.cs b/Engine/Source/Programs/AutomationTool/AutomationUtils/P4Utils.cs
--- a/Engine/Source/Programs/AutomationTool/AutomationUtils/P4Utils.cs
+++ b/Engine/Source/Programs/AutomationTool/AutomationUtils/P4Utils.cs
@@ -2481,32 +2481,49 @@ namespace AutomationTool
 		/// Updates a changelist with the given fields
 		/// </summary>
 		/// <param name="CL"></param>
-		/// <param name="NewOwner"></param>
+		/// <param name="NewClient"></param>
 		/// <param name="NewDescription"></param>
 		/// <param name="SpewIsVerbose"></param>
-		public void UpdateChange(int CL, string NewOwner, string NewDescription, bool SpewIsVerbose = false)
+		public void UpdateChange(int CL, string NewClient, string NewDescription, bool SpewIsVerbose = false)
+		{
+			UpdateChange(CL, null, NewClient, NewDescription, SpewIsVerbose);
+		}
+
+		/// <summary>
+		/// Updates a changelist with the given fields
+		/// </summary>
+		/// <param name="CL"></param>
+		/// <param name="NewUser"></param>
+		/// <param name="NewClient"></param>
+		/// <param name="NewDescription"></param>
+		/// <param name="SpewIsVerbose"></param>
+		public void UpdateChange(int CL, string NewUser, string NewClient, string NewDescription, bool SpewIsVerbose = false)
 		{
 			string CmdOutput;
-			if(!LogP4Output(out CmdOutput, "", String.Format("change -o {0}", CL), SpewIsVerbose: SpewIsVerbose))
+			if (!LogP4Output(out CmdOutput, "", String.Format("change -o {0}", CL), SpewIsVerbose: SpewIsVerbose))
 			{
 				throw new P4Exception("Couldn't describe changelist {0}", CL);
 			}
 
 			P4Spec Spec = P4Spec.FromString(CmdOutput);
-			if(NewOwner != null)
+			if (NewUser != null)
 			{
-				Spec.SetField("Client", NewOwner);
+				Spec.SetField("User", NewUser);
 			}
-			if(NewDescription != null)
+			if (NewClient != null)
+			{
+				Spec.SetField("Client", NewClient);
+			}			
+			if (NewDescription != null)
 			{
 				Spec.SetField("Description", NewDescription);
 			}
 
-			if(!LogP4Output(out CmdOutput, "", "change -i", Input: Spec.ToString(), SpewIsVerbose: SpewIsVerbose))
+			if (!LogP4Output(out CmdOutput, "", "change -i", Input: Spec.ToString(), SpewIsVerbose: SpewIsVerbose))
 			{
 				throw new P4Exception("Failed to update spec for changelist {0}", CL);
 			}
-			if(!CmdOutput.TrimEnd().EndsWith(String.Format("Change {0} updated.", CL)))
+			if (!CmdOutput.TrimEnd().EndsWith(String.Format("Change {0} updated.", CL)))
 			{
 				throw new P4Exception("Unexpected output from p4 change -i: {0}", CmdOutput);
 			}
diff --git a/Engine/Source/Programs/AutomationTool/Scripts/WorldPartitionBuilder.Automation.cs b/Engine/Source/Programs/AutomationTool/Scripts/WorldPartitionBuilder.Automation.cs
new file mode 100644
--- /dev/null
+++ b/Engine/Source/Programs/AutomationTool/Scripts/WorldPartitionBuilder.Automation.cs
@@ -0,0 +1,65 @@
+// Copyright Epic Games, Inc. All Rights Reserved.
+using AutomationTool;
+using EpicGames.Core;
+using System;
+using System.Diagnostics;
+using System.Threading;
+
+namespace AutomationScripts.Automation
+{
+	public class WorldPartitionBuilder : BuildCommand
+	{
+		public override void ExecuteBuild()
+		{
+			string Builder = ParseRequiredStringParam("Builder");
+			string CommandletArgs = ParseOptionalStringParam("CommandletArgs");
+
+			bool bSubmit = ParseParam("Submit");
+
+			string ShelveUser = ParseOptionalStringParam("ShelveUser");
+			string ShelveWorkspace = ParseOptionalStringParam("ShelveWorkspace");
+			bool bShelveResult = !String.IsNullOrEmpty(ShelveUser) && !String.IsNullOrEmpty(ShelveWorkspace);
+
+			if (!P4Enabled && (bSubmit || bShelveResult))
+			{
+				LogError("P4 required to submit or shelve build results");
+				return;
+			}
+
+			CommandletArgs = "-Builder=" + Builder + " " + CommandletArgs;
+
+			if (bSubmit)
+			{
+				CommandletArgs += " -Submit";
+			}
+
+			string EditorExe = "UnrealEditor-Cmd.exe";
+			EditorExe = AutomationTool.HostPlatform.Current.GetUE4ExePath(EditorExe);
+
+			FileReference ProjectPath = ParseProjectParam();
+			RunCommandlet(ProjectPath, EditorExe, "WorldPartitionBuilderCommandlet", CommandletArgs);
+			
+			if (bShelveResult)
+			{
+				LogInformation("### Shelving build results ###");
+
+				// Create a new changelist and move all checked out files to it
+				LogInformation("Creating pending changelist to shelve builder changes");
+				int PendingCL = P4.CreateChange(P4Env.Client);
+				P4.LogP4("", $"reopen -c {PendingCL} //...", AllowSpew: true);
+
+				// Shelve changelist & revert changes
+				LogInformation("Shelving changes...");
+				P4.Shelve(PendingCL);
+				LogInformation("Reverting local changes...");
+				P4.Revert($"-w -c {PendingCL} //...");
+
+				// Assign shelve to the provided user+workspace
+				LogInformation($"Changing ownership of CL {PendingCL} to user {ShelveUser}, workspace {ShelveWorkspace}");
+				P4.UpdateChange(PendingCL, ShelveUser, ShelveWorkspace, null, true);
+
+				LogInformation("### Shelving completed ###");
+			}
+		}
+	}
+}
\ No newline at end of file
