Added a warning icon in the auto-saved/restore UI aside a package to restore if the auto-saved package modification time is older than the file it is supposed to restore.

diff --git a/Engine/Source/Editor/UnrealEd/Private/PackageRestore.cpp b/Engine/Source/Editor/UnrealEd/Private/PackageRestore.cpp
--- a/Engine/Source/Editor/UnrealEd/Private/PackageRestore.cpp
+++ b/Engine/Source/Editor/UnrealEd/Private/PackageRestore.cpp
@@ -112,6 +112,18 @@ namespace PackageRestore
 			return FReply::Handled();
 		}
 
+		EVisibility GetRestoreOverMoreRecentPackageWarningVisibility() const
+		{
+			FString RestoreDstPathname = *FPaths::ConvertRelativePathToFull(PackageFilename);
+			if (IFileManager::Get().FileExists(*RestoreDstPathname))
+			{
+				FDateTime AutoSavedSrcModificationTime = IFileManager::Get().GetStatData(*FPaths::ConvertRelativePathToFull(AutoSaveFilename)).ModificationTime;
+				FDateTime RestoredDstModificationTime = IFileManager::Get().GetStatData(*RestoreDstPathname).ModificationTime;
+				return AutoSavedSrcModificationTime > RestoredDstModificationTime ? EVisibility::Collapsed : EVisibility::Visible;
+			}
+			return EVisibility::Collapsed; // The destination file doesn't exist, restoring the auto-saved file cannot overwrite anything.
+		}
+
 	private:
 		FString PackageName;
 		FString PackageFilename;
@@ -284,9 +296,20 @@ namespace PackageRestore
 				.Padding(FMargin(2, 0, 0, 0))
 				.VAlign(VAlign_Center)
 				.HAlign(HAlign_Left)
+				.AutoWidth()
+				[
+					SNew(SImage)
+					.Image(FEditorStyle::Get().GetBrush("Icons.Warning"))
+					.ToolTipText(LOCTEXT("OverwritingMoreRecentPackageFile", "The auto-saved file is older than the file it restores. You could lose work if the file was updated or modified after the crash but before the Editor was restarted."))
+					.Visibility(Item.Get(), &FPackageRestoreItem::GetRestoreOverMoreRecentPackageWarningVisibility)
+				]
+				+SHorizontalBox::Slot()
+				.Padding(FMargin(2, 0, 0, 0))
+				.VAlign(VAlign_Center)
+				.HAlign(HAlign_Left)
 				.FillWidth(1)
 				[
-					SNew(STextBlock) 
+					SNew(STextBlock)
 					.Text(FText::FromString(Item->GetPackageName()))
 				]
 				+SHorizontalBox::Slot()
@@ -547,7 +570,7 @@ FEditorFileUtils::EPromptReturnCode PackageRestore::PromptToRestorePackages(cons
 		FString PackageFilename;
 		if(FPackageName::DoesPackageExist(PackageFullPath, nullptr, &PackageFilename))
 		{
-			FPackageRestoreItemPtr PackageItemPtr = MakeShareable(new FPackageRestoreItem(PackageFullPath, PackageFilename, AutoSaveDir / AutoSavePath, true/*bIsExistingPackage*/));
+			FPackageRestoreItemPtr PackageItemPtr = MakeShared<FPackageRestoreItem>(PackageFullPath, PackageFilename, AutoSaveDir / AutoSavePath, true/*bIsExistingPackage*/);
 			PackageRestoreItems.Add(PackageItemPtr);
 		}
 		else
@@ -557,7 +580,7 @@ FEditorFileUtils::EPromptReturnCode PackageRestore::PromptToRestorePackages(cons
 			{
 				PackageFilename += FPaths::GetExtension(AutoSavePath, true/*bIncludeDot*/);
 
-				FPackageRestoreItemPtr PackageItemPtr = MakeShareable(new FPackageRestoreItem(PackageFullPath, PackageFilename, AutoSaveDir / AutoSavePath, false/*bIsExistingPackage*/));
+				FPackageRestoreItemPtr PackageItemPtr = MakeShared<FPackageRestoreItem>(PackageFullPath, PackageFilename, AutoSaveDir / AutoSavePath, false/*bIsExistingPackage*/);
 				PackageRestoreItems.Add(PackageItemPtr);
 			}
 		}
