CIS fix for non-unity compile issue.

#rnx

diff --git a/Engine/Source/Editor/UnrealEd/Private/Kismet2/WatchedPin.cpp b/Engine/Source/Editor/UnrealEd/Private/Kismet2/WatchedPin.cpp
--- a/Engine/Source/Editor/UnrealEd/Private/Kismet2/WatchedPin.cpp
+++ b/Engine/Source/Editor/UnrealEd/Private/Kismet2/WatchedPin.cpp
@@ -1,7 +1,7 @@
 // Copyright Epic Games, Inc. All Rights Reserved.
 
 #include "Kismet2/WatchedPin.h"
-#include "EdGraph/EdGraphPin.h"
+#include "EdGraph/EdGraphNode.h"
 
 FBlueprintWatchedPin::FBlueprintWatchedPin()
 {
diff --git a/Engine/Source/Editor/UnrealEd/Public/Kismet2/WatchedPin.h b/Engine/Source/Editor/UnrealEd/Public/Kismet2/WatchedPin.h
--- a/Engine/Source/Editor/UnrealEd/Public/Kismet2/WatchedPin.h
+++ b/Engine/Source/Editor/UnrealEd/Public/Kismet2/WatchedPin.h
@@ -3,9 +3,10 @@
 #pragma once
 
 #include "CoreMinimal.h"
+#include "EdGraph/EdGraphPin.h"
 #include "WatchedPin.generated.h"
 
-class UEdGraphPin;
+class UEdGraphNode;
 
 /** Contains information about a watched pin in a Blueprint graph for local settings data.
  */
