Load/Save BuildDefinitionList to the TargetDomain and EditorDomain.
#rnx

diff --git a/Engine/Source/Editor/UnrealEd/Classes/CookOnTheSide/CookOnTheFlyServer.h b/Engine/Source/Editor/UnrealEd/Classes/CookOnTheSide/CookOnTheFlyServer.h
--- a/Engine/Source/Editor/UnrealEd/Classes/CookOnTheSide/CookOnTheFlyServer.h
+++ b/Engine/Source/Editor/UnrealEd/Classes/CookOnTheSide/CookOnTheFlyServer.h
@@ -111,6 +111,7 @@ ENUM_CLASS_FLAGS(ECookTickFlags);
 
 namespace UE::Cook
 {
+	class FBuildDefinitions;
 	class FExternalRequests;
 	class FRequestCluster;
 	class FSaveCookedPackageContext;
@@ -1104,6 +1105,7 @@ private:
 	TUniquePtr<UE::Cook::FPackageTracker> PackageTracker;
 	TUniquePtr<UE::Cook::FPackageDatas> PackageDatas;
 	TUniquePtr<UE::Cook::FExternalRequests> ExternalRequests;
+	TUniquePtr<UE::Cook::FBuildDefinitions> BuildDefinitions;
 
 	TArray<UE::Cook::FCookSavePackageContext*> SavePackageContexts;
 	/** Objects that were collected during the single-threaded PreGarbageCollect callback and that should be reported as referenced in CookerAddReferencedObjects. */
diff --git a/Engine/Source/Editor/UnrealEd/Private/CookOnTheFlyServer.cpp b/Engine/Source/Editor/UnrealEd/Private/CookOnTheFlyServer.cpp
--- a/Engine/Source/Editor/UnrealEd/Private/CookOnTheFlyServer.cpp
+++ b/Engine/Source/Editor/UnrealEd/Private/CookOnTheFlyServer.cpp
@@ -2630,8 +2630,17 @@ bool UCookOnTheFlyServer::BeginPrepareSave(UE::Cook::FPackageData& PackageData,
 	// while BeginPrepareSave is active.
 	// Currently this does not cause significant cost since saving new platforms with some platforms already saved is a rare operation.
 
-	TArray<FWeakObjectPtr>& CachedObjectsInOuter = PackageData.GetCachedObjectsInOuter();
 	int32& CookedPlatformDataNextIndex = PackageData.GetCookedPlatformDataNextIndex();
+	if (CookedPlatformDataNextIndex == 0)
+	{
+		if (!BuildDefinitions->TryRemovePendingBuilds(PackageData.GetPackageName()))
+		{
+			// Builds are in progress; wait for them to complete
+			return false;
+		}
+	}
+
+	TArray<FWeakObjectPtr>& CachedObjectsInOuter = PackageData.GetCachedObjectsInOuter();
 	FWeakObjectPtr* CachedObjectsInOuterData = CachedObjectsInOuter.GetData();
 	TArray<const ITargetPlatform*, TInlineAllocator<ExpectedMaxNumPlatforms>> TargetPlatforms;
 	PackageData.GetRequestedPlatforms(TargetPlatforms);
@@ -4476,9 +4485,14 @@ void FSaveCookedPackageContext::FinishPlatform()
 		FAssetPackageData* AssetPackageData = Generator.GetAssetPackageData(Package->GetFName());
 		check(AssetPackageData);
 
+		// TODO_BuildDefinitionList: Calculate and store BuildDefinitionList on the PackageData, or collect it here from some other source.
+		TArray<UE::DerivedData::FBuildDefinition> BuildDefinitions;
+		FCbObject BuildDefinitionList = UE::TargetDomain::BuildDefinitionListToObject(BuildDefinitions);
 		FCbObject TargetDomainDependencies;
+		{
 			UE_SCOPED_HIERARCHICAL_COOKTIMER(TargetDomainDependencies);
 			TargetDomainDependencies = UE::TargetDomain::CollectDependenciesObject(Package, TargetPlatform, nullptr /* ErrorMessage */);
+		}
 
 		ICookedPackageWriter::FCommitPackageInfo Info;
 		Info.bSucceeded = bSuccessful;
@@ -4487,6 +4501,7 @@ void FSaveCookedPackageContext::FinishPlatform()
 		Info.PackageGuid = AssetPackageData->PackageGuid;
 		PRAGMA_ENABLE_DEPRECATION_WARNINGS
 		Info.Attachments.Add({ "Dependencies", TargetDomainDependencies });
+		Info.Attachments.Add({ "BuildDefinitionList", BuildDefinitionList });
 
 		PackageWriter->CommitPackage(Info);
 	}
@@ -4610,6 +4625,8 @@ void FSaveCookedPackageContext::FinishPackage()
 void UCookOnTheFlyServer::Initialize( ECookMode::Type DesiredCookMode, ECookInitializationFlags InCookFlags, const FString &InOutputDirectoryOverride )
 {
 	TRACE_CPUPROFILER_EVENT_SCOPE(UCookOnTheFlyServer::Initialize);
+
+	BuildDefinitions.Reset(new UE::Cook::FBuildDefinitions());
 	UE::Cook::InitializeTls();
 	UE::Cook::FPlatformManager::InitializeTls();
 
@@ -6814,6 +6831,7 @@ void UCookOnTheFlyServer::CookByTheBookFinished()
 	}
 
 	UPackage::WaitForAsyncFileWrites();
+	BuildDefinitions->Wait();
 	
 	GetDerivedDataCacheRef().WaitForQuiescence(true);
 	
diff --git a/Engine/Source/Editor/UnrealEd/Private/Cooker/CookRequestCluster.cpp b/Engine/Source/Editor/UnrealEd/Private/Cooker/CookRequestCluster.cpp
--- a/Engine/Source/Editor/UnrealEd/Private/Cooker/CookRequestCluster.cpp
+++ b/Engine/Source/Editor/UnrealEd/Private/Cooker/CookRequestCluster.cpp
@@ -18,7 +18,6 @@
 #include "EditorDomain/EditorDomainUtils.h"
 #include "Interfaces/ITargetPlatform.h"
 #include "Misc/StringBuilder.h"
-#include "TargetDomain/TargetDomainUtils.h"
 
 namespace UE::Cook
 {
@@ -29,6 +28,7 @@ FRequestCluster::FRequestCluster(UCookOnTheFlyServer& COTFS, TArray<const ITarge
 	, AssetRegistry(*IAssetRegistry::Get())
 	, PackageNameCache(COTFS.GetPackageNameCache())
 	, PackageTracker(*COTFS.PackageTracker)
+	, BuildDefinitions(*COTFS.BuildDefinitions)
 {
 }
 
@@ -38,6 +38,7 @@ FRequestCluster::FRequestCluster(UCookOnTheFlyServer& COTFS, TConstArrayView<con
 	, AssetRegistry(*IAssetRegistry::Get())
 	, PackageNameCache(COTFS.GetPackageNameCache())
 	, PackageTracker(*COTFS.PackageTracker)
+	, BuildDefinitions(*COTFS.BuildDefinitions)
 {
 }
 
@@ -196,6 +197,8 @@ void FRequestCluster::Initialize(UCookOnTheFlyServer& COTFS)
 		DLCPath = FPaths::Combine(*COTFS.GetBaseDirectoryForDLC(), TEXT("Content"));
 		FPaths::MakeStandardFilename(DLCPath);
 	}
+	GConfig->GetBool(TEXT("CookSettings"), TEXT("PreQueueBuildDefinitions"), bPreQueueBuildDefinitions, GEditorIni);
+
 
 	PackageWriters.Reserve(Platforms.Num());
 	bFullBuild = false;
@@ -376,7 +379,7 @@ void FRequestCluster::FetchDependencies(const FCookerTimer& CookerTimer, bool& b
 		for (FPackageData* PackageData : Requests)
 		{
 			bool bAlreadyCooked;
-			VisitPackageData(PackageData, nullptr /* Dependencies */, bAlreadyCooked);
+			VisitPackageData(PackageData, nullptr /* OutDependencies */, bAlreadyCooked);
 		}
 		bDependenciesComplete = true;
 		return;
@@ -639,13 +642,12 @@ FRequestCluster::EVisitStatus& FRequestCluster::AddVertex(FName PackageName, FPa
 		}
 	}
 	FStackData* StackData = nullptr;
+	bool bAlreadyCooked;
 	if (bExploreDependencies)
 	{
 		bOutPushedStack = true;
 		StackData = &PushStack(PackageName);
 	}
-
-	bool bAlreadyCooked;
 	VisitPackageData(PackageData, StackData ? &StackData->Dependencies : nullptr, bAlreadyCooked);
 
 	if (bCookable && !bAlreadyCooked)
@@ -687,14 +689,11 @@ void FRequestCluster::PopStack()
 	--StackNum;
 };
 
-void FRequestCluster::VisitPackageData(FPackageData* PackageData, TArray<FName>* OutDependencies, bool& bAlreadyCooked)
+void FRequestCluster::VisitPackageData(FPackageData* PackageData, TArray<FName>* OutDependencies, bool& bOutAlreadyCooked)
 {
 	if (OutDependencies)
 	{
-		TArray<FName>& BuildDependencies(*OutDependencies);
-		TArray<FName>& RuntimeDependencies(this->RuntimeScratch);
-		BuildDependencies.Reset();
-		RuntimeDependencies.Reset();
+		OutDependencies->Reset();
 		FName PackageName = PackageData->GetPackageName();
 
 		// TODO EditorOnly References: We only fetch Game dependencies, because the cooker explicitly loads all of
@@ -709,48 +708,64 @@ void FRequestCluster::VisitPackageData(FPackageData* PackageData, TArray<FName>*
 		{
 			DependencyQuery |= UE::AssetRegistry::EDependencyQuery::Hard;
 		}
-		AssetRegistry.GetDependencies(PackageName, BuildDependencies,
+		AssetRegistry.GetDependencies(PackageName, *OutDependencies,
 			UE::AssetRegistry::EDependencyCategory::Package, DependencyQuery);
-		if (bHybridIterativeEnabled && !bFullBuild)
+		bool bIncrementIterativeCounter = false;
+		bool bFoundBuildDefinitions = false;
+		for (int32 PlatIndex = 0; PlatIndex < Platforms.Num(); ++PlatIndex)
 		{
-			bool bExists = false;
-			for (int32 PlatIndex = 0; PlatIndex < Platforms.Num(); ++PlatIndex)
+			const ITargetPlatform* TargetPlatform = Platforms[PlatIndex];
+			ICookedPackageWriter* PackageWriter = PackageWriters[PlatIndex];
+			OplogDataScratch.Reset();
+
+			if (!bFullBuild && bHybridIterativeEnabled)
 			{
-				const ITargetPlatform* TargetPlatform = Platforms[PlatIndex];
-				ICookedPackageWriter* PackageWriter = PackageWriters[PlatIndex];
-				if (!PackageWriter)
+				if (UE::TargetDomain::TryFetchCookAttachments(PackageWriter, PackageName,
+					TargetPlatform, OplogDataScratch, nullptr /* OutErrorMessage */))
 				{
-					continue;
-				}
-				bool bExistsInTargetDomain = UE::TargetDomain::TryFetchKeyAndDependencies(PackageWriter, PackageName,
-					TargetPlatform, nullptr /* OutHash */, &BuildDependencies, &RuntimeDependencies,
-					nullptr /* OutErrorMessage */);
-				if (bExistsInTargetDomain && UE::TargetDomain::IsIterativeEnabled(PackageName))
-				{
-					bExists = true;
-					PackageData->SetPlatformCooked(TargetPlatform, true);
+					if (UE::TargetDomain::IsIterativeEnabled(PackageName))
+					{
+						bIncrementIterativeCounter = true;
+						PackageData->SetPlatformCooked(TargetPlatform, true);
+					}
+					OutDependencies->Append(OplogDataScratch.BuildDependencies);
+					if (bAllowSoftDependencies)
+					{
+						OutDependencies->Append(OplogDataScratch.RuntimeOnlyDependencies);
+					}
+
+					if (bPreQueueBuildDefinitions)
+					{
+						bFoundBuildDefinitions = true;
+						BuildDefinitions.AddBuildDefinitionList(PackageName, TargetPlatform,
+							OplogDataScratch.BuildDefinitionList);
+					}
 				}
 			}
-			if (bExists)
+		}
+		if (bPreQueueBuildDefinitions && !bFoundBuildDefinitions)
+		{
+			OplogDataScratch.Reset();
+			if (UE::TargetDomain::TryFetchCookAttachments(nullptr, PackageName, nullptr, OplogDataScratch, nullptr))
 			{
-				COOK_STAT(++DetailedCookStats::NumPackagesIterativelySkipped);
+				BuildDefinitions.AddBuildDefinitionList(PackageName, nullptr, OplogDataScratch.BuildDefinitionList);
 			}
 		}
-		if (bAllowSoftDependencies)
+		if (bIncrementIterativeCounter)
 		{
-			BuildDependencies.Append(RuntimeDependencies);
+			COOK_STAT(++DetailedCookStats::NumPackagesIterativelySkipped);
 		}
 
 		// Sort the list of BuildDependencies to check for uniqueness and make them deterministic
-		Algo::Sort(BuildDependencies, FNameLexicalLess());
-		BuildDependencies.SetNum(Algo::Unique(BuildDependencies));
+		Algo::Sort(*OutDependencies, FNameLexicalLess());
+		OutDependencies->SetNum(Algo::Unique(*OutDependencies));
 	}
 
 	for (const ITargetPlatform* TargetPlatform : Platforms)
 	{
 		PackageData->FindOrAddPlatformData(TargetPlatform).bExplored = true;
 	}
-	bAlreadyCooked = PackageData->AreAllRequestedPlatformsCooked(true /* bAllowFailedCooks */);
+	bOutAlreadyCooked = PackageData->AreAllRequestedPlatformsCooked(true /* bAllowFailedCooks */);
 }
 
 bool FRequestCluster::IsRequestCookable(FName PackageName, FName& InOutFileName, FPackageData*& PackageData)
diff --git a/Engine/Source/Editor/UnrealEd/Private/Cooker/CookRequestCluster.h b/Engine/Source/Editor/UnrealEd/Private/Cooker/CookRequestCluster.h
--- a/Engine/Source/Editor/UnrealEd/Private/Cooker/CookRequestCluster.h
+++ b/Engine/Source/Editor/UnrealEd/Private/Cooker/CookRequestCluster.h
@@ -7,6 +7,7 @@
 #include "Containers/Map.h"
 #include "Containers/RingBuffer.h"
 #include "Cooker/CookTypes.h"
+#include "TargetDomain/TargetDomainUtils.h"
 #include "UObject/NameTypes.h"
 
 class IAssetRegistry;
@@ -95,7 +96,7 @@ private:
 	void FetchDependencies(const FCookerTimer& CookerTimer, bool& bOutComplete);
 	void StartAsync(const FCookerTimer& CookerTimer, bool& bOutComplete);
 	bool TryTakeOwnership(FPackageData& PackageData, bool bUrgent, FCompletionCallback&& CompletionCallback);
-	void VisitPackageData(FPackageData* PackageData, TArray<FName>* OutDependencies, bool& bAlreadyCooked);
+	void VisitPackageData(FPackageData* PackageData, TArray<FName>* OutDependencies, bool& bOutAlreadyCooked);
 	bool IsRequestCookable(FName PackageName, FName& InOutFileName, FPackageData*& PackageData);
 	static bool IsRequestCookable(FName PackageName, FName& InOutFileName, FPackageData*& PackageData,
 		const FPackageNameCache& InPackageNameCache, FPackageDatas& InPackageDatas, FPackageTracker& InPackageTracker,
@@ -111,7 +112,7 @@ private:
 	TArray<FFileNameRequest> InRequests;
 	TArray<FPackageData*> Requests;
 	TArray<FPackageData*> RequestsToDemote;
-	TArray<FName> RuntimeScratch;
+	UE::TargetDomain::FCookAttachments OplogDataScratch;
 	TArray<const ITargetPlatform*> Platforms;
 	TArray<ICookedPackageWriter*> PackageWriters;
 	FPackageDataSet OwnedPackageDatas;
@@ -120,6 +121,7 @@ private:
 	IAssetRegistry& AssetRegistry;
 	const FPackageNameCache& PackageNameCache;
 	FPackageTracker& PackageTracker;
+	FBuildDefinitions& BuildDefinitions;
 	int32 NextRequest = 0;
 	bool bAllowHardDependencies = true;
 	bool bAllowSoftDependencies = true;
@@ -130,6 +132,7 @@ private:
 	bool bDependenciesComplete = false;
 	bool bStartAsyncComplete = false;
 	bool bFullBuild = false;
+	bool bPreQueueBuildDefinitions = true;
 
 
 	/** Data used during the reentrant dependencies and topological sort operation of FRequestCluster. */
diff --git a/Engine/Source/Editor/UnrealEd/Private/Cooker/CookTypes.cpp b/Engine/Source/Editor/UnrealEd/Private/Cooker/CookTypes.cpp
--- a/Engine/Source/Editor/UnrealEd/Private/Cooker/CookTypes.cpp
+++ b/Engine/Source/Editor/UnrealEd/Private/Cooker/CookTypes.cpp
@@ -2,8 +2,10 @@
 
 #include "CookTypes.h"
 
+#include "Cooker/CookPackageData.h"
 #include "HAL/PlatformTime.h"
 #include "Math/NumericLimits.h"
+#include "Misc/CommandLine.h"
 
 namespace UE::Cook
 {
@@ -95,4 +97,59 @@ FCookSavePackageContext::~FCookSavePackageContext()
 		delete PackageWriter;
 	}
 }
+
+FBuildDefinitions::FBuildDefinitions()
+{
+	bTestPendingBuilds = FParse::Param(FCommandLine::Get(), TEXT("CookTestPendingBuilds"));
+}
+
+FBuildDefinitions::~FBuildDefinitions()
+{
+	Cancel();
+}
+
+void FBuildDefinitions::AddBuildDefinitionList(FName PackageName, const ITargetPlatform* TargetPlatform, TConstArrayView<UE::DerivedData::FBuildDefinition> BuildDefinitionList)
+{
+	using namespace UE::DerivedData;
+
+	// TODO_BuildDefinitionList: Trigger the builds
+	if (!bTestPendingBuilds)
+	{
+		return;
+	}
+
+	FPendingBuildData& BuildData = PendingBuilds.FindOrAdd(PackageName);
+	BuildData.bTryRemoved = false; // overwrite any previous value
+}
+
+bool FBuildDefinitions::TryRemovePendingBuilds(FName PackageName)
+{
+	FPendingBuildData* BuildData = PendingBuilds.Find(PackageName);
+	if (BuildData)
+	{
+		if (!bTestPendingBuilds || BuildData->bTryRemoved)
+		{
+			PendingBuilds.Remove(PackageName);
+			return true;
+		}
+		else
+		{
+			BuildData->bTryRemoved = true;
+			return false;
+		}
+	}
+
+	return true;
+}
+
+void FBuildDefinitions::Wait()
+{
+	PendingBuilds.Empty();
+}
+
+void FBuildDefinitions::Cancel()
+{
+	PendingBuilds.Empty();
+}
+
 }
diff --git a/Engine/Source/Editor/UnrealEd/Private/Cooker/CookTypes.h b/Engine/Source/Editor/UnrealEd/Private/Cooker/CookTypes.h
--- a/Engine/Source/Editor/UnrealEd/Private/Cooker/CookTypes.h
+++ b/Engine/Source/Editor/UnrealEd/Private/Cooker/CookTypes.h
@@ -4,14 +4,18 @@
 
 #include "Containers/Map.h"
 #include "Containers/Set.h"
+#include "Containers/UnrealString.h"
 #include "CookOnTheSide/CookOnTheFlyServer.h" // ECookTickFlags
+#include "DerivedDataRequestOwner.h"
 #include "HAL/Platform.h"
 #include "Logging/TokenizedMessage.h"
 #include "Serialization/PackageWriter.h"
 #include "Templates/Function.h"
+#include "UObject/NameTypes.h"
 #include "UObject/SavePackage.h"
 
-class FString;
+class ITargetPlatform;
+namespace UE::DerivedData { class FBuildDefinition; }
 
 #define COOK_CHECKSLOW_PACKAGEDATA 1
 #define DEBUG_COOKONTHEFLY 0
@@ -182,6 +186,29 @@ namespace UE::Cook
 	void InitializeTls();
 	bool IsSchedulerThread();
 	void SetIsSchedulerThread(bool bValue);
+
+
+	/** Placeholder to handle executing BuildDefintions for requested but not-yet-loaded packages. */
+	class FBuildDefinitions
+	{
+	public:
+		FBuildDefinitions();
+		~FBuildDefinitions();
+
+		void AddBuildDefinitionList(FName PackageName, const ITargetPlatform* TargetPlatform,
+			TConstArrayView<UE::DerivedData::FBuildDefinition> BuildDefinitionList);
+		bool TryRemovePendingBuilds(FName PackageName);
+		void Wait();
+		void Cancel();
+
+	private:
+		bool bTestPendingBuilds = false;
+		struct FPendingBuildData
+		{
+			bool bTryRemoved = false;
+		};
+		TMap<FName, FPendingBuildData> PendingBuilds;
+	};
 }
 
 //////////////////////////////////////////////////////////////////////////
diff --git a/Engine/Source/Editor/UnrealEd/Private/EditorDomain/EditorDomain.cpp b/Engine/Source/Editor/UnrealEd/Private/EditorDomain/EditorDomain.cpp
--- a/Engine/Source/Editor/UnrealEd/Private/EditorDomain/EditorDomain.cpp
+++ b/Engine/Source/Editor/UnrealEd/Private/EditorDomain/EditorDomain.cpp
@@ -21,6 +21,7 @@
 #include "Misc/ScopeLock.h"
 #include "Serialization/Archive.h"
 #include "Serialization/CompactBinary.h"
+#include "TargetDomain/TargetDomainUtils.h"
 #include "Templates/RefCounting.h"
 #include "Templates/UniquePtr.h"
 #include "UObject/PackageResourceManagerFile.h"
@@ -49,6 +50,7 @@ public:
 				UE_LOG(LogEditorDomain, Error, TEXT("Editor.ini:[CookSettings]:EditorDomainEnabled is deprecated, use Editor.ini:[EditorDomain]:EditorDomainEnabled instead."));
 			}
 			UE::EditorDomain::UtilsInitialize();
+			UE::TargetDomain::UtilsInitialize(bEditorDomainEnabled);
 			if (bEditorDomainEnabled)
 			{
 				// Set values for config settings EditorDomain depends on
diff --git a/Engine/Source/Editor/UnrealEd/Private/EditorDomain/EditorDomainUtils.cpp b/Engine/Source/Editor/UnrealEd/Private/EditorDomain/EditorDomainUtils.cpp
--- a/Engine/Source/Editor/UnrealEd/Private/EditorDomain/EditorDomainUtils.cpp
+++ b/Engine/Source/Editor/UnrealEd/Private/EditorDomain/EditorDomainUtils.cpp
@@ -7,6 +7,7 @@
 #include "AssetRegistry/AssetData.h"
 #include "AssetRegistry/IAssetRegistry.h"
 #include "Containers/Array.h"
+#include "DerivedDataBuildDefinition.h"
 #include "DerivedDataCache.h"
 #include "DerivedDataCacheKey.h"
 #include "DerivedDataCacheRecord.h"
@@ -21,15 +22,14 @@
 #include "Misc/StringBuilder.h"
 #include "Serialization/CompactBinaryWriter.h"
 #include "Serialization/PackageWriter.h"
+#include "TargetDomain/TargetDomainUtils.h"
 #include "UObject/CoreRedirects.h"
 #include "UObject/ObjectVersion.h"
 #include "UObject/Package.h"
 #include "UObject/SavePackage.h"
 #include "UObject/UObjectHash.h"
 
-namespace UE
-{
-namespace EditorDomain
+namespace UE::EditorDomain
 {
 
 FClassDigestMap GClassDigests;
@@ -850,6 +850,18 @@ bool TrySavePackage(UPackage* Package)
 	FRequestOwner Owner(EPriority::Normal);
 	Cache.Put({RecordBuilder.Build()}, Package->GetName(), ECachePolicy::Default, Owner);
 	Owner.KeepAlive();
+
+	// TODO_BuildDefinitionList: Calculate and store BuildDefinitionList on the PackageData, or collect it here from some other source.
+	TArray<UE::DerivedData::FBuildDefinition> BuildDefinitions;
+	FCbObject BuildDefinitionList = UE::TargetDomain::BuildDefinitionListToObject(BuildDefinitions);
+	FCbObject TargetDomainDependencies = UE::TargetDomain::CollectDependenciesObject(Package, nullptr, nullptr);
+	if (TargetDomainDependencies)
+	{
+		TArray<IPackageWriter::FCommitAttachmentInfo, TInlineAllocator<2>> Attachments;
+		Attachments.Add({ "Dependencies", TargetDomainDependencies });
+		Attachments.Add({ "BuildDefinitionList", BuildDefinitionList });
+		UE::TargetDomain::CommitEditorDomainCookAttachments(Package->GetFName(), Attachments);
+	}
 	return true;
 }
 
@@ -988,7 +1000,4 @@ void PutBulkDataPayloadId(FName PackageName, const FGuid& BulkDataId, FSharedBuf
 	Owner.KeepAlive();
 }
 
-
-
-}
-}
+}
\ No newline at end of file
diff --git a/Engine/Source/Editor/UnrealEd/Private/TargetDomain/TargetDomainUtils.cpp b/Engine/Source/Editor/UnrealEd/Private/TargetDomain/TargetDomainUtils.cpp
--- a/Engine/Source/Editor/UnrealEd/Private/TargetDomain/TargetDomainUtils.cpp
+++ b/Engine/Source/Editor/UnrealEd/Private/TargetDomain/TargetDomainUtils.cpp
@@ -2,21 +2,71 @@
 
 #include "TargetDomain/TargetDomainUtils.h"
 
+#include "Algo/BinarySearch.h"
+#include "Algo/IsSorted.h"
+#include "Algo/Sort.h"
 #include "AssetRegistry/IAssetRegistry.h"
 #include "Containers/Set.h"
 #include "Containers/UnrealString.h"
 #include "Cooker/PackageBuildDependencyTracker.h"
+#include "DerivedDataBuildDefinition.h"
+#include "DerivedDataBuildKey.h"
+#include "EditorDomain/EditorDomain.h"
 #include "EditorDomain/EditorDomainUtils.h"
+#include "HAL/PlatformFile.h"
+#include "HAL/PlatformFileManager.h"
+#include "IO/IoDispatcher.h"
 #include "IO/IoHash.h"
+#include "Misc/App.h"
 #include "Misc/ScopeRWLock.h"
 #include "Misc/StringBuilder.h"
 #include "Serialization/CompactBinary.h"
 #include "Serialization/CompactBinaryWriter.h"
 #include "Serialization/PackageWriter.h"
+#include "ZenStoreHttpClient.h"
 
 namespace UE::TargetDomain
 {
 
+/**
+ * Reads / writes an oplog for EditorDomain BuildDefinitionLists.
+ * TODO: Reduce duplication between this class and FZenStoreWriter
+ */
+class FEditorDomainOplog
+{
+public:
+	FEditorDomainOplog();
+
+	bool IsValid() const;
+	void CommitPackage(FName PackageName, TArrayView<IPackageWriter::FCommitAttachmentInfo> Attachments);
+	FCbObject GetOplogAttachment(FName PackageName, FUtf8StringView AttachmentKey);
+
+private:
+	struct FOplogEntry
+	{
+		struct FAttachment
+		{
+			const UTF8CHAR* Key;
+			FIoHash Hash;
+		};
+
+		TArray<FAttachment> Attachments;
+	};
+
+	void InitializeRead();
+	static void StaticInit();
+	static bool IsReservedOplogKey(FUtf8StringView Key);
+
+	UE::FZenStoreHttpClient HttpClient;
+	FCriticalSection Lock;
+	TMap<FName, FOplogEntry> Entries;
+	bool bConnectSuccessful = false;
+	bool bInitializedRead = false;
+
+	static TArray<const UTF8CHAR*> ReservedOplogKeys;
+};
+TUniquePtr<FEditorDomainOplog> GEditorDomainOplog;
+
 bool TryCreateKey(FName PackageName, TArrayView<FName> SortedBuildDependencies, FIoHash* OutHash, FString* OutErrorMessage)
 {
 	IAssetRegistry* AssetRegistry = IAssetRegistry::Get();
@@ -179,41 +229,81 @@ FCbObject CollectDependenciesObject(UPackage* Package, const ITargetPlatform* Ta
 	return Writer.Save().AsObject();
 }
 
+FCbObject BuildDefinitionListToObject(TConstArrayView<UE::DerivedData::FBuildDefinition> BuildDefinitionList)
+{
+	using namespace UE::DerivedData;
+
+	if (BuildDefinitionList.IsEmpty())
+	{
+		return FCbObject();
+	}
+
+	TArray<const FBuildDefinition*> Sorted;
+	Sorted.Reserve(BuildDefinitionList.Num());
+	for (const FBuildDefinition& BuildDefinition : BuildDefinitionList)
+	{
+		Sorted.Add(&BuildDefinition);
+	}
+	Sorted.Sort([](const FBuildDefinition& A, const FBuildDefinition& B)
+		{
+			return A.GetKey().Hash < B.GetKey().Hash;
+		});
+
+	FCbWriter Writer;
+	Writer.BeginObject();
+	Writer.BeginArray("BuildDefinitions");
+	for (const FBuildDefinition* BuildDefinition : Sorted)
+	{
+		BuildDefinition->Save(Writer);
+	}
+	Writer.EndArray();
+	return Writer.Save().AsObject();
+}
 
-bool TryFetchKeyAndDependencies(ICookedPackageWriter* PackageWriter, FName PackageName, const ITargetPlatform* TargetPlatform,
-	FIoHash* OutHash, TArray<FName>* OutBuildDependencies, TArray<FName>* OutRuntimeOnlyDependencies, FString* OutErrorMessage)
+bool TryFetchCookAttachments(ICookedPackageWriter* PackageWriter, FName PackageName, const ITargetPlatform* TargetPlatform,
+	FCookAttachments& OutOplogData, FString* OutErrorMessage)
 {
-	FCbObject DependenciesObj = PackageWriter->GetOplogAttachment(PackageName, "Dependencies");
-	FIoHash StoredKey = DependenciesObj["targetdomainkey"].AsHash();
-	if (StoredKey.IsZero())
+	OutOplogData.Reset();
+
+	const ANSICHAR* DependenciesKey = "Dependencies";
+	FCbObject DependenciesObj;
+	if (TargetPlatform)
+	{
+		check(PackageWriter);
+		DependenciesObj = PackageWriter->GetOplogAttachment(PackageName, DependenciesKey);
+	}
+	else
+	{
+		if (!GEditorDomainOplog)
 		{
 			if (OutErrorMessage)
 			{
-			*OutErrorMessage = TEXT("Dependencies not in oplog.");
+				*OutErrorMessage = TEXT("EditorDomain oplog is not available.");
 			}
 			return false;
 		}
-
-	TArray<FName> BuildDependencies;
-	if (OutBuildDependencies)
-	{
-		OutBuildDependencies->Reset();
+		DependenciesObj = GEditorDomainOplog->GetOplogAttachment(PackageName, DependenciesKey);
 	}
-	else
+	FIoHash StoredKey = DependenciesObj["targetdomainkey"].AsHash();
+	if (StoredKey.IsZero())
+	{
+		if (OutErrorMessage)
 		{
-		OutBuildDependencies = &BuildDependencies;
+			*OutErrorMessage = TEXT("Dependencies not in oplog.");
+		}
+		return false;
 	}
 
 	for (FCbFieldView DepObj : DependenciesObj["builddependencies"])
 	{
 		if (FString DependencyName(DepObj.AsString()); !DependencyName.IsEmpty())
 		{
-			OutBuildDependencies->Add(FName(*DependencyName));
+			OutOplogData.BuildDependencies.Add(FName(*DependencyName));
 		}
 	}
 
 	FIoHash CurrentKey;
-	if (!TryCreateKey(PackageName, *OutBuildDependencies, &CurrentKey, OutErrorMessage))
+	if (!TryCreateKey(PackageName, OutOplogData.BuildDependencies, &CurrentKey, OutErrorMessage))
 	{
 		return false;
 	}
@@ -227,15 +317,36 @@ bool TryFetchKeyAndDependencies(ICookedPackageWriter* PackageWriter, FName Packa
 		return false;
 	}
 
-	if (OutRuntimeOnlyDependencies)
-	{
 	for (FCbFieldView DepObj : DependenciesObj["runtimeonlydependencies"])
 	{
 		if (FString DependencyName(DepObj.AsString()); !DependencyName.IsEmpty())
 		{
-				OutRuntimeOnlyDependencies->Add(FName(*DependencyName));
+			OutOplogData.RuntimeOnlyDependencies.Add(FName(*DependencyName));
 		}
 	}
+
+	const ANSICHAR* BuildDefinitionListKey = "BuildDefinitionList";
+	FCbObject BuildDefinitionListObj;
+	if (TargetPlatform)
+	{
+		BuildDefinitionListObj = PackageWriter->GetOplogAttachment(PackageName, BuildDefinitionListKey);
+	}
+	else
+	{
+		BuildDefinitionListObj = GEditorDomainOplog->GetOplogAttachment(PackageName, BuildDefinitionListKey);
+	}
+
+	for (FCbField BuildDefinitionObj : BuildDefinitionListObj)
+	{
+		UE::DerivedData::FOptionalBuildDefinition BuildDefinition =
+			UE::DerivedData::FBuildDefinition::Load(TEXT("TargetDomainBuildDefinitionList"),
+				BuildDefinitionObj.AsObject());
+		if (!BuildDefinition)
+		{
+			OutOplogData.BuildDefinitionList.Reset();
+			break;
+		}
+		OutOplogData.BuildDefinitionList.Add(MoveTemp(BuildDefinition).Get());
 	}
 
 	return true;
@@ -274,5 +385,256 @@ bool IsIterativeEnabled(FName PackageName)
 	return true;
 }
 
-} // namespace UE::TargetDomain
+TArray<const UTF8CHAR*> FEditorDomainOplog::ReservedOplogKeys;
+
+FEditorDomainOplog::FEditorDomainOplog()
+	: HttpClient(TEXT("localhost"), 1337)
+{
+	StaticInit();
+
+	FString ProjectId = FApp::GetProjectName();
+	FString OplogId = TEXT("EditorDomain");
+
+	FString RootDir = FPaths::RootDir();
+	FString EngineDir = FPaths::EngineDir();
+	FPaths::NormalizeDirectoryName(EngineDir);
+	FString ProjectDir = FPaths::ProjectDir();
+	FPaths::NormalizeDirectoryName(ProjectDir);
+
+	IPlatformFile& PlatformFile = FPlatformFileManager::Get().GetPlatformFile();
+	FString AbsServerRoot = PlatformFile.ConvertToAbsolutePathForExternalAppForRead(*RootDir);
+	FString AbsEngineDir = PlatformFile.ConvertToAbsolutePathForExternalAppForRead(*EngineDir);
+	FString AbsProjectDir = PlatformFile.ConvertToAbsolutePathForExternalAppForRead(*ProjectDir);
+
+	HttpClient.Initialize(ProjectId,
+		OplogId,
+		AbsServerRoot,
+		AbsEngineDir,
+		AbsProjectDir);
+	if (HttpClient.IsConnected())
+	{
+		HttpClient.EstablishWritableOpLog(ProjectId, OplogId, false /* bFullBuild */);
+	}
+}
+
+void FEditorDomainOplog::InitializeRead()
+{
+	if (bInitializedRead)
+	{
+		return;
+	}
+	UE_LOG(LogEditorDomain, Display, TEXT("Fetching EditorDomain oplog..."));
+
+	TFuture<FIoStatus> FutureOplogStatus = HttpClient.GetOplog().Next([this](TIoStatusOr<FCbObject> OplogStatus)
+		{
+			if (!OplogStatus.IsOk())
+			{
+				return OplogStatus.Status();
+			}
+
+			FCbObject Oplog = OplogStatus.ConsumeValueOrDie();
+
+			for (FCbField& EntryObject : Oplog["entries"])
+			{
+				FUtf8StringView PackageName = EntryObject["key"].AsString();
+				if (PackageName.IsEmpty())
+				{
+					continue;
+				}
+				FName PackageFName(PackageName);
+				FOplogEntry& Entry = Entries.FindOrAdd(PackageFName);
+				Entry.Attachments.Empty();
+
+				for (FCbFieldView Field : EntryObject)
+				{
+					FUtf8StringView FieldName = Field.GetName();
+					if (IsReservedOplogKey(FieldName))
+					{
+						continue;
+					}
+					if (Field.IsHash())
+					{
+						const UTF8CHAR* AttachmentId = UE::FZenStoreHttpClient::FindOrAddAttachmentId(FieldName);
+						Entry.Attachments.Add({ AttachmentId, Field.AsHash() });
+					}
+				}
+				Entry.Attachments.Shrink();
+				check(Algo::IsSorted(Entry.Attachments, [](const FOplogEntry::FAttachment& A, const FOplogEntry::FAttachment& B)
+					{
+						return FUtf8StringView(A.Key).Compare(FUtf8StringView(B.Key), ESearchCase::IgnoreCase) < 0;
+					}));
+			}
+
+			return FIoStatus::Ok;
+		});
+	FutureOplogStatus.Get();
+	bInitializedRead = true;
+}
+
+void FEditorDomainOplog::StaticInit()
+{
+	if (ReservedOplogKeys.Num() > 0)
+	{
+		return;
+	}
+
+	ReservedOplogKeys.Append({ UTF8TEXT("key") });
+	Algo::Sort(ReservedOplogKeys, [](const UTF8CHAR* A, const UTF8CHAR* B)
+		{
+			return FUtf8StringView(A).Compare(FUtf8StringView(B), ESearchCase::IgnoreCase) < 0;
+		});;
+}
+
+bool FEditorDomainOplog::IsReservedOplogKey(FUtf8StringView Key)
+{
+	int32 Index = Algo::LowerBound(ReservedOplogKeys, Key,
+		[](const UTF8CHAR* Existing, FUtf8StringView Key)
+		{
+			return FUtf8StringView(Existing).Compare(Key, ESearchCase::IgnoreCase) < 0;
+		});
+	return Index != ReservedOplogKeys.Num() &&
+		FUtf8StringView(ReservedOplogKeys[Index]).Equals(Key, ESearchCase::IgnoreCase);
+}
+
+bool FEditorDomainOplog::IsValid() const
+{
+	return HttpClient.IsConnected();
+}
+
+void FEditorDomainOplog::CommitPackage(FName PackageName, TArrayView<IPackageWriter::FCommitAttachmentInfo> Attachments)
+{
+	FScopeLock ScopeLock(&Lock);
+
+	FCbPackage Pkg;
+
+	TArray<FCbAttachment, TInlineAllocator<2>> CbAttachments;
+	int32 NumAttachments = Attachments.Num();
+	FOplogEntry& Entry = Entries.FindOrAdd(PackageName);
+	Entry.Attachments.Empty(NumAttachments);
+	if (NumAttachments)
+	{
+		TArray<const IPackageWriter::FCommitAttachmentInfo*, TInlineAllocator<2>> SortedAttachments;
+		SortedAttachments.Reserve(NumAttachments);
+		for (const IPackageWriter::FCommitAttachmentInfo& AttachmentInfo : Attachments)
+		{
+			SortedAttachments.Add(&AttachmentInfo);
+		}
+		SortedAttachments.Sort([](const IPackageWriter::FCommitAttachmentInfo& A, const IPackageWriter::FCommitAttachmentInfo& B)
+			{
+				return A.Key.Compare(B.Key, ESearchCase::IgnoreCase) < 0;
+			});
+		CbAttachments.Reserve(NumAttachments);
+		for (const IPackageWriter::FCommitAttachmentInfo* AttachmentInfo : SortedAttachments)
+		{
+			const FCbAttachment& CbAttachment = CbAttachments.Emplace_GetRef(AttachmentInfo->Value);
+			check(!IsReservedOplogKey(AttachmentInfo->Key));
+			Pkg.AddAttachment(CbAttachment);
+			Entry.Attachments.Add(FOplogEntry::FAttachment{
+				UE::FZenStoreHttpClient::FindOrAddAttachmentId(AttachmentInfo->Key), CbAttachment.GetHash() });
+		}
+	}
 
+	FCbWriter PackageObj;
+	FString PackageNameKey = PackageName.ToString();
+	PackageNameKey.ToLowerInline();
+	PackageObj.BeginObject();
+	PackageObj << "key" << PackageNameKey;
+	for (int32 Index = 0; Index < NumAttachments; ++Index)
+	{
+		FCbAttachment& CbAttachment = CbAttachments[Index];
+		FOplogEntry::FAttachment& EntryAttachment = Entry.Attachments[Index];
+		PackageObj << EntryAttachment.Key << CbAttachment;
+	}
+	PackageObj.EndObject();
+
+	FCbObject Obj = PackageObj.Save().AsObject();
+	Pkg.SetObject(Obj);
+	HttpClient.AppendOp(Pkg);
+}
+
+// Note that this is destructive - we yank out the buffer memory from the 
+// IoBuffer into the FSharedBuffer
+FSharedBuffer IoBufferToSharedBuffer(FIoBuffer& InBuffer)
+{
+	InBuffer.EnsureOwned();
+	const uint64 DataSize = InBuffer.DataSize();
+	uint8* DataPtr = InBuffer.Release().ValueOrDie();
+	return FSharedBuffer{ FSharedBuffer::TakeOwnership(DataPtr, DataSize, FMemory::Free) };
+};
+
+FCbObject FEditorDomainOplog::GetOplogAttachment(FName PackageName, FUtf8StringView AttachmentKey)
+{
+	FScopeLock ScopeLock(&Lock);
+	InitializeRead();
+
+	FOplogEntry* Entry = Entries.Find(PackageName);
+	if (!Entry)
+	{
+		return FCbObject();
+	}
+
+	const UTF8CHAR* AttachmentId = UE::FZenStoreHttpClient::FindAttachmentId(AttachmentKey);
+	if (!AttachmentId)
+	{
+		return FCbObject();
+	}
+	FUtf8StringView AttachmentIdView(AttachmentId);
+
+	int32 AttachmentIndex = Algo::LowerBound(Entry->Attachments, AttachmentIdView,
+		[](const FOplogEntry::FAttachment& Existing, FUtf8StringView AttachmentIdView)
+		{
+			return FUtf8StringView(Existing.Key).Compare(AttachmentIdView, ESearchCase::IgnoreCase) < 0;
+		});
+	if (AttachmentIndex == Entry->Attachments.Num())
+	{
+		return FCbObject();
+	}
+	const FOplogEntry::FAttachment& Existing = Entry->Attachments[AttachmentIndex];
+	if (!FUtf8StringView(Existing.Key).Equals(AttachmentIdView, ESearchCase::IgnoreCase))
+	{
+		return FCbObject();
+	}
+	TIoStatusOr<FIoBuffer> BufferResult = HttpClient.ReadOpLogAttachment(WriteToString<48>(Existing.Hash));
+	if (!BufferResult.IsOk())
+	{
+		return FCbObject();
+	}
+	FIoBuffer Buffer = BufferResult.ValueOrDie();
+	if (Buffer.DataSize() == 0)
+	{
+		return FCbObject();
+	}
+
+	FSharedBuffer SharedBuffer = IoBufferToSharedBuffer(Buffer);
+	return FCbObject(SharedBuffer);
+}
+
+void CommitEditorDomainCookAttachments(FName PackageName, TArrayView<IPackageWriter::FCommitAttachmentInfo> Attachments)
+{
+	if (!GEditorDomainOplog)
+	{
+		return;
+	}
+	GEditorDomainOplog->CommitPackage(PackageName, Attachments);
+}
+
+void UtilsInitialize(bool bEditorDomainEnabled)
+{
+	if (bEditorDomainEnabled)
+	{
+		bool bCookAttachmentsEnabled = true;
+		GConfig->GetBool(TEXT("EditorDomain"), TEXT("CookAttachmentsEnabled"), bCookAttachmentsEnabled, GEditorIni);
+		if (bCookAttachmentsEnabled)
+		{
+			GEditorDomainOplog = MakeUnique<FEditorDomainOplog>();
+			if (!GEditorDomainOplog->IsValid())
+			{
+				UE_LOG(LogEditorDomain, Warning, TEXT("Failed to connect to ZenServer; EditorDomain oplog is unavailable."));
+				GEditorDomainOplog.Reset();
+			}
+		}
+	}
+}
+
+
+} // namespace UE::TargetDomain
diff --git a/Engine/Source/Editor/UnrealEd/Private/TargetDomain/TargetDomainUtils.h b/Engine/Source/Editor/UnrealEd/Private/TargetDomain/TargetDomainUtils.h
--- a/Engine/Source/Editor/UnrealEd/Private/TargetDomain/TargetDomainUtils.h
+++ b/Engine/Source/Editor/UnrealEd/Private/TargetDomain/TargetDomainUtils.h
@@ -3,6 +3,9 @@
 #pragma once
 
 #include "Containers/Array.h"
+#include "Containers/ArrayView.h"
+#include "DerivedDataBuildDefinition.h"
+#include "Serialization/PackageWriter.h"
 #include "UObject/NameTypes.h"
 
 class ICookedPackageWriter;
@@ -11,10 +14,13 @@ class FCbObject;
 class FString;
 class UPackage;
 struct FIoHash;
+namespace UE::DerivedData { class FBuildDefinition; }
 
 namespace UE::TargetDomain
 {
 
+void UtilsInitialize(bool bEditorDomainEnabled);
+
 /** Create the TargetDomainKey based on the EditorDomainKeys of the Package and its dependencies. */
 bool TryCreateKey(FName PackageName, TArrayView<FName> SortedBuildDependencies, FIoHash* OutHash, FString* OutErrorMessage);
 
@@ -24,14 +30,30 @@ bool TryCollectKeyAndDependencies(UPackage* Package, const ITargetPlatform* Targ
 /** Collect the Package's dependencies, and create a FCbObject describing them for storage in the OpLog. */
 FCbObject CollectDependenciesObject(UPackage* Package, const ITargetPlatform* TargetPlatform, FString* ErrorMessage);
 
-/**
- * Read the oplog for the given packagename and fetch the dependencies and key out of it. Use the dependencies to calculate
- * the current key, and return the dependencies and key only if they key matches.
- */
-bool TryFetchKeyAndDependencies(ICookedPackageWriter* PackageStore, FName PackageName, const ITargetPlatform* TargetPlatform,
-	FIoHash* OutHash, TArray<FName>* OutBuildDependencies, TArray<FName>* OutRuntimeOnlyDependencies, FString* OutErrorMessage);
+/** Marshal the given BuildDefinitionList to FCbObject for storage in DDC. */
+FCbObject BuildDefinitionListToObject(TConstArrayView<UE::DerivedData::FBuildDefinition> BuildDefinitionList);
+
+struct FCookAttachments
+{
+	TArray<FName> BuildDependencies;
+	TArray<FName> RuntimeOnlyDependencies;
+	TArray<UE::DerivedData::FBuildDefinition> BuildDefinitionList;
+
+	void Reset()
+	{
+		BuildDependencies.Reset();
+		RuntimeOnlyDependencies.Reset();
+		BuildDefinitionList.Reset();
+	}
+};
+bool TryFetchCookAttachments(ICookedPackageWriter* PackageWriter, FName PackageName, const ITargetPlatform* TargetPlatform,
+	FCookAttachments& OutOplogData, FString* OutErrorMessage);
 
 /** Return whether iterative cook is enabled for the given packagename, based on used-class allowlist/blocklist. */
 bool IsIterativeEnabled(FName PackageName);
 
+
+/** Store extra information derived during save and used by the cooker for the given EditorDomain package. */
+void CommitEditorDomainCookAttachments(FName PackageName, TArrayView<IPackageWriter::FCommitAttachmentInfo> Attachments);
+
 }
