Cooker - disable saving of BuildDefinitionLists until the zenserver gets the CbPackage fix for empty CbObject attachments.
#rnx

diff --git a/Engine/Source/Editor/UnrealEd/Private/CookOnTheFlyServer.cpp b/Engine/Source/Editor/UnrealEd/Private/CookOnTheFlyServer.cpp
--- a/Engine/Source/Editor/UnrealEd/Private/CookOnTheFlyServer.cpp
+++ b/Engine/Source/Editor/UnrealEd/Private/CookOnTheFlyServer.cpp
@@ -4501,7 +4501,8 @@ void FSaveCookedPackageContext::FinishPlatform()
 		Info.PackageGuid = AssetPackageData->PackageGuid;
 		PRAGMA_ENABLE_DEPRECATION_WARNINGS
 		Info.Attachments.Add({ "Dependencies", TargetDomainDependencies });
-		Info.Attachments.Add({ "BuildDefinitionList", BuildDefinitionList });
+		// TODO: Reenable BuildDefinitionList once FCbPackage support for empty FCbObjects is in
+		//Info.Attachments.Add({ "BuildDefinitionList", BuildDefinitionList });
 
 		PackageWriter->CommitPackage(Info);
 	}
diff --git a/Engine/Source/Editor/UnrealEd/Private/EditorDomain/EditorDomainUtils.cpp b/Engine/Source/Editor/UnrealEd/Private/EditorDomain/EditorDomainUtils.cpp
--- a/Engine/Source/Editor/UnrealEd/Private/EditorDomain/EditorDomainUtils.cpp
+++ b/Engine/Source/Editor/UnrealEd/Private/EditorDomain/EditorDomainUtils.cpp
@@ -859,7 +859,8 @@ bool TrySavePackage(UPackage* Package)
 	{
 		TArray<IPackageWriter::FCommitAttachmentInfo, TInlineAllocator<2>> Attachments;
 		Attachments.Add({ "Dependencies", TargetDomainDependencies });
-		Attachments.Add({ "BuildDefinitionList", BuildDefinitionList });
+		// TODO: Reenable BuildDefinitionList once FCbPackage support for empty FCbObjects is in
+		//Attachments.Add({ "BuildDefinitionList", BuildDefinitionList });
 		UE::TargetDomain::CommitEditorDomainCookAttachments(Package->GetFName(), Attachments);
 	}
 	return true;
