COTF: Get read lock before locking tracked packages to prevent deadlock

#rnx

diff --git a/Engine/Source/Editor/UnrealEd/Private/Cooker/IoStoreCookOnTheFlyRequestManager.cpp b/Engine/Source/Editor/UnrealEd/Private/Cooker/IoStoreCookOnTheFlyRequestManager.cpp
--- a/Engine/Source/Editor/UnrealEd/Private/Cooker/IoStoreCookOnTheFlyRequestManager.cpp
+++ b/Engine/Source/Editor/UnrealEd/Private/Cooker/IoStoreCookOnTheFlyRequestManager.cpp
@@ -512,31 +512,31 @@ private:
 		using namespace UE::PackageStore::Messaging;
 
 		FPlatformContext& Context = GetContext(Client.PlatformName);
-		FScopeLock _(&Context.CriticalSection);
 		
-		const TArrayView<const int32> CookedEntryIndices = Context.PackageTracker.GetCookedEntryIndices();
-		const TArrayView<const FName> FailedPackages = Context.PackageTracker.GetFailedPackages();
-
 		UE::PackageStore::Messaging::FPackageStoreData PackageStoreData;
-
-		PackageStoreData.TotalCookedPackages	= CookedEntryIndices.Num();
-		PackageStoreData.TotalFailedPackages	= FailedPackages.Num();
-		//PackageStoreData.FailedPackages			= FailedPackages;
-
-		Context.PackageWriter->GetEntries([&CookedEntryIndices, &PackageStoreData](TArrayView<const FPackageStoreEntryResource> Entries)
+		
+		Context.PackageWriter->GetEntries([&Context, &PackageStoreData](TArrayView<const FPackageStoreEntryResource> Entries)
 		{
+			FScopeLock _(&Context.CriticalSection);
+			
+			const TArrayView<const int32> CookedEntryIndices = Context.PackageTracker.GetCookedEntryIndices();
+			const TArrayView<const FName> FailedPackages = Context.PackageTracker.GetFailedPackages();
+			
 			for (int32 EntryIndex : CookedEntryIndices)
 			{
 				check(EntryIndex != INDEX_NONE);
 				PackageStoreData.CookedPackages.Add(Entries[EntryIndex]);
 			}
+			
+			for (const FName& FailedPackageName : FailedPackages)
+			{
+				PackageStoreData.FailedPackages.Add(FPackageId::FromName(FailedPackageName));
+			}
+			
+			PackageStoreData.TotalCookedPackages	= CookedEntryIndices.Num();
+			PackageStoreData.TotalFailedPackages	= FailedPackages.Num();
 		});
-
-		for (const FName& FailedPackageName : FailedPackages)
-		{
-			PackageStoreData.FailedPackages.Add(FPackageId::FromName(FailedPackageName));
-		}
-
+		
 		Response.SetBodyTo(FGetCookedPackagesResponse { MoveTemp(PackageStoreData) });
 		Response.SetStatus(ECookOnTheFlyMessageStatus::Ok);
 
@@ -778,4 +778,4 @@ TUniquePtr<ICookOnTheFlyRequestManager> MakeIoStoreCookOnTheFlyRequestManager(IC
 
 }}
 
-#endif
\ No newline at end of file
+#endif
