WorldPartition: Foliage builder allows changing grid size

diff --git a/Engine/Source/Editor/UnrealEd/Private/WorldPartition/WorldPartitionFoliageBuilder.cpp b/Engine/Source/Editor/UnrealEd/Private/WorldPartition/WorldPartitionFoliageBuilder.cpp
new file mode 100644
--- /dev/null
+++ b/Engine/Source/Editor/UnrealEd/Private/WorldPartition/WorldPartitionFoliageBuilder.cpp
@@ -0,0 +1,240 @@
+// Copyright Epic Games, Inc. All Rights Reserved.
+
+#include "WorldPartition/WorldPartitionFoliageBuilder.h"
+
+#include "WorldPartition/WorldPartition.h"
+#include "WorldPartition/WorldPartitionHelpers.h"
+#include "PackageSourceControlHelper.h"
+#include "SourceControlHelpers.h"
+#include "HAL/PlatformFileManager.h"
+#include "Engine/World.h"
+#include "GameFramework/Actor.h"
+#include "GameFramework/WorldSettings.h"
+#include "InstancedFoliageActor.h"
+#include "ActorPartition/ActorPartitionSubsystem.h"
+
+DEFINE_LOG_CATEGORY_STATIC(LogWorldPartitionFoliageBuilder, All, All);
+
+UWorldPartitionFoliageBuilder::UWorldPartitionFoliageBuilder(const FObjectInitializer& ObjectInitializer)
+	: Super(ObjectInitializer)
+	, NewGridSize(0)
+{
+	FParse::Value(FCommandLine::Get(), TEXT("NewGridSize="), NewGridSize);
+}
+
+bool UWorldPartitionFoliageBuilder::RunInternal(UWorld* World, const FBox& Bounds, FPackageSourceControlHelper& PackageHelper)
+{			
+	UWorldPartition* WorldPartition = World->GetWorldPartition();
+	if (!WorldPartition)
+	{
+		UE_LOG(LogWorldPartitionFoliageBuilder, Error, TEXT("Failed to retrieve WorldPartition."));
+		return false;
+	}
+
+	// Validate New Grid Size
+	if (World->GetWorldSettings()->InstancedFoliageGridSize == NewGridSize)
+	{
+		UE_LOG(LogWorldPartitionFoliageBuilder, Display, TEXT("Foliage Grid Size already is %d."), NewGridSize);
+		return true;
+	}
+	else if (NewGridSize == 0)
+	{
+		UE_LOG(LogWorldPartitionFoliageBuilder, Error, TEXT("Foliage Grid Size is invalid %d. Please specify NewGridSize parameter."), NewGridSize);
+		return false;
+	}
+
+	UE_LOG(LogWorldPartitionFoliageBuilder, Display, TEXT("Changing Foliage Grid Size from %d to %d"), World->GetWorldSettings()->InstancedFoliageGridSize, NewGridSize);
+			
+	// Update Setting
+	World->GetWorldSettings()->Modify();
+
+	World->GetWorldSettings()->InstancedFoliageGridSize = NewGridSize;
+	World->GetWorldSettings()->bIncludeGridSizeInNameForFoliageActors = true;
+		
+	FPackageSourceControlHelper SCCHelper;
+	TArray<FGuid> ExistingActors;
+	
+	// Snapshot existing ActorDescs
+	FWorldPartitionHelpers::ForEachActorDesc(WorldPartition, AInstancedFoliageActor::StaticClass(), [&ExistingActors](const FWorldPartitionActorDesc* ActorDesc)
+	{
+		ExistingActors.Add(ActorDesc->GetGuid());
+		return true;
+	});
+		
+	TMap<UActorPartitionSubsystem::FCellCoord, FGuid> NewActors;
+	UActorPartitionSubsystem* ActorPartitionSubsystem = World->GetSubsystem<UActorPartitionSubsystem>();
+	int32 NumInstances = 0;
+	int32 NumInstancesProcessed = 0;
+
+	// This is the loop:
+	// - Load 1 Existing IFA (OldGridSize)
+	// - Copy all its instance data and compute its overall instance bounds
+	// - Iterate over its bounds using the new grid size and create/reload new IFAs (NewGridSize).
+	// - Add instance data to the new IFAs
+	// - Save new/reload IFAs that were modified
+	// - Unload the Existing IFA and all the new IFAs
+	// - Repeat for next Existing IFA
+	FWorldPartitionHelpers::ForEachActorWithLoading(WorldPartition, ExistingActors, [this, World, WorldPartition, ActorPartitionSubsystem, &SCCHelper, &NewActors, &NumInstances, &NumInstancesProcessed](const FWorldPartitionActorDesc* ActorDesc)
+	{
+		TMap<UFoliageType*, TArray<FFoliageInstance>> FoliageToAdd;
+		FBox InstanceBounds(ForceInit);
+				
+		AInstancedFoliageActor* IFA = Cast<AInstancedFoliageActor>(ActorDesc->GetActor());
+		if (!IFA)
+		{
+			UE_LOG(LogWorldPartitionFoliageBuilder, Error, TEXT("Foliage actor failed to load: %s (%s)"), *ActorDesc->GetActorName().ToString(), *ActorDesc->GetActorPackage().ToString());
+			return false;
+		}
+	
+		check(IFA->GridSize != NewGridSize);
+        // Harvest Instances from existing IFA and build up instance bounds
+		UE_LOG(LogWorldPartitionFoliageBuilder, Display, TEXT("Processing existing foliage actor: %s (%s)"), *ActorDesc->GetActorName().ToString(), *ActorDesc->GetActorPackage().ToString());
+		
+        IFA->ForEachFoliageInfo([IFA, &FoliageToAdd, &InstanceBounds, &NumInstances](UFoliageType* FoliageType, FFoliageInfo& FoliageInfo)
+        {
+            TMap<AInstancedFoliageActor*, TArray<const FFoliageInstance*>> InstancesToAdd;
+
+            check(FoliageType->GetTypedOuter<AInstancedFoliageActor>() == nullptr);
+            FoliageToAdd.FindOrAdd(FoliageType).Append(FoliageInfo.Instances);
+            for (const FFoliageInstance& Instance : FoliageInfo.Instances)
+            {
+                InstanceBounds += Instance.GetInstanceWorldTransform().GetLocation();
+            }
+
+			if(FoliageInfo.Instances.Num() > 0)
+			{
+				UE_LOG(LogWorldPartitionFoliageBuilder, Display, TEXT("    FoliageType: %s Count: %d"), *FoliageType->GetName(), FoliageInfo.Instances.Num());
+			}
+			NumInstances += FoliageInfo.Instances.Num();
+
+            return true;
+        });
+
+        TMap<UActorPartitionSubsystem::FCellCoord, AInstancedFoliageActor*> IntersectingActors;
+
+        // Maintain a list of loaded 
+        TSet<FWorldPartitionReference> ActorReferences;
+
+        // Create/Reload new actors
+        FActorPartitionGridHelper::ForEachIntersectingCell(AInstancedFoliageActor::StaticClass(), InstanceBounds, World->PersistentLevel, [&](const UActorPartitionSubsystem::FCellCoord& InCellCoord, const FBox& InCellBounds)
+        {
+            // Create or Get (might have already been created and not GCed)
+            if (AInstancedFoliageActor* NewOrExistingActor = Cast<AInstancedFoliageActor>(ActorPartitionSubsystem->GetActor(AInstancedFoliageActor::StaticClass(), InCellCoord, /*bInCreate=*/true)))
+            {
+                IntersectingActors.Add(InCellCoord, NewOrExistingActor);
+            } 
+            else // If we haven't created or found the actor it means it was created/saved/unloaded and should be in the NewActorDescs so we can reload it
+            {
+                const FGuid& NewActorGuid = NewActors.FindChecked(InCellCoord);
+                FWorldPartitionReference Reference(WorldPartition, NewActorGuid);
+                ActorReferences.Add(Reference);
+                IntersectingActors.Add(InCellCoord, CastChecked<AInstancedFoliageActor>(Reference->GetActor()));
+            }
+
+            return true;
+        });
+                        
+        // Maintain list of modified IFAs so we don't save empty intersecting actors
+        TSet<AInstancedFoliageActor*> ModifiedIFAs;
+            
+        // Add Foliage to those actors
+        for (auto& InstancesPerFoliageType : FoliageToAdd)
+        {
+            for (const FFoliageInstance& Instance : InstancesPerFoliageType.Value)
+            {
+                // Here we should get back one of the IntersectingActors that was just created or reloaded so that is why we pass in false to bCreateIfNone
+                AInstancedFoliageActor* ModifiedIFA = AInstancedFoliageActor::Get(World, /*bCreateIfNone=*/false, World->PersistentLevel, Instance.Location);
+                check(ModifiedIFA);
+                ModifiedIFAs.Add(ModifiedIFA);
+                check(ModifiedIFA->GridSize == NewGridSize);
+                FFoliageInfo* NewFoliageInfo = nullptr;
+                UFoliageType* NewFoliageType = ModifiedIFA->AddFoliageType(InstancesPerFoliageType.Key, &NewFoliageInfo);
+                NewFoliageInfo->AddInstance(NewFoliageType, Instance);
+				NumInstancesProcessed++;
+            }
+        }
+
+        // Save only modified ones
+        for (auto Pair : IntersectingActors)
+        {
+            AInstancedFoliageActor* IntersectingActor = Pair.Value;
+            // No need to save actor as it hasn't changed
+            if (!ModifiedIFAs.Contains(IntersectingActor))
+            {
+                continue;
+            }
+
+            UPackage* Package = IntersectingActor->GetExternalPackage();
+            check(Package);
+
+            FString PackageFileName = SourceControlHelpers::PackageFilename(Package);
+                                
+            // Save package
+            if (!UPackage::SavePackage(Package, nullptr, RF_Standalone, *PackageFileName))
+            {
+                UE_LOG(LogWorldPartitionFoliageBuilder, Error, TEXT("Error saving package %s."), *Package->GetName());
+                return false;
+            }
+
+            // Actor Desc didn't exist yet which means we need to add the file to SCC and add it to the new actors descs and reference it so it gets unloaded
+            if (!NewActors.Contains(Pair.Key))
+            {
+                if (!SCCHelper.AddToSourceControl(Package))
+                {
+                    // It is possible the resave can't checkout everything. Continue processing.
+                    UE_LOG(LogWorldPartitionFoliageBuilder, Error, TEXT("Error adding package to source control %s."), *Package->GetName());
+                    return false;
+                }
+
+                FGuid NewActorGuid = IntersectingActor->GetActorGuid();
+                FWorldPartitionActorDesc* NewActorDesc = WorldPartition->GetActorDesc(NewActorGuid);
+                check(NewActorDesc);
+                NewActors.Add(Pair.Key, NewActorGuid);
+
+                // Add ref so actor gets unloaded 
+                FWorldPartitionReference Reference(WorldPartition, NewActorGuid);
+                ActorReferences.Add(Reference);
+            }
+                
+            UE_LOG(LogWorldPartitionFoliageBuilder, Display, TEXT("Saved foliage actor %s (%s)."), *IntersectingActor->GetName(), *Package->GetName());
+        }
+					
+		return true;
+	}, /*bGCPerActor=*/ false);
+
+	check(NumInstances == NumInstancesProcessed);
+
+	// Checkout World for WorldSettings change
+	UPackage* WorldSettingsPackage = World->GetWorldSettings()->GetPackage();
+	if (!SCCHelper.Checkout(WorldSettingsPackage))
+	{
+		UE_LOG(LogWorldPartitionFoliageBuilder, Error, TEXT("Error checking out package %s."), *WorldSettingsPackage->GetName());
+		return false;
+	}
+
+	// Save World 
+	{
+		FString PackageFileName = SourceControlHelpers::PackageFilename(WorldSettingsPackage);
+		if (!UPackage::SavePackage(World->GetPackage(), nullptr, RF_Standalone, *PackageFileName))
+		{
+			UE_LOG(LogWorldPartitionFoliageBuilder, Error, TEXT("Error saving package %s."), *WorldSettingsPackage->GetName());
+			return false;
+		}
+	}
+
+	// Delete old IFAs
+	for (const FGuid& ActorGuid : ExistingActors)
+	{
+		FWorldPartitionActorDesc* ActorDesc = WorldPartition->GetActorDesc(ActorGuid);
+		check(ActorDesc);
+		const FString PackageToDelete = ActorDesc->GetActorPackage().ToString();
+		if (!SCCHelper.Delete(PackageToDelete))
+		{
+			UE_LOG(LogWorldPartitionFoliageBuilder, Error, TEXT("Error deleting package %s."), *PackageToDelete);
+			return false;
+		}
+	}
+
+	return true;
+}
+
diff --git a/Engine/Source/Editor/UnrealEd/Public/WorldPartition/WorldPartitionFoliageBuilder.h b/Engine/Source/Editor/UnrealEd/Public/WorldPartition/WorldPartitionFoliageBuilder.h
new file mode 100644
--- /dev/null
+++ b/Engine/Source/Editor/UnrealEd/Public/WorldPartition/WorldPartitionFoliageBuilder.h
@@ -0,0 +1,23 @@
+// Copyright Epic Games, Inc. All Rights Reserved.
+
+#pragma once
+
+#include "WorldPartition/WorldPartitionBuilder.h"
+#include "WorldPartitionFoliageBuilder.generated.h"
+
+// Example Command Line: ProjectName MapName -run=WorldPartitionBuilderCommandlet -SCCProvider=Perforce -Builder=WorldPartitionFoliageBuilder -NewGridSize=Value
+UCLASS()
+class UWorldPartitionFoliageBuilder : public UWorldPartitionBuilder
+{
+	GENERATED_UCLASS_BODY()
+
+public:
+	// UWorldPartitionBuilder interface begin
+	virtual bool RequiresCommandletRendering() const override { return false; }
+	virtual ELoadingMode GetLoadingMode() const override { return ELoadingMode::Custom; }
+	virtual bool RunInternal(UWorld* World, const FBox& Bounds, FPackageSourceControlHelper& PackageHelper) override;
+	// UWorldPartitionBuilder interface end
+
+private:
+	uint32 NewGridSize;
+};
\ No newline at end of file
diff --git a/Engine/Source/Runtime/Engine/Private/WorldPartition/WorldPartitionHelpers.cpp b/Engine/Source/Runtime/Engine/Private/WorldPartition/WorldPartitionHelpers.cpp
--- a/Engine/Source/Runtime/Engine/Private/WorldPartition/WorldPartitionHelpers.cpp
+++ b/Engine/Source/Runtime/Engine/Private/WorldPartition/WorldPartitionHelpers.cpp
@@ -32,10 +32,9 @@ void FWorldPartitionHelpers::ForEachActorDesc(UWorldPartition* WorldPartition, T
 	}
 }
 
-void FWorldPartitionHelpers::ForEachActorWithLoading(UWorldPartition* WorldPartition, TSubclassOf<AActor> ActorClass, TFunctionRef<bool(const FWorldPartitionActorDesc*)> Predicate)
+namespace WorldPartitionHelpers
 {
-	// Recursive loading of references 
-	TFunction<void(const FGuid&, TMap<FGuid, FWorldPartitionReference>&)> LoadReferences = [WorldPartition, &LoadReferences](const FGuid& ActorGuid, TMap<FGuid, FWorldPartitionReference>& InOutActorReferences)
+	void LoadReferences(UWorldPartition* WorldPartition, const FGuid& ActorGuid, TMap<FGuid, FWorldPartitionReference>& InOutActorReferences)
 	{
 		if (InOutActorReferences.Contains(ActorGuid))
 		{
@@ -45,37 +44,19 @@ void FWorldPartitionHelpers::ForEachActorWithLoading(UWorldPartition* WorldParti
 		if (const FWorldPartitionActorDesc* ActorDesc = WorldPartition->GetActorDesc(ActorGuid))
 		{
 			InOutActorReferences.Emplace(ActorGuid);
-			
+
 			for (FGuid ReferenceGuid : ActorDesc->GetReferences())
 			{
-				LoadReferences(ReferenceGuid, InOutActorReferences);
+				LoadReferences(WorldPartition, ReferenceGuid, InOutActorReferences);
 			}
 
 			InOutActorReferences[ActorGuid] = FWorldPartitionReference(WorldPartition, ActorGuid);
 		}
-	};
-
-	TMap<FGuid, FWorldPartitionReference> ActorReferences;
-	TArray<FGuid> ActorsToProcess;
-
-	auto DoCollectGarbage = [&ActorReferences]()
-	{
-		ActorReferences.Empty();
-
-		const FPlatformMemoryStats MemStatsBefore = FPlatformMemory::GetStats();
-		CollectGarbage(RF_NoFlags, true);
-		const FPlatformMemoryStats MemStatsAfter = FPlatformMemory::GetStats();
-
-		UE_LOG(LogWorldPartition, Log, TEXT("GC Performed - Available Physical: %.2fGB, Available Virtual: %.2fGB"),
-			(int64)MemStatsAfter.AvailablePhysical / (1024.0 * 1024.0 * 1024.0),
-			(int64)MemStatsAfter.AvailableVirtual / (1024.0 * 1024.0 * 1024.0)
-		);
-	};
+	}
 
-	ForEachActorDesc(WorldPartition, ActorClass, [&](const FWorldPartitionActorDesc* ActorDesc)
+	bool ForEachActorWithLoadingBody(const FGuid& ActorGuid, UWorldPartition* WorldPartition, TFunctionRef<bool(const FWorldPartitionActorDesc*)> Predicate, bool bGCPerActor, TMap<FGuid, FWorldPartitionReference>& ActorReferences)
 	{
-		const FGuid& ActorGuid = ActorDesc->GetGuid();
-		LoadReferences(ActorGuid, ActorReferences);
+		WorldPartitionHelpers::LoadReferences(WorldPartition, ActorGuid, ActorReferences);
 
 		FWorldPartitionReference ActorReference(WorldPartition, ActorGuid);
 		if (!Predicate(ActorReference.Get()))
@@ -83,14 +64,42 @@ void FWorldPartitionHelpers::ForEachActorWithLoading(UWorldPartition* WorldParti
 			return false;
 		}
 
-		if (HasExceededMaxMemory())
+		if (bGCPerActor || FWorldPartitionHelpers::HasExceededMaxMemory())
 		{
-			DoCollectGarbage();
+			ActorReferences.Empty();
+			FWorldPartitionHelpers::DoCollectGarbage();
 		}
 
 		return true;
+	}
+}
+
+void FWorldPartitionHelpers::ForEachActorWithLoading(UWorldPartition* WorldPartition, TSubclassOf<AActor> ActorClass, TFunctionRef<bool(const FWorldPartitionActorDesc*)> Predicate, bool bGCPerActor)
+{
+	TMap<FGuid, FWorldPartitionReference> ActorReferences;
+		
+	ForEachActorDesc(WorldPartition, ActorClass, [&](const FWorldPartitionActorDesc* ActorDesc)
+	{
+		return WorldPartitionHelpers::ForEachActorWithLoadingBody(ActorDesc->GetGuid(), WorldPartition, Predicate, bGCPerActor, ActorReferences);
 	});
 
+	ActorReferences.Empty();
+	DoCollectGarbage();
+}
+
+void FWorldPartitionHelpers::ForEachActorWithLoading(UWorldPartition* WorldPartition, const TArray<FGuid>& ActorGuids, TFunctionRef<bool(const FWorldPartitionActorDesc*)> Predicate, bool bGCPerActor)
+{
+	TMap<FGuid, FWorldPartitionReference> ActorReferences;
+
+	for(const FGuid& ActorGuid : ActorGuids)
+	{
+		if (!WorldPartitionHelpers::ForEachActorWithLoadingBody(ActorGuid, WorldPartition, Predicate, bGCPerActor, ActorReferences))
+		{
+			break;
+		}
+	}
+
+	ActorReferences.Empty();
 	DoCollectGarbage();
 }
 
diff --git a/Engine/Source/Runtime/Engine/Public/WorldPartition/WorldPartitionHelpers.h b/Engine/Source/Runtime/Engine/Public/WorldPartition/WorldPartitionHelpers.h
--- a/Engine/Source/Runtime/Engine/Public/WorldPartition/WorldPartitionHelpers.h
+++ b/Engine/Source/Runtime/Engine/Public/WorldPartition/WorldPartitionHelpers.h
@@ -32,7 +32,8 @@ public:
 
 	static void ForEachIntersectingActorDesc(UWorldPartition* WorldPartition, const FBox& Box, TSubclassOf<AActor> ActorClass, TFunctionRef<bool(const FWorldPartitionActorDesc*)> Func);
 	static void ForEachActorDesc(UWorldPartition* WorldPartition, TSubclassOf<AActor> ActorClass, TFunctionRef<bool(const FWorldPartitionActorDesc*)> Func);
-	static void ForEachActorWithLoading(UWorldPartition* WorldPartition, TSubclassOf<AActor> ActorClass, TFunctionRef<bool(const FWorldPartitionActorDesc*)> Func);
+	static void ForEachActorWithLoading(UWorldPartition* WorldPartition, TSubclassOf<AActor> ActorClass, TFunctionRef<bool(const FWorldPartitionActorDesc*)> Func, bool bGCPerActor = false);
+	static void ForEachActorWithLoading(UWorldPartition* WorldPartition, const TArray<FGuid>& ActorGuids, TFunctionRef<bool(const FWorldPartitionActorDesc*)> Func, bool bGCPerActor = false);
 
 	static bool HasExceededMaxMemory();
 	static void DoCollectGarbage();
