Add a new Blueprint editor workflow setting to allow the user to choose how to restore breakpoints on Blueprint asset reload.

diff --git a/Engine/Source/Editor/BlueprintGraph/Private/BlueprintEditorSettings.cpp b/Engine/Source/Editor/BlueprintGraph/Private/BlueprintEditorSettings.cpp
--- a/Engine/Source/Editor/BlueprintGraph/Private/BlueprintEditorSettings.cpp
+++ b/Engine/Source/Editor/BlueprintGraph/Private/BlueprintEditorSettings.cpp
@@ -34,6 +34,7 @@ UBlueprintEditorSettings::UBlueprintEditorSettings(const FObjectInitializer& Obj
 	, bDoubleClickNavigatesToParent(true)
 	, bEnableTypePromotion(true)
 	, TypePromotionPinBlacklist { UEdGraphSchema_K2::PC_String, UEdGraphSchema_K2::PC_Text }
+	, BreakpointReloadMethod(EBlueprintBreakpointReloadMethod::RestoreAll)
 	// Experimental
 	, bEnableNamespaceFilteringFeatures(false)
 	, bEnableNamespaceImportingFeatures(false)
diff --git a/Engine/Source/Editor/BlueprintGraph/Public/BlueprintEditorSettings.h b/Engine/Source/Editor/BlueprintGraph/Public/BlueprintEditorSettings.h
--- a/Engine/Source/Editor/BlueprintGraph/Public/BlueprintEditorSettings.h
+++ b/Engine/Source/Editor/BlueprintGraph/Public/BlueprintEditorSettings.h
@@ -135,6 +135,10 @@ public:
 	UPROPERTY(config, EditAnywhere, Category = Workflow, meta=(EditCondition="bEnableTypePromotion"))
 	TSet<FName> TypePromotionPinBlacklist;
 
+	/** How to handle previously-set breakpoints on reload. */
+	UPROPERTY(config, EditAnywhere, Category = Workflow)
+	EBlueprintBreakpointReloadMethod BreakpointReloadMethod;
+
 // Experimental
 public:
 	/** Whether to enable namespace filtering features in the Blueprint editor */
diff --git a/Engine/Source/Editor/UnrealEd/Private/Kismet2/KismetDebugUtilities.cpp b/Engine/Source/Editor/UnrealEd/Private/Kismet2/KismetDebugUtilities.cpp
--- a/Engine/Source/Editor/UnrealEd/Private/Kismet2/KismetDebugUtilities.cpp
+++ b/Engine/Source/Editor/UnrealEd/Private/Kismet2/KismetDebugUtilities.cpp
@@ -1187,6 +1187,47 @@ FBlueprintBreakpoint* FKismetDebugUtilities::FindBreakpointForNode(const UEdGrap
 	);
 }
 
+void FKismetDebugUtilities::RestoreBreakpointsOnLoad(const UBlueprint* Blueprint)
+{
+	if (!Blueprint)
+	{
+		return;
+	}
+
+	// Remove stale breakpoints
+	RemoveBreakpointsByPredicate(
+		Blueprint,
+		[Blueprint](const FBlueprintBreakpoint& Breakpoint)
+		{
+			const UEdGraphNode* const Location = Breakpoint.GetLocation();
+			return !Location || !Location->IsIn(Blueprint);
+		}
+	);
+
+	// Restore breakpoints based on preferred method
+	const UBlueprintEditorSettings* BlueprintEditorSettings = GetDefault<UBlueprintEditorSettings>();
+	switch (BlueprintEditorSettings->BreakpointReloadMethod)
+	{
+	case EBlueprintBreakpointReloadMethod::RestoreAllAndDisable:
+		ForeachBreakpoint(
+			Blueprint,
+			[](FBlueprintBreakpoint& Breakpoint)
+			{
+				SetBreakpointEnabled(Breakpoint, false);
+			}
+		);
+		break;
+
+	case EBlueprintBreakpointReloadMethod::DiscardAll:
+		ClearBreakpoints(Blueprint);
+		break;
+
+	case EBlueprintBreakpointReloadMethod::RestoreAll:
+	default:
+		break;
+	}
+}
+
 bool FKismetDebugUtilities::HasDebuggingData(const UBlueprint* Blueprint)
 {
 	return Cast<UBlueprintGeneratedClass>(*Blueprint->GeneratedClass)->GetDebugData().IsValid();
diff --git a/Engine/Source/Editor/UnrealEd/Public/Kismet2/KismetDebugUtilities.h b/Engine/Source/Editor/UnrealEd/Public/Kismet2/KismetDebugUtilities.h
--- a/Engine/Source/Editor/UnrealEd/Public/Kismet2/KismetDebugUtilities.h
+++ b/Engine/Source/Editor/UnrealEd/Public/Kismet2/KismetDebugUtilities.h
@@ -17,6 +17,21 @@ template<typename ElementType> class TSimpleRingBuffer;
 
 DECLARE_LOG_CATEGORY_EXTERN(LogBlueprintDebug, Log, All);
 
+//////////////////////////////////////////////////////////////////////////
+// EBlueprintBreakpointReloadMethod
+
+/** Indicates how to restore breakpoints when a Blueprint asset is reloaded. */
+UENUM()
+enum class EBlueprintBreakpointReloadMethod
+{
+	/** Restore all breakpoints and keep their saved enabled/disabled state. */
+	RestoreAll,
+	/** Restore all breakpoints and disable on reload. */
+	RestoreAllAndDisable,
+	/** Discard all breakpoints on reload. */
+	DiscardAll
+};
+
 //////////////////////////////////////////////////////////////////////////
 // FKismetTraceSample
 
@@ -266,6 +281,9 @@ public:
 
 	/** Queries whether a blueprint has breakpoints in it */
 	static bool BlueprintHasBreakpoints(const UBlueprint* Blueprint);
+
+	/** Handles breakpoint validation/restoration after loading the given Blueprint */
+	static void RestoreBreakpointsOnLoad(const UBlueprint* Blueprint);
 	
 	// Blueprint utils 
 
diff --git a/Engine/Source/Runtime/Engine/Private/Blueprint.cpp b/Engine/Source/Runtime/Engine/Private/Blueprint.cpp
--- a/Engine/Source/Runtime/Engine/Private/Blueprint.cpp
+++ b/Engine/Source/Runtime/Engine/Private/Blueprint.cpp
@@ -706,15 +706,8 @@ void UBlueprint::PostLoad()
 	FBlueprintEditorUtils::PurgeNullGraphs(this);
 
 #if WITH_EDITOR
-	// Remove stale breakpoints
-	FKismetDebugUtilities::RemoveBreakpointsByPredicate(
-		this,
-		[this](const FBlueprintBreakpoint& Breakpoint)
-		{
-			const UEdGraphNode* const Location = Breakpoint.GetLocation();
-			return !Location || !Location->IsIn(this);
-		}
-	);
+	// Restore breakpoints for this Blueprint
+	FKismetDebugUtilities::RestoreBreakpointsOnLoad(this);
 # endif
 	Breakpoints_DEPRECATED.Empty();
 
