Fix dangling lambda pointer due to being stored in TFunctionRef.
#rnx

diff --git a/Engine/Source/Editor/UnrealEd/Private/Cooker/AsyncIODeleteTest.cpp b/Engine/Source/Editor/UnrealEd/Private/Cooker/AsyncIODeleteTest.cpp
--- a/Engine/Source/Editor/UnrealEd/Private/Cooker/AsyncIODeleteTest.cpp
+++ b/Engine/Source/Editor/UnrealEd/Private/Cooker/AsyncIODeleteTest.cpp
@@ -38,7 +38,7 @@ bool FAsyncIODeleteTest::RunTest(const FString& Parameters)
 		FFileHelper::SaveStringToFile(TestText, *TestFile1);
 		FileManager.MakeDirectory(*TestDir1, bApplyToTreeTrue);
 	};
-	IPlatformFile::FDirectoryVisitorFunc CountFilesAndDirs = [&NumFiles, &NumDirs](const TCHAR* VisitFilename, bool VisitIsDir)
+	auto CountFilesAndDirs = [&NumFiles, &NumDirs](const TCHAR* VisitFilename, bool VisitIsDir)
 	{
 		NumFiles += VisitIsDir ? 1 : 0;
 		NumDirs += VisitIsDir ? 0 : 1;
