CookOnTheFly - handle already-cooked packages from iterative builds.
#rnx

diff --git a/Engine/Source/Developer/IoStoreUtilities/Private/ZenStoreWriter.cpp b/Engine/Source/Developer/IoStoreUtilities/Private/ZenStoreWriter.cpp
--- a/Engine/Source/Developer/IoStoreUtilities/Private/ZenStoreWriter.cpp
+++ b/Engine/Source/Developer/IoStoreUtilities/Private/ZenStoreWriter.cpp
@@ -647,7 +647,7 @@ void FZenStoreWriter::CommitPackage(const FCommitPackageInfo& Info)
 
 	IPackageStoreWriter::FCommitEventArgs CommitEventArgs;
 
-	CommitEventArgs.PlatformName	= FName(*TargetPlatform.PlatformName()),
+	CommitEventArgs.PlatformName	= FName(*TargetPlatform.PlatformName());
 	CommitEventArgs.PackageName		= Info.PackageName;
 	CommitEventArgs.EntryIndex		= INDEX_NONE;
 	
@@ -666,16 +666,25 @@ void FZenStoreWriter::CommitPackage(const FCommitPackageInfo& Info)
 			check(PackageState->PackageData.IsValid);
 
 			FPackageDataEntry& PkgData = PackageState->PackageData;
-			
+
 			const int64 PkgDiskSize = PkgData.Payload.DataSize();
-			
+
 			FCbAttachment PkgDataAttachment(IoBufferToSharedBuffer(PkgData.Payload));
 			Pkg.AddAttachment(PkgDataAttachment);
 
-			CommitEventArgs.EntryIndex = PackageStoreEntries.Num();
-			PackageStoreEntries.Add(PkgData.PackageStoreEntry);
-			FOplogCookInfo& CookInfo = CookedPackagesInfo.Add_GetRef(FOplogCookInfo{
-				FCookedPackageInfo { Info.PackageName, IoHashToMD5(PkgDataAttachment.GetHash()), Info.PackageGuid, PkgDiskSize } });
+			CommitEventArgs.EntryIndex = PackageNameToIndex.FindOrAdd(Info.PackageName, PackageStoreEntries.Num());
+			if (CommitEventArgs.EntryIndex == PackageStoreEntries.Num())
+			{
+				PackageStoreEntries.Emplace();
+				CookedPackagesInfo.Emplace();
+			}
+
+			PackageStoreEntries[CommitEventArgs.EntryIndex] = PkgData.PackageStoreEntry;
+			FOplogCookInfo& CookInfo = CookedPackagesInfo[CommitEventArgs.EntryIndex];
+			CookInfo = FOplogCookInfo{
+					FCookedPackageInfo { Info.PackageName, IoHashToMD5(PkgDataAttachment.GetHash()), Info.PackageGuid, PkgDiskSize } };
+			CookInfo.bUpToDate = true;
+
 			int32 NumAttachments = Info.Attachments.Num();
 			TArray<FCbAttachment, TInlineAllocator<2>> CbAttachments;
 			if (NumAttachments)
@@ -702,8 +711,6 @@ void FZenStoreWriter::CommitPackage(const FCommitPackageInfo& Info)
 				}
 			}
 
-			PackageNameToIndex.Add(Info.PackageName, CookedPackagesInfo.Num()-1);
-
 			FCbWriter PackageObj;
 			PackageObj.BeginObject();
 			FString PackageNameKey = Info.PackageName.ToString();
@@ -800,10 +807,10 @@ void FZenStoreWriter::CommitPackage(const FCommitPackageInfo& Info)
 		PendingPackages.Remove(Info.PackageName);
 	}
 }
-void FZenStoreWriter::GetEntries(TFunction<void(TArrayView<const FPackageStoreEntryResource>)>&& Callback)
+void FZenStoreWriter::GetEntries(TFunction<void(TArrayView<const FPackageStoreEntryResource>, TArrayView<const FOplogCookInfo>)>&& Callback)
 {
 	FReadScopeLock _(PackagesLock);
-	Callback(PackageStoreEntries);
+	Callback(PackageStoreEntries, CookedPackagesInfo);
 }
 
 void FZenStoreWriter::Flush()
@@ -912,6 +919,36 @@ void FZenStoreWriter::RemoveCookedPackages()
 	PackageNameToIndex.Empty();
 }
 
+void FZenStoreWriter::MarkPackagesUpToDate(TArrayView<const FName> UpToDatePackages)
+{
+	TRACE_CPUPROFILER_EVENT_SCOPE(FZenStoreWriter::MarkPackagesUpToDate);
+
+	IPackageStoreWriter::FMarkUpToDateEventArgs MarkUpToDateEventArgs;
+
+	MarkUpToDateEventArgs.PackageIndexes.Reserve(UpToDatePackages.Num());
+
+	{
+		FWriteScopeLock _(PackagesLock);
+		for (FName PackageName : UpToDatePackages)
+		{
+			int32* Index = PackageNameToIndex.Find(PackageName);
+			if (!Index)
+			{
+				UE_LOG(LogZenStoreWriter, Error, TEXT("MarkPackagesUpToDate called with package %s that is not in the oplog."),
+					*PackageName.ToString());
+				continue;
+			}
+
+			MarkUpToDateEventArgs.PackageIndexes.Add(*Index);
+			CookedPackagesInfo[*Index].bUpToDate = true;
+		}
+	}
+	if (MarkUpToDateEventArgs.PackageIndexes.Num())
+	{
+		BroadcastMarkUpToDate(MarkUpToDateEventArgs);
+	}
+}
+
 void FZenStoreWriter::CreateProjectMetaData(FCbPackage& Pkg, FCbWriter& PackageObj, bool bGenerateContainerHeader)
 {
 	// File Manifest
@@ -1004,3 +1041,18 @@ void FZenStoreWriter::BroadcastCommit(IPackageStoreWriter::FCommitEventArgs& Eve
 		CommitEvent.Broadcast(EventArgs);
 	}
 }
+
+void FZenStoreWriter::BroadcastMarkUpToDate(IPackageStoreWriter::FMarkUpToDateEventArgs& EventArgs)
+{
+	FScopeLock CommitEventLock(&CommitEventCriticalSection);
+
+	if (MarkUpToDateEvent.IsBound())
+	{
+		FReadScopeLock EntriesLock(PackagesLock);
+		EventArgs.PlatformName = FName(*TargetPlatform.PlatformName());
+		EventArgs.Entries = PackageStoreEntries;
+		EventArgs.CookInfos = CookedPackagesInfo;
+		MarkUpToDateEvent.Broadcast(EventArgs);
+	}
+}
+
diff --git a/Engine/Source/Developer/IoStoreUtilities/Public/PackageStoreWriter.h b/Engine/Source/Developer/IoStoreUtilities/Public/PackageStoreWriter.h
--- a/Engine/Source/Developer/IoStoreUtilities/Public/PackageStoreWriter.h
+++ b/Engine/Source/Developer/IoStoreUtilities/Public/PackageStoreWriter.h
@@ -14,10 +14,23 @@ public:
 		return this;
 	}
 
+	struct FOplogCookInfo
+	{
+		struct FAttachment
+		{
+			const UTF8CHAR* Key;
+			FIoHash Hash;
+		};
+
+		ICookedPackageWriter::FCookedPackageInfo CookInfo;
+		TArray<FAttachment> Attachments;
+		bool bUpToDate = false;
+	};
+
 	/**
 	 * Returns all cooked package store entries.
 	 */
-	virtual void GetEntries(TFunction<void(TArrayView<const FPackageStoreEntryResource>)>&&) = 0;
+	virtual void GetEntries(TFunction<void(TArrayView<const FPackageStoreEntryResource>, TArrayView<const FOplogCookInfo>)>&&) = 0;
 
 	/**
 	 * Package commit event arguments
@@ -28,6 +41,7 @@ public:
 		FName PackageName;
 		int32 EntryIndex = INDEX_NONE;
 		TArrayView<const FPackageStoreEntryResource> Entries;
+		TArrayView<const FOplogCookInfo> CookInfos;
 		TArray<FAdditionalFileInfo> AdditionalFiles;
 	};
 
@@ -36,4 +50,18 @@ public:
 	 */
 	DECLARE_EVENT_OneParam(IPackageStoreWriter, FCommitEvent, const FCommitEventArgs&);
 	virtual FCommitEvent& OnCommit() = 0;
+
+	struct FMarkUpToDateEventArgs
+	{
+		FName PlatformName;
+		TArray<int32> PackageIndexes;
+		TArrayView<const FPackageStoreEntryResource> Entries;
+		TArrayView<const FOplogCookInfo> CookInfos;
+		TArray<FAdditionalFileInfo> AdditionalFiles;
+	};
+	/**
+	 * Broadcasted after a set of packages have been found to be up to date.
+	 */
+	DECLARE_EVENT_OneParam(IPackageStoreWriter, FMarkUpToDateEvent, const FMarkUpToDateEventArgs&);
+	virtual FMarkUpToDateEvent& OnMarkUpToDate() = 0;
 };
\ No newline at end of file
diff --git a/Engine/Source/Developer/IoStoreUtilities/Public/ZenStoreWriter.h b/Engine/Source/Developer/IoStoreUtilities/Public/ZenStoreWriter.h
--- a/Engine/Source/Developer/IoStoreUtilities/Public/ZenStoreWriter.h
+++ b/Engine/Source/Developer/IoStoreUtilities/Public/ZenStoreWriter.h
@@ -44,13 +44,19 @@ public:
 	IOSTOREUTILITIES_API virtual void BeginCook(const FCookInfo& Info) override;
 	IOSTOREUTILITIES_API virtual void EndCook() override;
 
-	IOSTOREUTILITIES_API virtual void GetEntries(TFunction<void(TArrayView<const FPackageStoreEntryResource>)>&& Callback) override;
+	IOSTOREUTILITIES_API virtual void GetEntries(TFunction<void(TArrayView<const FPackageStoreEntryResource>, TArrayView<const FOplogCookInfo>)>&& Callback) override;
 
 	DECLARE_DERIVED_EVENT(FZenStoreWriter, IPackageStoreWriter::FCommitEvent, FCommitEvent);
 	IOSTOREUTILITIES_API virtual FCommitEvent& OnCommit() override
 	{
 		return CommitEvent;
 	}
+	DECLARE_DERIVED_EVENT(FZenStoreWriter, IPackageStoreWriter::FMarkUpToDateEvent, FMarkUpToDateEvent);
+	IOSTOREUTILITIES_API virtual FMarkUpToDateEvent& OnMarkUpToDate() override
+	{
+		return MarkUpToDateEvent;
+	}
+
 
 	virtual void Flush() override;
 
@@ -60,10 +66,12 @@ public:
 	IOSTOREUTILITIES_API virtual FCbObject GetOplogAttachment(FName PackageName, FUtf8StringView AttachmentKey) override;
 	IOSTOREUTILITIES_API virtual void RemoveCookedPackages(TArrayView<const FName> PackageNamesToRemove) override;
 	IOSTOREUTILITIES_API virtual void RemoveCookedPackages() override;
+	IOSTOREUTILITIES_API virtual void MarkPackagesUpToDate(TArrayView<const FName> UpToDatePackages) override;
 
 private:
 	void CreateProjectMetaData(FCbPackage& Pkg, FCbWriter& PackageObj, bool bGenerateContainerHeader);
 	void BroadcastCommit(IPackageStoreWriter::FCommitEventArgs& EventArgs);
+	void BroadcastMarkUpToDate(IPackageStoreWriter::FMarkUpToDateEventArgs& EventArgs);
 
 	struct FBulkDataEntry
 	{
@@ -104,18 +112,6 @@ private:
 		double TotalRequestTime = 0.0;
 	};
 
-	struct FOplogCookInfo
-	{
-		struct FAttachment
-		{
-			const UTF8CHAR* Key;
-			FIoHash Hash;
-		};
-
-		ICookedPackageWriter::FCookedPackageInfo CookInfo;
-		TArray<FAttachment> Attachments;
-	};
-
 	FRWLock								PackagesLock;
 	TMap<FName,FPendingPackageState>	PendingPackages;
 	TUniquePtr<UE::FZenStoreHttpClient>	HttpClient;
@@ -135,6 +131,7 @@ private:
 	
 	FCriticalSection					CommitEventCriticalSection;
 	FCommitEvent						CommitEvent;
+	FMarkUpToDateEvent					MarkUpToDateEvent;
 
 	class FZenStoreHttpQueue;
 	TUniquePtr<FZenStoreHttpQueue>		HttpQueue;
diff --git a/Engine/Source/Editor/UnrealEd/Private/CookOnTheFlyServer.cpp b/Engine/Source/Editor/UnrealEd/Private/CookOnTheFlyServer.cpp
--- a/Engine/Source/Editor/UnrealEd/Private/CookOnTheFlyServer.cpp
+++ b/Engine/Source/Editor/UnrealEd/Private/CookOnTheFlyServer.cpp
@@ -5778,6 +5778,7 @@ void UCookOnTheFlyServer::PopulateCookedPackages(TArrayView<const ITargetPlatfor
 					PackageData->SetPlatformCooked(TargetPlatform, true /* bSucceeded */);
 				}
 			}
+			PackageWriter.MarkPackagesUpToDate(Difference.IdenticalCookedPackages.Array());
 			for (FName UncookedPackage : Difference.IdenticalUncookedPackages)
 			{
 				// Mark this failed-to-cook package as cooked so that we don't unnecessarily try to cook it again
diff --git a/Engine/Source/Editor/UnrealEd/Private/Cooker/CookPackageData.cpp b/Engine/Source/Editor/UnrealEd/Private/Cooker/CookPackageData.cpp
--- a/Engine/Source/Editor/UnrealEd/Private/Cooker/CookPackageData.cpp
+++ b/Engine/Source/Editor/UnrealEd/Private/Cooker/CookPackageData.cpp
@@ -181,7 +181,7 @@ namespace Cook
 	}
 
 	void FPackageData::UpdateRequestData(const TConstArrayView<const ITargetPlatform*> InRequestedPlatforms,
-		bool bInIsUrgent, FCompletionCallback&& InCompletionCallback)
+		bool bInIsUrgent, FCompletionCallback&& InCompletionCallback, bool bAllowUpdateUrgency)
 	{
 		if (IsInProgress())
 		{
@@ -203,7 +203,7 @@ namespace Cook
 				}
 				SetPlatformsRequested(InRequestedPlatforms, true);
 			}
-			else if (bUrgencyChanged)
+			else if (bUrgencyChanged && bAllowUpdateUrgency)
 			{
 				SendToState(GetState(), ESendFlags::QueueAddAndRemove);
 			}
diff --git a/Engine/Source/Editor/UnrealEd/Private/Cooker/CookPackageData.h b/Engine/Source/Editor/UnrealEd/Private/Cooker/CookPackageData.h
--- a/Engine/Source/Editor/UnrealEd/Private/Cooker/CookPackageData.h
+++ b/Engine/Source/Editor/UnrealEd/Private/Cooker/CookPackageData.h
@@ -201,7 +201,7 @@ namespace UE::Cook
 		 * RequestData is set to the union of the previous and new values.
 		 */
 		void UpdateRequestData(const TConstArrayView<const ITargetPlatform*> InRequestedPlatforms, bool bInIsUrgent,
-			FCompletionCallback&& InCompletionCallback);
+			FCompletionCallback&& InCompletionCallback, bool bAllowUpdateUrgency=true);
 		/*
 		 * Set the given data onto the existing request data.
 		 * It is invalid to call with empty InRequestedPlatforms.
diff --git a/Engine/Source/Editor/UnrealEd/Private/Cooker/CookRequestCluster.cpp b/Engine/Source/Editor/UnrealEd/Private/Cooker/CookRequestCluster.cpp
--- a/Engine/Source/Editor/UnrealEd/Private/Cooker/CookRequestCluster.cpp
+++ b/Engine/Source/Editor/UnrealEd/Private/Cooker/CookRequestCluster.cpp
@@ -190,6 +190,13 @@ void FRequestCluster::Initialize(UCookOnTheFlyServer& COTFS)
 		bAllowSoftDependencies = !Options->bSkipSoftReferences;
 		bErrorOnEngineContentUse = Options->bErrorOnEngineContentUse;
 	}
+	else
+	{
+		// Do not queue soft-dependencies during CookOnTheFly; wait for them to be requested
+		// TODO: Report soft dependencies separately, and mark them as normal priority,
+		// and mark all hard dependencies as high priority in cook on the fly.
+		bAllowSoftDependencies = false;
+	}
 	bHybridIterativeEnabled = COTFS.bHybridIterativeEnabled;
 	bPreexploreDependenciesEnabled = COTFS.bPreexploreDependenciesEnabled;
 	if (bErrorOnEngineContentUse)
@@ -358,7 +365,7 @@ bool FRequestCluster::TryTakeOwnership(FPackageData& PackageData, bool bUrgent,
 				// This might steal it from another RequestCluster or from the UnclusteredRequests if it's in request
 				// Doing that steal is wasteful but okay; one of the RequestClusters will win it and keep it
 				PackageData.SendToState(EPackageState::Request, ESendFlags::QueueRemove);
-				PackageData.UpdateRequestData(GetPlatforms(), bUrgent, MoveTemp(CompletionCallback));
+				PackageData.UpdateRequestData(GetPlatforms(), bUrgent, MoveTemp(CompletionCallback), false /* bAllowUpdateUrgency */);
 			}
 			return true;
 		}
@@ -473,7 +480,7 @@ void FRequestCluster::FetchDependencies(const FCookerTimer& CookerTimer, bool& b
 	}
 
 	// Reverse the TransitiveRequests array to put it in LeafToRoot order, and push it into the output Requests
-	COOK_STAT(DetailedCookStats::NumPreloadedDependencies += TransitiveRequests.Num() - Requests.Num());
+	COOK_STAT(DetailedCookStats::NumPreloadedDependencies += FMath::Max(0,TransitiveRequests.Num() - Requests.Num()));
 	Requests.Reset(TransitiveRequests.Num());
 	for (int32 Index = TransitiveRequests.Num() - 1; Index >= 0; --Index)
 	{
@@ -727,6 +734,7 @@ void FRequestCluster::VisitPackageData(FPackageData* PackageData, TArray<FName>*
 					{
 						bIncrementIterativeCounter = true;
 						PackageData->SetPlatformCooked(TargetPlatform, true);
+						PackageWriter->MarkPackagesUpToDate({ PackageName });
 					}
 					OutDependencies->Append(OplogDataScratch.BuildDependencies);
 					if (bAllowSoftDependencies)
diff --git a/Engine/Source/Editor/UnrealEd/Private/Cooker/IoStoreCookOnTheFlyRequestManager.cpp b/Engine/Source/Editor/UnrealEd/Private/Cooker/IoStoreCookOnTheFlyRequestManager.cpp
--- a/Engine/Source/Editor/UnrealEd/Private/Cooker/IoStoreCookOnTheFlyRequestManager.cpp
+++ b/Engine/Source/Editor/UnrealEd/Private/Cooker/IoStoreCookOnTheFlyRequestManager.cpp
@@ -116,10 +116,12 @@ public:
 
 	FCompletedPackages MarkAsFailed(const FName& PackageName)
 	{
-		return MarkAsCompleted(PackageName, INDEX_NONE, TArrayView<const FPackageStoreEntryResource>());
+		FCompletedPackages ResultPackages;
+		MarkAsCompleted(PackageName, INDEX_NONE, TArrayView<const FPackageStoreEntryResource>(), ResultPackages);
+		return ResultPackages;
 	}
 
-	FCompletedPackages MarkAsCompleted(const FName& PackageName, const int32 EntryIndex, TArrayView<const FPackageStoreEntryResource> Entries)
+	void MarkAsCompleted(const FName& PackageName, const int32 EntryIndex, TArrayView<const FPackageStoreEntryResource> Entries, FCompletedPackages& OutPackages)
 	{
 		const int32 NumCookedPackages = CookedEntries.Num();
 		const int32 NumFailedPackages = FailedPackages.Num();
@@ -174,28 +176,25 @@ public:
 			AddWaitingPackage(Package);
 		}
 
-		FCompletedPackages CompletedPackages;
-		CompletedPackages.TotalCooked = CookedEntries.Num();
-		CompletedPackages.TotalFailed = FailedPackages.Num();
-		CompletedPackages.TotalTracked = Packages.Num();
+		OutPackages.TotalCooked = CookedEntries.Num();
+		OutPackages.TotalFailed = FailedPackages.Num();
+		OutPackages.TotalTracked = Packages.Num();
 
 		if (CookedEntries.Num() > NumCookedPackages || FailedPackages.Num() > NumFailedPackages)
 		{
 			for (int32 Idx = NumCookedPackages; Idx < CookedEntries.Num(); ++Idx)
 			{
-				CompletedPackages.CookedEntryIndices.Add(CookedEntries[Idx]);
+				OutPackages.CookedEntryIndices.Add(CookedEntries[Idx]);
 			}
 
 			for (int32 Idx = NumFailedPackages; Idx < FailedPackages.Num(); ++Idx)
 			{
-				CompletedPackages.FailedPackages.Add(FPackageId::FromName(FailedPackages[Idx]));
+				OutPackages.FailedPackages.Add(FPackageId::FromName(FailedPackages[Idx]));
 			}
 		}
-
-		return CompletedPackages;
 	}
 
-	void AddCompletedPackages(TArrayView<const FPackageStoreEntryResource> Entries)
+	void AddExistingPackages(TArrayView<const FPackageStoreEntryResource> Entries, TArrayView<const IPackageStoreWriter::FOplogCookInfo> CookInfos)
 	{
 		Packages.Reserve(Entries.Num());
 		CookedEntries.Reserve(Entries.Num());
@@ -203,16 +202,23 @@ public:
 		for (int32 EntryIndex = 0, Num = Entries.Num(); EntryIndex < Num; ++EntryIndex)
 		{
 			const FPackageStoreEntryResource& Entry = Entries[EntryIndex];
+			const IPackageStoreWriter::FOplogCookInfo& CookInfo = CookInfos[EntryIndex];
 			const FPackageId PackageId = Entry.GetPackageId();
 
 			FPackage& Package	= GetPackage(PackageId);
-
 			Package.PackageId	= PackageId;
 			Package.PackageName = Entry.PackageName;
-			Package.Status		= EPackageStatus::Cooked;
 			Package.EntryIndex	= EntryIndex;
 
-			CookedEntries.Add(Package.EntryIndex);
+			if (CookInfo.bUpToDate)
+			{
+				CookedEntries.Add(Package.EntryIndex);
+				Package.Status = EPackageStatus::Cooked;
+			}
+			else
+			{
+				Package.Status = EPackageStatus::None;
+			}
 		}
 	}
 
@@ -457,12 +463,14 @@ private:
 					check(PackageWriter); // This class should not be used except when COTFS is using an IPackageStoreWriter
 					Context->PackageWriter = PackageWriter;
 					
-					PackageWriter->GetEntries([&Context](TArrayView<const FPackageStoreEntryResource> Entries)
+					PackageWriter->GetEntries([&Context](TArrayView<const FPackageStoreEntryResource> Entries,
+						TArrayView<const IPackageStoreWriter::FOplogCookInfo> CookInfos)
 					{
-						Context->PackageTracker.AddCompletedPackages(Entries);
+						Context->PackageTracker.AddExistingPackages(Entries, CookInfos);
 					});
 
 					PackageWriter->OnCommit().AddRaw(this, &FIoStoreCookOnTheFlyRequestManager::OnPackageCooked);
+					PackageWriter->OnMarkUpToDate().AddRaw(this, &FIoStoreCookOnTheFlyRequestManager::OnPackagesMarkedUpToDate);
 				}
 
 				return true;
@@ -521,7 +529,9 @@ private:
 		
 		UE::PackageStore::Messaging::FPackageStoreData PackageStoreData;
 		
-		Context.PackageWriter->GetEntries([&Context, &PackageStoreData](TArrayView<const FPackageStoreEntryResource> Entries)
+		Context.PackageWriter->GetEntries(
+			[&Context, &PackageStoreData](TArrayView<const FPackageStoreEntryResource> Entries,
+				TArrayView<const IPackageStoreWriter::FOplogCookInfo> CookInfos)
 		{
 			FScopeLock _(&Context.CriticalSection);
 			
@@ -619,17 +629,44 @@ private:
 
 			ConnectionServer->BroadcastMessage(Message, EventArgs.PlatformName);
 		}
-		
+
 		FCookOnTheFlyPackageTracker::FCompletedPackages NewCompletedPackages;
 		{
 			TRACE_CPUPROFILER_EVENT_SCOPE(CookOnTheFly::MarkAsCooked);
-			
+
 			FPlatformContext& Context = GetContext(EventArgs.PlatformName);
 			FScopeLock _(&Context.CriticalSection);
 
-			NewCompletedPackages = Context.PackageTracker.MarkAsCompleted(EventArgs.PackageName, EventArgs.EntryIndex, EventArgs.Entries);
+			Context.PackageTracker.MarkAsCompleted(EventArgs.PackageName, EventArgs.EntryIndex, EventArgs.Entries, NewCompletedPackages);
 		}
+		BroadcastCompletedPackages(EventArgs.PlatformName, MoveTemp(NewCompletedPackages), EventArgs.Entries);
+	}
+
+	void OnPackagesMarkedUpToDate(const IPackageStoreWriter::FMarkUpToDateEventArgs& EventArgs)
+	{
+		TRACE_CPUPROFILER_EVENT_SCOPE(CookOnTheFly::OnPackagesMarkedUpToDate);
 
+		FCookOnTheFlyPackageTracker::FCompletedPackages NewCompletedPackages;
+		{
+			TRACE_CPUPROFILER_EVENT_SCOPE(CookOnTheFly::MarkAsCooked);
+
+			FPlatformContext& Context = GetContext(EventArgs.PlatformName);
+			FScopeLock _(&Context.CriticalSection);
+
+			for (int32 EntryIndex : EventArgs.PackageIndexes)
+			{
+				FName PackageName = EventArgs.Entries[EntryIndex].PackageName;
+				Context.PackageTracker.MarkAsCompleted(PackageName, EntryIndex, EventArgs.Entries, NewCompletedPackages);
+			}
+		}
+		BroadcastCompletedPackages(EventArgs.PlatformName, MoveTemp(NewCompletedPackages), EventArgs.Entries);
+	}
+
+	void BroadcastCompletedPackages(FName PlatformName,
+		FCookOnTheFlyPackageTracker::FCompletedPackages&& NewCompletedPackages, TArrayView<const FPackageStoreEntryResource> Entries)
+	{
+		using namespace UE::Cook;
+		using namespace UE::PackageStore::Messaging;
 		if (NewCompletedPackages.CookedEntryIndices.Num() || NewCompletedPackages.FailedPackages.Num())
 		{
 			TRACE_CPUPROFILER_EVENT_SCOPE(CookOnTheFly::SendCookedPackagesMessage);
@@ -642,7 +679,7 @@ private:
 			for (int32 EntryIndex : NewCompletedPackages.CookedEntryIndices)
 			{
 				check(EntryIndex != INDEX_NONE);
-				PackageStoreData.CookedPackages.Add(EventArgs.Entries[EntryIndex]);
+				PackageStoreData.CookedPackages.Add(Entries[EntryIndex]);
 			}
 
 			UE_LOG(LogCookOnTheFly, Verbose, TEXT("Sending '%s' message, Cooked='%d', Failed='%d', Server total='%d/%d/%d' (Tracked/Cooked/Failed)"),
@@ -655,7 +692,7 @@ private:
 
 			FCookOnTheFlyMessage Message(ECookOnTheFlyMessage::PackagesCooked);
 			Message.SetBodyTo<FPackagesCookedMessage>(FPackagesCookedMessage { MoveTemp(PackageStoreData) });
-			ConnectionServer->BroadcastMessage(Message, EventArgs.PlatformName);
+			ConnectionServer->BroadcastMessage(Message, PlatformName);
 		}
 	}
 
@@ -741,7 +778,8 @@ private:
 				PackageStoreData.TotalFailedPackages = Context.PackageTracker.GetFailedPackages().Num();
 				PackageStoreData.FailedPackages = MoveTemp(CompletedPackages.FailedPackages);
 
-				Context.PackageWriter->GetEntries([&CompletedPackages, &PackageStoreData](TArrayView<const FPackageStoreEntryResource> Entries)
+				Context.PackageWriter->GetEntries([&CompletedPackages, &PackageStoreData](TArrayView<const FPackageStoreEntryResource> Entries,
+					TArrayView<const IPackageStoreWriter::FOplogCookInfo> CookInfos)
 				{
 					for (int32 EntryIndex : CompletedPackages.CookedEntryIndices)
 					{
@@ -785,4 +823,4 @@ TUniquePtr<ICookOnTheFlyRequestManager> MakeIoStoreCookOnTheFlyRequestManager(IC
 
 }}
 
-#endif
+#endif
\ No newline at end of file
diff --git a/Engine/Source/Editor/UnrealEd/Private/Cooker/LooseCookedPackageWriter.cpp b/Engine/Source/Editor/UnrealEd/Private/Cooker/LooseCookedPackageWriter.cpp
--- a/Engine/Source/Editor/UnrealEd/Private/Cooker/LooseCookedPackageWriter.cpp
+++ b/Engine/Source/Editor/UnrealEd/Private/Cooker/LooseCookedPackageWriter.cpp
@@ -217,6 +217,10 @@ void FLooseCookedPackageWriter::RemoveCookedPackages(TArrayView<const FName> Pac
 	UncookedPathToCookedPath.Empty();
 }
 
+void FLooseCookedPackageWriter::MarkPackagesUpToDate(TArrayView<const FName> UpToDatePackages)
+{
+}
+
 void FLooseCookedPackageWriter::RemoveCookedPackages()
 {
 	DeleteSandboxDirectory();
diff --git a/Engine/Source/Editor/UnrealEd/Private/Cooker/LooseCookedPackageWriter.h b/Engine/Source/Editor/UnrealEd/Private/Cooker/LooseCookedPackageWriter.h
--- a/Engine/Source/Editor/UnrealEd/Private/Cooker/LooseCookedPackageWriter.h
+++ b/Engine/Source/Editor/UnrealEd/Private/Cooker/LooseCookedPackageWriter.h
@@ -45,6 +45,7 @@ public:
 	FCbObject GetOplogAttachment(FName PackageName, FUtf8StringView AttachmentKey) override;
 	void RemoveCookedPackages(TArrayView<const FName> PackageNamesToRemove) override;
 	void RemoveCookedPackages() override;
+	void MarkPackagesUpToDate(TArrayView<const FName> UpToDatePackages) override;
 	FAssetRegistryState* ReleasePreviousAssetRegistry() override;
 
 private:
diff --git a/Engine/Source/Runtime/Core/Public/Serialization/PackageWriter.h b/Engine/Source/Runtime/Core/Public/Serialization/PackageWriter.h
--- a/Engine/Source/Runtime/Core/Public/Serialization/PackageWriter.h
+++ b/Engine/Source/Runtime/Core/Public/Serialization/PackageWriter.h
@@ -211,6 +211,11 @@ public:
 	 */
 	virtual void RemoveCookedPackages() = 0;
 
+	/**
+	 * Signal the given cooked package(s) have been checked for changes and have not been modified since the last cook.
+	 */
+	virtual void MarkPackagesUpToDate(TArrayView<const FName> UpToDatePackages) = 0;
+
 	/**
 	 * If a previous-version AssetRegistryState was allocated during GetCookedPackages, drop the reference to it
 	 * and return it to the caller who takes responsibility for deleting it.
