Delete obsolete and unused CookedIterativeFile module

#rnx

diff --git a/Engine/Source/Editor/UnrealEd/Private/CookOnTheFlyServer.cpp b/Engine/Source/Editor/UnrealEd/Private/CookOnTheFlyServer.cpp
--- a/Engine/Source/Editor/UnrealEd/Private/CookOnTheFlyServer.cpp
+++ b/Engine/Source/Editor/UnrealEd/Private/CookOnTheFlyServer.cpp
@@ -300,31 +300,6 @@ public:
 		Cooker.PlatformManager->ReleaseCookOnTheFlyPlatform(PlatformName);
 	}
 
-	virtual void GetPrecookedFileList(const FName& PlatformName, UE::Cook::FPrecookedFileList& OutPrecookedFiles) override
-	{
-		ITargetPlatformManagerModule& TPM = GetTargetPlatformManagerRef();
-		const ITargetPlatform* TargetPlatform = TPM.FindTargetPlatform(PlatformName);
-		if (!TargetPlatform)
-		{
-			UE_LOG(LogCook, Warning, TEXT("Trying to get precooked file(s) on the fly for invalid platform '%s'"), *PlatformName.ToString());
-			return;
-		}
-
-		TArray<FName> CookedPlatformFiles;
-		Cooker.PackageDatas->GetCookedFileNamesForPlatform(TargetPlatform, CookedPlatformFiles, /* include failed */ true, /* include successful */ true);
-
-		for (const FName& CookedFile : CookedPlatformFiles)
-		{
-			const FString SandboxFilename = Cooker.ConvertToFullSandboxPath(CookedFile.ToString(), true, PlatformName.ToString());
-			if (IFileManager::Get().FileExists(*SandboxFilename))
-			{
-				continue;
-			}
-
-			OutPrecookedFiles.Add(CookedFile.ToString(),FDateTime::MinValue());
-		}
-	}
-
 	virtual void GetUnsolicitedFiles(const FName& PlatformName, const FString& Filename, const bool bIsCookable, TArray<FString>& OutUnsolicitedFiles) override
 	{
 		while (true)
diff --git a/Engine/Source/Editor/UnrealEd/Private/Cooker/CookOnTheFlyServerInterface.h b/Engine/Source/Editor/UnrealEd/Private/Cooker/CookOnTheFlyServerInterface.h
--- a/Engine/Source/Editor/UnrealEd/Private/Cooker/CookOnTheFlyServerInterface.h
+++ b/Engine/Source/Editor/UnrealEd/Private/Cooker/CookOnTheFlyServerInterface.h
@@ -79,14 +79,6 @@ public:
 	 */
 	virtual void RemovePlatform(const FName& PlatformName) = 0;
 
-	/**
-	 * Returns all asset file(s) that has been cooked.
-	 *
-	 * @param PlatformName The platform name.
-	 * @param OutPrecookedFiles Filename(s) returned to the caller.
-	 */
-	virtual void GetPrecookedFileList(const FName& PlatformName, FPrecookedFileList& OutPrecookedFiles) = 0;
-
 	/**
 	 * Returns all unsolicited files that has been produced has a result of a cook request.
 	 *
diff --git a/Engine/Source/Editor/UnrealEd/Private/Cooker/NetworkFileCookOnTheFlyRequestManager.cpp b/Engine/Source/Editor/UnrealEd/Private/Cooker/NetworkFileCookOnTheFlyRequestManager.cpp
--- a/Engine/Source/Editor/UnrealEd/Private/Cooker/NetworkFileCookOnTheFlyRequestManager.cpp
+++ b/Engine/Source/Editor/UnrealEd/Private/Cooker/NetworkFileCookOnTheFlyRequestManager.cpp
@@ -14,7 +14,6 @@ public:
 		: CookOnTheFlyServer(InCookOnTheFlyServer)
 	{
 		FileServerOptions.Delegates.NewConnectionDelegate			= FNewConnectionDelegate::CreateRaw(this, &FNetworkFileCookOnTheFlyRequestManager::OnNewConnection);
-		FileServerOptions.Delegates.InitialPrecookedListDelegate	= FInitialPrecookedListDelegate::CreateRaw(this, &FNetworkFileCookOnTheFlyRequestManager::OnGetPrecookedList);
 		FileServerOptions.Delegates.FileRequestDelegate				= FFileRequestDelegate::CreateRaw(this, &FNetworkFileCookOnTheFlyRequestManager::OnFileRequest);
 		FileServerOptions.Delegates.RecompileShadersDelegate		= FRecompileShadersDelegate::CreateRaw(this, &FNetworkFileCookOnTheFlyRequestManager::OnRecompileShaders);
 		FileServerOptions.Delegates.SandboxPathOverrideDelegate		= FSandboxPathDelegate::CreateRaw(this, &FNetworkFileCookOnTheFlyRequestManager::OnGetSandboxPath);
@@ -50,11 +49,6 @@ private:
 		return CookOnTheFlyServer.AddPlatform(FName(*PlatformName)) != nullptr;
 	}
 
-	void OnGetPrecookedList(const FString& PlatformName, TMap<FString, FDateTime>& PrecookedFileList)
-	{
-		CookOnTheFlyServer.GetPrecookedFileList(FName(*PlatformName), PrecookedFileList);
-	}
-
 	void OnFileRequest(FString& Filename, const FString& PlatformNameStr, TArray<FString>& UnsolicitedFiles)
 	{
 		using namespace UE::Cook;
diff --git a/Engine/Source/Runtime/CookedIterativeFile/CookedIterativeFile.Build.cs b/Engine/Source/Runtime/CookedIterativeFile/CookedIterativeFile.Build.cs
deleted file mode 100644
diff --git a/Engine/Source/Runtime/CookedIterativeFile/Private/CookedIterativeNetworkFile.cpp b/Engine/Source/Runtime/CookedIterativeFile/Private/CookedIterativeNetworkFile.cpp
deleted file mode 100644
diff --git a/Engine/Source/Runtime/CookedIterativeFile/Public/CookedIterativeNetworkFile.h b/Engine/Source/Runtime/CookedIterativeFile/Public/CookedIterativeNetworkFile.h
deleted file mode 100644
diff --git a/Engine/Source/Runtime/Launch/Launch.Build.cs b/Engine/Source/Runtime/Launch/Launch.Build.cs
--- a/Engine/Source/Runtime/Launch/Launch.Build.cs
+++ b/Engine/Source/Runtime/Launch/Launch.Build.cs
@@ -122,7 +122,6 @@ public class Launch : ModuleRules
 			PrivateDependencyModuleNames.AddRange(new string[] {
 					"NetworkFile",
 					"StreamingFile",
-					"CookedIterativeFile",
 					"AutomationWorker"
 			});
 		}
diff --git a/Engine/Source/Runtime/Launch/Private/LaunchEngineLoop.cpp b/Engine/Source/Runtime/Launch/Private/LaunchEngineLoop.cpp
--- a/Engine/Source/Runtime/Launch/Private/LaunchEngineLoop.cpp
+++ b/Engine/Source/Runtime/Launch/Private/LaunchEngineLoop.cpp
@@ -863,20 +863,9 @@ bool LaunchCheckForFileOverride(const TCHAR* CmdLine, bool& OutFileOverrideFound
 			FPlatformFileManager::Get().SetPlatformFile(*CurrentPlatformFile);
 		}
 
-		bool bShouldUseCookedIterativeFile = false;
-		if ( !bShouldUseStreamingFile && !NetworkPlatformFile )
-		{
-			NetworkPlatformFile = ConditionallyCreateFileWrapper(TEXT("CookedIterativeFile"), CurrentPlatformFile, CmdLine, &bNetworkFailedToInitialize, &bShouldUseCookedIterativeFile);
-			if (NetworkPlatformFile)
-			{
-				CurrentPlatformFile = NetworkPlatformFile;
-				FPlatformFileManager::Get().SetPlatformFile(*CurrentPlatformFile);
-			}
-		}
-
 		// if streaming network platform file was tried this loop don't try this one
 		// Network file wrapper (only create if the streaming wrapper hasn't been created)
-		if ( !bShouldUseStreamingFile && !bShouldUseCookedIterativeFile && !NetworkPlatformFile)
+		if (!bShouldUseStreamingFile && !NetworkPlatformFile)
 		{
 			NetworkPlatformFile = ConditionallyCreateFileWrapper(TEXT("NetworkFile"), CurrentPlatformFile, CmdLine, &bNetworkFailedToInitialize);
 			if (NetworkPlatformFile)
diff --git a/Engine/Source/Runtime/NetworkFileSystem/Private/NetworkFileServerConnection.cpp b/Engine/Source/Runtime/NetworkFileSystem/Private/NetworkFileServerConnection.cpp
--- a/Engine/Source/Runtime/NetworkFileSystem/Private/NetworkFileServerConnection.cpp
+++ b/Engine/Source/Runtime/NetworkFileSystem/Private/NetworkFileServerConnection.cpp
@@ -887,7 +887,6 @@ bool FNetworkFileServerClientConnection::ProcessGetFileList( FArchive& In, FArch
 	}
 
 	const bool bIsStreamingRequest = (ConnectionFlags & EConnectionFlags::Streaming) == EConnectionFlags::Streaming;
-	const bool bIsPrecookedIterativeRequest = (ConnectionFlags & EConnectionFlags::PreCookedIterative) == EConnectionFlags::PreCookedIterative;
 
 	ConnectedPlatformName = TEXT("");
 	ConnectedTargetPlatform = nullptr;
@@ -1168,17 +1167,6 @@ bool FNetworkFileServerClientConnection::ProcessGetFileList( FArchive& In, FArch
 		Out << FixedTimes;
 	}
 
-
-
-	if ( bIsPrecookedIterativeRequest )
-	{
-		TMap<FString, FDateTime> PrecookedList;
-		NetworkFileDelegates->InitialPrecookedListDelegate.ExecuteIfBound(ConnectedPlatformName, PrecookedList);
-
-		FixedTimes = FixupSandboxPathsForClient(PrecookedList);
-		Out << FixedTimes;
-	}
-
 	return true;
 }
 
diff --git a/Engine/Source/Runtime/NetworkFileSystem/Public/INetworkFileSystemModule.h b/Engine/Source/Runtime/NetworkFileSystem/Public/INetworkFileSystemModule.h
--- a/Engine/Source/Runtime/NetworkFileSystem/Public/INetworkFileSystemModule.h
+++ b/Engine/Source/Runtime/NetworkFileSystem/Public/INetworkFileSystemModule.h
@@ -75,31 +75,18 @@ DECLARE_MULTICAST_DELEGATE_OneParam( FOnFileModifiedDelegate, const FString& );
 DECLARE_DELEGATE_RetVal_TwoParams( bool, FNewConnectionDelegate, const FString&, const FString& );
 
 
-/**
- * Delegate which returns a list of all the files which should already be deployed to the devkit
- *
- * @param 1 IN, Platform to get precooked file list
- * @param 2 OUT, list of precooked files 
- */
-typedef TMap<FString,FDateTime> FFileTimeMap;
-DECLARE_DELEGATE_TwoParams( FInitialPrecookedListDelegate, const FString&, FFileTimeMap& );
-
-
-
 // container struct for delegates which the network file system uses
 struct FNetworkFileDelegateContainer
 {
 public:
 	FNetworkFileDelegateContainer() : 
 		NewConnectionDelegate(nullptr), 
-		InitialPrecookedListDelegate(nullptr),
 		SandboxPathOverrideDelegate(nullptr),
 		FileRequestDelegate(nullptr),
 		RecompileShadersDelegate(nullptr),
 		OnFileModifiedCallback(nullptr)
 	{}
 	FNewConnectionDelegate NewConnectionDelegate; 
-	FInitialPrecookedListDelegate InitialPrecookedListDelegate;
 	FSandboxPathDelegate SandboxPathOverrideDelegate;
 	FFileRequestDelegate FileRequestDelegate;
 	FRecompileShadersDelegate RecompileShadersDelegate;
diff --git a/Engine/Source/Runtime/Sockets/Public/NetworkMessage.h b/Engine/Source/Runtime/Sockets/Public/NetworkMessage.h
--- a/Engine/Source/Runtime/Sockets/Public/NetworkMessage.h
+++ b/Engine/Source/Runtime/Sockets/Public/NetworkMessage.h
@@ -66,7 +66,7 @@ enum class EConnectionFlags : uint8
 {
 	None = 0x00000000,
 	Streaming = 0x00000001,
-	PreCookedIterative = 0x00000002,
+	PreCookedIterative = 0x00000002, //DEPRECATED
 };
 ENUM_CLASS_FLAGS(EConnectionFlags);
 
