[MLDeformer]
- Initial icons.
- Introduced an editor style and linked colors and icons to it.
- Different wireframe colors for linear skinned, ml deformed, and ground truth meshes.
- Added support for adjusting wireframe color on UDebugSkelMeshComponent.
- Added support for adjusting wireframe color on UGeometryCacheComponent and its scene proxy.
- Some layout size tweaks.

diff --git a/Engine/Plugins/Experimental/GeometryCache/Source/GeometryCache/Classes/GeometryCacheComponent.h b/Engine/Plugins/Experimental/GeometryCache/Source/GeometryCache/Classes/GeometryCacheComponent.h
--- a/Engine/Plugins/Experimental/GeometryCache/Source/GeometryCache/Classes/GeometryCacheComponent.h
+++ b/Engine/Plugins/Experimental/GeometryCache/Source/GeometryCache/Classes/GeometryCacheComponent.h
@@ -183,6 +183,23 @@ public:
 	UFUNCTION(BlueprintCallable, Category = "Components|GeometryCache")
 	int32 GetNumberOfFrames() const;
 
+	/** Override wireframe color? */
+	UFUNCTION(BlueprintCallable, Category = "Components|GeometryCache")
+	void SetOverrideWireframeColor(bool bOverride);
+
+	/** Check whether we are overriding the wireframe color or not. */
+	UFUNCTION(BlueprintCallable, Category = "Components|GeometryCache")
+	bool GetOverrideWireframeColor() const;
+
+	/** Set the color, used when overriding the wireframe color is enabled. */
+	UFUNCTION(BlueprintCallable, Category = "Components|GeometryCache")
+	void SetWireframeOverrideColor(const FLinearColor Color);
+
+	/** Get the wireframe override color, used when overriding the wireframe color is enabled. */
+	UFUNCTION(BlueprintCallable, Category = "Components|GeometryCache")
+	FLinearColor GetWireframeOverrideColor() const;
+
+
 public:
 	/** Helper to get the frame of the ABC asset at this time*/
 	int32 GetFrameAtTime(const float Time) const;
@@ -278,4 +295,12 @@ protected:
 
 	UPROPERTY(EditAnywhere, Category = GeometryCache)
 	bool bManualTick;
+
+	/** Do we override the wireframe rendering color? */
+	UPROPERTY(EditAnywhere, Category = GeometryCache)
+	bool bOverrideWireframeColor;
+
+	/** The wireframe override color. */
+	UPROPERTY(EditAnywhere, Category = GeometryCache, meta = (EditCondition = "bOVerrideWireframeColor"))
+	FLinearColor WireframeOverrideColor;
 };
diff --git a/Engine/Plugins/Experimental/GeometryCache/Source/GeometryCache/Private/GeometryCacheComponent.cpp b/Engine/Plugins/Experimental/GeometryCache/Source/GeometryCache/Private/GeometryCacheComponent.cpp
--- a/Engine/Plugins/Experimental/GeometryCache/Source/GeometryCache/Private/GeometryCacheComponent.cpp
+++ b/Engine/Plugins/Experimental/GeometryCache/Source/GeometryCache/Private/GeometryCacheComponent.cpp
@@ -30,6 +30,8 @@ UGeometryCacheComponent::UGeometryCacheComponent(const FObjectInitializer& Objec
 	Duration = 0.0f;
 	bManualTick = false;
 	MotionVectorScale = 1.f;
+	bOverrideWireframeColor = false;
+	WireframeOverrideColor = FLinearColor::Green;
 }
 
 void UGeometryCacheComponent::BeginDestroy()
@@ -502,7 +504,6 @@ void UGeometryCacheComponent::ReleaseResources()
 
 int32 UGeometryCacheComponent::GetFrameAtTime(const float Time) const
 {
-
 	const float FrameTime = GetNumberOfFrames() > 1 ? Duration / (float)(GetNumberOfFrames() - 1) : 0.0f;
 	const int32 NormalizedFrame =  FMath::Clamp(FMath::RoundToInt(Time / FrameTime), 0, GetNumberOfFrames() - 1);
 	const int32 StartFrame = GeometryCache != nullptr ? GeometryCache->GetStartFrame() : 0;
@@ -525,6 +526,27 @@ int32 UGeometryCacheComponent::GetNumberOfFrames() const
 	}
 	return 0;
 }
+
+void UGeometryCacheComponent::SetOverrideWireframeColor(bool bOverride)
+{
+	bOverrideWireframeColor = bOverride;
+}
+
+bool UGeometryCacheComponent::GetOverrideWireframeColor() const
+{
+	return bOverrideWireframeColor;
+}
+
+void UGeometryCacheComponent::SetWireframeOverrideColor(FLinearColor Color)
+{
+	WireframeOverrideColor = Color;
+}
+
+FLinearColor UGeometryCacheComponent::GetWireframeOverrideColor() const
+{
+	return WireframeOverrideColor;
+}
+
 #if WITH_EDITOR
 void UGeometryCacheComponent::PreEditUndo()
 {
diff --git a/Engine/Plugins/Experimental/GeometryCache/Source/GeometryCache/Private/GeometryCacheSceneProxy.cpp b/Engine/Plugins/Experimental/GeometryCache/Source/GeometryCache/Private/GeometryCacheSceneProxy.cpp
--- a/Engine/Plugins/Experimental/GeometryCache/Source/GeometryCache/Private/GeometryCacheSceneProxy.cpp
+++ b/Engine/Plugins/Experimental/GeometryCache/Source/GeometryCache/Private/GeometryCacheSceneProxy.cpp
@@ -70,6 +70,9 @@ FGeometryCacheSceneProxy::FGeometryCacheSceneProxy(UGeometryCacheComponent* Comp
 	bAlwaysHasVelocity = true;
 	PlaybackSpeed = (Component->IsPlaying()) ? Component->GetPlaybackSpeed() : 0.0f;
 	MotionVectorScale = Component->GetMotionVectorScale();
+	bOverrideWireframeColor = Component->GetOverrideWireframeColor();
+	WireframeOverrideColor = Component->GetWireframeOverrideColor();
+	
 	UpdatedFrameNum = 0;
 
 	bVFRequiresPrimitiveUniformBuffer = !UseGPUScene(GMaxRHIShaderPlatform, FeatureLevel) || (FeatureLevel == ERHIFeatureLevel::ES3_1);
@@ -414,7 +417,8 @@ void FGeometryCacheSceneProxy::GetDynamicMeshElements(const TArray<const FSceneV
 		const bool bLevelColorationEnabled = EngineShowFlags.LevelColoration;
 		const bool bPropertyColorationEnabled = EngineShowFlags.PropertyColoration;
 
-		FLinearColor ViewWireframeColor(bLevelColorationEnabled ? GetLevelColor() : GetWireframeColor());
+		const FLinearColor WireColor = bOverrideWireframeColor ? WireframeOverrideColor : GetWireframeColor();
+		FLinearColor ViewWireframeColor(bLevelColorationEnabled ? GetLevelColor() : WireColor);
 		if (bPropertyColorationEnabled)
 		{
 			ViewWireframeColor = GetPropertyColor();
diff --git a/Engine/Plugins/Experimental/GeometryCache/Source/GeometryCache/Public/GeometryCacheSceneProxy.h b/Engine/Plugins/Experimental/GeometryCache/Source/GeometryCache/Public/GeometryCacheSceneProxy.h
--- a/Engine/Plugins/Experimental/GeometryCache/Source/GeometryCache/Public/GeometryCacheSceneProxy.h
+++ b/Engine/Plugins/Experimental/GeometryCache/Source/GeometryCache/Public/GeometryCacheSceneProxy.h
@@ -388,6 +388,9 @@ private:
 	float PlaybackSpeed;
 	float MotionVectorScale;
 
+	bool bOverrideWireframeColor = false;
+	FLinearColor WireframeOverrideColor = FLinearColor::Green;
+
 	FMaterialRelevance MaterialRelevance;
 	uint32 bLooping : 1;
 	uint32 bIsPlayingBackwards : 1;
diff --git a/Engine/Source/Editor/UnrealEd/Classes/Animation/DebugSkelMeshComponent.h b/Engine/Source/Editor/UnrealEd/Classes/Animation/DebugSkelMeshComponent.h
--- a/Engine/Source/Editor/UnrealEd/Classes/Animation/DebugSkelMeshComponent.h
+++ b/Engine/Source/Editor/UnrealEd/Classes/Animation/DebugSkelMeshComponent.h
@@ -124,7 +124,7 @@ public:
 	* Constructor.
 	* @param	Component - skeletal mesh primitive being added
 	*/
-	FDebugSkelMeshSceneProxy(const UDebugSkelMeshComponent* InComponent, FSkeletalMeshRenderData* InSkelMeshRenderData, const FColor& InWireframeOverlayColor = FColor::White);
+	FDebugSkelMeshSceneProxy(const UDebugSkelMeshComponent* InComponent, FSkeletalMeshRenderData* InSkelMeshRenderData, FLinearColor InWireframeOverlayColor = FLinearColor::White);
 
 	virtual ~FDebugSkelMeshSceneProxy()
 	{}
@@ -227,6 +227,9 @@ class UNREALED_API UDebugSkelMeshComponent : public USkeletalMeshComponent
 	UPROPERTY(transient)
 	bool bDisplayVertexColors;
 
+	UPROPERTY(transient)
+	FLinearColor WireframeMeshOverlayColor;
+
 	UE_DEPRECATED(5.0, "This variable is no longer used. Use ProcessRootMotionMode instead.")
 	UPROPERTY()
 	uint32 bPreviewRootMotion_DEPRECATED : 1;
@@ -437,6 +440,13 @@ class UNREALED_API UDebugSkelMeshComponent : public USkeletalMeshComponent
 	/** Whether the current LOD of the debug mesh is being synced with the attached (preview) mesh instance. */
 	bool IsTrackingAttachedLOD() const;
 
+	/** Set the wireframe mesh overlay color, which basically controls the color of the wireframe. */
+	void SetWireframeMeshOverlayColor(FLinearColor Color) { WireframeMeshOverlayColor = Color; }
+
+	/** Get the wireframe mesh overlay color, which basically controls the color of the wireframe. */
+	FLinearColor GetWireframeMeshOverlayColor() const { return WireframeMeshOverlayColor; }
+
+
 #if WITH_EDITOR
 	//TODO - This is a really poor way to post errors to the user. Work out a better way.
 	struct FAnimNotifyErrors
diff --git a/Engine/Source/Editor/UnrealEd/Private/Animation/DebugSkelMeshComponent.cpp b/Engine/Source/Editor/UnrealEd/Private/Animation/DebugSkelMeshComponent.cpp
--- a/Engine/Source/Editor/UnrealEd/Private/Animation/DebugSkelMeshComponent.cpp
+++ b/Engine/Source/Editor/UnrealEd/Private/Animation/DebugSkelMeshComponent.cpp
@@ -46,6 +46,8 @@ UDebugSkelMeshComponent::UDebugSkelMeshComponent(const FObjectInitializer& Objec
 	bTrackAttachedInstanceLOD = false;
 
 	BoneRadiusMultiplier = 1.f;
+
+	WireframeMeshOverlayColor = FLinearColor(0.4f, 0.8f, 0.66f);
 	
 	CachedClothBounds = FBoxSphereBounds(ForceInit);
 
@@ -411,7 +413,6 @@ FPrimitiveSceneProxy* UDebugSkelMeshComponent::CreateSceneProxy()
 		!bHideSkin &&
 		MeshObject)
 	{
-		const FColor WireframeMeshOverlayColor(102,205,170,255);
 		Result = ::new FDebugSkelMeshSceneProxy(this, SkelMeshRenderData, WireframeMeshOverlayColor);
 	}
 
@@ -1120,12 +1121,12 @@ bool UDebugSkelMeshComponent::IsReferencePoseShown() const
 /***************************************************
  * FDebugSkelMeshSceneProxy 
  ***************************************************/
-FDebugSkelMeshSceneProxy::FDebugSkelMeshSceneProxy(const UDebugSkelMeshComponent* InComponent, FSkeletalMeshRenderData* InSkelMeshRenderData, const FColor& InWireframeOverlayColor /*= FColor::White*/) :
+FDebugSkelMeshSceneProxy::FDebugSkelMeshSceneProxy(const UDebugSkelMeshComponent* InComponent, FSkeletalMeshRenderData* InSkelMeshRenderData, FLinearColor InWireframeOverlayColor /*= FLinearColor::White*/) :
 	FSkeletalMeshSceneProxy(InComponent, InSkelMeshRenderData)
 	, bSelectable(false)
 {
 	DynamicData = nullptr;
-	SetWireframeColor(FLinearColor(InWireframeOverlayColor));
+	SetWireframeColor(InWireframeOverlayColor);
 
 	if(GEngine->ClothPaintMaterial)
 	{
