CookedPackageWriter iterative builds - change to use FAssetRegistryState as the list of previous files. Load the AssetRegistry from the AssetRegistry.bin file even when using ZenStore for the package storage.
#rnx

diff --git a/Engine/Source/Developer/IoStoreUtilities/Private/ZenStoreWriter.cpp b/Engine/Source/Developer/IoStoreUtilities/Private/ZenStoreWriter.cpp
--- a/Engine/Source/Developer/IoStoreUtilities/Private/ZenStoreWriter.cpp
+++ b/Engine/Source/Developer/IoStoreUtilities/Private/ZenStoreWriter.cpp
@@ -8,7 +8,9 @@
 #include "Algo/BinarySearch.h"
 #include "Algo/IsSorted.h"
 #include "Algo/Sort.h"
+#include "AssetRegistry/AssetRegistryState.h"
 #include "Async/Async.h"
+#include "Serialization/ArrayReader.h"
 #include "Serialization/CompactBinaryPackage.h"
 #include "Serialization/CompactBinaryWriter.h"
 #include "Serialization/LargeMemoryWriter.h" 
@@ -16,6 +18,7 @@
 #include "HAL/PlatformFileManager.h"
 #include "IO/IoDispatcher.h"
 #include "Misc/App.h"
+#include "Misc/FileHelper.h"
 #include "Misc/StringBuilder.h"
 #include "Interfaces/ITargetPlatform.h"
 #include "Misc/Paths.h"
@@ -819,13 +822,53 @@ void FZenStoreWriter::Flush()
 	HttpQueue->Flush();
 }
 
-void FZenStoreWriter::GetCookedPackages(TArray<FCookedPackageInfo>& OutCookedPackages) 
+TUniquePtr<FAssetRegistryState> FZenStoreWriter::LoadPreviousAssetRegistry() 
 {
-	OutCookedPackages.Reserve(OutCookedPackages.Num() + CookedPackagesInfo.Num());
-	for (FOplogCookInfo& CookedPackageInfo : CookedPackagesInfo)
+	// Load the previous asset registry to return to CookOnTheFlyServer, and set the packages enumerated in both *this and
+	// the returned asset registry to the intersection of the oplog and the previous asset registry;
+	// to report a package as already cooked we have to have the information from both sources.
+	FString PreviousAssetRegistryFile = FPaths::Combine(MetadataDirectoryPath, GetDevelopmentAssetRegistryFilename());
+	FArrayReader SerializedAssetData;
+	if (!IFileManager::Get().FileExists(*PreviousAssetRegistryFile) ||
+		!FFileHelper::LoadFileToArray(SerializedAssetData, *PreviousAssetRegistryFile))
 	{
-		OutCookedPackages.Add(CookedPackageInfo.CookInfo);
+		RemoveCookedPackages();
+		return TUniquePtr<FAssetRegistryState>();
 	}
+
+	TUniquePtr<FAssetRegistryState> PreviousState = MakeUnique<FAssetRegistryState>();
+	PreviousState->Load(SerializedAssetData);
+
+	TSet<FName> RemoveSet;
+	const TMap<FName, const FAssetPackageData*>& PreviousStatePackages = PreviousState->GetAssetPackageDataMap(); 
+	for (const TPair<FName, const FAssetPackageData*>& Pair : PreviousStatePackages)
+	{
+		FName PackageName = Pair.Key;
+		if (!PackageNameToIndex.Find(PackageName))
+		{
+			RemoveSet.Add(PackageName);
+		}
+	}
+	if (RemoveSet.Num())
+	{
+		PreviousState->PruneAssetData(TSet<FName>(), RemoveSet, FAssetRegistrySerializationOptions());
+	}
+
+	TArray<FName> RemoveArray;
+	for (TPair<FName, int32>& Pair : PackageNameToIndex)
+	{
+		FName PackageName = Pair.Key;
+		if (!PreviousStatePackages.Find(PackageName))
+		{
+			RemoveArray.Add(PackageName);
+		}
+	}
+	if (RemoveArray.Num())
+	{
+		RemoveCookedPackages(RemoveArray);
+	}
+
+	return PreviousState;
 }
 
 FCbObject FZenStoreWriter::GetOplogAttachment(FName PackageName, FUtf8StringView AttachmentKey)
@@ -935,8 +978,11 @@ void FZenStoreWriter::MarkPackagesUpToDate(TArrayView<const FName> UpToDatePacka
 			int32* Index = PackageNameToIndex.Find(PackageName);
 			if (!Index)
 			{
-				UE_LOG(LogZenStoreWriter, Error, TEXT("MarkPackagesUpToDate called with package %s that is not in the oplog."),
-					*PackageName.ToString());
+				if (!FPackageName::IsScriptPackage(WriteToString<128>(PackageName)))
+				{
+					UE_LOG(LogZenStoreWriter, Warning, TEXT("MarkPackagesUpToDate called with package %s that is not in the oplog."),
+						*PackageName.ToString());
+				}
 				continue;
 			}
 
diff --git a/Engine/Source/Developer/IoStoreUtilities/Public/ZenStoreWriter.h b/Engine/Source/Developer/IoStoreUtilities/Public/ZenStoreWriter.h
--- a/Engine/Source/Developer/IoStoreUtilities/Public/ZenStoreWriter.h
+++ b/Engine/Source/Developer/IoStoreUtilities/Public/ZenStoreWriter.h
@@ -62,7 +62,7 @@ public:
 
 	IOSTOREUTILITIES_API void WriteIoStorePackageData(const FPackageInfo& Info, const FIoBuffer& PackageData, const FPackageStoreEntryResource& PackageStoreEntry, const TArray<FFileRegion>& FileRegions);
 
-	IOSTOREUTILITIES_API virtual void GetCookedPackages(TArray<FCookedPackageInfo>& OutCookedPackages) override;
+	IOSTOREUTILITIES_API TUniquePtr<FAssetRegistryState> LoadPreviousAssetRegistry() override;
 	IOSTOREUTILITIES_API virtual FCbObject GetOplogAttachment(FName PackageName, FUtf8StringView AttachmentKey) override;
 	IOSTOREUTILITIES_API virtual void RemoveCookedPackages(TArrayView<const FName> PackageNamesToRemove) override;
 	IOSTOREUTILITIES_API virtual void RemoveCookedPackages() override;
diff --git a/Engine/Source/Editor/UnrealEd/Private/Commandlets/AssetRegistryGenerator.cpp b/Engine/Source/Editor/UnrealEd/Private/Commandlets/AssetRegistryGenerator.cpp
--- a/Engine/Source/Editor/UnrealEd/Private/Commandlets/AssetRegistryGenerator.cpp
+++ b/Engine/Source/Editor/UnrealEd/Private/Commandlets/AssetRegistryGenerator.cpp
@@ -375,6 +375,7 @@ bool FAssetRegistryGenerator::GenerateStreamingInstallManifest(int64 InExtraFlav
 	TMap<FString, int64> PackageFileSizes;
 	if (MaxChunkSize > 0)
 	{
+		// ZENTODO: Pull file size information from the AssetRegistryState, since the packages do not exist on disk in zen
 		FString SandboxPath = InSandboxFile->GetSandboxDirectory();
 		SandboxPath.ReplaceInline(TEXT("[Platform]"), *Platform);
 		FPackageFileSizeVisitor PackageSearch(PackageFileSizes);
@@ -795,20 +796,18 @@ void FAssetRegistryGenerator::CleanManifestDirectories()
 	CleanTempPackagingDirectory(TargetPlatform->PlatformName());
 }
 
-void FAssetRegistryGenerator::SetPreviousAssetRegistry(TUniquePtr<FAssetRegistryState>&& InPreviousState,
-	TMap<FName, FAssetPackageData*>&& PreviousAssetPackageData, TArray<FName>& OutTombstonePackages)
+void FAssetRegistryGenerator::SetPreviousAssetRegistry(TUniquePtr<FAssetRegistryState>&& InPreviousState)
 {
-	OutTombstonePackages.Reset();
+	PreviousPackagesToUpdate.Empty();
 	if (InPreviousState)
 	{
-		for (const TPair<FName, const FAssetPackageData*>& Pair : InPreviousState->GetAssetPackageDataMap())
+		const TMap<FName, const FAssetPackageData*>& PreviousPackageDataMap = InPreviousState->GetAssetPackageDataMap();
+		PreviousPackagesToUpdate.Reserve(PreviousPackageDataMap.Num());
+		for (const TPair<FName, const FAssetPackageData*>& Pair : PreviousPackageDataMap)
 		{
 			FName PackageName = Pair.Key;
-			const FAssetPackageData& PreviousPackageData = *Pair.Value;
 			if (!State.GetAssetPackageData(PackageName))
 			{
-				// Tombstone means a package from the previous cook no longer exists in perforce
-				OutTombstonePackages.Add(PackageName);
 				continue;
 			}
 			TArrayView<FAssetData const * const> PreviousAssetDatas = InPreviousState->GetAssetsByPackageName(PackageName);
@@ -822,26 +821,10 @@ void FAssetRegistryGenerator::SetPreviousAssetRegistry(TUniquePtr<FAssetRegistry
 					UpdateAssetDatas.Add(FAssetData(*AssetData));
 				}
 			}
+			const FAssetPackageData& PreviousPackageData = *Pair.Value;
 			UpdateData.Get<1>() = PreviousPackageData;
 		}
 	}
-	else
-	{
-		for (TPair<FName, FAssetPackageData*>& Pair : PreviousAssetPackageData)
-		{
-			FName PackageName = Pair.Key;
-			FAssetPackageData& PreviousPackageData = *Pair.Value;
-			if (!State.GetAssetPackageData(PackageName))
-			{
-				// Tombstone means a package from the previous cook no longer exists in perforce
-				OutTombstonePackages.Add(PackageName);
-				continue;
-			}
-			TPair<TArray<FAssetData>, FAssetPackageData>& UpdateData = PreviousPackagesToUpdate.FindOrAdd(PackageName);
-			// No AssetDatas available in this case
-			UpdateData.Get<1>() = MoveTemp(PreviousPackageData);
-		}
-	}
 }
 
 void FAssetRegistryGenerator::InjectEncryptionData(FAssetRegistryState& TargetState)
@@ -1076,9 +1059,10 @@ void FAssetRegistryGenerator::Initialize(const TArray<FName> &InStartupPackages)
 	FGameDelegates::Get().GetAssignLayerChunkDelegate() = FAssignLayerChunkDelegate::CreateStatic(AssignLayerChunkDelegate);
 }
 
-void FAssetRegistryGenerator::ComputePackageDifferences(const FComputeDifferenceOptions& Options, const TMap<FName, const FAssetPackageData*>& PreviousAssetPackageDataMap, FAssetRegistryDifference& OutDifference)
+void FAssetRegistryGenerator::ComputePackageDifferences(const FComputeDifferenceOptions& Options, const FAssetRegistryState& PreviousState, FAssetRegistryDifference& OutDifference)
 {
 	TArray<FName> ModifiedScriptPackages;
+	const TMap<FName, const FAssetPackageData*>& PreviousAssetPackageDataMap = PreviousState.GetAssetPackageDataMap();
 
 	for (const TPair<FName, const FAssetPackageData*>& PackagePair : State.GetAssetPackageDataMap())
 	{
@@ -1120,10 +1104,7 @@ void FAssetRegistryGenerator::ComputePackageDifferences(const FComputeDifference
 	for (const TPair<FName, const FAssetPackageData*>& PackagePair : PreviousAssetPackageDataMap)
 	{
 		FName PackageName = PackagePair.Key;
-		const FAssetPackageData* PreviousPackageData = PackagePair.Value;
-
 		const FAssetPackageData* CurrentPackageData = State.GetAssetPackageData(PackageName);
-
 		if (!CurrentPackageData)
 		{
 			OutDifference.RemovedPackages.Add(PackageName);
@@ -1163,6 +1144,19 @@ void FAssetRegistryGenerator::ComputePackageDifferences(const FComputeDifference
 	}
 }
 
+void FAssetRegistryGenerator::ComputePackageRemovals(const FAssetRegistryState& PreviousState, TArray<FName>& RemovedPackages)
+{
+	for (const TPair<FName, const FAssetPackageData*>& PackagePair : PreviousState.GetAssetPackageDataMap())
+	{
+		FName PackageName = PackagePair.Key;
+		const FAssetPackageData* CurrentPackageData = State.GetAssetPackageData(PackageName);
+		if (!CurrentPackageData)
+		{
+			RemovedPackages.Add(PackageName);
+		}
+	}
+}
+
 void FAssetRegistryGenerator::BuildChunkManifest(const TSet<FName>& InCookedPackages, const TSet<FName>& InDevelopmentOnlyPackages, FSandboxPlatformFile* InSandboxFile, bool bGenerateStreamingInstallManifest)
 {
 	// If we were asked to generate a streaming install manifest explicitly and we did not have bGenerateNoChunks set, we will generate chunks.
diff --git a/Engine/Source/Editor/UnrealEd/Private/CookOnTheFlyServer.cpp b/Engine/Source/Editor/UnrealEd/Private/CookOnTheFlyServer.cpp
--- a/Engine/Source/Editor/UnrealEd/Private/CookOnTheFlyServer.cpp
+++ b/Engine/Source/Editor/UnrealEd/Private/CookOnTheFlyServer.cpp
@@ -5650,36 +5650,6 @@ void UCookOnTheFlyServer::PopulateCookedPackages(TArrayView<const ITargetPlatfor
 	using namespace UE::Cook;
 	TRACE_CPUPROFILER_EVENT_SCOPE(UCookOnTheFlyServer::PopulateCookedPackages);
 
-	struct FPreviousAssetPackageData
-	{
-		FPreviousAssetPackageData(const TArray<ICookedPackageWriter::FCookedPackageInfo>& InCookedPackages)
-		{
-			AssetPackageDatas.SetNum(InCookedPackages.Num());
-			AssetPackageDataMap.Reserve(InCookedPackages.Num());
-
-			for (int32 Idx = 0, Count = InCookedPackages.Num(); Idx < Count; ++Idx)
-			{
-				const ICookedPackageWriter::FCookedPackageInfo& CookedPackageInfo = InCookedPackages[Idx];
-				FAssetPackageData& AssetPackageData = AssetPackageDatas[Idx];
-				
-				AssetPackageData.CookedHash = CookedPackageInfo.Hash;
-				PRAGMA_DISABLE_DEPRECATION_WARNINGS
-				AssetPackageData.PackageGuid = CookedPackageInfo.PackageGuid;
-				PRAGMA_ENABLE_DEPRECATION_WARNINGS
-				AssetPackageData.DiskSize = CookedPackageInfo.DiskSize;
-
-				AssetPackageDataMap.Add(CookedPackageInfo.PackageName, &AssetPackageData);
-			}
-		}
-		TMap<FName, const FAssetPackageData*>& GetConstAssetPackageDataMap()
-		{
-			return reinterpret_cast<TMap<FName, const FAssetPackageData*>&>(AssetPackageDataMap);
-		}
-
-		TArray<FAssetPackageData> AssetPackageDatas;
-		TMap<FName, FAssetPackageData*> AssetPackageDataMap;
-	};
-
 	// TODO: NumPackagesIterativelySkipped is only counted for the first platform; to count all platforms we would
 	// have to check whether each one is already cooked.
 	bool bFirstPlatform = true;
@@ -5692,20 +5662,26 @@ void UCookOnTheFlyServer::PopulateCookedPackages(TArrayView<const ITargetPlatfor
 		UE_LOG(LogCook, Display, TEXT("Populating cooked package(s) from %s package store on platform '%s'"),
 			*CookSavePackageContext.WriterDebugName, *TargetPlatform->PlatformName());
 
-		TArray<ICookedPackageWriter::FCookedPackageInfo> CookedPackages;
-		PackageWriter.GetCookedPackages(CookedPackages);
-		// TODO: Implementation ReleasePreviousAssetRegistry for all PackageWriters and add an assertion here instead of handling null.
-		TUniquePtr<FAssetRegistryState> PreviousAssetRegistry(PackageWriter.ReleasePreviousAssetRegistry());
-
-		UE_LOG(LogCook, Display, TEXT("Found '%d' cooked package(s) in package store"), CookedPackages.Num());
-
-		FPreviousAssetPackageData PreviousAssetPackageData(CookedPackages);
+		TUniquePtr<FAssetRegistryState> PreviousAssetRegistry(PackageWriter.LoadPreviousAssetRegistry());
+		int32 NumPreviousPackages = PreviousAssetRegistry ? PreviousAssetRegistry->GetNumPackages() : 0;
+		UE_LOG(LogCook, Display, TEXT("Found '%d' cooked package(s) in package store"), NumPreviousPackages);
+		if (NumPreviousPackages == 0)
+		{
+			continue;
+		}
 
 		if (bHybridIterativeEnabled)
 		{
 			// HybridIterative does the equivalent operation of bRecurseModifications=true and bRecurseScriptModifications=true,
 			// but checks for out-of-datedness are done by FRequestCluster using the TargetDomainKey (which is built
 			// from dependencies), so we do not need to check for out-of-datedness here
+
+			// Remove packages that no longer exist in the WorkspaceDomain from the TargetDomain. We have to
+			// check this for all packages in the previous cook rather than just the currently referenced
+			// set because when packages are removed the references to them are usually also removed.
+			TArray<FName> TombstonePackages;
+			PlatformAssetRegistry.ComputePackageRemovals(*PreviousAssetRegistry, TombstonePackages);
+			PackageWriter.RemoveCookedPackages(TombstonePackages);
 		}
 		else
 		{
@@ -5718,7 +5694,7 @@ void UCookOnTheFlyServer::PopulateCookedPackages(TArrayView<const ITargetPlatfor
 			Options.bRecurseModifications = true;
 			Options.bRecurseScriptModifications = !IsCookFlagSet(ECookInitializationFlags::IgnoreScriptPackagesOutOfDate);
 			FAssetRegistryGenerator::FAssetRegistryDifference Difference;
-			PlatformAssetRegistry.ComputePackageDifferences(Options, PreviousAssetPackageData.GetConstAssetPackageDataMap(), Difference);
+			PlatformAssetRegistry.ComputePackageDifferences(Options, *PreviousAssetRegistry, Difference);
 
 			TArray<FName> PackagesToRemove;
 			PackagesToRemove.Reserve(Difference.ModifiedPackages.Num() + Difference.RemovedPackages.Num() + Difference.IdenticalUncookedPackages.Num());
@@ -5780,17 +5756,7 @@ void UCookOnTheFlyServer::PopulateCookedPackages(TArrayView<const ITargetPlatfor
 			}
 		}
 
-		TArray<FName> TombstonePackages;
-		PlatformAssetRegistry.SetPreviousAssetRegistry(MoveTemp(PreviousAssetRegistry),
-			MoveTemp(PreviousAssetPackageData.AssetPackageDataMap), TombstonePackages);
-
-		if (bHybridIterativeEnabled)
-		{
-			// Remove packages that no longer exist in the WorkspaceDomain from the TargetDomain. We have to
-			// check this for all packages in the previous cook rather than just the currently referenced
-			// set because when packages are removed the references to them are usually also removed.
-			PackageWriter.RemoveCookedPackages(TombstonePackages);
-		}
+		PlatformAssetRegistry.SetPreviousAssetRegistry(MoveTemp(PreviousAssetRegistry));
 	}
 }
 
diff --git a/Engine/Source/Editor/UnrealEd/Private/Cooker/LooseCookedPackageWriter.cpp b/Engine/Source/Editor/UnrealEd/Private/Cooker/LooseCookedPackageWriter.cpp
--- a/Engine/Source/Editor/UnrealEd/Private/Cooker/LooseCookedPackageWriter.cpp
+++ b/Engine/Source/Editor/UnrealEd/Private/Cooker/LooseCookedPackageWriter.cpp
@@ -24,8 +24,6 @@
 #include "UObject/Package.h"
 #include "PackageStoreOptimizer.h"
 
-const static FName NAME_DummyCookedFilename(TEXT("DummyCookedFilename"));
-
 FLooseCookedPackageWriter::FLooseCookedPackageWriter(const FString& InOutputPath, const FString& InMetadataDirectoryPath, const ITargetPlatform* InTargetPlatform,
 	FAsyncIODelete& InAsyncIODelete, const FPackageNameCache& InPackageNameCache, const TArray<TSharedRef<IPlugin>>& InPluginsToRemap)
 	: OutputPath(InOutputPath)
@@ -114,11 +112,8 @@ void FLooseCookedPackageWriter::Flush()
 	UPackage::WaitForAsyncFileWrites();
 }
 
-void FLooseCookedPackageWriter::GetCookedPackages(TArray<FCookedPackageInfo>& OutCookedPackages)
+TUniquePtr<FAssetRegistryState> FLooseCookedPackageWriter::LoadPreviousAssetRegistry()
 {
-	// GetCookedPackages is currently only called in BeginCook to read the list of previous cooked files
-	// Read that list from the platform cooked asset registry file
-
 	// Report files from the shared build if the option is set
 	FString PreviousAssetRegistryFile;
 	if (bIterateSharedBuild)
@@ -127,10 +122,6 @@ void FLooseCookedPackageWriter::GetCookedPackages(TArray<FCookedPackageInfo>& Ou
 		DeleteSandboxDirectory();
 		PreviousAssetRegistryFile = FPaths::Combine(*FPaths::ProjectSavedDir(), TEXT("SharedIterativeBuild"),
 			*TargetPlatform.PlatformName(), TEXT("Metadata"), GetDevelopmentAssetRegistryFilename());
-
-		// We use NAME_DummyCookedFilename as a placeholder in UncookedPathToCookedPath to indicate the file comes from the shared pak file
-		// If an actual filename of this name exists, our placeholder system will break
-		check(!IFileManager::Get().FileExists(*NAME_DummyCookedFilename.ToString()));
 	}
 	else
 	{
@@ -138,52 +129,58 @@ void FLooseCookedPackageWriter::GetCookedPackages(TArray<FCookedPackageInfo>& Ou
 	}
 
 	UncookedPathToCookedPath.Reset();
-	GetAllCookedFiles();
 
 	FArrayReader SerializedAssetData;
-	if (IFileManager::Get().FileExists(*PreviousAssetRegistryFile) && FFileHelper::LoadFileToArray(SerializedAssetData, *PreviousAssetRegistryFile))
-	{
-		PreviousState = MakeUnique<FAssetRegistryState>();
-		PreviousState->Load(SerializedAssetData);
+	if (!IFileManager::Get().FileExists(*PreviousAssetRegistryFile) || !FFileHelper::LoadFileToArray(SerializedAssetData, *PreviousAssetRegistryFile))
+	{ 
+		RemoveCookedPackages();
+		return TUniquePtr<FAssetRegistryState>();
+	}
 
-		const TMap<FName, const FAssetPackageData*>& DataMap = PreviousState->GetAssetPackageDataMap();
-		OutCookedPackages.Reserve(OutCookedPackages.Num() + DataMap.Num());
-		FString UncookedFileNameStr;
+	TUniquePtr<FAssetRegistryState> PreviousState = MakeUnique<FAssetRegistryState>();
+	PreviousState->Load(SerializedAssetData);
 
-		for (const TPair<FName, const FAssetPackageData*>& Pair : DataMap)
+	// If we are iterating from a shared build the cooked files do not exist in the local cooked directory;
+	// we assume they are packaged in the pak file (which we don't want to extract to confirm) and keep them all.
+	if (!bIterateSharedBuild)
+	{
+		// For regular iteration, remove every file from the PreviousState that no longer exists in the cooked directory
+		// and remove every cooked file from disk that is not present in the AssetRegistry
+		GetAllCookedFiles();
+		TSet<FName> ExistsOnlyInRegistry;
+		TSet<FName> ExistsOnlyOnDisk;
+		ExistsOnlyOnDisk.Reserve(UncookedPathToCookedPath.Num());
+		for (TPair<FName, FName>& Pair : UncookedPathToCookedPath)
+		{
+			ExistsOnlyOnDisk.Add(Pair.Key);
+		}
+		for (const TPair<FName, const FAssetPackageData*>& Pair : PreviousState->GetAssetPackageDataMap())
 		{
 			FName PackageName = Pair.Key;
 			const FAssetPackageData* PackageData = Pair.Value;
+			const FName UncookedFilename = PackageNameCache.GetCachedStandardFileName(PackageName);
 			bool bExistsOnDisk = false;
-			if (bIterateSharedBuild)
-			{
-				// if we are iterating of a shared build the cooked files might not exist in the cooked directory because we assume they are packaged in the pak file (which we don't want to extract)
-				// Overwrite the value in UncookedPathToCookpath with a placeholder that indicates the file exists even if it's not present in the Cooked directory
-				if (FPackageName::DoesPackageExist(PackageName.ToString(), nullptr, &UncookedFileNameStr))
-				{
-					UncookedPathToCookedPath.Add(FName(*UncookedFileNameStr), NAME_DummyCookedFilename);
-				}
-				bExistsOnDisk = true;
-			}
-			else
+			if (!UncookedFilename.IsNone())
 			{
-				const FName UncookedFilename = PackageNameCache.GetCachedStandardFileName(PackageName);
-				if (!UncookedFilename.IsNone())
-				{
-					FName* CookedFilename = UncookedPathToCookedPath.Find(UncookedFilename);
-					bExistsOnDisk = CookedFilename != nullptr;
-				}
+				bExistsOnDisk = (ExistsOnlyOnDisk.Remove(UncookedFilename) == 1);
 			}
-			if (bExistsOnDisk)
+			if (!bExistsOnDisk)
 			{
-				PRAGMA_DISABLE_DEPRECATION_WARNINGS;
-				OutCookedPackages.Add(FCookedPackageInfo{ PackageName, PackageData->CookedHash,
-					PackageData->PackageGuid, PackageData->DiskSize });
-				// Oplog attachments are not implemented by FLooseCookedPackageWriter
-				PRAGMA_ENABLE_DEPRECATION_WARNINGS;
+				ExistsOnlyInRegistry.Add(PackageName);
 			}
 		}
+
+		if (ExistsOnlyInRegistry.Num())
+		{
+			PreviousState->PruneAssetData(TSet<FName>(), ExistsOnlyInRegistry, FAssetRegistrySerializationOptions());
+		}
+		if (ExistsOnlyOnDisk.Num())
+		{
+			RemoveCookedPackagesByUncookedFilename(ExistsOnlyOnDisk.Array());
+		}
 	}
+
+	return PreviousState;
 }
 
 FCbObject FLooseCookedPackageWriter::GetOplogAttachment(FName PackageName, FUtf8StringView AttachmentKey)
@@ -199,12 +196,12 @@ void FLooseCookedPackageWriter::RemoveCookedPackages(TArrayView<const FName> Pac
 		return;
 	}
 
-	// if we are going to clear the cooked packages it is conceivable that we will recook the packages which we just cooked 
-	// that means it's also conceivable that we will recook the same package which currently has an outstanding async write request
-	UPackage::WaitForAsyncFileWrites();
-	
 	if (PackageNamesToRemove.Num() > 0)
 	{
+		// if we are going to clear the cooked packages it is conceivable that we will recook the packages which we just cooked 
+		// that means it's also conceivable that we will recook the same package which currently has an outstanding async write request
+		UPackage::WaitForAsyncFileWrites();
+
 		// PackageNameCache is read-game-thread-only, so we have to read it before calling parallel-for
 		TArray<FName> UncookedFileNamesToRemove;
 		UncookedFileNamesToRemove.Reserve(PackageNamesToRemove.Num());
@@ -216,21 +213,32 @@ void FLooseCookedPackageWriter::RemoveCookedPackages(TArrayView<const FName> Pac
 				UncookedFileNamesToRemove.Add(UncookedFileName);
 			}
 		}
+		RemoveCookedPackagesByUncookedFilename(UncookedFileNamesToRemove);
+	}
 
-		auto DeletePackageLambda = [&UncookedFileNamesToRemove, this](int32 PackageIndex)
+	// We no longer have a use for UncookedPathToCookedPath, after the RemoveCookedPackages call at the beginning of the cook.
+	UncookedPathToCookedPath.Empty();
+}
+
+void FLooseCookedPackageWriter::RemoveCookedPackagesByUncookedFilename(const TArray<FName>& UncookedFileNamesToRemove)
+{
+	auto DeletePackageLambda = [&UncookedFileNamesToRemove, this](int32 PackageIndex)
+	{
+		FName UncookedFileName = UncookedFileNamesToRemove[PackageIndex];
+		FName* CookedFileName = UncookedPathToCookedPath.Find(UncookedFileName);
+		if (CookedFileName)
 		{
-			FName UncookedFileName = UncookedFileNamesToRemove[PackageIndex];
-			FName* CookedFileName = UncookedPathToCookedPath.Find(UncookedFileName);
-			if (CookedFileName)
-			{
-				TStringBuilder<256> FilePath;
-				CookedFileName->ToString(FilePath);
-				IFileManager::Get().Delete(*FilePath, true, true, true);
-			}
-		};
-		ParallelFor(UncookedFileNamesToRemove.Num(), DeletePackageLambda);
+			TStringBuilder<256> FilePath;
+			CookedFileName->ToString(FilePath);
+			IFileManager::Get().Delete(*FilePath, true, true, true);
+		}
+	};
+	ParallelFor(UncookedFileNamesToRemove.Num(), DeletePackageLambda);
+
+	for (FName UncookedFilename : UncookedFileNamesToRemove)
+	{
+		UncookedPathToCookedPath.Remove(UncookedFilename);
 	}
-	UncookedPathToCookedPath.Empty();
 }
 
 void FLooseCookedPackageWriter::MarkPackagesUpToDate(TArrayView<const FName> UpToDatePackages)
@@ -242,11 +250,6 @@ void FLooseCookedPackageWriter::RemoveCookedPackages()
 	DeleteSandboxDirectory();
 }
 
-FAssetRegistryState* FLooseCookedPackageWriter::ReleasePreviousAssetRegistry()
-{
-	return PreviousState.Release();
-}
-
 void FLooseCookedPackageWriter::DeleteSandboxDirectory()
 {
 	// if we are going to clear the cooked packages it is conceivable that we will recook the packages which we just cooked 
diff --git a/Engine/Source/Editor/UnrealEd/Private/Cooker/LooseCookedPackageWriter.h b/Engine/Source/Editor/UnrealEd/Private/Cooker/LooseCookedPackageWriter.h
--- a/Engine/Source/Editor/UnrealEd/Private/Cooker/LooseCookedPackageWriter.h
+++ b/Engine/Source/Editor/UnrealEd/Private/Cooker/LooseCookedPackageWriter.h
@@ -42,12 +42,11 @@ public:
 	void BeginCook(const FCookInfo& Info) override;
 	void EndCook() override;
 	void Flush() override;
-	void GetCookedPackages(TArray<FCookedPackageInfo>& OutCookedPackages) override;
+	TUniquePtr<FAssetRegistryState> LoadPreviousAssetRegistry() override;
 	FCbObject GetOplogAttachment(FName PackageName, FUtf8StringView AttachmentKey) override;
 	void RemoveCookedPackages(TArrayView<const FName> PackageNamesToRemove) override;
 	void RemoveCookedPackages() override;
 	void MarkPackagesUpToDate(TArrayView<const FName> UpToDatePackages) override;
-	FAssetRegistryState* ReleasePreviousAssetRegistry() override;
 
 private:
 
@@ -66,6 +65,7 @@ private:
 		const FString& SandboxRootDir, const FString& RelativeRootDir,
 		const FString& SandboxProjectDir, const FString& RelativeProjectDir,
 		const FString& CookedPath, FString& OutUncookedPath) const;
+	void RemoveCookedPackagesByUncookedFilename(const TArray<FName>& UncookedFileNamesToRemove);
 
 	TMap<FName, FName> UncookedPathToCookedPath;
 	FString OutputPath;
@@ -73,7 +73,6 @@ private:
 	const ITargetPlatform& TargetPlatform;
 	const FPackageNameCache& PackageNameCache;
 	FPackageStoreManifest PackageStoreManifest;
-	TUniquePtr<FAssetRegistryState> PreviousState;
 	const TArray<TSharedRef<IPlugin>>& PluginsToRemap;
 	FAsyncIODelete& AsyncIODelete;
 	bool bIterateSharedBuild;
diff --git a/Engine/Source/Editor/UnrealEd/Public/Commandlets/AssetRegistryGenerator.h b/Engine/Source/Editor/UnrealEd/Public/Commandlets/AssetRegistryGenerator.h
--- a/Engine/Source/Editor/UnrealEd/Public/Commandlets/AssetRegistryGenerator.h
+++ b/Engine/Source/Editor/UnrealEd/Public/Commandlets/AssetRegistryGenerator.h
@@ -42,10 +42,9 @@ public:
 	const ITargetPlatform* GetTargetPlatform() const { return TargetPlatform; }
 
 	/** 
-	 * Sets asset registry from a previous run that is used for iterative or DLC cooking
+	 * Sets asset registry from a previous run that is used for iterative cooking
 	 */
-	void SetPreviousAssetRegistry(TUniquePtr<FAssetRegistryState>&& PreviousState,
-		TMap<FName, FAssetPackageData*>&& PreviousAssetPackageData, TArray<FName>& OutTombstonePackages);
+	void SetPreviousAssetRegistry(TUniquePtr<FAssetRegistryState>&& PreviousState);
 
 	/**
 	 * Options when computing the differences between current and previous state.
@@ -82,7 +81,10 @@ public:
 	 * @param PreviousAssetPackageDataMap previously cooked asset package data
 	 * @param OutDifference the differences between the current and the previous state
 	 */
-	void ComputePackageDifferences(const FComputeDifferenceOptions& Options, const TMap<FName, const FAssetPackageData*>& PreviousAssetPackageDataMap, FAssetRegistryDifference& OutDifference);
+	void ComputePackageDifferences(const FComputeDifferenceOptions& Options, const FAssetRegistryState& PreviousState, FAssetRegistryDifference& OutDifference);
+
+	/** Computes just the list of packages in the PreviousState that no longer exist in the current state. */
+	void ComputePackageRemovals(const FAssetRegistryState& PreviousState, TArray<FName>& RemovedPackages);
 
 	/**
 	 * GenerateChunkManifest 
diff --git a/Engine/Source/Runtime/AssetRegistry/Public/AssetRegistry/AssetRegistryState.h b/Engine/Source/Runtime/AssetRegistry/Public/AssetRegistry/AssetRegistryState.h
--- a/Engine/Source/Runtime/AssetRegistry/Public/AssetRegistry/AssetRegistryState.h
+++ b/Engine/Source/Runtime/AssetRegistry/Public/AssetRegistry/AssetRegistryState.h
@@ -392,6 +392,9 @@ public:
 	/** Returns the number of assets in this state */
 	int32 GetNumAssets() const { return NumAssets; }
 
+	/** Returns the number of packages in this state */
+	int32 GetNumPackages() const { return CachedAssetsByPackageName.Num(); }
+
 #if ASSET_REGISTRY_STATE_DUMPING_ENABLED
 	/**
 	 * Writes out the state in textual form. Use arguments to control which segments to emit.
diff --git a/Engine/Source/Runtime/Core/Public/Serialization/PackageWriter.h b/Engine/Source/Runtime/Core/Public/Serialization/PackageWriter.h
--- a/Engine/Source/Runtime/Core/Public/Serialization/PackageWriter.h
+++ b/Engine/Source/Runtime/Core/Public/Serialization/PackageWriter.h
@@ -7,6 +7,7 @@
 #include "Misc/DateTime.h"
 #include "Misc/SecureHash.h"
 #include "Serialization/CompactBinary.h"
+#include "Templates/UniquePtr.h"
 
 class FAssetRegistryState;
 class IPackageStoreWriter;
@@ -192,9 +193,9 @@ public:
 	};
 
 	/**
-	 * Appends list of cooked package(s) to the given output array.
+	 * Returns an AssetRegistry describing the previous cook results.
 	 */
-	virtual void GetCookedPackages(TArray<FCookedPackageInfo>& OutCookedPackages) = 0;
+	virtual TUniquePtr<FAssetRegistryState> LoadPreviousAssetRegistry() = 0;
 
 	/**
 	 * Returns an Attachment that was previously commited for the given PackageName.
@@ -217,15 +218,6 @@ public:
 	 */
 	virtual void MarkPackagesUpToDate(TArrayView<const FName> UpToDatePackages) = 0;
 
-	/**
-	 * If a previous-version AssetRegistryState was allocated during GetCookedPackages, drop the reference to it
-	 * and return it to the caller who takes responsibility for deleting it.
-	 */
-	virtual FAssetRegistryState* ReleasePreviousAssetRegistry()
-	{
-		return nullptr;
-	}
-
 	/** Downcast function for ICookedPackageWriters that implement the IPackageStoreWriter inherited interface. */
 	virtual IPackageStoreWriter* AsPackageStoreWriter()
 	{
