Update Remote Compute to use new compact binary api

#rnx

diff --git a/Engine/Source/Editor/Experimental/BazelExecutor/BazelExecutor.Build.cs b/Engine/Source/Editor/Experimental/BazelExecutor/BazelExecutor.Build.cs
deleted file mode 100644
diff --git a/Engine/Source/Editor/Experimental/BazelExecutor/Generate.bat b/Engine/Source/Editor/Experimental/BazelExecutor/Generate.bat
deleted file mode 100755
diff --git a/Engine/Source/Editor/Experimental/BazelExecutor/Generate.sh b/Engine/Source/Editor/Experimental/BazelExecutor/Generate.sh
deleted file mode 100644
diff --git a/Engine/Source/Editor/Experimental/BazelExecutor/Private/BazelCompletionQueueRunnable.cpp b/Engine/Source/Editor/Experimental/BazelExecutor/Private/BazelCompletionQueueRunnable.cpp
deleted file mode 100644
diff --git a/Engine/Source/Editor/Experimental/BazelExecutor/Private/BazelCompletionQueueRunnable.h b/Engine/Source/Editor/Experimental/BazelExecutor/Private/BazelCompletionQueueRunnable.h
deleted file mode 100644
diff --git a/Engine/Source/Editor/Experimental/BazelExecutor/Private/BazelCompletionQueueRunnable.inl b/Engine/Source/Editor/Experimental/BazelExecutor/Private/BazelCompletionQueueRunnable.inl
deleted file mode 100644
diff --git a/Engine/Source/Editor/Experimental/BazelExecutor/Private/BazelExecutor.cpp b/Engine/Source/Editor/Experimental/BazelExecutor/Private/BazelExecutor.cpp
deleted file mode 100644
diff --git a/Engine/Source/Editor/Experimental/BazelExecutor/Private/BazelExecutor.h b/Engine/Source/Editor/Experimental/BazelExecutor/Private/BazelExecutor.h
deleted file mode 100644
diff --git a/Engine/Source/Editor/Experimental/BazelExecutor/Private/BazelExecutorModule.cpp b/Engine/Source/Editor/Experimental/BazelExecutor/Private/BazelExecutorModule.cpp
deleted file mode 100644
diff --git a/Engine/Source/Editor/Experimental/BazelExecutor/Private/BazelExecutorModule.h b/Engine/Source/Editor/Experimental/BazelExecutor/Private/BazelExecutorModule.h
deleted file mode 100644
diff --git a/Engine/Source/Editor/Experimental/BazelExecutor/Private/BazelExecutorSettings.cpp b/Engine/Source/Editor/Experimental/BazelExecutor/Private/BazelExecutorSettings.cpp
deleted file mode 100644
diff --git a/Engine/Source/Editor/Experimental/BazelExecutor/Private/BazelExecutorSettings.h b/Engine/Source/Editor/Experimental/BazelExecutor/Private/BazelExecutorSettings.h
deleted file mode 100644
diff --git a/Engine/Source/Editor/Experimental/BazelExecutor/Private/ContentAddressableStorage.cpp b/Engine/Source/Editor/Experimental/BazelExecutor/Private/ContentAddressableStorage.cpp
deleted file mode 100644
diff --git a/Engine/Source/Editor/Experimental/BazelExecutor/Private/ContentAddressableStorage.h b/Engine/Source/Editor/Experimental/BazelExecutor/Private/ContentAddressableStorage.h
deleted file mode 100644
diff --git a/Engine/Source/Editor/Experimental/BazelExecutor/Private/Execution.cpp b/Engine/Source/Editor/Experimental/BazelExecutor/Private/Execution.cpp
deleted file mode 100644
diff --git a/Engine/Source/Editor/Experimental/BazelExecutor/Private/Execution.h b/Engine/Source/Editor/Experimental/BazelExecutor/Private/Execution.h
deleted file mode 100644
diff --git a/Engine/Source/Editor/Experimental/BazelExecutor/Private/ProtoConverter.cpp b/Engine/Source/Editor/Experimental/BazelExecutor/Private/ProtoConverter.cpp
deleted file mode 100644
diff --git a/Engine/Source/Editor/Experimental/BazelExecutor/Private/ProtoConverter.h b/Engine/Source/Editor/Experimental/BazelExecutor/Private/ProtoConverter.h
deleted file mode 100644
diff --git a/Engine/Source/Editor/Experimental/BazelExecutor/SchemaAutogenHeader.txt b/Engine/Source/Editor/Experimental/BazelExecutor/SchemaAutogenHeader.txt
deleted file mode 100644
diff --git a/Engine/Source/Editor/Experimental/BazelExecutor/SchemaFooter.txt b/Engine/Source/Editor/Experimental/BazelExecutor/SchemaFooter.txt
deleted file mode 100644
diff --git a/Engine/Source/Editor/Experimental/BazelExecutor/SchemaHeader.txt b/Engine/Source/Editor/Experimental/BazelExecutor/SchemaHeader.txt
deleted file mode 100644
diff --git a/Engine/Source/Editor/Experimental/HordeExecutor/HordeExecutor.Build.cs b/Engine/Source/Editor/Experimental/HordeExecutor/HordeExecutor.Build.cs
new file mode 100644
--- /dev/null
+++ b/Engine/Source/Editor/Experimental/HordeExecutor/HordeExecutor.Build.cs
@@ -0,0 +1,25 @@
+// Copyright Epic Games, Inc. All Rights Reserved.
+
+using UnrealBuildTool;
+using System.IO;
+
+namespace UnrealBuildTool.Rules
+{
+	public class HordeExecutor : ModuleRules
+	{
+		public HordeExecutor(ReadOnlyTargetRules Target) : base(Target)
+		{
+			PrivateDependencyModuleNames.AddRange(
+				new string[] {
+				"Core",
+				"CoreUObject",
+				"Settings",
+				"HTTP",
+				"RemoteExecution",
+				}
+			);
+
+			OptimizeCode = CodeOptimization.Never;
+		}
+	}
+}
diff --git a/Engine/Source/Editor/Experimental/HordeExecutor/Private/ContentAddressableStorage.cpp b/Engine/Source/Editor/Experimental/HordeExecutor/Private/ContentAddressableStorage.cpp
new file mode 100644
--- /dev/null
+++ b/Engine/Source/Editor/Experimental/HordeExecutor/Private/ContentAddressableStorage.cpp
@@ -0,0 +1,335 @@
+// Copyright Epic Games, Inc. All Rights Reserved.
+
+#include "ContentAddressableStorage.h"
+#include "HordeExecutorModule.h"
+
+#include "IO/IoHash.h"
+#include "HttpModule.h"
+#include "Interfaces/IHttpRequest.h"
+#include "Interfaces/IHttpResponse.h"
+#include "Serialization/CompactBinary.h"
+
+#include "RemoteMessages.h"
+#include "HordeMessages.h"
+
+
+namespace UE::RemoteExecution
+{
+	FContentAddressableStorage::FContentAddressableStorage(const FString& BaseURL, const TMap<FString, FString> AdditionalHeaders)
+		: BaseURL(BaseURL)
+		, AdditionalHeaders(AdditionalHeaders)
+	{
+	}
+
+	bool FContentAddressableStorage::ToBlob(const IRemoteMessage& Message, TArray<uint8>& OutBlob, FIoHash& OutIoHash)
+	{
+		FCbObject CbObject = Message.Save();
+		FMemoryView MemoryView;
+		if (CbObject.TryGetView(MemoryView))
+		{
+			OutBlob.Empty();
+			OutBlob.Append((const uint8*)MemoryView.GetData(), MemoryView.GetSize());
+			OutIoHash = FIoHash(CbObject.GetHash());
+			return true;
+		}
+		return false;
+	}
+
+	FIoHash FContentAddressableStorage::ToDigest(const TArray<uint8>& Blob)
+	{
+		return FIoHash::HashBuffer(Blob.GetData(), Blob.Num());
+	}
+
+	TFuture<EStatusCode> FContentAddressableStorage::DoesBlobExistAsync(const FString& NameSpaceId, const FIoHash& Hash)
+	{
+		FStringFormatOrderedArguments Args;
+		Args.Add(NameSpaceId);
+		Args.Add(FString::FromHexBlob(Hash.GetBytes(), sizeof(FIoHash::ByteArray)));
+		const FString Route = FString::Format(TEXT("/api/v1/blobs/{0}/{1}"), Args);
+
+		TSharedRef<IHttpRequest> Request = FHttpModule::Get().CreateRequest();
+		Request->SetVerb(TEXT("HEAD"));
+		Request->SetURL(BaseURL + Route);
+		for (const TPair<FString, FString>& Header : AdditionalHeaders)
+		{
+			Request->SetHeader(Header.Key, Header.Value);
+		}
+
+		TSharedPtr<TPromise<EStatusCode>> ReturnPromise = MakeShared<TPromise<EStatusCode>>();
+		Request->OnProcessRequestComplete().BindLambda([ReturnPromise](FHttpRequestPtr Req, FHttpResponsePtr HttpResponse, bool bSucceeded) {
+			ReturnPromise->EmplaceValue((EStatusCode)HttpResponse->GetResponseCode());
+			});
+		Request->ProcessRequest();
+		return ReturnPromise->GetFuture();
+	}
+
+	TFuture<TMap<FIoHash, EStatusCode>> FContentAddressableStorage::DoBlobsExistAsync(const FString& NameSpaceId, const TSet<FIoHash>& Hashes)
+	{
+		FStringFormatOrderedArguments Args;
+		Args.Add(NameSpaceId);
+		const FString Route = FString::Format(TEXT("/api/v1/blobs/{0}/exists?"), Args);
+		FString Query;
+		for (const FIoHash& Hash : Hashes)
+		{
+			FStringFormatOrderedArguments QueryArgs;
+			QueryArgs.Add(Query);
+			QueryArgs.Add(FString::FromHexBlob(Hash.GetBytes(), sizeof(FIoHash::ByteArray)));
+			if (Query.IsEmpty())
+			{
+				Query = FString::Format(TEXT("id={1}"), QueryArgs);
+			}
+			else
+			{
+				Query = FString::Format(TEXT("{0}&id={1}"), QueryArgs);
+			}
+		}
+
+		TSharedRef<IHttpRequest> Request = FHttpModule::Get().CreateRequest();
+		Request->SetVerb(TEXT("POST"));
+		Request->SetURL(BaseURL + Route + Query);
+		Request->SetHeader(TEXT("Accept"), TEXT("application/x-ue-cb"));
+		for (const TPair<FString, FString>& Header : AdditionalHeaders)
+		{
+			Request->SetHeader(Header.Key, Header.Value);
+		}
+
+		TSharedPtr<TPromise<TMap<FIoHash, EStatusCode>>> ReturnPromise = MakeShared<TPromise<TMap<FIoHash, EStatusCode>>>();
+		Request->OnProcessRequestComplete().BindLambda([ReturnPromise, Hashes](FHttpRequestPtr Req, FHttpResponsePtr HttpResponse, bool bSucceeded) {
+			TMap<FIoHash, EStatusCode> Results;
+			if (bSucceeded && HttpResponse->GetResponseCode() == EHttpResponseCodes::Ok)
+			{
+				FExistsResponse ExistsResponse;
+				ExistsResponse.Load(FCbObjectView(HttpResponse->GetContent().GetData()));
+				for (const FIoHash& Hash : Hashes)
+				{
+					Results.Add(Hash, ExistsResponse.Id.Contains(Hash) ? EStatusCode::Ok : EStatusCode::NotFound);
+				}
+			}
+			else
+			{
+				for (const FIoHash& Hash : Hashes)
+				{
+					Results.Add(Hash, (EStatusCode)HttpResponse->GetResponseCode());
+				}
+			}
+			ReturnPromise->EmplaceValue(TMap<FIoHash, EStatusCode>(MoveTemp(Results)));
+			});
+		Request->ProcessRequest();
+		return ReturnPromise->GetFuture();
+	}
+
+	TFuture<EStatusCode> FContentAddressableStorage::PutBlobAsync(const FString& NameSpaceId, const FIoHash& Hash, const TArray<uint8>& Data)
+	{
+		FStringFormatOrderedArguments Args;
+		Args.Add(NameSpaceId);
+		Args.Add(FString::FromHexBlob(Hash.GetBytes(), sizeof(FIoHash::ByteArray)));
+		const FString Route = FString::Format(TEXT("/api/v1/blobs/{0}/{1}"), Args);
+
+		TSharedRef<IHttpRequest> Request = FHttpModule::Get().CreateRequest();
+		Request->SetVerb(TEXT("PUT"));
+		Request->SetURL(BaseURL + Route);
+		Request->SetHeader(TEXT("Content-Type"), TEXT("application/octet-stream"));
+		Request->SetHeader(TEXT("Accept"), TEXT("application/x-ue-cb"));
+		for (const TPair<FString, FString>& Header : AdditionalHeaders)
+		{
+			Request->SetHeader(Header.Key, Header.Value);
+		}
+		Request->SetContent(Data);
+
+		TSharedPtr<TPromise<EStatusCode>> ReturnPromise = MakeShared<TPromise<EStatusCode>>();
+		Request->OnProcessRequestComplete().BindLambda([ReturnPromise](FHttpRequestPtr Req, FHttpResponsePtr HttpResponse, bool bSucceeded) {
+			ReturnPromise->EmplaceValue((EHttpResponseCodes::Type)HttpResponse->GetResponseCode());
+			});
+		Request->ProcessRequest();
+		return ReturnPromise->GetFuture();
+	}
+
+	TFuture<TMap<FIoHash, EStatusCode>> FContentAddressableStorage::PutBlobsAsync(const FString& NameSpaceId, const TMap<FIoHash, TArray<uint8>>& Blobs)
+	{
+		TSharedPtr<TPromise<TMap<FIoHash, EStatusCode>>> ReturnPromise = MakeShared<TPromise<TMap<FIoHash, EStatusCode>>>();
+		TSharedPtr<TMap<FIoHash, EStatusCode>> Results = MakeShared<TMap<FIoHash, EStatusCode>>();
+		TSharedPtr<FCriticalSection> CriticalSection = MakeShared<FCriticalSection>();
+		for (const TPair<FIoHash, TArray<uint8>>& Blob : Blobs)
+		{
+			PutBlobAsync(NameSpaceId, Blob.Key, Blob.Value).Next([ReturnPromise, Results, CriticalSection, Hash = Blob.Key, Total = Blobs.Num()](EStatusCode&& Response) {
+				FScopeLock Lock(CriticalSection.Get());
+				Results->Add(Hash, MoveTemp(Response));
+				if (Results->Num() == Total) {
+					ReturnPromise->EmplaceValue(TMap<FIoHash, EStatusCode>(*Results.Get()));
+				}
+			});
+		}
+		return ReturnPromise->GetFuture();
+	}
+
+	TFuture<EStatusCode> FContentAddressableStorage::PutMissingBlobAsync(const FString& NameSpaceId, const FIoHash& Hash, const TArray<uint8>& Data)
+	{
+		TSharedPtr<TPromise<EStatusCode>> ReturnPromise = MakeShared<TPromise<EStatusCode>>();
+
+		DoesBlobExistAsync(NameSpaceId, Hash).Next([this, NameSpaceId, ReturnPromise, Hash, Data](EStatusCode&& ExistResponse) {
+			if (ExistResponse == EStatusCode::Ok)
+			{
+				ReturnPromise->EmplaceValue(ExistResponse);
+				return;
+			}
+
+			PutBlobAsync(NameSpaceId, Hash, Data).Next([ReturnPromise](EStatusCode&& PutResponse) {
+				ReturnPromise->EmplaceValue(PutResponse);
+				});
+			});
+		return ReturnPromise->GetFuture();
+	}
+
+	TFuture<TMap<FIoHash, EStatusCode>> FContentAddressableStorage::PutMissingBlobsAsync(const FString& NameSpaceId, const TMap<FIoHash, TArray<uint8>>& Blobs)
+	{
+		TSharedPtr<TPromise<TMap<FIoHash, EStatusCode>>> ReturnPromise = MakeShared<TPromise<TMap<FIoHash, EStatusCode>>>();
+
+		TSet<FIoHash> Hashes;
+		Blobs.GetKeys(Hashes);
+		DoBlobsExistAsync(NameSpaceId, Hashes).Next([this, NameSpaceId, ReturnPromise, Blobs](TMap<FIoHash, EStatusCode>&& ExistsResponse) {
+			const TMap<FIoHash, TArray<uint8>>& Remaining = Blobs.FilterByPredicate([ExistsResponse](const TPair<FIoHash, TArray<uint8>>& Pair) {
+				return ExistsResponse[Pair.Key] != EStatusCode::Ok;
+				});
+
+			if (Remaining.IsEmpty())
+			{
+				ReturnPromise->EmplaceValue(ExistsResponse);
+				return;
+			}
+
+			PutBlobsAsync(NameSpaceId, Remaining).Next([ReturnPromise, ExistsResponse](TMap<FIoHash, EStatusCode>&& PutResponse) {
+				for (const TPair<FIoHash, EStatusCode>& Exists : ExistsResponse)
+				{
+					if (!PutResponse.Contains(Exists.Key))
+					{
+						PutResponse.Add(Exists.Key, Exists.Value);
+					}
+				}
+				ReturnPromise->EmplaceValue(PutResponse);
+				});
+			});
+		return ReturnPromise->GetFuture();
+	}
+
+	TFuture<TPair<EStatusCode, TArray<uint8>>> FContentAddressableStorage::GetBlobAsync(const FString& NameSpaceId, const FIoHash& Hash)
+	{
+		FStringFormatOrderedArguments Args;
+		Args.Add(NameSpaceId);
+		Args.Add(FString::FromHexBlob(Hash.GetBytes(), sizeof(FIoHash::ByteArray)));
+		const FString Route = FString::Format(TEXT("/api/v1/blobs/{0}/{1}"), Args);
+
+		TSharedRef<IHttpRequest> Request = FHttpModule::Get().CreateRequest();
+		Request->SetVerb("GET");
+		Request->SetURL(BaseURL + Route);
+		Request->SetHeader(TEXT("Accept"), TEXT("application/octet-stream"));
+		for (const TPair<FString, FString>& Header : AdditionalHeaders)
+		{
+			Request->SetHeader(Header.Key, Header.Value);
+		}
+
+		TSharedPtr<TPromise<TPair<EStatusCode, TArray<uint8>>>> ReturnPromise = MakeShared<TPromise<TPair<EStatusCode, TArray<uint8>>>>();
+		Request->OnProcessRequestComplete().BindLambda([ReturnPromise](FHttpRequestPtr Req, FHttpResponsePtr HttpResponse, bool bSucceeded) {
+			ReturnPromise->EmplaceValue(TPair<EStatusCode, TArray<uint8>>((EStatusCode)HttpResponse->GetResponseCode(), HttpResponse->GetContent()));
+			});
+		Request->ProcessRequest();
+		return ReturnPromise->GetFuture();
+	}
+
+	TFuture<TMap<FIoHash, TPair<EStatusCode, TArray<uint8>>>> FContentAddressableStorage::GetBlobsAsync(const FString& NameSpaceId, const TSet<FIoHash>& Hashes)
+	{
+		TSharedPtr<TPromise<TMap<FIoHash, TPair<EStatusCode, TArray<uint8>>>>> ReturnPromise = MakeShared<TPromise<TMap<FIoHash, TPair<EStatusCode, TArray<uint8>>>>>();
+		TSharedPtr<TMap<FIoHash, TPair<EStatusCode, TArray<uint8>>>> Results = MakeShared<TMap<FIoHash, TPair<EStatusCode, TArray<uint8>>>>();
+		TSharedPtr<FCriticalSection> CriticalSection = MakeShared<FCriticalSection>();
+		for (const FIoHash& Hash : Hashes)
+		{
+			GetBlobAsync(NameSpaceId, Hash).Next([ReturnPromise, Results, CriticalSection, Hash, Total = Hashes.Num()](TPair<EStatusCode, TArray<uint8>>&& Response) {
+				FScopeLock Lock(CriticalSection.Get());
+				Results->Add(Hash, Response);
+				if (Results->Num() == Total) {
+					ReturnPromise->EmplaceValue(TMap<FIoHash, TPair<EStatusCode, TArray<uint8>>>(*Results.Get()));
+				}
+			});
+		}
+		return ReturnPromise->GetFuture();
+	}
+
+	TFuture<TPair<EStatusCode, FGetObjectTreeResponse>> FContentAddressableStorage::GetObjectTreeAsync(const FString& NameSpaceId, const FIoHash& Hash, const TSet<FIoHash>& HaveHashes)
+	{
+		FStringFormatOrderedArguments Args;
+		Args.Add(NameSpaceId);
+		Args.Add(FString::FromHexBlob(Hash.GetBytes(), sizeof(FIoHash::ByteArray)));
+		const FString Route = FString::Format(TEXT("/api/v1/objects/{0}/{1}/tree"), Args);
+
+		TSharedRef<IHttpRequest> Request = FHttpModule::Get().CreateRequest();
+		Request->SetVerb("POST");
+		Request->SetURL(BaseURL + Route);
+		Request->SetHeader(TEXT("Accept"), TEXT("application/octet-stream"));
+		Request->SetHeader(TEXT("Content-Type"), TEXT("application/x-ue-cb"));
+		for (const TPair<FString, FString>& Header : AdditionalHeaders)
+		{
+			Request->SetHeader(Header.Key, Header.Value);
+		}
+		{
+			FGetObjectTreeRequest ObjectTreeRequest;
+			ObjectTreeRequest.Have = HaveHashes;
+			FCbObject CbObject = ObjectTreeRequest.Save();
+			FMemoryView MemoryView;
+			if (!CbObject.TryGetView(MemoryView))
+			{
+				TPromise<TPair<EStatusCode, FGetObjectTreeResponse>> ReturnPromise;
+				ReturnPromise.SetValue(TPair<EStatusCode, FGetObjectTreeResponse>());
+				return ReturnPromise.GetFuture();
+			}
+			Request->SetContent(TArray<uint8>((const uint8*)MemoryView.GetData(), MemoryView.GetSize()));
+		}
+
+		TSharedPtr<TPromise<TPair<EStatusCode, FGetObjectTreeResponse>>> ReturnPromise = MakeShared<TPromise<TPair<EStatusCode, FGetObjectTreeResponse>>>();
+		Request->OnProcessRequestComplete().BindLambda([this, NameSpaceId, ReturnPromise](FHttpRequestPtr Req, FHttpResponsePtr HttpResponse, bool bSucceeded) {
+			EStatusCode ResponseCode = (EStatusCode)HttpResponse->GetResponseCode();
+			FGetObjectTreeResponse Response;
+			if (ResponseCode == EStatusCode::Ok)
+			{
+				const TArray<uint8>& Content = HttpResponse->GetContent();
+				int Index = 0;
+				// Repeating array of Hash then object
+				while (Index < Content.Num())
+				{
+					FCbFieldView Field = FCbFieldView(&Content[Index]);
+					Index += Field.GetSize();
+					if (Field.IsObjectAttachment())
+					{
+						FIoHash Hash = Field.AsObjectAttachment();
+						Field = FCbFieldView(&Content[Index]);
+						Index += Field.GetSize();
+						if (!Field.IsObject())
+						{
+							Response.Objects.Add(MoveTemp(Hash), TArray<uint8>());
+							continue;
+						}
+						FCbObjectView ObjectView = Field.AsObjectView();
+						FMemoryView MemoryView;
+						if (ObjectView.TryGetView(MemoryView))
+						{
+							Response.Objects.Add(MoveTemp(Hash), TArray<uint8>((const uint8*)MemoryView.GetData(), MemoryView.GetSize()));
+						}
+					}
+					else if (Field.IsBinaryAttachment())
+					{
+						Response.BinaryAttachments.Add(Field.AsBinaryAttachment());
+					}
+					else // Unknown type
+					{
+					}
+				}
+			}
+			ReturnPromise->EmplaceValue(TPair<EStatusCode, FGetObjectTreeResponse>(MoveTemp(ResponseCode), MoveTemp(Response)));
+			});
+		Request->ProcessRequest();
+		return ReturnPromise->GetFuture();
+	}
+
+	TFuture<TPair<EStatusCode, FGetObjectTreeResponse>> FContentAddressableStorage::GetObjectTreeAsync(const FString& NameSpaceId, const FIoHash& Hash)
+	{
+		return GetObjectTreeAsync(NameSpaceId, Hash, TSet<FIoHash>());
+	}
+}
diff --git a/Engine/Source/Editor/Experimental/HordeExecutor/Private/ContentAddressableStorage.h b/Engine/Source/Editor/Experimental/HordeExecutor/Private/ContentAddressableStorage.h
new file mode 100644
--- /dev/null
+++ b/Engine/Source/Editor/Experimental/HordeExecutor/Private/ContentAddressableStorage.h
@@ -0,0 +1,38 @@
+// Copyright Epic Games, Inc. All Rights Reserved.
+
+#pragma once
+
+#include "CoreMinimal.h"
+#include "IContentAddressableStorage.h"
+
+
+namespace UE::RemoteExecution
+{
+	class FContentAddressableStorage : public IContentAddressableStorage
+	{
+	private:
+		FString BaseURL;
+		TMap<FString, FString> AdditionalHeaders;
+
+	public:
+		FContentAddressableStorage(const FString& BaseURL, const TMap<FString, FString> AdditionalHeaders);
+
+		virtual bool ToBlob(const IRemoteMessage& Message, TArray<uint8>& OutBlob, FIoHash& OutIoHash) override;
+		virtual FIoHash ToDigest(const TArray<uint8>& Blob) override;
+
+		virtual TFuture<EStatusCode> DoesBlobExistAsync(const FString& NameSpaceId, const FIoHash& Hash) override;
+		virtual TFuture<TMap<FIoHash, EStatusCode>> DoBlobsExistAsync(const FString& NameSpaceId, const TSet<FIoHash>& Hashes) override;
+
+		virtual TFuture<EStatusCode> PutBlobAsync(const FString& NameSpaceId, const FIoHash& Hash, const TArray<uint8>& Data) override;
+		virtual TFuture<TMap<FIoHash, EStatusCode>> PutBlobsAsync(const FString& NameSpaceId, const TMap<FIoHash, TArray<uint8>>& Blobs) override;
+
+		virtual TFuture<EStatusCode> PutMissingBlobAsync(const FString& NameSpaceId, const FIoHash& Hash, const TArray<uint8>& Data) override;
+		virtual TFuture<TMap<FIoHash, EStatusCode>> PutMissingBlobsAsync(const FString& NameSpaceId, const TMap<FIoHash, TArray<uint8>>& Blobs) override;
+
+		virtual TFuture<TPair<EStatusCode, TArray<uint8>>> GetBlobAsync(const FString& NameSpaceId, const FIoHash& Hash) override;
+		virtual TFuture<TMap<FIoHash, TPair<EStatusCode, TArray<uint8>>>> GetBlobsAsync(const FString& NameSpaceId, const TSet<FIoHash>& Hashes) override;
+
+		virtual TFuture<TPair<EStatusCode, FGetObjectTreeResponse>> GetObjectTreeAsync(const FString& NameSpaceId, const FIoHash& Hash, const TSet<FIoHash>& HaveHashes) override;
+		virtual TFuture<TPair<EStatusCode, FGetObjectTreeResponse>> GetObjectTreeAsync(const FString& NameSpaceId, const FIoHash& Hash) override;
+	};
+}
diff --git a/Engine/Source/Editor/Experimental/HordeExecutor/Private/Execution.cpp b/Engine/Source/Editor/Experimental/HordeExecutor/Private/Execution.cpp
new file mode 100644
--- /dev/null
+++ b/Engine/Source/Editor/Experimental/HordeExecutor/Private/Execution.cpp
@@ -0,0 +1,146 @@
+// Copyright Epic Games, Inc. All Rights Reserved.
+
+#include "Execution.h"
+#include "HordeExecutorModule.h"
+#include "RemoteMessages.h"
+
+#include "HttpModule.h"
+#include "Interfaces/IHttpRequest.h"
+#include "Interfaces/IHttpResponse.h"
+#include "Serialization/CompactBinary.h"
+
+
+namespace UE::RemoteExecution
+{
+	struct FExecution::FWaitTask
+	{
+		TPromise<TPair<EStatusCode, FGetTaskUpdatesResponse>> Promise;
+		FString ChannelId;
+		TArray<FIoHash> TaskHashes;
+		TMap<FIoHash, FGetTaskUpdateResponse> Complete;
+		FDateTime Timeout;
+	};
+
+	FExecution::FExecution(const FString& BaseURL, const TMap<FString, FString> AdditionalHeaders)
+		: BaseURL(BaseURL)
+		, AdditionalHeaders(AdditionalHeaders)
+	{
+	}
+
+	TFuture<EStatusCode> FExecution::AddTasksAsync(const FString& ChannelId, const FAddTasksRequest& AddTaskRequest)
+	{
+		FStringFormatOrderedArguments Args;
+		Args.Add(ChannelId);
+		const FString Route = FString::Format(TEXT("/api/v1/compute/{0}"), Args);
+
+		TSharedRef<IHttpRequest> Request = FHttpModule::Get().CreateRequest();
+		Request->SetVerb(TEXT("POST"));
+		Request->SetURL(BaseURL + Route);
+		Request->SetHeader(TEXT("Content-Type"), TEXT("application/x-ue-cb"));
+		for (const TPair<FString, FString>& Header : AdditionalHeaders)
+		{
+			Request->SetHeader(Header.Key, Header.Value);
+		}
+
+		{
+			FCbObject CbObject = AddTaskRequest.Save();
+			FMemoryView MemoryView;
+			if (!CbObject.TryGetView(MemoryView))
+			{
+				TPromise<EStatusCode> ReturnPromise;
+				ReturnPromise.SetValue(EStatusCode::BadRequest);
+				return ReturnPromise.GetFuture();
+			}
+
+			TArray<uint8> Content = TArray<uint8>((const uint8*)MemoryView.GetData(), MemoryView.GetSize());
+			Request->SetContent(MoveTemp(Content));
+		}
+
+		TSharedPtr<TPromise<EStatusCode>> ReturnPromise = MakeShared<TPromise<EStatusCode>>();
+		Request->OnProcessRequestComplete().BindLambda([ReturnPromise](FHttpRequestPtr Req, FHttpResponsePtr HttpResponse, bool bSucceeded) {
+			ReturnPromise->EmplaceValue((EStatusCode)HttpResponse->GetResponseCode());
+			});
+		Request->ProcessRequest();
+		return ReturnPromise->GetFuture();
+	}
+
+	TFuture<TPair<EStatusCode, FGetTaskUpdatesResponse>> FExecution::GetUpdatesAsync(const FString& ChannelId, const int32 WaitSeconds)
+	{
+		FStringFormatOrderedArguments Args;
+		Args.Add(ChannelId);
+		Args.Add(WaitSeconds);
+		const FString Route = FString::Format(TEXT("/api/v1/compute/{0}/updates?wait={1}"), Args);
+
+		TSharedRef<IHttpRequest> Request = FHttpModule::Get().CreateRequest();
+		Request->SetVerb(TEXT("POST"));
+		Request->SetURL(BaseURL + Route);
+		Request->SetHeader(TEXT("Accept"), TEXT("application/x-ue-cb"));
+		for (const TPair<FString, FString>& Header : AdditionalHeaders)
+		{
+			Request->SetHeader(Header.Key, Header.Value);
+		}
+
+		TSharedPtr<TPromise<TPair<EStatusCode, FGetTaskUpdatesResponse>>> ReturnPromise = MakeShared<TPromise<TPair<EStatusCode, FGetTaskUpdatesResponse>>>();
+		Request->OnProcessRequestComplete().BindLambda([ReturnPromise](FHttpRequestPtr Req, FHttpResponsePtr HttpResponse, bool bSucceeded) {
+			FGetTaskUpdatesResponse GetTaskUpdatesResponse;
+			if (bSucceeded && HttpResponse->GetResponseCode() == EHttpResponseCodes::Ok)
+			{
+				const FCbObjectView View = FCbObjectView(HttpResponse->GetContent().GetData());
+				GetTaskUpdatesResponse.Load(View);
+			}
+			ReturnPromise->EmplaceValue(TPair<EStatusCode, FGetTaskUpdatesResponse>((EStatusCode)HttpResponse->GetResponseCode(), MoveTemp(GetTaskUpdatesResponse)));
+			});
+		Request->ProcessRequest();
+		return ReturnPromise->GetFuture();
+	}
+
+	TFuture<TPair<EStatusCode, FGetTaskUpdatesResponse>> FExecution::RunTasksAsync(const FAddTasksRequest& AddTaskRequest, const int32 TimeoutSeconds)
+	{
+		TSharedPtr<FWaitTask> WaitTask = MakeShared<FWaitTask>();
+		WaitTask->ChannelId = FGuid::NewGuid().ToString();
+		WaitTask->TaskHashes.Append(AddTaskRequest.TaskHashes);
+		WaitTask->Timeout = TimeoutSeconds > 0 ? FDateTime::Now() + FTimespan::FromSeconds(TimeoutSeconds) : FDateTime::MaxValue();
+
+		AddTasksAsync(WaitTask->ChannelId, AddTaskRequest).Next([this, WaitTask](EStatusCode&& Response) {
+			if (Response != EStatusCode::Ok)
+			{
+				WaitTask->Promise.EmplaceValue(TPair<EStatusCode, FGetTaskUpdatesResponse>(MoveTemp(Response), FGetTaskUpdatesResponse()));
+				return;
+			}
+			WaitForTasksAsync(WaitTask);
+			});
+
+		return WaitTask->Promise.GetFuture();
+
+	}
+
+	void FExecution::WaitForTasksAsync(TSharedPtr<FWaitTask> WaitTask)
+	{
+		GetUpdatesAsync(WaitTask->ChannelId, 5).Next([this, WaitTask](TPair<EStatusCode, FGetTaskUpdatesResponse>&& Response) {
+			for (FGetTaskUpdateResponse& Update : Response.Value.Updates)
+			{
+				UE_LOG(LogHordeExecutor, Verbose, TEXT("WaitForTasksAsync: %s %s"),
+					*TaskStateString(Update.State),
+					*FString::FromHexBlob(Update.TaskHash.GetBytes(), sizeof(FIoHash::ByteArray)));
+				if (Update.State == ETaskState::Complete)
+				{
+					WaitTask->Complete.Add(Update.TaskHash, MoveTemp(Update));
+				}
+			}
+
+			if (WaitTask->Complete.Num() != WaitTask->TaskHashes.Num() && Response.Key == EStatusCode::Ok && FDateTime::Now() > WaitTask->Timeout)
+			{
+				Response.Key = EStatusCode::RequestTimeout;
+			}
+
+			if (WaitTask->Complete.Num() == WaitTask->TaskHashes.Num() || Response.Key != EStatusCode::Ok)
+			{
+				FGetTaskUpdatesResponse Result;
+				WaitTask->Complete.GenerateValueArray(Result.Updates);
+				WaitTask->Promise.EmplaceValue(TPair<EStatusCode, FGetTaskUpdatesResponse>(MoveTemp(Response.Key), MoveTemp(Result)));
+				return;
+			}
+			WaitForTasksAsync(WaitTask);
+			});
+	}
+}
diff --git a/Engine/Source/Editor/Experimental/HordeExecutor/Private/Execution.h b/Engine/Source/Editor/Experimental/HordeExecutor/Private/Execution.h
new file mode 100644
--- /dev/null
+++ b/Engine/Source/Editor/Experimental/HordeExecutor/Private/Execution.h
@@ -0,0 +1,31 @@
+// Copyright Epic Games, Inc. All Rights Reserved.
+
+#pragma once
+
+#include "CoreMinimal.h"
+#include "IExecution.h"
+
+
+namespace UE::RemoteExecution
+{
+	class FExecution : public IExecution
+	{
+	private:
+		struct FWaitTask;
+
+		FString BaseURL;
+		TMap<FString, FString> AdditionalHeaders;
+
+	public:
+		FExecution(const FString& BaseURL, const TMap<FString, FString> AdditionalHeaders);
+
+		// Inherited via IExecution
+		virtual TFuture<EStatusCode> AddTasksAsync(const FString& ChannelId, const FAddTasksRequest& AddTaskRequest) override;
+		virtual TFuture<TPair<EStatusCode, FGetTaskUpdatesResponse>> GetUpdatesAsync(const FString& ChannelId, const int32 WaitSeconds = 0) override;
+
+		virtual TFuture<TPair<EStatusCode, FGetTaskUpdatesResponse>> RunTasksAsync(const FAddTasksRequest& AddTaskRequest, const int32 TimeoutSeconds = 0) override;
+
+	private:
+		void WaitForTasksAsync(TSharedPtr<FWaitTask> WaitTask);
+	};
+}
diff --git a/Engine/Source/Editor/Experimental/HordeExecutor/Private/HordeExecutor.cpp b/Engine/Source/Editor/Experimental/HordeExecutor/Private/HordeExecutor.cpp
new file mode 100644
--- /dev/null
+++ b/Engine/Source/Editor/Experimental/HordeExecutor/Private/HordeExecutor.cpp
@@ -0,0 +1,54 @@
+// Copyright Epic Games, Inc. All Rights Reserved.
+
+#include "HordeExecutor.h"
+
+#include "Execution.h"
+#include "ContentAddressableStorage.h"
+
+#define LOCTEXT_NAMESPACE "HordeExecutor"
+
+
+namespace UE::RemoteExecution
+{
+	void FHordeExecutor::Initialize(const FSettings& Settings)
+	{
+		Shutdown();
+
+		ContentAddressableStorage.Reset(new FContentAddressableStorage(Settings.ContentAddressableStorageTarget, Settings.ContentAddressableStorageHeaders));
+
+		Execution.Reset(new FExecution(Settings.ExecutionTarget, Settings.ExecutionHeaders));
+	}
+
+	void FHordeExecutor::Shutdown()
+	{
+		if (Thread.IsValid())
+		{
+			Thread->Kill();
+		}
+		Thread.Reset();
+		ContentAddressableStorage.Reset();
+		Execution.Reset();
+	}
+
+	FName FHordeExecutor::GetFName() const
+	{
+		return FName("Horde");
+	}
+
+	FText FHordeExecutor::GetNameText() const
+	{
+		return LOCTEXT("DefaultDisplayName", "Horde");
+	}
+
+	FText FHordeExecutor::GetDescriptionText() const
+	{
+		return LOCTEXT("DefaultDisplayDesc", "Horde remote execution.");
+	}
+
+	bool FHordeExecutor::CanRemoteExecute() const
+	{
+		return ContentAddressableStorage.IsValid() && Execution.IsValid();
+	}
+}
+
+#undef LOCTEXT_NAMESPACE
diff --git a/Engine/Source/Editor/Experimental/HordeExecutor/Private/HordeExecutor.h b/Engine/Source/Editor/Experimental/HordeExecutor/Private/HordeExecutor.h
new file mode 100644
--- /dev/null
+++ b/Engine/Source/Editor/Experimental/HordeExecutor/Private/HordeExecutor.h
@@ -0,0 +1,44 @@
+// Copyright Epic Games, Inc. All Rights Reserved.
+
+#pragma once
+
+#include "CoreMinimal.h"
+
+#include "IRemoteExecutor.h"
+#include "IContentAddressableStorage.h"
+#include "IExecution.h"
+#include "HAL/RunnableThread.h"
+
+
+namespace UE::RemoteExecution
+{
+	class FHordeExecutor : public IRemoteExecutor
+	{
+	public:
+		struct FSettings
+		{
+			FString ContentAddressableStorageTarget;
+			FString ExecutionTarget;
+			TMap<FString, FString> ContentAddressableStorageHeaders;
+			TMap<FString, FString> ExecutionHeaders;
+		};
+
+	private:
+		TUniquePtr<IContentAddressableStorage> ContentAddressableStorage;
+		TUniquePtr<IExecution> Execution;
+		TUniquePtr<FRunnableThread> Thread;
+
+	public:
+		void Initialize(const FSettings& Settings);
+		void Shutdown();
+
+		virtual FName GetFName() const override;
+		virtual FText GetNameText() const override;
+		virtual FText GetDescriptionText() const override;
+
+		virtual bool CanRemoteExecute() const override;
+
+		IContentAddressableStorage* GetContentAddressableStorage() const override { return ContentAddressableStorage.Get(); }
+		IExecution* GetExecution() const override { return Execution.Get(); }
+	};
+}
diff --git a/Engine/Source/Editor/Experimental/HordeExecutor/Private/HordeExecutorModule.cpp b/Engine/Source/Editor/Experimental/HordeExecutor/Private/HordeExecutorModule.cpp
new file mode 100644
--- /dev/null
+++ b/Engine/Source/Editor/Experimental/HordeExecutor/Private/HordeExecutorModule.cpp
@@ -0,0 +1,50 @@
+// Copyright Epic Games, Inc. All Rights Reserved.
+
+#include "HordeExecutorModule.h"
+#include "Misc/FileHelper.h"
+#include "Misc/Paths.h"
+#include "Modules/ModuleManager.h"
+#include "ISettingsModule.h"
+#include "HordeExecutorSettings.h"
+
+IMPLEMENT_MODULE(UE::RemoteExecution::FHordeExecutorModule, HordeExecutor);
+
+#define LOCTEXT_NAMESPACE "HordeExecutorModule"
+
+DEFINE_LOG_CATEGORY(LogHordeExecutor);
+
+static FName RemoteExecutionFeatureName(TEXT("RemoteExecution"));
+
+
+namespace UE::RemoteExecution
+{
+	void FHordeExecutorModule::StartupModule()
+	{
+		GetMutableDefault<UHordeExecutorSettings>()->LoadConfig();
+
+		const UHordeExecutorSettings* HordeExecutorSettings = GetDefault<UHordeExecutorSettings>();
+
+		FHordeExecutor::FSettings Settings;
+		Settings.ContentAddressableStorageTarget = HordeExecutorSettings->ContentAddressableStorageTarget;
+		Settings.ExecutionTarget = HordeExecutorSettings->ExecutionTarget;
+		Settings.ContentAddressableStorageHeaders = HordeExecutorSettings->ContentAddressableStorageHeaders;
+		Settings.ExecutionHeaders = HordeExecutorSettings->ExecutionHeaders;
+
+		HordeExecution.Initialize(Settings);
+
+		IModularFeatures::Get().RegisterModularFeature(RemoteExecutionFeatureName, &HordeExecution);
+	}
+
+	void FHordeExecutorModule::ShutdownModule()
+	{
+		HordeExecution.Shutdown();
+		IModularFeatures::Get().UnregisterModularFeature(RemoteExecutionFeatureName, &HordeExecution);
+	}
+
+	bool FHordeExecutorModule::SupportsDynamicReloading()
+	{
+		return true;
+	}
+}
+
+#undef LOCTEXT_NAMESPACE
diff --git a/Engine/Source/Editor/Experimental/HordeExecutor/Private/HordeExecutorModule.h b/Engine/Source/Editor/Experimental/HordeExecutor/Private/HordeExecutorModule.h
new file mode 100644
--- /dev/null
+++ b/Engine/Source/Editor/Experimental/HordeExecutor/Private/HordeExecutorModule.h
@@ -0,0 +1,29 @@
+// Copyright Epic Games, Inc. All Rights Reserved.
+
+#pragma once
+
+#include "CoreMinimal.h"
+#include "Modules/ModuleInterface.h"
+#include "Features/IModularFeatures.h"
+
+#include "HordeExecutor.h"
+
+HORDEEXECUTOR_API DECLARE_LOG_CATEGORY_EXTERN(LogHordeExecutor, Display, All);
+
+
+namespace UE::RemoteExecution
+{
+	class FHordeExecutorModule : public IModuleInterface
+	{
+	public:
+		/** IModuleInterface implementation */
+		virtual void StartupModule() override;
+
+		virtual void ShutdownModule() override;
+
+		virtual bool SupportsDynamicReloading() override;
+
+	private:
+		FHordeExecutor HordeExecution;
+	};
+}
diff --git a/Engine/Source/Editor/Experimental/HordeExecutor/Private/HordeExecutorSettings.cpp b/Engine/Source/Editor/Experimental/HordeExecutor/Private/HordeExecutorSettings.cpp
new file mode 100644
--- /dev/null
+++ b/Engine/Source/Editor/Experimental/HordeExecutor/Private/HordeExecutorSettings.cpp
@@ -0,0 +1,16 @@
+// Copyright Epic Games, Inc. All Rights Reserved.
+
+#include "HordeExecutorSettings.h"
+#include "Misc/ConfigCacheIni.h"
+
+UHordeExecutorSettings::UHordeExecutorSettings(const FObjectInitializer& ObjectInitializer)
+	: Super(ObjectInitializer)
+{
+	ContentAddressableStorageTarget = TEXT("https://horde.devtools-dev.epicgames.com");
+
+	ExecutionTarget = TEXT("https://horde.devtools-dev.epicgames.com");
+
+	ContentAddressableStorageHeaders.Add("Authorization", "ServiceAccount 0f8056b30bd0df0959be55fc3338159b6f938456d3474aed0087fb965268d079");
+
+	ExecutionHeaders.Add("Authorization", "ServiceAccount 0f8056b30bd0df0959be55fc3338159b6f938456d3474aed0087fb965268d079");
+}
diff --git a/Engine/Source/Editor/Experimental/HordeExecutor/Private/HordeExecutorSettings.h b/Engine/Source/Editor/Experimental/HordeExecutor/Private/HordeExecutorSettings.h
new file mode 100644
--- /dev/null
+++ b/Engine/Source/Editor/Experimental/HordeExecutor/Private/HordeExecutorSettings.h
@@ -0,0 +1,30 @@
+// Copyright Epic Games, Inc. All Rights Reserved.
+
+#pragma once
+
+#include "CoreMinimal.h"
+#include "UObject/Object.h"
+#include "HordeExecutorSettings.generated.h"
+
+UCLASS(config = EditorSettings)
+class UHordeExecutorSettings : public UObject
+{
+	GENERATED_UCLASS_BODY()
+
+public:
+	/** The Horde server content addressable storage address. */
+	UPROPERTY(Config, EditAnywhere, Category = "Horde", meta = (DisplayName = "Content Addressable Storage Target", ConfigRestartRequired = true))
+	FString ContentAddressableStorageTarget;
+
+	/** The Horde server execution address. */
+	UPROPERTY(Config, EditAnywhere, Category = "Horde", meta = (DisplayName = "Execution Target", ConfigRestartRequired = true))
+	FString ExecutionTarget;
+
+	/** Extra headers required for content addressable storage requests. */
+	UPROPERTY(Config, EditAnywhere, Category = "Horde", meta = (DisplayName = "Content Addressable Storage Headers", ConfigRestartRequired = true))
+	TMap<FString, FString> ContentAddressableStorageHeaders;
+
+	/** Extra headers required for execution requests. */
+	UPROPERTY(Config, EditAnywhere, Category = "Horde", meta = (DisplayName = "Execution Headers", ConfigRestartRequired = true))
+	TMap<FString, FString> ExecutionHeaders;
+};
diff --git a/Engine/Source/Editor/Experimental/HordeExecutor/Private/HordeMessages.cpp b/Engine/Source/Editor/Experimental/HordeExecutor/Private/HordeMessages.cpp
new file mode 100644
--- /dev/null
+++ b/Engine/Source/Editor/Experimental/HordeExecutor/Private/HordeMessages.cpp
@@ -0,0 +1,88 @@
+// Copyright Epic Games, Inc. All Rights Reserved.
+
+#include "HordeMessages.h"
+#include "Serialization/CompactBinaryWriter.h"
+
+
+namespace UE::RemoteExecution
+{
+	FCbObject FPutBlobResponse::Save() const
+	{
+		FCbWriter Writer;
+		Writer.BeginObject();
+		Writer.AddHash("h", Hash);
+		Writer.EndObject();
+		return Writer.Save().AsObject();
+	}
+
+	void FPutBlobResponse::Load(const FCbObjectView& CbObjectView)
+	{
+		Hash = CbObjectView["h"].AsHash();
+	}
+
+	FCbObject FExistsResponse::Save() const
+	{
+		FCbWriter Writer;
+		Writer.BeginObject();
+		if (!Id.IsEmpty())
+		{
+			Writer.BeginArray("id");
+			for (const FIoHash& Hash : Id)
+			{
+				Writer.AddHash(Hash);
+			}
+			Writer.EndArray();
+		}
+		Writer.EndObject();
+		return Writer.Save().AsObject();
+	}
+
+	void FExistsResponse::Load(const FCbObjectView& CbObjectView)
+	{
+		Id.Empty();
+		for (FCbFieldView& CbFieldView : CbObjectView["id"].AsArrayView())
+		{
+			Id.Add(CbFieldView.AsHash());
+		}
+	}
+
+	FCbObject FGetObjectTreeRequest::Save() const
+	{
+		FCbWriter Writer;
+		Writer.BeginObject();
+		if (!Have.IsEmpty())
+		{
+			Writer.BeginArray("h");
+			for (const FIoHash& Hash : Have)
+			{
+				Writer.AddObjectAttachment(Hash);
+			}
+			Writer.EndArray();
+		}
+		Writer.EndObject();
+		return Writer.Save().AsObject();
+	}
+
+	void FGetObjectTreeRequest::Load(const FCbObjectView& CbObjectView)
+	{
+		Have.Empty();
+		for (FCbFieldView& CbFieldView : CbObjectView["h"].AsArrayView())
+		{
+			Have.Add(CbFieldView.AsObjectAttachment());
+		}
+	}
+
+	FCbObject FPutObjectResponse::Save() const
+	{
+		FCbWriter Writer;
+		Writer.BeginObject();
+		Writer.AddObjectAttachment("id", Id);
+		Writer.EndObject();
+		return Writer.Save().AsObject();
+	}
+
+	void FPutObjectResponse::Load(const FCbObjectView& CbObjectView)
+	{
+		Id = CbObjectView["id"].AsObjectAttachment();
+	}
+}
diff --git a/Engine/Source/Editor/Experimental/HordeExecutor/Private/HordeMessages.h b/Engine/Source/Editor/Experimental/HordeExecutor/Private/HordeMessages.h
new file mode 100644
--- /dev/null
+++ b/Engine/Source/Editor/Experimental/HordeExecutor/Private/HordeMessages.h
@@ -0,0 +1,51 @@
+// Copyright Epic Games, Inc. All Rights Reserved.
+
+#pragma once
+
+#include "CoreMinimal.h"
+#include "IRemoteMessage.h"
+#include "IO/IoHash.h"
+
+
+namespace UE::RemoteExecution
+{
+	class FPutBlobResponse : public IRemoteMessage
+	{
+	public:
+		FIoHash Hash;
+
+		// Inherited via IMessage
+		virtual FCbObject Save() const override;
+		virtual void Load(const FCbObjectView& CbObjectView) override;
+	};
+
+	class FExistsResponse : public IRemoteMessage
+	{
+	public:
+		TSet<FIoHash> Id;
+
+		// Inherited via IMessage
+		virtual FCbObject Save() const override;
+		virtual void Load(const FCbObjectView& CbObjectView) override;
+	};
+
+	class FGetObjectTreeRequest : public IRemoteMessage
+	{
+	public:
+		TSet<FIoHash> Have;
+
+		// Inherited via IMessage
+		virtual FCbObject Save() const override;
+		virtual void Load(const FCbObjectView& CbObjectView) override;
+	};
+
+	class FPutObjectResponse : public IRemoteMessage
+	{
+	public:
+		FIoHash Id;
+
+		// Inherited via IMessage
+		virtual FCbObject Save() const override;
+		virtual void Load(const FCbObjectView& CbObjectView) override;
+	};
+}
diff --git a/Engine/Source/Editor/Experimental/RemoteExecution/Private/DefaultRemoteExecutor.cpp b/Engine/Source/Editor/Experimental/RemoteExecution/Private/DefaultRemoteExecutor.cpp
--- a/Engine/Source/Editor/Experimental/RemoteExecution/Private/DefaultRemoteExecutor.cpp
+++ b/Engine/Source/Editor/Experimental/RemoteExecution/Private/DefaultRemoteExecutor.cpp
@@ -4,34 +4,38 @@
 
 #define LOCTEXT_NAMESPACE "DefaultRemoteExecutor"
 
-FName FDefaultRemoteExecutor::GetFName() const
-{
-	return FName("None");
-}
-
-FText FDefaultRemoteExecutor::GetNameText() const
-{
-	return LOCTEXT("DefaultDisplayName", "None");
-}
-
-FText FDefaultRemoteExecutor::GetDescriptionText() const
-{
-	return LOCTEXT("DefaultDisplayDesc", "Disable remote execution.");
-}
-
-bool FDefaultRemoteExecutor::CanRemoteExecute() const
-{
-	return false;
-}
-
-IContentAddressableStorage* FDefaultRemoteExecutor::GetContentAddressableStorage() const
-{
-	return nullptr;
-}
 
-IExecution* FDefaultRemoteExecutor::GetExecution() const
+namespace UE::RemoteExecution
 {
-	return nullptr;
+	FName FDefaultRemoteExecutor::GetFName() const
+	{
+		return FName("None");
+	}
+
+	FText FDefaultRemoteExecutor::GetNameText() const
+	{
+		return LOCTEXT("DefaultDisplayName", "None");
+	}
+
+	FText FDefaultRemoteExecutor::GetDescriptionText() const
+	{
+		return LOCTEXT("DefaultDisplayDesc", "Disable remote execution.");
+	}
+
+	bool FDefaultRemoteExecutor::CanRemoteExecute() const
+	{
+		return false;
+	}
+
+	IContentAddressableStorage* FDefaultRemoteExecutor::GetContentAddressableStorage() const
+	{
+		return nullptr;
+	}
+
+	IExecution* FDefaultRemoteExecutor::GetExecution() const
+	{
+		return nullptr;
+	}
 }
 
 #undef LOCTEXT_NAMESPACE
diff --git a/Engine/Source/Editor/Experimental/RemoteExecution/Private/DefaultRemoteExecutor.h b/Engine/Source/Editor/Experimental/RemoteExecution/Private/DefaultRemoteExecutor.h
--- a/Engine/Source/Editor/Experimental/RemoteExecution/Private/DefaultRemoteExecutor.h
+++ b/Engine/Source/Editor/Experimental/RemoteExecution/Private/DefaultRemoteExecutor.h
@@ -5,15 +5,19 @@
 #include "CoreMinimal.h"
 #include "IRemoteExecutor.h"
 
-class FDefaultRemoteExecutor : public IRemoteExecutor
+
+namespace UE::RemoteExecution
 {
-public:
-	virtual FName GetFName() const override;
-	virtual FText GetNameText() const override;
-	virtual FText GetDescriptionText() const override;
+	class FDefaultRemoteExecutor : public IRemoteExecutor
+	{
+	public:
+		virtual FName GetFName() const override;
+		virtual FText GetNameText() const override;
+		virtual FText GetDescriptionText() const override;
 
-	virtual bool CanRemoteExecute() const override;
+		virtual bool CanRemoteExecute() const override;
 
-	virtual IContentAddressableStorage* GetContentAddressableStorage() const;
-	virtual IExecution* GetExecution() const;
-};
+		virtual IContentAddressableStorage* GetContentAddressableStorage() const;
+		virtual IExecution* GetExecution() const;
+	};
+}
diff --git a/Engine/Source/Editor/Experimental/RemoteExecution/Private/RemoteExecutionModule.cpp b/Engine/Source/Editor/Experimental/RemoteExecution/Private/RemoteExecutionModule.cpp
--- a/Engine/Source/Editor/Experimental/RemoteExecution/Private/RemoteExecutionModule.cpp
+++ b/Engine/Source/Editor/Experimental/RemoteExecution/Private/RemoteExecutionModule.cpp
@@ -4,10 +4,9 @@
 #include "Features/IModularFeatures.h"
 #include "ISettingsModule.h"
 #include "RemoteExecutionSettings.h"
-#include "Messages.h"
 #include "Modules/ModuleManager.h"
 
-IMPLEMENT_MODULE(FRemoteExecutionModule, RemoteExecution);
+IMPLEMENT_MODULE(UE::RemoteExecution::FRemoteExecutionModule, RemoteExecution);
 
 #define LOCTEXT_NAMESPACE "RemoteExecutionModule"
 
@@ -16,76 +15,79 @@ DEFINE_LOG_CATEGORY(LogRemoteExecution);
 static FName RemoteExecutionFeatureName(TEXT("RemoteExecution"));
 
 
-FRemoteExecutionModule::FRemoteExecutionModule()
-	: CurrentExecutor(nullptr)
+namespace UE::RemoteExecution
 {
-}
-
-void FRemoteExecutionModule::StartupModule()
-{
-	GetMutableDefault<URemoteExecutionSettings>()->LoadConfig();
+	FRemoteExecutionModule::FRemoteExecutionModule()
+		: CurrentExecutor(nullptr)
+	{
+	}
 
-	// Register to check for remote execution features
-	IModularFeatures::Get().OnModularFeatureRegistered().AddRaw(this, &FRemoteExecutionModule::HandleModularFeatureRegistered);
-	IModularFeatures::Get().OnModularFeatureUnregistered().AddRaw(this, &FRemoteExecutionModule::HandleModularFeatureUnregistered);
+	void FRemoteExecutionModule::StartupModule()
+	{
+		GetMutableDefault<URemoteExecutionSettings>()->LoadConfig();
 
-	// bind default accessor to editor
-	IModularFeatures::Get().RegisterModularFeature(RemoteExecutionFeatureName, &DefaultExecutor);
-}
+		// Register to check for remote execution features
+		IModularFeatures::Get().OnModularFeatureRegistered().AddRaw(this, &FRemoteExecutionModule::HandleModularFeatureRegistered);
+		IModularFeatures::Get().OnModularFeatureUnregistered().AddRaw(this, &FRemoteExecutionModule::HandleModularFeatureUnregistered);
 
-void FRemoteExecutionModule::ShutdownModule()
-{
-	// unbind default provider from editor
-	IModularFeatures::Get().UnregisterModularFeature(RemoteExecutionFeatureName, &DefaultExecutor);
+		// bind default accessor to editor
+		IModularFeatures::Get().RegisterModularFeature(RemoteExecutionFeatureName, &DefaultExecutor);
+	}
 
-	// we don't care about modular features any more
-	IModularFeatures::Get().OnModularFeatureRegistered().RemoveAll(this);
-	IModularFeatures::Get().OnModularFeatureUnregistered().RemoveAll(this);
-}
+	void FRemoteExecutionModule::ShutdownModule()
+	{
+		// unbind default provider from editor
+		IModularFeatures::Get().UnregisterModularFeature(RemoteExecutionFeatureName, &DefaultExecutor);
 
-bool FRemoteExecutionModule::CanRemoteExecute() const
-{
-	return CurrentExecutor->CanRemoteExecute();;
-}
+		// we don't care about modular features any more
+		IModularFeatures::Get().OnModularFeatureRegistered().RemoveAll(this);
+		IModularFeatures::Get().OnModularFeatureUnregistered().RemoveAll(this);
+	}
 
-IRemoteExecutor& FRemoteExecutionModule::GetRemoteExecutor() const
-{
-	return *CurrentExecutor;
-}
+	bool FRemoteExecutionModule::CanRemoteExecute() const
+	{
+		return CurrentExecutor->CanRemoteExecute();;
+	}
 
-void FRemoteExecutionModule::SetRemoteExecutor(const FName& InName)
-{
-	const int32 FeatureCount = IModularFeatures::Get().GetModularFeatureImplementationCount(RemoteExecutionFeatureName);
-	for (int32 FeatureIndex = 0; FeatureIndex < FeatureCount; FeatureIndex++)
+	IRemoteExecutor& FRemoteExecutionModule::GetRemoteExecutor() const
 	{
-		IModularFeature* Feature = IModularFeatures::Get().GetModularFeatureImplementation(RemoteExecutionFeatureName, FeatureIndex);
-		check(Feature);
+		return *CurrentExecutor;
+	}
 
-		IRemoteExecutor& RemoteExecution = *static_cast<IRemoteExecutor*>(Feature);
-		if (InName == RemoteExecution.GetFName())
+	void FRemoteExecutionModule::SetRemoteExecutor(const FName& InName)
+	{
+		const int32 FeatureCount = IModularFeatures::Get().GetModularFeatureImplementationCount(RemoteExecutionFeatureName);
+		for (int32 FeatureIndex = 0; FeatureIndex < FeatureCount; FeatureIndex++)
 		{
-			CurrentExecutor = static_cast<IRemoteExecutor*>(Feature);
-			break;
+			IModularFeature* Feature = IModularFeatures::Get().GetModularFeatureImplementation(RemoteExecutionFeatureName, FeatureIndex);
+			check(Feature);
+
+			IRemoteExecutor& RemoteExecution = *static_cast<IRemoteExecutor*>(Feature);
+			if (InName == RemoteExecution.GetFName())
+			{
+				CurrentExecutor = static_cast<IRemoteExecutor*>(Feature);
+				break;
+			}
 		}
 	}
-}
 
-void FRemoteExecutionModule::HandleModularFeatureRegistered(const FName& Type, IModularFeature* ModularFeature)
-{
-	if (Type == RemoteExecutionFeatureName)
+	void FRemoteExecutionModule::HandleModularFeatureRegistered(const FName& Type, IModularFeature* ModularFeature)
 	{
-		const FString PreferredRemoteExecutor = GetDefault<URemoteExecutionSettings>()->PreferredRemoteExecutor;
+		if (Type == RemoteExecutionFeatureName)
+		{
+			const FString PreferredRemoteExecutor = GetDefault<URemoteExecutionSettings>()->PreferredRemoteExecutor;
 
-		CurrentExecutor = &DefaultExecutor;
-		SetRemoteExecutor(FName(PreferredRemoteExecutor));
+			CurrentExecutor = &DefaultExecutor;
+			SetRemoteExecutor(FName(PreferredRemoteExecutor));
+		}
 	}
-}
 
-void FRemoteExecutionModule::HandleModularFeatureUnregistered(const FName& Type, IModularFeature* ModularFeature)
-{
-	if (Type == RemoteExecutionFeatureName && CurrentExecutor == static_cast<IRemoteExecutor*>(ModularFeature))
+	void FRemoteExecutionModule::HandleModularFeatureUnregistered(const FName& Type, IModularFeature* ModularFeature)
 	{
-		CurrentExecutor = &DefaultExecutor;
+		if (Type == RemoteExecutionFeatureName && CurrentExecutor == static_cast<IRemoteExecutor*>(ModularFeature))
+		{
+			CurrentExecutor = &DefaultExecutor;
+		}
 	}
 }
 
diff --git a/Engine/Source/Editor/Experimental/RemoteExecution/Private/RemoteExecutionModule.h b/Engine/Source/Editor/Experimental/RemoteExecution/Private/RemoteExecutionModule.h
--- a/Engine/Source/Editor/Experimental/RemoteExecution/Private/RemoteExecutionModule.h
+++ b/Engine/Source/Editor/Experimental/RemoteExecution/Private/RemoteExecutionModule.h
@@ -6,36 +6,36 @@
 #include "IRemoteExecutionModule.h"
 #include "DefaultRemoteExecutor.h"
 
-
 REMOTEEXECUTION_API DECLARE_LOG_CATEGORY_EXTERN(LogRemoteExecution, Display, All);
 
-class FRemoteExecutionModule : public IRemoteExecutionModule
-{
-
-public:
 
-	/** Default constructor. */
-	FRemoteExecutionModule();
-
-	// IModuleInterface interface
+namespace UE::RemoteExecution
+{
+	class FRemoteExecutionModule : public IRemoteExecutionModule
+	{
+	public:
+		/** Default constructor. */
+		FRemoteExecutionModule();
 
-	virtual void StartupModule() override;
-	virtual void ShutdownModule() override;
+		// IModuleInterface interface
 
-	// IRemoteExecutionModule interface
-	virtual bool CanRemoteExecute() const override;
-	virtual IRemoteExecutor& GetRemoteExecutor() const override;
-	virtual void SetRemoteExecutor(const FName& InName);
+		virtual void StartupModule() override;
+		virtual void ShutdownModule() override;
 
-private:
-	/** Handle when one of the modular features we are interested in is registered */
-	void HandleModularFeatureRegistered(const FName& Type, IModularFeature* ModularFeature);
+		// IRemoteExecutionModule interface
+		virtual bool CanRemoteExecute() const override;
+		virtual IRemoteExecutor& GetRemoteExecutor() const override;
+		virtual void SetRemoteExecutor(const FName& InName);
 
-	/** Handle when one of the modular features we are interested in is unregistered */
-	void HandleModularFeatureUnregistered(const FName& Type, IModularFeature* ModularFeature);
+	private:
+		/** Handle when one of the modular features we are interested in is registered */
+		void HandleModularFeatureRegistered(const FName& Type, IModularFeature* ModularFeature);
 
-	FDefaultRemoteExecutor DefaultExecutor;
+		/** Handle when one of the modular features we are interested in is unregistered */
+		void HandleModularFeatureUnregistered(const FName& Type, IModularFeature* ModularFeature);
 
-	IRemoteExecutor* CurrentExecutor;
+		FDefaultRemoteExecutor DefaultExecutor;
 
-};
+		IRemoteExecutor* CurrentExecutor;
+	};
+}
diff --git a/Engine/Source/Editor/Experimental/RemoteExecution/Private/RemoteMessages.cpp b/Engine/Source/Editor/Experimental/RemoteExecution/Private/RemoteMessages.cpp
new file mode 100644
--- /dev/null
+++ b/Engine/Source/Editor/Experimental/RemoteExecution/Private/RemoteMessages.cpp
@@ -0,0 +1,355 @@
+// Copyright Epic Games, Inc. All Rights Reserved.
+
+#include "RemoteMessages.h"
+#include "Serialization/CompactBinaryWriter.h"
+
+
+namespace UE::RemoteExecution
+{
+	FAddTasksRequest::FAddTasksRequest()
+		: DoNotCache(false)
+	{
+	}
+
+	FCbObject FAddTasksRequest::Save() const
+	{
+		FCbWriter Writer;
+		Writer.BeginObject();
+		if (RequirementsHash != FIoHash::Zero)
+		{
+			Writer.AddObjectAttachment("r", RequirementsHash);
+		}
+		if (!TaskHashes.IsEmpty())
+		{
+			Writer.BeginArray("t");
+			for (const FIoHash& Hash : TaskHashes)
+			{
+				Writer.AddObjectAttachment(Hash);
+			}
+			Writer.EndArray();
+		}
+		Writer.AddBool("nc", DoNotCache);
+		Writer.EndObject();
+		return Writer.Save().AsObject();
+	}
+
+	void FAddTasksRequest::Load(const FCbObjectView& CbObjectView)
+	{
+		TaskHashes.Empty();
+		RequirementsHash = CbObjectView["n"].AsObjectAttachment();
+		for (FCbFieldView& CbFieldView : CbObjectView["t"].AsArrayView())
+		{
+			TaskHashes.Add(CbFieldView.AsObjectAttachment());
+		}
+		DoNotCache = CbObjectView["nc"].AsBool();
+	}
+
+	FGetTaskUpdateResponse::FGetTaskUpdateResponse()
+		: Time(0)
+		, State(ETaskState::Queued)
+	{
+	}
+
+	FCbObject FGetTaskUpdateResponse::Save() const
+	{
+		FCbWriter Writer;
+		Writer.BeginObject();
+		if (TaskHash != FIoHash::Zero)
+		{
+			Writer.AddObjectAttachment("h", TaskHash);
+		}
+		Writer.AddDateTime("t", Time);
+		Writer.AddInteger("s", (int32)State);
+		if (ResultHash != FIoHash::Zero)
+		{
+			Writer.AddObjectAttachment("r", ResultHash);
+		}
+		Writer.AddString("a", AgentId);
+		Writer.AddString("l", LeaseId);
+		Writer.EndObject();
+		return Writer.Save().AsObject();
+	}
+
+	void FGetTaskUpdateResponse::Load(const FCbObjectView& CbObjectView)
+	{
+		TaskHash = CbObjectView["h"].AsObjectAttachment();
+		Time = CbObjectView["t"].AsDateTime();
+		State = (ETaskState)CbObjectView["s"].AsInt32();
+		ResultHash = CbObjectView["r"].AsObjectAttachment();
+		AgentId = FString(CbObjectView["a"].AsString());
+		LeaseId = FString(CbObjectView["l"].AsString());
+	}
+
+	FCbObject FGetTaskUpdatesResponse::Save() const
+	{
+		FCbWriter Writer;
+		Writer.BeginObject();
+		if (!Updates.IsEmpty())
+		{
+			Writer.BeginArray("u");
+			for (const FGetTaskUpdateResponse& Update : Updates)
+			{
+				Writer.AddObject(Update.Save());
+			}
+			Writer.EndArray();
+		}
+		Writer.EndObject();
+		return Writer.Save().AsObject();
+	}
+
+	void FGetTaskUpdatesResponse::Load(const FCbObjectView& CbObjectView)
+	{
+		Updates.Empty();
+		for (FCbFieldView& CbFieldView : CbObjectView["u"].AsArrayView())
+		{
+			FGetTaskUpdateResponse Response;
+			Response.Load(CbFieldView.AsObjectView());
+			Updates.Add(MoveTemp(Response));
+		}
+	}
+
+	FFileNode::FFileNode()
+		: Size(0)
+		, Attributes(0)
+	{
+	}
+
+	FCbObject FFileNode::Save() const
+	{
+		FCbWriter Writer;
+		Writer.BeginObject();
+		Writer.AddString("n", Name);
+		Writer.AddBinaryAttachment("h", Hash);
+		Writer.AddInteger("s", Size);
+		Writer.AddInteger("a", Attributes);
+		Writer.EndObject();
+		return Writer.Save().AsObject();
+	}
+
+	void FFileNode::Load(const FCbObjectView& CbObjectView)
+	{
+		Name = FString(CbObjectView["n"].AsString());
+		Hash = CbObjectView["h"].AsBinaryAttachment();
+		Size = CbObjectView["s"].AsInt64();
+		Attributes = CbObjectView["a"].AsInt32();
+	}
+
+	FCbObject FDirectoryNode::Save() const
+	{
+		FCbWriter Writer;
+		Writer.BeginObject();
+		Writer.AddString("n", Name);
+		Writer.AddObjectAttachment("h", Hash);
+		Writer.EndObject();
+		return Writer.Save().AsObject();
+	}
+
+	void FDirectoryNode::Load(const FCbObjectView& CbObjectView)
+	{
+		Name = FString(CbObjectView["n"].AsString());
+		Hash = CbObjectView["h"].AsObjectAttachment();
+	}
+
+	FCbObject FDirectoryTree::Save() const
+	{
+		FCbWriter Writer;
+		Writer.BeginObject();
+		if (!Files.IsEmpty())
+		{
+			Writer.BeginArray("f");
+			for (const FFileNode& File : Files)
+			{
+				Writer.AddObject(File.Save().AsView());
+			}
+			Writer.EndArray();
+		}
+		if (!Directories.IsEmpty())
+		{
+			Writer.BeginArray("d");
+			for (const FDirectoryNode& Directory : Directories)
+			{
+				Writer.AddObject(Directory.Save().AsView());
+			}
+			Writer.EndArray();
+		}
+		Writer.EndObject();
+		return Writer.Save().AsObject();
+	}
+
+	void FDirectoryTree::Load(const FCbObjectView& CbObjectView)
+	{
+		Files.Empty();
+		Directories.Empty();
+
+		for (FCbFieldView& CbFieldView : CbObjectView["f"].AsArrayView())
+		{
+			FFileNode FileNode;
+			FileNode.Load(CbFieldView.AsObjectView());
+			Files.Add(MoveTemp(FileNode));
+		}
+		for (FCbFieldView& CbFieldView : CbObjectView["d"].AsArrayView())
+		{
+			FDirectoryNode DirectoryNode;
+			DirectoryNode.Load(CbFieldView.AsObjectView());
+			Directories.Add(MoveTemp(DirectoryNode));
+		}
+	}
+
+	FRequirements::FRequirements()
+		: Exclusive(false)
+	{
+	}
+
+	FCbObject FRequirements::Save() const
+	{
+		FCbWriter Writer;
+		Writer.BeginObject();
+		Writer.AddString("c", Condition);
+		if (!Resources.IsEmpty())
+		{
+			Writer.BeginArray("r");
+			for (const TPair<FString, int32>& Resource : Resources)
+			{
+				Writer.BeginArray();
+				Writer.AddString(Resource.Key);
+				Writer.AddInteger(Resource.Value);
+				Writer.EndArray();
+			}
+			Writer.EndArray();
+		}
+		Writer.AddBool("e", Exclusive);
+		Writer.EndObject();
+		return Writer.Save().AsObject();
+	}
+
+	void FRequirements::Load(const FCbObjectView& CbObjectView)
+	{
+		Resources.Empty();
+
+		Condition = FString(CbObjectView["c"].AsString());
+		for (FCbFieldView& CbFieldView : CbObjectView["r"].AsArrayView())
+		{
+			FCbArrayView Pair = CbFieldView.AsArrayView();
+			if (Pair.Num() != 2)
+			{
+				continue;
+			}
+			FCbFieldViewIterator Iter = Pair.CreateViewIterator();
+			FString Key = FString(Iter.AsString());
+			int32 Value = (++Iter).AsInt32();
+			Resources.Add(MoveTemp(Key), MoveTemp(Value));
+		}
+		Exclusive = CbObjectView["e"].AsBool();
+	}
+
+	FCbObject FTask::Save() const
+	{
+		FCbWriter Writer;
+		Writer.BeginObject();
+		Writer.AddString("e", Executable);
+		if (!Arguments.IsEmpty())
+		{
+			Writer.BeginArray("a");
+			for (const FString& Argument : Arguments)
+			{
+				Writer.AddString(Argument);
+			}
+			Writer.EndArray();
+		}
+		if (!EnvVars.IsEmpty())
+		{
+			Writer.BeginArray("v");
+			for (const TPair<FString, FString>& EnvVar : EnvVars)
+			{
+				Writer.BeginArray();
+				Writer.AddString(EnvVar.Key);
+				Writer.AddString(EnvVar.Value);
+				Writer.EndArray();
+			}
+			Writer.EndArray();
+		}
+		Writer.AddString("w", WorkingDirectory);
+		Writer.AddObjectAttachment("s", SandboxHash);
+		if (RequirementsHash != FIoHash::Zero)
+		{
+			Writer.AddObjectAttachment("r", RequirementsHash);
+		}
+		if (!OutputPaths.IsEmpty())
+		{
+			Writer.BeginArray("o");
+			for (const FString& OutputPath : OutputPaths)
+			{
+				Writer.AddString(OutputPath);
+			}
+			Writer.EndArray();
+		}
+		Writer.EndObject();
+		return Writer.Save().AsObject();
+	}
+
+	void FTask::Load(const FCbObjectView& CbObjectView)
+	{
+		Arguments.Empty();
+		EnvVars.Empty();
+		OutputPaths.Empty();
+
+		Executable = FString(CbObjectView["e"].AsString());
+		for (FCbFieldView& CbFieldView : CbObjectView["a"].AsArrayView())
+		{
+			Arguments.Add(FString(CbFieldView.AsString()));
+		}
+		for (FCbFieldView& CbFieldView : CbObjectView["v"].AsArrayView())
+		{
+			FCbArrayView Pair = CbFieldView.AsArrayView();
+			if (Pair.Num() != 2)
+			{
+				continue;
+			}
+			FCbFieldViewIterator Iter = Pair.CreateViewIterator();
+			FString Key = FString(Iter.AsString());
+			FString Value = FString((++Iter).AsString());
+			EnvVars.Add(MoveTemp(Key), MoveTemp(Value));
+		}
+		WorkingDirectory = FString(CbObjectView["w"].AsString());
+		SandboxHash = CbObjectView["s"].AsObjectAttachment();
+		RequirementsHash = CbObjectView["r"].AsObjectAttachment();
+		for (FCbFieldView& CbFieldView : CbObjectView["o"].AsArrayView())
+		{
+			OutputPaths.Add(FString(CbFieldView.AsString()));
+		}
+	}
+
+	FTaskResult::FTaskResult()
+		: ExitCode(0)
+	{
+	}
+
+	FCbObject FTaskResult::Save() const
+	{
+		FCbWriter Writer;
+		Writer.BeginObject();
+		Writer.AddInteger("e", ExitCode);
+		if (!StdOutHash.IsZero())
+		{
+			Writer.AddBinaryAttachment("so", StdOutHash);
+		}
+		if (!StdErrHash.IsZero())
+		{
+			Writer.AddBinaryAttachment("se", StdErrHash);
+		}
+		if (!OutputHash.IsZero())
+		{
+			Writer.AddObjectAttachment("o", OutputHash);
+		}
+		Writer.EndObject();
+		return Writer.Save().AsObject();
+	}
+
+	void FTaskResult::Load(const FCbObjectView& CbObjectView)
+	{
+		ExitCode = CbObjectView["e"].AsInt32();
+		StdOutHash = CbObjectView["so"].AsBinaryAttachment();
+		StdErrHash = CbObjectView["se"].AsBinaryAttachment();
+		OutputHash = CbObjectView["o"].AsObjectAttachment();
+	}
+}
diff --git a/Engine/Source/Editor/Experimental/RemoteExecution/Public/IContentAddressableStorage.h b/Engine/Source/Editor/Experimental/RemoteExecution/Public/IContentAddressableStorage.h
--- a/Engine/Source/Editor/Experimental/RemoteExecution/Public/IContentAddressableStorage.h
+++ b/Engine/Source/Editor/Experimental/RemoteExecution/Public/IContentAddressableStorage.h
@@ -5,35 +5,37 @@
 #include "CoreMinimal.h"
 #include "Async/Future.h"
 
-struct FAction;
-struct FBatchReadBlobsRequest;
-struct FBatchReadBlobsResponse;
-struct FBatchUpdateBlobsRequest;
-struct FBatchUpdateBlobsResponse;
-struct FCommand;
-struct FDigest;
-struct FDirectory;
-struct FFindMissingBlobsRequest;
-struct FFindMissingBlobsResponse;
-struct FStatus;
-
-
-class IContentAddressableStorage
+struct FIoHash;
+
+
+namespace UE::RemoteExecution
 {
-public:
-	/** Virtual destructor */
-	virtual ~IContentAddressableStorage() {}
-
-	virtual bool ToDigest(const TArray<uint8>& InData, FDigest& OutDigest) = 0;
-	virtual bool ToBlob(const FDirectory& InDirectory, TArray<uint8>& OutData, FDigest& OutDigest) = 0;
-	virtual bool ToBlob(const FCommand& InCommand, TArray<uint8>& OutData, FDigest& OutDigest) = 0;
-	virtual bool ToBlob(const FAction& InAction, TArray<uint8>& OutData, FDigest& OutDigest) = 0;
-
-	virtual FStatus FindMissingBlobs(const FFindMissingBlobsRequest& Request, FFindMissingBlobsResponse& Response, int64 TimeoutMs = 0) = 0;
-	virtual FStatus BatchUpdateBlobs(const FBatchUpdateBlobsRequest& Request, FBatchUpdateBlobsResponse& Response, int64 TimeoutMs = 0) = 0;
-	virtual FStatus BatchReadBlobs(const FBatchReadBlobsRequest& Request, FBatchReadBlobsResponse& Response, int64 TimeoutMs = 0) = 0;
-
-	virtual TFuture<TPair<FStatus, FFindMissingBlobsResponse>> FindMissingBlobsAsync(const FFindMissingBlobsRequest& Request, TUniqueFunction<void()> CompletionCallback = nullptr, int64 TimeoutMs = 0) = 0;
-	virtual TFuture<TPair<FStatus, FBatchUpdateBlobsResponse>> BatchUpdateBlobsAsync(const FBatchUpdateBlobsRequest& Request, TUniqueFunction<void()> CompletionCallback = nullptr, int64 TimeoutMs = 0) = 0;
-	virtual TFuture<TPair<FStatus, FBatchReadBlobsResponse>> BatchReadBlobsAsync(const FBatchReadBlobsRequest& Request, TUniqueFunction<void()> CompletionCallback = nullptr, int64 TimeoutMs = 0) = 0;
-};
+	enum class EStatusCode;
+	class IRemoteMessage;
+	struct FGetObjectTreeResponse;
+
+	class IContentAddressableStorage
+	{
+	public:
+		/** Virtual destructor */
+		virtual ~IContentAddressableStorage() {}
+
+		virtual bool ToBlob(const IRemoteMessage& Message, TArray<uint8>& OutBlob, FIoHash& OutIoHash) = 0;
+		virtual FIoHash ToDigest(const TArray<uint8>& Blob) = 0;
+
+		virtual TFuture<EStatusCode> DoesBlobExistAsync(const FString& NameSpaceId, const FIoHash& Hash) = 0;
+		virtual TFuture<TMap<FIoHash, EStatusCode>> DoBlobsExistAsync(const FString& NameSpaceId, const TSet<FIoHash>& Hashes) = 0;
+
+		virtual TFuture<EStatusCode> PutBlobAsync(const FString& NameSpaceId, const FIoHash& Hash, const TArray<uint8>& Data) = 0;
+		virtual TFuture<TMap<FIoHash, EStatusCode>> PutBlobsAsync(const FString& NameSpaceId, const TMap<FIoHash, TArray<uint8>>& Blobs) = 0;
+
+		virtual TFuture<EStatusCode> PutMissingBlobAsync(const FString& NameSpaceId, const FIoHash& Hash, const TArray<uint8>& Data) = 0;
+		virtual TFuture<TMap<FIoHash, EStatusCode>> PutMissingBlobsAsync(const FString& NameSpaceId, const TMap<FIoHash, TArray<uint8>>& Blobs) = 0;
+
+		virtual TFuture<TPair<EStatusCode, TArray<uint8>>> GetBlobAsync(const FString& NameSpaceId, const FIoHash& Hash) = 0;
+		virtual TFuture<TMap<FIoHash, TPair<EStatusCode, TArray<uint8>>>> GetBlobsAsync(const FString& NameSpaceId, const TSet<FIoHash>& Hashes) = 0;
+
+		virtual TFuture<TPair<EStatusCode, FGetObjectTreeResponse>> GetObjectTreeAsync(const FString& NameSpaceId, const FIoHash& Hash, const TSet<FIoHash>& HaveHashes) = 0;
+		virtual TFuture<TPair<EStatusCode, FGetObjectTreeResponse>> GetObjectTreeAsync(const FString& NameSpaceId, const FIoHash& Hash) = 0;
+	};
+}
diff --git a/Engine/Source/Editor/Experimental/RemoteExecution/Public/IExecution.h b/Engine/Source/Editor/Experimental/RemoteExecution/Public/IExecution.h
--- a/Engine/Source/Editor/Experimental/RemoteExecution/Public/IExecution.h
+++ b/Engine/Source/Editor/Experimental/RemoteExecution/Public/IExecution.h
@@ -5,17 +5,22 @@
 #include "CoreMinimal.h"
 #include "Async/Future.h"
 
-struct FExecuteRequest;
-struct FExecuteResponse;
 
-
-class IExecution
+namespace UE::RemoteExecution
 {
-public:
-	/** Virtual destructor */
-	virtual ~IExecution() {}
+	enum class EStatusCode;
+	class FAddTasksRequest;
+	class FGetTaskUpdatesResponse;
+
+	class IExecution
+	{
+	public:
+		/** Virtual destructor */
+		virtual ~IExecution() {}
 
-	virtual bool Execute(const FExecuteRequest& Request, FExecuteResponse& Response, int64 TimeoutMs = 0) = 0;
+		virtual TFuture<EStatusCode> AddTasksAsync(const FString& ChannelId, const FAddTasksRequest& Request) = 0;
+		virtual TFuture<TPair<EStatusCode, FGetTaskUpdatesResponse>> GetUpdatesAsync(const FString& ChannelId, const int32 WaitSeconds = 0) = 0;
 
-	virtual TFuture<FExecuteResponse> ExecuteAsync(const FExecuteRequest& Request, TUniqueFunction<void()> CompletionCallback = nullptr, int64 TimeoutMs = 0) = 0;
-};
+		virtual TFuture<TPair<EStatusCode, FGetTaskUpdatesResponse>> RunTasksAsync(const FAddTasksRequest& AddTaskRequest, const int32 TimeoutSeconds = 0) = 0;
+	};
+}
diff --git a/Engine/Source/Editor/Experimental/RemoteExecution/Public/IRemoteExecutionModule.h b/Engine/Source/Editor/Experimental/RemoteExecution/Public/IRemoteExecutionModule.h
--- a/Engine/Source/Editor/Experimental/RemoteExecution/Public/IRemoteExecutionModule.h
+++ b/Engine/Source/Editor/Experimental/RemoteExecution/Public/IRemoteExecutionModule.h
@@ -5,28 +5,32 @@
 #include "CoreMinimal.h"
 #include "Modules/ModuleInterface.h"
 
-class IRemoteExecutor;
 
-/**
- * Module used to remotely execute computation
- */
-class IRemoteExecutionModule : public IModuleInterface
+namespace UE::RemoteExecution
 {
-public:
-	/**
-	 * Check to see if remote execution can be used.
-	 * @return true if remote execution can be use.
-	 */
-	virtual bool CanRemoteExecute() const = 0;
-	/**
-	 * Get the accessor to allow us to remotely execute computation
-	 * @return the executor
-	 */
-	virtual IRemoteExecutor& GetRemoteExecutor() const = 0;
+	class IRemoteExecutor;
 
 	/**
-	 * Set the remote executor we want to use to do computation
-	 * @param	InName	The name of the executor we want to use
+	 * Module used to remotely execute computation
 	 */
-	virtual void SetRemoteExecutor(const FName& InName) = 0;
-};
+	class IRemoteExecutionModule : public IModuleInterface
+	{
+	public:
+		/**
+		 * Check to see if remote execution can be used.
+		 * @return true if remote execution can be use.
+		 */
+		virtual bool CanRemoteExecute() const = 0;
+		/**
+		 * Get the accessor to allow us to remotely execute computation
+		 * @return the executor
+		 */
+		virtual IRemoteExecutor& GetRemoteExecutor() const = 0;
+
+		/**
+		 * Set the remote executor we want to use to do computation
+		 * @param	InName	The name of the executor we want to use
+		 */
+		virtual void SetRemoteExecutor(const FName& InName) = 0;
+	};
+}
diff --git a/Engine/Source/Editor/Experimental/RemoteExecution/Public/IRemoteExecutor.h b/Engine/Source/Editor/Experimental/RemoteExecution/Public/IRemoteExecutor.h
--- a/Engine/Source/Editor/Experimental/RemoteExecution/Public/IRemoteExecutor.h
+++ b/Engine/Source/Editor/Experimental/RemoteExecution/Public/IRemoteExecutor.h
@@ -12,18 +12,21 @@
 /**
  * Interface for remote execution functionality
  */
-class IRemoteExecutor : public IModularFeature
+namespace UE::RemoteExecution
 {
-public:
-	/** Virtual destructor */
-	virtual ~IRemoteExecutor() {}
-
-	virtual FName GetFName() const = 0;
-	virtual FText GetNameText() const = 0;
-	virtual FText GetDescriptionText() const = 0;
-
-	virtual bool CanRemoteExecute() const = 0;
-
-	virtual IContentAddressableStorage* GetContentAddressableStorage() const = 0;
-	virtual IExecution* GetExecution() const = 0;
-};
\ No newline at end of file
+	class IRemoteExecutor : public IModularFeature
+	{
+	public:
+		/** Virtual destructor */
+		virtual ~IRemoteExecutor() {}
+
+		virtual FName GetFName() const = 0;
+		virtual FText GetNameText() const = 0;
+		virtual FText GetDescriptionText() const = 0;
+
+		virtual bool CanRemoteExecute() const = 0;
+
+		virtual IContentAddressableStorage* GetContentAddressableStorage() const = 0;
+		virtual IExecution* GetExecution() const = 0;
+	};
+}
diff --git a/Engine/Source/Editor/Experimental/RemoteExecution/Public/IRemoteMessage.h b/Engine/Source/Editor/Experimental/RemoteExecution/Public/IRemoteMessage.h
new file mode 100644
--- /dev/null
+++ b/Engine/Source/Editor/Experimental/RemoteExecution/Public/IRemoteMessage.h
@@ -0,0 +1,22 @@
+// Copyright Epic Games, Inc. All Rights Reserved.
+
+#pragma once
+
+#include "CoreMinimal.h"
+
+class FCbObject;
+class FCbObjectView;
+
+
+namespace UE::RemoteExecution
+{
+	class IRemoteMessage
+	{
+	public:
+		/** Virtual destructor */
+		virtual ~IRemoteMessage() {}
+
+		virtual FCbObject Save() const = 0;
+		virtual void Load(const FCbObjectView& CbObjectView) = 0;
+	};
+}
diff --git a/Engine/Source/Editor/Experimental/RemoteExecution/Public/Messages.h b/Engine/Source/Editor/Experimental/RemoteExecution/Public/Messages.h
deleted file mode 100644
diff --git a/Engine/Source/Editor/Experimental/RemoteExecution/Public/RemoteMessages.h b/Engine/Source/Editor/Experimental/RemoteExecution/Public/RemoteMessages.h
new file mode 100644
--- /dev/null
+++ b/Engine/Source/Editor/Experimental/RemoteExecution/Public/RemoteMessages.h
@@ -0,0 +1,180 @@
+// Copyright Epic Games, Inc. All Rights Reserved.
+
+#pragma once
+
+#include "CoreMinimal.h"
+#include "IO/IoHash.h"
+#include "IRemoteMessage.h"
+
+
+namespace UE::RemoteExecution
+{
+	// Subset of http response codes
+	enum class EStatusCode : int32
+	{
+		Unknown = 0,
+		Ok = 200,
+		BadRequest = 400,
+		Denied = 401,
+		Forbidden = 403,
+		NotFound = 404,
+		BadMethod = 405,
+		RequestTimeout = 408,
+		RequestTooLarge = 413,
+		UriTooLong = 414,
+		UnsupportedMedia = 415,
+		TooManyRequests = 429,
+		ServerError = 500,
+		NotSupported = 501,
+		BadGateway = 502,
+		ServiceUnavail = 503,
+		GatewayTimeout = 504,
+	};
+
+	enum class ETaskState : int32
+	{
+		Queued = 0,
+		Executing = 1,
+		Complete = 2,
+	};
+
+	static FString TaskStateString(const ETaskState TaskState)
+	{
+		switch (TaskState)
+		{
+		case ETaskState::Queued: return TEXT("Queued");
+		case ETaskState::Executing: return TEXT("Executing");
+		case ETaskState::Complete: return TEXT("Complete");
+		}
+		return TEXT("Unknown");
+	}
+
+	class FAddTasksRequest : public IRemoteMessage
+	{
+	public:
+		FIoHash RequirementsHash;
+		TArray<FIoHash> TaskHashes;
+		bool DoNotCache;
+
+		REMOTEEXECUTION_API FAddTasksRequest();
+
+		// Inherited via IMessage
+		REMOTEEXECUTION_API virtual FCbObject Save() const override;
+		REMOTEEXECUTION_API virtual void Load(const FCbObjectView& CbObjectView) override;
+	};
+
+	class FGetTaskUpdateResponse : public IRemoteMessage
+	{
+	public:
+		FIoHash TaskHash;
+		FDateTime Time;
+		ETaskState State;
+		FIoHash ResultHash;
+		FString AgentId;
+		FString LeaseId;
+
+		REMOTEEXECUTION_API FGetTaskUpdateResponse();
+
+		// Inherited via IMessage
+		REMOTEEXECUTION_API virtual FCbObject Save() const override;
+		REMOTEEXECUTION_API virtual void Load(const FCbObjectView& CbObjectView) override;
+	};
+
+	class FGetTaskUpdatesResponse : public IRemoteMessage
+	{
+	public:
+		TArray<FGetTaskUpdateResponse> Updates;
+
+		// Inherited via IMessage
+		REMOTEEXECUTION_API virtual FCbObject Save() const override;
+		REMOTEEXECUTION_API virtual void Load(const FCbObjectView& CbObjectView) override;
+	};
+
+	struct FGetObjectTreeResponse
+	{
+		TMap<FIoHash, TArray<uint8>> Objects;
+		TSet<FIoHash> BinaryAttachments;
+	};
+
+	class FFileNode : public IRemoteMessage
+	{
+	public:
+		FString Name;
+		FIoHash Hash;
+		int64 Size;
+		int32 Attributes;
+
+		REMOTEEXECUTION_API FFileNode();
+
+		// Inherited via IMessage
+		REMOTEEXECUTION_API virtual FCbObject Save() const override;
+		REMOTEEXECUTION_API virtual void Load(const FCbObjectView& CbObjectView) override;
+	};
+
+	class FDirectoryNode : public IRemoteMessage
+	{
+	public:
+		FString Name;
+		FIoHash Hash;
+
+		// Inherited via IMessage
+		REMOTEEXECUTION_API virtual FCbObject Save() const override;
+		REMOTEEXECUTION_API virtual void Load(const FCbObjectView& CbObjectView) override;
+	};
+
+	class FDirectoryTree : public IRemoteMessage
+	{
+	public:
+		TArray<FFileNode> Files;
+		TArray<FDirectoryNode> Directories;
+
+		// Inherited via IMessage
+		REMOTEEXECUTION_API virtual FCbObject Save() const override;
+		REMOTEEXECUTION_API virtual void Load(const FCbObjectView& CbObjectView) override;
+	};
+
+	class FRequirements : public IRemoteMessage
+	{
+	public:
+		FString Condition;
+		TMap<FString, int32> Resources;
+		bool Exclusive;
+
+		REMOTEEXECUTION_API FRequirements();
+
+		// Inherited via IMessage
+		REMOTEEXECUTION_API virtual FCbObject Save() const override;
+		REMOTEEXECUTION_API virtual void Load(const FCbObjectView& CbObjectView) override;
+	};
+
+	class FTask : public IRemoteMessage
+	{
+	public:
+		FString Executable;
+		TArray<FString> Arguments;
+		TMap<FString, FString> EnvVars;
+		FString WorkingDirectory;
+		FIoHash SandboxHash;
+		FIoHash RequirementsHash;
+		TArray<FString> OutputPaths;
+
+		// Inherited via IMessage
+		REMOTEEXECUTION_API virtual FCbObject Save() const override;
+		REMOTEEXECUTION_API virtual void Load(const FCbObjectView& CbObjectView) override;
+	};
+
+	class FTaskResult : public IRemoteMessage
+	{
+	public:
+		int32 ExitCode;
+		FIoHash StdOutHash;
+		FIoHash StdErrHash;
+		FIoHash OutputHash;
+
+		REMOTEEXECUTION_API FTaskResult();
+
+		// Inherited via IMessage
+		REMOTEEXECUTION_API virtual FCbObject Save() const override;
+		REMOTEEXECUTION_API virtual void Load(const FCbObjectView& CbObjectView) override;
+	};
+}
diff --git a/Engine/Source/Editor/UnrealEd/Private/DerivedDataBuildRemoteExecutor.cpp b/Engine/Source/Editor/UnrealEd/Private/DerivedDataBuildRemoteExecutor.cpp
--- a/Engine/Source/Editor/UnrealEd/Private/DerivedDataBuildRemoteExecutor.cpp
+++ b/Engine/Source/Editor/UnrealEd/Private/DerivedDataBuildRemoteExecutor.cpp
@@ -16,7 +16,6 @@
 #include "HAL/PlatformProcess.h"
 #include "IContentAddressableStorage.h"
 #include "IRemoteExecutor.h"
-#include "Messages.h"
 #include "Misc/CommandLine.h"
 #include "Misc/ConfigCacheIni.h"
 #include "Misc/Optional.h"
@@ -24,17 +23,14 @@
 #include "Modules/ModuleManager.h"
 #include "Serialization/CompactBinaryValidation.h"
 #include "Serialization/CompactBinaryWriter.h"
+#include "TickableEditorObject.h"
 #include <atomic>
 
-static uint32 GetTypeHash(const FDigest& Digest)
-{
-	return GetTypeHash(Digest.Hash);
-}
+#include "HttpManager.h"
+#include "HttpModule.h"
+#include "RemoteMessages.h"
+
 
-static bool operator==(const FDigest& A, const FDigest& B)
-{
-	return A.Hash == B.Hash;
-}
 
 namespace UE::DerivedData
 {
@@ -119,7 +115,7 @@ private:
 	struct FMerkleTreeFileBuilder
 	{
 		FStringView Path;
-		FFileNode File;
+		UE::RemoteExecution::FFileNode File;
 		EFileType Type;
 		FCompositeBuffer ContentBytes;
 	};
@@ -128,8 +124,8 @@ private:
 	{
 		FStringView Name;
 		FStringView Path;
-		FDirectory Directory;
-		TOptional<FDigest> Digest;
+		UE::RemoteExecution::FDirectoryTree Directory;
+		TOptional<FIoHash> Digest;
 		TArray<uint8> ContentBytes;
 		TArray<int32> SubDirIndices;
 	};
@@ -148,34 +144,38 @@ private:
 		TArray<FMerkleTreeFileBuilder> Files;
 
 		// Lookup tables for indexing in different scenarios
-		TMultiMap<FDigest, FVariantIndex> DigestFilesystemIndex;
+		TMultiMap<FIoHash, FVariantIndex> DigestFilesystemIndex;
 		TMap<FStringView, int32> PathToDirectoryIndex;
 		TMap<int32, FString> FileIndexToInputKey;
 		FString BaseDirectoryPath;
 
 		// Unique items in the tree
-		FCommand Command;
-		TArray<uint8> CommandContentBytes;
-		FAction Action;
-		TArray<uint8> ActionContentBytes;
-		FDigest ActionDigest;
+		UE::RemoteExecution::FTask Task;
+		TArray<uint8>TaskContentBytes;
+		FIoHash TaskDigest;
+		UE::RemoteExecution::FRequirements Requirements;
+		TArray<uint8> RequirementsContentBytes;
 		FCompositeBuffer BuildActionContentBytes {};
-		FDigest BuildActionDigest;
+		FIoHash BuildActionDigest;
 		TArray<FString> InputPaths;
 
 		// Step 1: Find missing blobs
-		FFindMissingBlobsRequest FindMissingBlobsRequest;
-		FFindMissingBlobsResponse FindMissingBlobsResponse;
+		TSet<FIoHash> DoBlobsExistRequest;
+		TMap<FIoHash, UE::RemoteExecution::EStatusCode> DoBlobsExistResponse;
 
 		// Step 3: Batch update blobs (upload)
-		FBatchUpdateBlobsRequest BatchUpdateBlobsRequest;
+		TMap<FIoHash, TArray<uint8>> PutBlobsRequest;
+
+		// Step 3: Execute
+		UE::RemoteExecution::FAddTasksRequest ExecuteRequest;
+		UE::RemoteExecution::FGetTaskUpdateResponse ExecuteResponse;
 
-		// Step 4: Execute
-		FExecuteRequest ExecuteRequest;
-		FExecuteResponse ExecuteResponse;
+		// Step 4: Get execute results
+		UE::RemoteExecution::FTaskResult ExecuteResult;
+		TMap<FString, FIoHash> ExecuteResultFiles;
 
 		// Step 5: Batch read blobs (download)
-		FBatchReadBlobsRequest BatchReadBlobsRequest;
+		TSet<FIoHash> GetBlobsRequest;
 	};
 
 	FRemoteExecutionState State;
@@ -188,28 +188,31 @@ private:
 	// Merkle tree operations
 	FMerkleTreeDirectoryBuilder& GetOrAddMerkleTreeDirectory(FStringView Path, int32& OutDirectoryBuilderIndex);
 	int32 AddMerkleTreeFile(FStringView Path, const FIoHash& RawHash, uint64 RawSize, bool bIsExecutable, EFileType FileType, FCompositeBuffer ContentBytes = FCompositeBuffer::Null);
-	const FDigest& BuildMerkleTreeDirectoryDigest(int32 Index);
+	const FIoHash& BuildMerkleTreeDirectoryDigest(int32 Index);
 	void BuildMerkleTreeNodes();
 
 	// General utility methods
 	void GatherMissingInputFileBlobs(TArray<FStringView>& OutMissingInputs);
-	bool ValidateUploadSuccess(const FBatchUpdateBlobsResponse& BatchUpdateBlobsResponse);
-	FOptionalBuildOutput ComposeBuildOutput(FBatchReadBlobsResponse& BatchReadBlobsResponse, EStatus& OutStatus);
+	bool ValidateUploadSuccess(const TMap<FIoHash, UE::RemoteExecution::EStatusCode>& PutBlobsResponse);
+	void GatherExecuteFileResults(const TMap<FIoHash, TArray<uint8>>& Data, const FString& Parent, const FIoHash& DirectoryTreeHash);
+	FOptionalBuildOutput ComposeBuildOutput(TMap<FIoHash, TPair<UE::RemoteExecution::EStatusCode, TArray<uint8>>>&& GetBlobsResponse, EStatus& OutStatus);
 	bool ProcessCancellation();
-	bool IsStatusOk(FStatus Status, const TCHAR* OperationDesc);
+	bool IsStatusOk(UE::RemoteExecution::EStatusCode Status, const TCHAR* OperationDesc);
 
 	// Async steps
-	TFuture<TPair<FStatus, FFindMissingBlobsResponse>> DetermineMissingBlobsAsync();
+	TFuture<TMap<FIoHash, UE::RemoteExecution::EStatusCode>> DetermineMissingBlobsAsync();
 	void LoadMissingWorkerFileBlobsAsync();
-	TFuture<TPair<FStatus, FBatchUpdateBlobsResponse>> UploadMissingBlobsAsync();
-	TFuture<FExecuteResponse> ExecuteBuildAsync();
-	TFuture<TPair<FStatus, FBatchReadBlobsResponse>> DownloadResultsAsync();
+	TFuture<TMap<FIoHash, UE::RemoteExecution::EStatusCode>> UploadMissingBlobsAsync();
+	TFuture<TPair<UE::RemoteExecution::EStatusCode, UE::RemoteExecution::FGetTaskUpdatesResponse>> ExecuteBuildAsync();
+	TFuture<TPair<UE::RemoteExecution::EStatusCode, UE::RemoteExecution::FGetObjectTreeResponse>> GetTaskResultAsync();
+	TFuture<TMap<FIoHash, TPair<UE::RemoteExecution::EStatusCode, TArray<uint8>>>> DownloadResultsAsync();
 
 	// Post-step flow
-	void OnMissingBlobsDetermined(FStatus Status, FFindMissingBlobsResponse&& Result);
-	void OnMissingBlobsUploaded(FStatus Status, const FBatchUpdateBlobsResponse& Result);
-	void OnExecutionCompleted(FExecuteResponse&& Result);
-	void OnOutputBlobsDownloaded(FStatus Status, FBatchReadBlobsResponse&& Result);
+	void OnMissingBlobsDetermined(TMap<FIoHash, UE::RemoteExecution::EStatusCode>&& Result);
+	void OnMissingBlobsUploaded(const TMap<FIoHash, UE::RemoteExecution::EStatusCode>& Result);
+	void OnExecutionCompleted(UE::RemoteExecution::EStatusCode Status, UE::RemoteExecution::FGetTaskUpdatesResponse&& Result);
+	void OnTaskResultDownloaded(UE::RemoteExecution::EStatusCode Status, UE::RemoteExecution::FGetObjectTreeResponse&& Result);
+	void OnOutputBlobsDownloaded(TMap<FIoHash, TPair<UE::RemoteExecution::EStatusCode, TArray<uint8>>>&& Result);
 };
 
 class FRemoteBuildWorkerExecutor final: public IBuildWorkerExecutor
@@ -217,15 +220,15 @@ class FRemoteBuildWorkerExecutor final: public IBuildWorkerExecutor
 public:
 	FRemoteBuildWorkerExecutor()
 	: GlobalExecutionTimeoutSeconds(-1)
-	, Salt(TEXT("807c6a49-0657-49f3-b498-fd457213c0a7"))
 	, RemoteExecutor(nullptr)
 	, ContentAddressableStorage(nullptr)
+	, Execution(nullptr)
 	, bEnabled(false)
 	{
 		check(IsInGameThread()); // initialization from the main thread is expected to allow config reading for the limiting heuristics
 		check(GConfig && GConfig->IsReadyForUse());
 
-		GConfig->GetString(TEXT("DerivedDataBuildRemoteExecutor"), TEXT("InstanceName"), InstanceName, GEngineIni);
+		GConfig->GetString(TEXT("DerivedDataBuildRemoteExecutor"), TEXT("NameSpaceId"), NameSpaceId, GEngineIni);
 		GConfig->GetInt(TEXT("DerivedDataBuildRemoteExecutor"), TEXT("GlobalExecutionTimeoutSeconds"), GlobalExecutionTimeoutSeconds, GEngineIni);
 
 		const FName RemoteExecutionFeatureName(TEXT("RemoteExecution"));
@@ -233,22 +236,24 @@ public:
 
 		if (FParse::Param(FCommandLine::Get(), TEXT("DDC2RemoteExecution")))
 		{
-			FModuleManager::Get().LoadModule("BazelExecutor");
+			FModuleManager::Get().LoadModule("HordeExecutor");
 			if (ModularFeatures.IsModularFeatureAvailable(RemoteExecutionFeatureName))
 			{
-				RemoteExecutor = &ModularFeatures.GetModularFeature<IRemoteExecutor>(RemoteExecutionFeatureName);
+				RemoteExecutor = &ModularFeatures.GetModularFeature<UE::RemoteExecution::IRemoteExecutor>(RemoteExecutionFeatureName);
 				if (RemoteExecutor)
 				{
 					ContentAddressableStorage = RemoteExecutor->GetContentAddressableStorage();
+					Execution = RemoteExecutor->GetExecution();
 				}
 			}
 		}
 
-		bEnabled = (RemoteExecutor != nullptr) && (ContentAddressableStorage != nullptr) && !InstanceName.IsEmpty();
+		bEnabled = (RemoteExecutor != nullptr) && (ContentAddressableStorage != nullptr) && (Execution != nullptr) && !NameSpaceId.IsEmpty();
 		if (bEnabled)
 		{
 			ModularFeatures.RegisterModularFeature(IBuildWorkerExecutor::GetFeatureName(), this);
 		}
+		RemoteBuildTickable.SetTickable(bEnabled);
 	}
 
 	virtual ~FRemoteBuildWorkerExecutor()
@@ -343,7 +348,7 @@ private:
 			}
 		};
 		FBlobStat TotalActionBlobsUploaded;
-		FBlobStat TotalCommandBlobsUploaded;
+		FBlobStat TotalRequirementBlobsUploaded;
 		FBlobStat TotalDirectoryBlobsUploaded;
 		FBlobStat TotalFileBlobsUploaded;
 		FBlobStat TotalFileBlobsDownloaded;
@@ -357,8 +362,8 @@ private:
 			UE_LOG(LogDerivedDataBuildRemoteExecutor, Display, TEXT("%-35s=%10") UINT64_FMT, TEXT("Successful remote builds"), TotalSuccessfulRemoteBuilds.load());
 			UE_LOG(LogDerivedDataBuildRemoteExecutor, Display, TEXT("%-35s=%10") UINT64_FMT, TEXT("Uploaded actions (quantity)"), TotalActionBlobsUploaded.Quantity.load());
 			UE_LOG(LogDerivedDataBuildRemoteExecutor, Display, TEXT("%-35s=%10") UINT64_FMT, TEXT("Uploaded actions (KB)"), TotalActionBlobsUploaded.Bytes.load()/1024);
-			UE_LOG(LogDerivedDataBuildRemoteExecutor, Display, TEXT("%-35s=%10") UINT64_FMT, TEXT("Uploaded commands (quantity)"), TotalCommandBlobsUploaded.Quantity.load());
-			UE_LOG(LogDerivedDataBuildRemoteExecutor, Display, TEXT("%-35s=%10") UINT64_FMT, TEXT("Uploaded commands (KB)"), TotalCommandBlobsUploaded.Bytes.load()/1024);
+			UE_LOG(LogDerivedDataBuildRemoteExecutor, Display, TEXT("%-35s=%10") UINT64_FMT, TEXT("Uploaded requirements (quantity)"), TotalRequirementBlobsUploaded.Quantity.load());
+			UE_LOG(LogDerivedDataBuildRemoteExecutor, Display, TEXT("%-35s=%10") UINT64_FMT, TEXT("Uploaded requirements (KB)"), TotalRequirementBlobsUploaded.Bytes.load()/1024);
 			UE_LOG(LogDerivedDataBuildRemoteExecutor, Display, TEXT("%-35s=%10") UINT64_FMT, TEXT("Uploaded directories (quantity)"), TotalDirectoryBlobsUploaded.Quantity.load());
 			UE_LOG(LogDerivedDataBuildRemoteExecutor, Display, TEXT("%-35s=%10") UINT64_FMT, TEXT("Uploaded directories (KB)"), TotalDirectoryBlobsUploaded.Bytes.load()/1024);
 			UE_LOG(LogDerivedDataBuildRemoteExecutor, Display, TEXT("%-35s=%10") UINT64_FMT, TEXT("Uploaded files (quantity)"), TotalFileBlobsUploaded.Quantity.load());
@@ -459,15 +464,36 @@ private:
 		bool bEnableLimits{false};
 	};
 
+	class FRemoteBuildTickableObject : public FTickableEditorObject, public FTickableCookObject
+	{
+		bool bIsTickable;
+		// FTickableEditorObject/FTickableCookObject interface
+		virtual void Tick(float DeltaTime) override
+		{
+			TickCook(DeltaTime, false  /* bTickComplete */);
+		}
+		virtual void TickCook(float DeltaTime, bool bTickComplete) override
+		{
+			FHttpModule::Get().GetHttpManager().Tick(DeltaTime);
+		}
+		virtual bool IsTickable() const { return bIsTickable; }
+		virtual ETickableTickType GetTickableTickType() const override { return ETickableTickType::Conditional; }
+		virtual TStatId GetStatId() const override { return TStatId(); }
+	public:
+		FRemoteBuildTickableObject() : bIsTickable(false) {}
+		void SetTickable(bool IsTickable) { bIsTickable = IsTickable; }
+	};
+
 	friend class FRemoteBuildExecutionRequest;
 
 	FStats Stats;
 	FLimitingHeuristics LimitingHeuristics;
-	FString InstanceName;
+	FRemoteBuildTickableObject RemoteBuildTickable;
+	FString NameSpaceId;
 	int GlobalExecutionTimeoutSeconds;
-	const FStringView Salt;
-	IRemoteExecutor* RemoteExecutor;
-	IContentAddressableStorage* ContentAddressableStorage;
+	UE::RemoteExecution::IRemoteExecutor* RemoteExecutor;
+	UE::RemoteExecution::IContentAddressableStorage* ContentAddressableStorage;
+	UE::RemoteExecution::IExecution* Execution;
 	bool bEnabled;
 };
 
@@ -488,7 +514,7 @@ FRemoteBuildExecutionRequest::FRemoteBuildExecutionRequest(
 {
 	Owner.Begin(this);
 	DetermineMissingBlobsAsync()
-		.Next([this] (TPair<FStatus, FFindMissingBlobsResponse>&& Result) { OnMissingBlobsDetermined(Result.Key, MoveTemp(Result.Value)); });
+		.Next([this] (TMap<FIoHash, UE::RemoteExecution::EStatusCode>&& Result) { OnMissingBlobsDetermined(MoveTemp(Result)); });
 }
 
 FRemoteBuildExecutionRequest::~FRemoteBuildExecutionRequest()
@@ -528,16 +554,19 @@ int32 FRemoteBuildExecutionRequest::AddMerkleTreeFile(FStringView Path, const FI
 	FStringView ContainingDirectoryPath = FPathViews::GetPath(Path);
 	int32 ContainingDirectoryIndex = INDEX_NONE;
 	FMerkleTreeDirectoryBuilder* ContainingDirectoryBuilder = &GetOrAddMerkleTreeDirectory(ContainingDirectoryPath, ContainingDirectoryIndex);
-	FFileNode& NewNode = ContainingDirectoryBuilder->Directory.Files.AddDefaulted_GetRef();
+	UE::RemoteExecution::FFileNode& NewNode = ContainingDirectoryBuilder->Directory.Files.AddDefaulted_GetRef();
 	NewNode.Name = FPathViews::GetCleanFilename(Path);
-	NewNode.Digest.Hash = RawHash;
-	NewNode.Digest.SizeBytes = RawSize;
-	NewNode.IsExecutable = bIsExecutable;
+	NewNode.Hash = RawHash;
+	NewNode.Size = RawSize;
+
+	constexpr int32 ExecutableFileAttribute = 0x1ED; // 0755 octal
+	constexpr int32 NormalFileAttribute = 0x1A4; // 0644 octal
+	NewNode.Attributes = bIsExecutable ? ExecutableFileAttribute : NormalFileAttribute;
 
 	UE_LOG(LogDerivedDataBuildRemoteExecutor, Verbose, TEXT("Remote execution: added merkle tree file '%.*s' (hash: %s, size: %u)"), Path.Len(), Path.GetData(), *::LexToString(RawHash), RawSize);
 
 	int32 NewFileIndex = State.Files.Num();
-	State.DigestFilesystemIndex.Add(NewNode.Digest, FVariantIndex(ENodeType::File, NewFileIndex));
+	State.DigestFilesystemIndex.Add(NewNode.Hash, FVariantIndex(ENodeType::File, NewFileIndex));
 	FMerkleTreeFileBuilder& FileBuilder = State.Files.AddDefaulted_GetRef();
 	FileBuilder.Path = Path;
 	FileBuilder.File = NewNode; // Duplicates the node in the state's file array
@@ -550,7 +579,7 @@ int32 FRemoteBuildExecutionRequest::AddMerkleTreeFile(FStringView Path, const FI
 	return NewFileIndex;
 }
 
-const FDigest& FRemoteBuildExecutionRequest::BuildMerkleTreeDirectoryDigest(int32 Index)
+const FIoHash& FRemoteBuildExecutionRequest::BuildMerkleTreeDirectoryDigest(int32 Index)
 {
 	FMerkleTreeDirectoryBuilder& DirBuilder = State.Directories[Index];
 
@@ -561,18 +590,18 @@ const FDigest& FRemoteBuildExecutionRequest::BuildMerkleTreeDirectoryDigest(int3
 
 	for (int32 SubDirIndex : DirBuilder.SubDirIndices)
 	{
-		FDirectoryNode& SubDirNode = DirBuilder.Directory.Directories.AddDefaulted_GetRef();
+		UE::RemoteExecution::FDirectoryNode& SubDirNode = DirBuilder.Directory.Directories.AddDefaulted_GetRef();
 		SubDirNode.Name = State.Directories[SubDirIndex].Name;
-		SubDirNode.Digest = BuildMerkleTreeDirectoryDigest(SubDirIndex);
-		State.DigestFilesystemIndex.Add(SubDirNode.Digest, FVariantIndex(ENodeType::Directory, SubDirIndex));
+		SubDirNode.Hash = BuildMerkleTreeDirectoryDigest(SubDirIndex);
+		State.DigestFilesystemIndex.Add(SubDirNode.Hash, FVariantIndex(ENodeType::Directory, SubDirIndex));
 	}
 
-	DirBuilder.Directory.Directories.Sort([] (const FDirectoryNode& A, const FDirectoryNode& B)
+	DirBuilder.Directory.Directories.Sort([] (const UE::RemoteExecution::FDirectoryNode& A, const UE::RemoteExecution::FDirectoryNode& B)
 	{
 		return A.Name < B.Name;
 	});
 
-	DirBuilder.Directory.Files.Sort([] (const FFileNode& A, const FFileNode& B)
+	DirBuilder.Directory.Files.Sort([] (const UE::RemoteExecution::FFileNode& A, const UE::RemoteExecution::FFileNode& B)
 	{
 		return A.Name < B.Name;
 	});
@@ -636,56 +665,44 @@ void FRemoteBuildExecutionRequest::BuildMerkleTreeNodes()
 	FUniqueBuffer UncompressedBuildActionContentBytes = FUniqueBuffer::Alloc(BuildActionWriter.GetSaveSize());
 	BuildActionWriter.Save(UncompressedBuildActionContentBytes);
 	State.BuildActionContentBytes = FCompositeBuffer(UncompressedBuildActionContentBytes.MoveToShared());
-	State.BuildActionDigest.Hash = FIoHash::HashBuffer(State.BuildActionContentBytes);
-	State.BuildActionDigest.SizeBytes = State.BuildActionContentBytes.GetSize();
-	AddMerkleTreeFile(TEXT("Build.action"), State.BuildActionDigest.Hash, State.BuildActionDigest.SizeBytes, false, EFileType::BuildAction, State.BuildActionContentBytes);
+	State.BuildActionDigest = FIoHash::HashBuffer(State.BuildActionContentBytes);
+	AddMerkleTreeFile(TEXT("Build.action"), State.BuildActionDigest, State.BuildActionContentBytes.GetSize(), false, EFileType::BuildAction, State.BuildActionContentBytes);
 
 
 	if (!State.PathToDirectoryIndex.IsEmpty())
 	{
 		int32 RootDirectoryIndex = State.PathToDirectoryIndex.FindChecked(TEXT(""));
-		State.Action.InputRootDigest = BuildMerkleTreeDirectoryDigest(RootDirectoryIndex);
-		State.DigestFilesystemIndex.Add(State.Action.InputRootDigest, FVariantIndex(ENodeType::Directory, RootDirectoryIndex));
+		State.Task.SandboxHash = BuildMerkleTreeDirectoryDigest(RootDirectoryIndex);
+		State.DigestFilesystemIndex.Add(State.Task.SandboxHash, FVariantIndex(ENodeType::Directory, RootDirectoryIndex));
 	}
 
-	State.Command.OutputPaths.Add("Outputs");
-	State.Command.OutputPaths.Add("Build.output");
-	State.Command.Arguments.Add(FString(State.BuildWorker.GetPath()));
-	State.Command.Arguments.Add("-Build=Build.action");
+	{
+		State.Requirements.Condition = TEXT("OSFamily == 'Windows'");
+		Executor.ContentAddressableStorage->ToBlob(State.Requirements, State.RequirementsContentBytes, State.Task.RequirementsHash);
+	}
+
+	State.Task.OutputPaths.Add("Outputs");
+	State.Task.OutputPaths.Add("Build.output");
+	State.Task.Executable = State.BuildWorker.GetPath();
+	State.Task.Arguments.Add("-Build=Build.action");
 	State.BuildWorker.IterateEnvironment([this] (FStringView Name, FStringView Value)
 		{
-			FCommand::FEnvironmentVariable& EnvVar = State.Command.EnvironmentVariables.AddDefaulted_GetRef();
-			EnvVar.Name = Name;
-			EnvVar.Value = Value;
+			State.Task.EnvVars.Add(FString(Name), FString(Value));
 		});
-	State.Command.EnvironmentVariables.Sort([] (const FCommand::FEnvironmentVariable& A, const FCommand::FEnvironmentVariable& B)
-		{
-			if (A.Name < B.Name)
-			{
-				return true;
-			}
-			else if (B.Name < A.Name)
-			{
-				return false;
-			}
-			return A.Value < B.Value;
-		});
-	Executor.ContentAddressableStorage->ToBlob(State.Command, State.CommandContentBytes, State.Action.CommandDigest);
-
-	if (Executor.GlobalExecutionTimeoutSeconds > 0)
-	{
-		State.Action.Timeout = FTimespan::FromSeconds(Executor.GlobalExecutionTimeoutSeconds);
-	}
-	State.Action.Salt = Executor.Salt;
-	Executor.ContentAddressableStorage->ToBlob(State.Action, State.ActionContentBytes, State.ActionDigest);
+	Executor.ContentAddressableStorage->ToBlob(State.Task, State.TaskContentBytes, State.TaskDigest);
 }
 
 void FRemoteBuildExecutionRequest::GatherMissingInputFileBlobs(TArray<FStringView>& OutMissingInputs)
 {
-	for (const FDigest& MissingItem : State.FindMissingBlobsResponse.MissingBlobDigests)
+	for (const TPair<FIoHash, UE::RemoteExecution::EStatusCode>& MissingItem : State.DoBlobsExistResponse)
 	{
+		if (MissingItem.Value == UE::RemoteExecution::EStatusCode::Ok)
+		{
+			continue;
+		}
+
 		TArray<FVariantIndex> DigestFilesystemEntries;
-		State.DigestFilesystemIndex.MultiFind(MissingItem, DigestFilesystemEntries);
+		State.DigestFilesystemIndex.MultiFind(MissingItem.Key, DigestFilesystemEntries);
 
 		for (FVariantIndex& VariantIndex : DigestFilesystemEntries)
 		{
@@ -701,62 +718,82 @@ void FRemoteBuildExecutionRequest::GatherMissingInputFileBlobs(TArray<FStringVie
 	}
 }
 
-bool FRemoteBuildExecutionRequest::ValidateUploadSuccess(const FBatchUpdateBlobsResponse& BatchUpdateBlobsResponse)
+bool FRemoteBuildExecutionRequest::ValidateUploadSuccess(const TMap<FIoHash, UE::RemoteExecution::EStatusCode>& PutBlobsResponse)
 {
 	bool bSuccess = true;
-	for (const FBatchUpdateBlobsRequest::FRequest& RequestedUpload : State.BatchUpdateBlobsRequest.Requests)
+	for (const TPair<FIoHash, TArray<uint8>>& BlobRequest : State.PutBlobsRequest)
 	{
-		if (const FBatchUpdateBlobsResponse::FResponse* FoundResponse = Algo::FindByPredicate(BatchUpdateBlobsResponse.Responses, [&RequestedUpload] (const FBatchUpdateBlobsResponse::FResponse& Response)
-			{
-				return Response.Digest == RequestedUpload.Digest;
-			}))
+		const UE::RemoteExecution::EStatusCode* StatusCode = PutBlobsResponse.Find(BlobRequest.Key);
+		if (StatusCode == nullptr || *StatusCode != UE::RemoteExecution::EStatusCode::Ok)
 		{
-			if (!FoundResponse->Status.Ok())
-			{
-				FStringView ActionName = State.BuildAction.GetName();
-				UE_LOG(LogDerivedDataBuildRemoteExecutor, Log, TEXT("Remote execution system error: data for action '%.*s' could not be uploaded (hash: %s, size: %u)"), ActionName.Len(), ActionName.GetData(), *::LexToString(RequestedUpload.Digest.Hash), RequestedUpload.Digest.SizeBytes);
-				bSuccess = false;
-			}
+			FStringView ActionName = State.BuildAction.GetName();
+			UE_LOG(LogDerivedDataBuildRemoteExecutor, Log, TEXT("Remote execution system error: data for action '%.*s' could not be uploaded (hash: %s, size: %u)"), ActionName.Len(), ActionName.GetData(), *::LexToString(BlobRequest.Key), BlobRequest.Value.Num());
+			bSuccess = false;
 		}
 	}
 	return bSuccess;
 }
 
-FOptionalBuildOutput FRemoteBuildExecutionRequest::ComposeBuildOutput(FBatchReadBlobsResponse& BatchReadBlobsResponse, EStatus& OutStatus)
+void FRemoteBuildExecutionRequest::GatherExecuteFileResults(const TMap<FIoHash, TArray<uint8>>& Data, const FString& Parent, const FIoHash& DirectoryTreeHash)
+{
+	if (DirectoryTreeHash == FIoHash::Zero)
+	{
+		return;
+	}
+
+	UE::RemoteExecution::FDirectoryTree DirectoryTree;
+	{
+		const FCbObjectView View = FCbObjectView(Data[DirectoryTreeHash].GetData());
+		DirectoryTree.Load(View);
+	}
+
+	for (UE::RemoteExecution::FFileNode& FileNode : DirectoryTree.Files)
+	{
+		FString Path = Parent;
+		Path.Append(MoveTemp(FileNode.Name));
+		State.ExecuteResultFiles.Add(MoveTemp(Path), MoveTemp(FileNode.Hash));
+	}
+
+	for (UE::RemoteExecution::FDirectoryNode& DirectoryNode : DirectoryTree.Directories)
+	{
+		FString Path = Parent;
+		Path.Append(MoveTemp(DirectoryNode.Name));
+		Path.AppendChar('/');
+		GatherExecuteFileResults(Data, Parent, DirectoryNode.Hash);
+	}
+}
+
+FOptionalBuildOutput FRemoteBuildExecutionRequest::ComposeBuildOutput(TMap<FIoHash, TPair<UE::RemoteExecution::EStatusCode, TArray<uint8>>>&& GetBlobsResponse, EStatus& OutStatus)
 {
 	OutStatus = EStatus::Error;
 
-	TOptional<FDigest> BuildOutputDigest;
-	for (const FOutputFile& ExecuteOutputFile : State.ExecuteResponse.Result.OutputFiles)
+	for (const TPair<FIoHash, TPair<UE::RemoteExecution::EStatusCode, TArray<uint8>>>& BlobResponse : GetBlobsResponse)
 	{
-		if (ExecuteOutputFile.Path == TEXT("Build.output"))
+		if (BlobResponse.Value.Key == UE::RemoteExecution::EStatusCode::Ok)
 		{
-			BuildOutputDigest = ExecuteOutputFile.Digest;
-			break;
+			continue;
 		}
 	}
 
-	if (!BuildOutputDigest.IsSet())
+	if (!State.ExecuteResultFiles.Contains(TEXT("Build.output")))
 	{
 		UE_LOG(LogDerivedDataBuildRemoteExecutor, Warning, TEXT("Worker error: build output structure not produced!"));
 		return FOptionalBuildOutput();
-	}
+	};
+
+	FIoHash BuildOutputDigest = State.ExecuteResultFiles[TEXT("Build.output")];
 
 	FOptionalBuildOutput RemoteBuildOutput;
 
 	TMap<FIoHash, FCompressedBuffer> PayloadResponses;
-	for (FBatchReadBlobsResponse::FResponse& OutputReadResponse : BatchReadBlobsResponse.Responses)
+	for (TPair<FString, FIoHash>& ResultFile : State.ExecuteResultFiles)
 	{
-		if (!OutputReadResponse.Status.Ok())
-		{
-			UE_LOG(LogDerivedDataBuildRemoteExecutor, Warning, TEXT("Remote execution system error: output not downloaded!"));
-			return FOptionalBuildOutput();
-		}
-		Executor.Stats.TotalFileBlobsDownloaded.AddBlob(OutputReadResponse.Data.Num());
+		TArray<uint8>&& FileData = MoveTemp(GetBlobsResponse[ResultFile.Value].Value);
+		Executor.Stats.TotalFileBlobsDownloaded.AddBlob(FileData.Num());
 
-		if (OutputReadResponse.Digest == BuildOutputDigest)
+		if (ResultFile.Value == BuildOutputDigest)
 		{
-			FSharedBuffer BuildOutputBuffer = MakeSharedBufferFromArray(MoveTemp(OutputReadResponse.Data));
+			FSharedBuffer BuildOutputBuffer = MakeSharedBufferFromArray(MoveTemp(FileData));
 
 			if (ValidateCompactBinary(BuildOutputBuffer, ECbValidateMode::Default) != ECbValidateError::None)
 			{
@@ -768,7 +805,7 @@ FOptionalBuildOutput FRemoteBuildExecutionRequest::ComposeBuildOutput(FBatchRead
 		}
 		else
 		{
-			FCompressedBuffer NewBuffer = FCompressedBuffer::FromCompressed(MakeSharedBufferFromArray(MoveTemp(OutputReadResponse.Data)));
+			FCompressedBuffer NewBuffer = FCompressedBuffer::FromCompressed(MakeSharedBufferFromArray(MoveTemp(FileData)));
 			PayloadResponses.Add(NewBuffer.GetRawHash(), MoveTemp(NewBuffer));
 		}
 	}
@@ -805,7 +842,7 @@ FOptionalBuildOutput FRemoteBuildExecutionRequest::ComposeBuildOutput(FBatchRead
 		OutputBuilder.AddPayload(FPayload(Payload.GetId(), *BufferForPayload));
 	}
 
-	OutStatus = State.ExecuteResponse.Result.ExitCode == 0 ? EStatus::Ok : EStatus::Error;
+	OutStatus = State.ExecuteResult.ExitCode == 0 ? EStatus::Ok : EStatus::Error;
 	return OutputBuilder.Build();
 }
 
@@ -823,11 +860,11 @@ bool FRemoteBuildExecutionRequest::ProcessCancellation()
 	return false;
 }
 
-bool FRemoteBuildExecutionRequest::IsStatusOk(FStatus Status, const TCHAR* OperationDesc)
+bool FRemoteBuildExecutionRequest::IsStatusOk(UE::RemoteExecution::EStatusCode Status, const TCHAR* OperationDesc)
 {
-	if (!Status.Ok())
+	if (Status != UE::RemoteExecution::EStatusCode::Ok)
 	{
-		UE_LOG(LogDerivedDataBuildRemoteExecutor, Warning, TEXT("Remote execution system error: operation '%s' produced an error result (%d)!"), OperationDesc, (int)Status.Code);
+		UE_LOG(LogDerivedDataBuildRemoteExecutor, Warning, TEXT("Remote execution system error: operation '%s' produced an error result (%d)!"), OperationDesc, (int)Status);
 
 		State.Owner.End(this, [this]
 		{
@@ -839,44 +876,53 @@ bool FRemoteBuildExecutionRequest::IsStatusOk(FStatus Status, const TCHAR* Opera
 	return true;
 }
 
-TFuture<TPair<FStatus, FFindMissingBlobsResponse>> FRemoteBuildExecutionRequest::DetermineMissingBlobsAsync()
+TFuture<TMap<FIoHash, UE::RemoteExecution::EStatusCode>> FRemoteBuildExecutionRequest::DetermineMissingBlobsAsync()
 {
 	BuildMerkleTreeNodes();
 
-	State.FindMissingBlobsRequest.InstanceName = Executor.InstanceName;
-	State.FindMissingBlobsRequest.BlobDigests.Add(State.ActionDigest);
-	UE_LOG(LogDerivedDataBuildRemoteExecutor, Verbose, TEXT("Checking CAS presence of action (hash: %s) of size %d."), *::LexToString(State.ActionDigest.Hash), State.ActionContentBytes.Num());
+	State.DoBlobsExistRequest.Add(State.TaskDigest);
+	UE_LOG(LogDerivedDataBuildRemoteExecutor, Verbose, TEXT("Checking CAS presence of task (hash: %s) of size %d."), *::LexToString(State.TaskDigest), State.TaskContentBytes.Num());
 
-	State.FindMissingBlobsRequest.BlobDigests.Add(State.Action.CommandDigest);
-	UE_LOG(LogDerivedDataBuildRemoteExecutor, Verbose, TEXT("Checking CAS presence of command (hash: %s) of size %d."), *::LexToString(State.Action.CommandDigest.Hash), State.CommandContentBytes.Num());
+	State.DoBlobsExistRequest.Add(State.Task.RequirementsHash);
+	UE_LOG(LogDerivedDataBuildRemoteExecutor, Verbose, TEXT("Checking CAS presence of requirements (hash: %s) of size %d."), *::LexToString(State.Task.RequirementsHash), State.RequirementsContentBytes.Num());
 
-	for (const TPair<FDigest, FVariantIndex>& FilesystemItem : State.DigestFilesystemIndex)
+	for (const TPair<FIoHash, FVariantIndex>& FilesystemItem : State.DigestFilesystemIndex)
 	{
-		State.FindMissingBlobsRequest.BlobDigests.AddUnique(FilesystemItem.Key);
+		State.DoBlobsExistRequest.Add(FilesystemItem.Key);
 		switch (FilesystemItem.Value.NodeType)
 		{
 		case ENodeType::Directory:
-			UE_LOG(LogDerivedDataBuildRemoteExecutor, Verbose, TEXT("Checking CAS presence of directory '%s' (hash: %s) of size %d."), *FString(State.Directories[FilesystemItem.Value.Index].Path), *::LexToString(FilesystemItem.Key.Hash), State.Directories[FilesystemItem.Value.Index].ContentBytes.Num());
+			UE_LOG(LogDerivedDataBuildRemoteExecutor, Verbose, TEXT("Checking CAS presence of directory '%s' (hash: %s) of size %d."),
+				*FString(State.Directories[FilesystemItem.Value.Index].Path),
+				*::LexToString(FilesystemItem.Key),
+				State.Directories[FilesystemItem.Value.Index].ContentBytes.Num());
 			break;
 		case ENodeType::File:
-			UE_LOG(LogDerivedDataBuildRemoteExecutor, Verbose, TEXT("Checking CAS presence of file '%s' (hash: %s, type: %s) of size %d."), *FString(State.Files[FilesystemItem.Value.Index].Path), *::LexToString(FilesystemItem.Key.Hash), LexToString(State.Files[FilesystemItem.Value.Index].Type), FilesystemItem.Key.SizeBytes);
+			UE_LOG(LogDerivedDataBuildRemoteExecutor, Verbose, TEXT("Checking CAS presence of file '%s' (hash: %s, type: %s) of size %d."),
+				*FString(State.Files[FilesystemItem.Value.Index].Path),
+				*::LexToString(FilesystemItem.Key),
+				LexToString(State.Files[FilesystemItem.Value.Index].Type), State.Files[FilesystemItem.Value.Index].File.Size);
 			break;
 		default:
 			checkNoEntry();
 			break;
 		}
 	}
-	return Executor.ContentAddressableStorage->FindMissingBlobsAsync(State.FindMissingBlobsRequest);
+	return Executor.ContentAddressableStorage->DoBlobsExistAsync(Executor.NameSpaceId, State.DoBlobsExistRequest);
 }
 
 void FRemoteBuildExecutionRequest::LoadMissingWorkerFileBlobsAsync()
 {
 	TArray<FIoHash> WorkerFileHashes;
 	TMultiMap<FIoHash, uint32> WorkerFileMapping;
-	for (const FDigest& MissingItem : State.FindMissingBlobsResponse.MissingBlobDigests)
+	for (const TPair<FIoHash, UE::RemoteExecution::EStatusCode>& MissingItem : State.DoBlobsExistResponse)
 	{
+		if (MissingItem.Value == UE::RemoteExecution::EStatusCode::Ok)
+		{
+			continue;
+		}
 		TArray<FVariantIndex> DigestFilesystemEntries;
-		State.DigestFilesystemIndex.MultiFind(MissingItem, DigestFilesystemEntries);
+		State.DigestFilesystemIndex.MultiFind(MissingItem.Key, DigestFilesystemEntries);
 
 		for (FVariantIndex& VariantIndex : DigestFilesystemEntries)
 		{
@@ -885,8 +931,8 @@ void FRemoteBuildExecutionRequest::LoadMissingWorkerFileBlobsAsync()
 				const FMerkleTreeFileBuilder& File = State.Files[VariantIndex.Index];
 				if ((File.Type == EFileType::Worker) && File.ContentBytes.IsNull())
 				{
-					WorkerFileHashes.Emplace(MissingItem.Hash);
-					WorkerFileMapping.Add(MissingItem.Hash, VariantIndex.Index);
+					WorkerFileHashes.Emplace(MissingItem.Key);
+					WorkerFileMapping.Add(MissingItem.Key, VariantIndex.Index);
 				}
 			}
 		}
@@ -912,54 +958,57 @@ void FRemoteBuildExecutionRequest::LoadMissingWorkerFileBlobsAsync()
 	BlockingOwner.Wait();
 }
 
-TFuture<TPair<FStatus, FBatchUpdateBlobsResponse>> FRemoteBuildExecutionRequest::UploadMissingBlobsAsync()
+TFuture<TMap<FIoHash, UE::RemoteExecution::EStatusCode>> FRemoteBuildExecutionRequest::UploadMissingBlobsAsync()
 {
-	State.BatchUpdateBlobsRequest.InstanceName = Executor.InstanceName;
-	for (const FDigest& MissingItem : State.FindMissingBlobsResponse.MissingBlobDigests)
+	for (const TPair<FIoHash, UE::RemoteExecution::EStatusCode>& MissingItem : State.DoBlobsExistResponse)
 	{
-		if (MissingItem == State.ActionDigest)
+		if (MissingItem.Value == UE::RemoteExecution::EStatusCode::Ok)
 		{
-			Executor.Stats.TotalActionBlobsUploaded.AddBlob(State.ActionContentBytes.Num());
-			FBatchUpdateBlobsRequest::FRequest& NewRequest = State.BatchUpdateBlobsRequest.Requests.AddDefaulted_GetRef();
-			NewRequest.Digest = State.ActionDigest;
-			NewRequest.Data = MoveTemp(State.ActionContentBytes);
-			UE_LOG(LogDerivedDataBuildRemoteExecutor, Verbose, TEXT("Uploading action (hash: %s) of upload size %d."), *::LexToString(NewRequest.Digest.Hash), NewRequest.Data.Num());
+			continue;
 		}
-		else if (MissingItem == State.Action.CommandDigest)
+
+		if (MissingItem.Key == State.TaskDigest)
+		{
+			Executor.Stats.TotalActionBlobsUploaded.AddBlob(State.TaskContentBytes.Num());
+			State.PutBlobsRequest.Add(State.TaskDigest, MoveTemp(State.TaskContentBytes));
+			UE_LOG(LogDerivedDataBuildRemoteExecutor, Verbose, TEXT("Uploading task (hash: %s) of upload size %d."), *::LexToString(State.TaskDigest), State.PutBlobsRequest[State.TaskDigest].Num());
+		}
+		else if (MissingItem.Key == State.Task.RequirementsHash)
 		{
-			Executor.Stats.TotalCommandBlobsUploaded.AddBlob(State.CommandContentBytes.Num());
-			FBatchUpdateBlobsRequest::FRequest& NewRequest = State.BatchUpdateBlobsRequest.Requests.AddDefaulted_GetRef();
-			NewRequest.Digest = State.Action.CommandDigest;
-			NewRequest.Data = MoveTemp(State.CommandContentBytes);
-			UE_LOG(LogDerivedDataBuildRemoteExecutor, Verbose, TEXT("Uploading command (hash: %s) of upload size %d."), *::LexToString(NewRequest.Digest.Hash), NewRequest.Data.Num());
+			Executor.Stats.TotalRequirementBlobsUploaded.AddBlob(State.RequirementsContentBytes.Num());
+			State.PutBlobsRequest.Add(State.Task.RequirementsHash, MoveTemp(State.RequirementsContentBytes));
+			UE_LOG(LogDerivedDataBuildRemoteExecutor, Verbose, TEXT("Uploading requirements (hash: %s) of upload size %d."), *::LexToString(State.Task.RequirementsHash), State.PutBlobsRequest[State.Task.RequirementsHash].Num());
 		}
 		else
 		{
-			const FVariantIndex VariantIndex = State.DigestFilesystemIndex.FindChecked(MissingItem);
-			FBatchUpdateBlobsRequest::FRequest& NewRequest = State.BatchUpdateBlobsRequest.Requests.AddDefaulted_GetRef();
-			NewRequest.Digest = MissingItem;
+			const FVariantIndex VariantIndex = State.DigestFilesystemIndex.FindChecked(MissingItem.Key);
 			switch (VariantIndex.NodeType)
 			{
 			case ENodeType::Directory:
 				Executor.Stats.TotalDirectoryBlobsUploaded.AddBlob(State.Directories[VariantIndex.Index].ContentBytes.Num());
-				NewRequest.Digest = State.Directories[VariantIndex.Index].Digest.GetValue();
-				NewRequest.Data = MoveTemp(State.Directories[VariantIndex.Index].ContentBytes);
-				UE_LOG(LogDerivedDataBuildRemoteExecutor, Verbose, TEXT("Uploading directory '%s' (hash: %s) of upload size %d."), *FString(State.Directories[VariantIndex.Index].Path), *::LexToString(NewRequest.Digest.Hash), NewRequest.Data.Num());
+				State.PutBlobsRequest.Add(State.Directories[VariantIndex.Index].Digest.GetValue(), MoveTemp(State.Directories[VariantIndex.Index].ContentBytes));
+				UE_LOG(LogDerivedDataBuildRemoteExecutor, Verbose, TEXT("Uploading directory '%s' (hash: %s) of upload size %d."),
+					*FString(State.Directories[VariantIndex.Index].Path),
+					*::LexToString(State.Directories[VariantIndex.Index].Digest.GetValue()),
+					State.PutBlobsRequest[State.Directories[VariantIndex.Index].Digest.GetValue()].Num());
 				break;
 			case ENodeType::File:
 				{
 					Executor.Stats.TotalFileBlobsUploaded.AddBlob(State.Files[VariantIndex.Index].ContentBytes.GetSize());
-					NewRequest.Digest.Hash = FIoHash::HashBuffer(State.Files[VariantIndex.Index].ContentBytes);
-					NewRequest.Digest.SizeBytes = State.Files[VariantIndex.Index].ContentBytes.GetSize();
 					FCompositeBuffer& FileBuffer = State.Files[VariantIndex.Index].ContentBytes;
 					check(!FileBuffer.IsNull());
-					NewRequest.Data.Reserve(FileBuffer.GetSize());
+					TArray<uint8> FileData;
+					FileData.Reserve(FileBuffer.GetSize());
 					for (const FSharedBuffer& Segment : State.Files[VariantIndex.Index].ContentBytes.GetSegments())
 					{
-						NewRequest.Data.Append((const uint8 *)Segment.GetData(), Segment.GetSize());
+						FileData.Append((const uint8 *)Segment.GetData(), Segment.GetSize());
 					}
 					FileBuffer.Reset();
-					UE_LOG(LogDerivedDataBuildRemoteExecutor, Verbose, TEXT("Uploading file '%s' (hash: %s, type: %s) of upload size %d."), *FString(State.Files[VariantIndex.Index].Path), *::LexToString(NewRequest.Digest.Hash), LexToString(State.Files[VariantIndex.Index].Type), NewRequest.Data.Num());
+					State.PutBlobsRequest.Add(State.Files[VariantIndex.Index].File.Hash, MoveTemp(FileData));
+					UE_LOG(LogDerivedDataBuildRemoteExecutor, Verbose, TEXT("Uploading file '%s' (hash: %s, type: %s) of upload size %d."),
+						*FString(State.Files[VariantIndex.Index].Path),
+						*::LexToString(State.Files[VariantIndex.Index].File.Hash), LexToString(State.Files[VariantIndex.Index].Type),
+						State.PutBlobsRequest[State.Files[VariantIndex.Index].File.Hash].Num());
 				}
 				break;
 			default:
@@ -969,43 +1018,51 @@ TFuture<TPair<FStatus, FBatchUpdateBlobsResponse>> FRemoteBuildExecutionRequest:
 		}
 	}
 
-	return Executor.ContentAddressableStorage->BatchUpdateBlobsAsync(State.BatchUpdateBlobsRequest);
+	return Executor.ContentAddressableStorage->PutBlobsAsync(Executor.NameSpaceId, State.PutBlobsRequest);
 }
 
-TFuture<FExecuteResponse> FRemoteBuildExecutionRequest::ExecuteBuildAsync()
+TFuture<TPair<UE::RemoteExecution::EStatusCode, UE::RemoteExecution::FGetTaskUpdatesResponse>> FRemoteBuildExecutionRequest::ExecuteBuildAsync()
 {
-	State.ExecuteRequest.InstanceName = Executor.InstanceName;
-	State.ExecuteRequest.ActionDigest = State.ActionDigest;
-	State.ExecuteRequest.SkipCacheLookup = true;
-	return Executor.RemoteExecutor->GetExecution()->ExecuteAsync(State.ExecuteRequest);
+	State.ExecuteRequest.RequirementsHash = State.Task.RequirementsHash;
+	State.ExecuteRequest.TaskHashes.Empty();
+	State.ExecuteRequest.TaskHashes.Add(State.TaskDigest);
+	//State.ExecuteRequest.DoNotCache = true;
+	return Executor.Execution->RunTasksAsync(State.ExecuteRequest, Executor.GlobalExecutionTimeoutSeconds);
 }
 
-TFuture<TPair<FStatus, FBatchReadBlobsResponse>> FRemoteBuildExecutionRequest::DownloadResultsAsync()
+TFuture<TPair<UE::RemoteExecution::EStatusCode, UE::RemoteExecution::FGetObjectTreeResponse>> FRemoteBuildExecutionRequest::GetTaskResultAsync()
 {
-	State.BatchReadBlobsRequest.InstanceName = Executor.InstanceName;
-	//State.BatchReadBlobsRequest.Digests.Add(State.ExecuteResponse.Result.StdoutDigest);
-	//State.BatchReadBlobsRequest.Digests.Add(State.ExecuteResponse.Result.StderrDigest);
-
-	for (const FOutputFile& OutputFile : State.ExecuteResponse.Result.OutputFiles)
-	{
-		State.BatchReadBlobsRequest.Digests.Add(OutputFile.Digest);
-	}
+	return Executor.ContentAddressableStorage->GetObjectTreeAsync(Executor.NameSpaceId, State.ExecuteResponse.ResultHash);
+}
 
-	return Executor.ContentAddressableStorage->BatchReadBlobsAsync(State.BatchReadBlobsRequest);
+TFuture<TMap<FIoHash, TPair<UE::RemoteExecution::EStatusCode, TArray<uint8>>>> FRemoteBuildExecutionRequest::DownloadResultsAsync()
+{
+	return Executor.ContentAddressableStorage->GetBlobsAsync(Executor.NameSpaceId, State.GetBlobsRequest);
 }
 
-void FRemoteBuildExecutionRequest::OnMissingBlobsDetermined(FStatus Status, FFindMissingBlobsResponse&& Result)
+void FRemoteBuildExecutionRequest::OnMissingBlobsDetermined(TMap<FIoHash, UE::RemoteExecution::EStatusCode>&& Result)
 {
-	if (ProcessCancellation() || !IsStatusOk(Status, TEXT("FindMissingBlobs")))
+	if (ProcessCancellation())
 	{
 		return;
 	}
 
-	State.FindMissingBlobsResponse = MoveTemp(Result);
+	for (const TPair<FIoHash, UE::RemoteExecution::EStatusCode>& Entry : Result)
+	{
+		if (Entry.Value != UE::RemoteExecution::EStatusCode::NotFound && !IsStatusOk(Entry.Value, TEXT("FindMissingBlobs")))
+		{
+			return;
+		}
+	}
+
+	State.DoBlobsExistResponse = MoveTemp(Result);
 	constexpr bool bForceUploads = false;
 	if (bForceUploads)
 	{
-		State.FindMissingBlobsResponse.MissingBlobDigests = State.FindMissingBlobsRequest.BlobDigests;
+		for (TPair<FIoHash, UE::RemoteExecution::EStatusCode>& Entry : State.DoBlobsExistResponse)
+		{
+			Entry.Value = UE::RemoteExecution::EStatusCode::NotFound;
+		}
 	}
 
 	TArray<FStringView> MissingInputViews;
@@ -1035,26 +1092,36 @@ void FRemoteBuildExecutionRequest::OnMissingBlobsDetermined(FStatus Status, FFin
 	// TODO: This should be async but isn't.  Requires IRequest chaining.
 	LoadMissingWorkerFileBlobsAsync();
 
-	if (!State.FindMissingBlobsResponse.MissingBlobDigests.IsEmpty())
+	TMap<FIoHash, UE::RemoteExecution::EStatusCode> Missing = State.DoBlobsExistResponse.FilterByPredicate([](const TPair<FIoHash, UE::RemoteExecution::EStatusCode>& Entry) { return Entry.Value == UE::RemoteExecution::EStatusCode::NotFound; });
+
+	if (!Missing.IsEmpty())
 	{
 		UploadMissingBlobsAsync()
-			.Next([this] (const TPair<FStatus, FBatchUpdateBlobsResponse>& InnerResult) { OnMissingBlobsUploaded(InnerResult.Key, InnerResult.Value); });
+			.Next([this] (const TMap<FIoHash, UE::RemoteExecution::EStatusCode>& InnerResult) { OnMissingBlobsUploaded(InnerResult); });
 	}
 	else
 	{
 		ExecuteBuildAsync()
-			.Next([this] (FExecuteResponse&& Result) { OnExecutionCompleted(MoveTemp(Result)); });
+			.Next([this] (TPair<UE::RemoteExecution::EStatusCode, UE::RemoteExecution::FGetTaskUpdatesResponse>&& Result) { OnExecutionCompleted(Result.Key, MoveTemp(Result.Value)); });
 	}
 }
 
-void FRemoteBuildExecutionRequest::OnMissingBlobsUploaded(FStatus Status, const FBatchUpdateBlobsResponse& Result)
+void FRemoteBuildExecutionRequest::OnMissingBlobsUploaded(const TMap<FIoHash, UE::RemoteExecution::EStatusCode>& Result)
 {
-	if (ProcessCancellation() || !IsStatusOk(Status, TEXT("BatchUploadBlobs")))
+	if (ProcessCancellation())
 	{
 		return;
 	}
 
-	UE_LOG(LogDerivedDataBuildRemoteExecutor, Verbose, TEXT("Uploaded %d data blobs for remote execution."), State.BatchUpdateBlobsRequest.Requests.Num());
+	for (const TPair<FIoHash, UE::RemoteExecution::EStatusCode>& Entry : Result)
+	{
+		if (!IsStatusOk(Entry.Value, TEXT("BatchUploadBlobs")))
+		{
+			return;
+		}
+	}
+
+	UE_LOG(LogDerivedDataBuildRemoteExecutor, Verbose, TEXT("Uploaded %d data blobs for remote execution."), State.PutBlobsRequest.Num());
 	if (!ValidateUploadSuccess(Result))
 	{
 		State.Owner.End(this, [this]
@@ -1066,47 +1133,86 @@ void FRemoteBuildExecutionRequest::OnMissingBlobsUploaded(FStatus Status, const
 	}
 
 	ExecuteBuildAsync()
-		.Next([this] (FExecuteResponse&& Result) { OnExecutionCompleted(MoveTemp(Result)); });
+		.Next([this] (TPair<UE::RemoteExecution::EStatusCode, UE::RemoteExecution::FGetTaskUpdatesResponse>&& Result) { OnExecutionCompleted(Result.Key, MoveTemp(Result.Value)); });
 
 }
 
-void FRemoteBuildExecutionRequest::OnExecutionCompleted(FExecuteResponse&& Result)
+void FRemoteBuildExecutionRequest::OnExecutionCompleted(UE::RemoteExecution::EStatusCode Status, UE::RemoteExecution::FGetTaskUpdatesResponse&& Result)
 {
 	if (ProcessCancellation())
 	{
 		return;
 	}
 
-	State.ExecuteResponse = MoveTemp(Result);
-	if (State.ExecuteResponse.Status.Ok())
+	if (!IsStatusOk(Status, TEXT("OnExecutionCompleted")))
 	{
-		DownloadResultsAsync()
-			.Next([this] (TPair<FStatus, FBatchReadBlobsResponse>&& InnerResult)
-				{
-					OnOutputBlobsDownloaded(InnerResult.Key, MoveTemp(InnerResult.Value));
-				});
+		return;
 	}
-	else
+
+	if (Result.Updates.IsEmpty())
 	{
-		UE_LOG(LogDerivedDataBuildRemoteExecutor, Warning, TEXT("Remote execution system error: Failed to execute build operation!  Response message: %s"), *State.ExecuteResponse.Status.Message);
+		UE_LOG(LogDerivedDataBuildRemoteExecutor, Warning, TEXT("Remote execution system error: Failed to get results from remote build operation!"));
 		State.Owner.End(this, [this]
-		{
-			CompletionCallback({State.BuildAction.GetKey(), {}, {}, EStatus::Error});
-			CompletionEvent->Trigger();
-		});
+			{
+				CompletionCallback({ State.BuildAction.GetKey(), {}, {}, EStatus::Error });
+				CompletionEvent->Trigger();
+			});
+		return;
 	}
 
+	if (Result.Updates[0].ResultHash == FIoHash::Zero)
+	{
+		UE_LOG(LogDerivedDataBuildRemoteExecutor, Warning, TEXT("Remote execution system error: Zero ResultHash returned from remote build operation!"));
+		State.Owner.End(this, [this]
+			{
+				CompletionCallback({ State.BuildAction.GetKey(), {}, {}, EStatus::Error });
+				CompletionEvent->Trigger();
+			});
+		return;
+	}
+
+	State.ExecuteResponse = MoveTemp(Result.Updates[0]);
+	GetTaskResultAsync()
+		.Next([this] (TPair<UE::RemoteExecution::EStatusCode, UE::RemoteExecution::FGetObjectTreeResponse>&& InnerResult)
+			{
+				OnTaskResultDownloaded(InnerResult.Key, MoveTemp(InnerResult.Value));
+			});
+
 }
 
-void FRemoteBuildExecutionRequest::OnOutputBlobsDownloaded(FStatus Status, FBatchReadBlobsResponse&& Result)
+void FRemoteBuildExecutionRequest::OnTaskResultDownloaded(UE::RemoteExecution::EStatusCode Status, UE::RemoteExecution::FGetObjectTreeResponse&& Result)
 {
-	if (ProcessCancellation() || !IsStatusOk(Status, TEXT("BatchReadBlobs")))
+	if (ProcessCancellation())
+	{
+		return;
+	}
+
+	if (!IsStatusOk(Status, TEXT("OnTaskResultDownloaded")))
+	{
+		return;
+	}
+
+	const FCbObjectView View = FCbObjectView(Result.Objects[State.ExecuteResponse.ResultHash].GetData());
+	State.ExecuteResult.Load(View);
+	State.GetBlobsRequest = MoveTemp(Result.BinaryAttachments);
+	GatherExecuteFileResults(Result.Objects, TEXT(""), State.ExecuteResult.OutputHash);
+
+	DownloadResultsAsync()
+		.Next([this](TMap<FIoHash, TPair<UE::RemoteExecution::EStatusCode, TArray<uint8>>>&& InnerResult)
+			{
+				OnOutputBlobsDownloaded(MoveTemp(InnerResult));
+			});
+}
+
+void FRemoteBuildExecutionRequest::OnOutputBlobsDownloaded(TMap<FIoHash, TPair<UE::RemoteExecution::EStatusCode, TArray<uint8>>>&& Result)
+{
+	if (ProcessCancellation())
 	{
 		return;
 	}
 
 	EStatus BuildStatus = EStatus::Error;
-	FOptionalBuildOutput BuildOutput = ComposeBuildOutput(Result, BuildStatus);
+	FOptionalBuildOutput BuildOutput = ComposeBuildOutput(MoveTemp(Result), BuildStatus);
 	if (BuildStatus == EStatus::Ok)
 	{
 		Executor.Stats.TotalSuccessfulRemoteBuilds.fetch_add(1, std::memory_order_relaxed);
diff --git a/Engine/Source/Editor/UnrealEd/UnrealEd.Build.cs b/Engine/Source/Editor/UnrealEd/UnrealEd.Build.cs
--- a/Engine/Source/Editor/UnrealEd/UnrealEd.Build.cs
+++ b/Engine/Source/Editor/UnrealEd/UnrealEd.Build.cs
@@ -68,7 +68,6 @@ public class UnrealEd : ModuleRules
 				"NavigationSystem",
 				"Media",
 				"VirtualTexturingEditor",
-				"RemoteExecution",
 				"TextureBuild"
 			}
 		);
@@ -210,6 +209,8 @@ public class UnrealEd : ModuleRules
 				"AnimationBlueprintLibrary",
 				"MaterialBaking",
 				"CookOnTheFly",
+				"RemoteExecution",
+				"HordeExecutor",
 			}
 		);
 
@@ -280,7 +281,6 @@ public class UnrealEd : ModuleRules
 					"ProjectTargetPlatformEditor",
 					"PListEditor",
 					"TraceInsights",
-					"BazelExecutor",
 				}
 			);
 		}
