Improvements to texture debug view modes to help debug VT resolution.
Fix velocity+depth pass being removed in texture debug view modes.
Add [VT] suffix to debug view texture menu entries.

diff --git a/Engine/Source/Editor/UnrealEd/Private/EditorViewportCommands.cpp b/Engine/Source/Editor/UnrealEd/Private/EditorViewportCommands.cpp
--- a/Engine/Source/Editor/UnrealEd/Private/EditorViewportCommands.cpp
+++ b/Engine/Source/Editor/UnrealEd/Private/EditorViewportCommands.cpp
@@ -301,7 +301,7 @@ static void GetSelectedMaterials(TArray<const UMaterialInterface*>& SelectedMate
 	}
 }
 
-static void AppendTextureStreamingInfoToMenu(const UMaterialInterface* MaterialInterface, bool bSingleMaterial, TMap<int32, TArray<FString> >& DataPerTextureIndex)
+static void AppendTextureStreamingInfoToMenu(const UMaterialInterface* MaterialInterface, bool bSingleMaterial, TMap<int32, TArray<FString>>& DataPerTextureIndex)
 {
 	check(MaterialInterface);
 	for (const FMaterialTextureInfo& TextureData : MaterialInterface->GetTextureStreamingData())
@@ -332,7 +332,15 @@ static void AppendTextureStreamingInfoToMenu(const UMaterialInterface* MaterialI
 	}
 }
 
-static void AppendMaterialInfoToMenu(const UMaterialInterface* MaterialInterface, ERHIFeatureLevel::Type FeatureLevel, const FString& MenuName, TMap<int32, TArray<FString> >& DataPerTextureIndex, TMap<FName, TArray<FString> >& DataPerTextureName)
+struct FMenuDataPerTexture
+{
+	// Menu display name.
+	FString DisplayName;
+	// Menu tool tip lines.
+	TArray<FString> ToolTips;
+};
+
+static void AppendMaterialInfoToMenu(const UMaterialInterface* MaterialInterface, ERHIFeatureLevel::Type FeatureLevel, const FString& MenuName, TMap<int32, TArray<FString> >& DataPerTextureIndex, TMap<FName, FMenuDataPerTexture>& DataPerTextureName)
 {
 	check(MaterialInterface);
 	const FMaterialResource* Material = MaterialInterface->GetMaterialResource(FeatureLevel);
@@ -351,8 +359,19 @@ static void AppendMaterialInfoToMenu(const UMaterialInterface* MaterialInterface
 				if (Texture2D)
 				{
 					const FMaterialTextureParameterInfo& Parameter = UniformExpressions.GetTextureParameter(TextureType, i);
+					
 					DataPerTextureIndex.FindOrAdd(Parameter.TextureIndex).AddUnique(FString::Printf(TEXT("%s.%s"), *MaterialInterface->GetName(), *Texture2D->GetName()));
-					DataPerTextureName.FindOrAdd(*Texture2D->GetName()).AddUnique(FString::Printf(TEXT("%s %d : %s"), *MenuName, Parameter.TextureIndex, *MaterialInterface->GetName()));
+					
+					FMenuDataPerTexture& MenuData = DataPerTextureName.FindOrAdd(*Texture2D->GetName());
+					if (MenuData.DisplayName.IsEmpty())
+					{
+						MenuData.DisplayName = *Texture2D->GetName();
+						if (TextureType == EMaterialTextureParameterType::Virtual)
+						{
+							MenuData.DisplayName += TEXT(" [VT]");
+						}
+					}
+					MenuData.ToolTips.AddUnique(FString::Printf(TEXT("%s %d : %s"), *MenuName, Parameter.TextureIndex, *MaterialInterface->GetName()));
 				}
 			}
 		}
@@ -382,7 +401,7 @@ TSharedRef<SWidget> BuildViewModeOptionsMenu(TSharedPtr<FUICommandList> CommandL
 		GetSelectedMaterials(SelectedMaterials);
 
 		TMap<int32, TArray<FString> > DataPerTextureIndex;
-		TMap<FName, TArray<FString> > DataPerTextureName;
+		TMap<FName, FMenuDataPerTexture> DataPerTextureName;
 
 		if (ViewModeIndex == VMI_MaterialTextureScaleAccuracy)
 		{
@@ -416,13 +435,15 @@ TSharedRef<SWidget> BuildViewModeOptionsMenu(TSharedPtr<FUICommandList> CommandL
 			int32 CommandIndex = 0;
 			for (auto Ite = SortedFNames.CreateConstIterator(); Ite; ++Ite)
 			{ 
-				const TArray<FString>& MaterialInfos = DataPerTextureName.FindOrAdd(Ite.Value());
-				FString ToolTipOverride = MaterialInfos[0];
-				for (int32 MaterialIndex = 1; MaterialIndex < MaterialInfos.Num(); ++MaterialIndex)
+				const FMenuDataPerTexture& MaterialInfos = DataPerTextureName.FindOrAdd(Ite.Value());
+				FString ToolTipOverride = MaterialInfos.ToolTips[0];
+				for (int32 MaterialIndex = 1; MaterialIndex < MaterialInfos.ToolTips.Num(); ++MaterialIndex)
 				{
-					ToolTipOverride = FString::Printf(TEXT("%s\n%s"), *ToolTipOverride, *MaterialInfos[MaterialIndex]);
+					ToolTipOverride = FString::Printf(TEXT("%s\n%s"), *ToolTipOverride, *MaterialInfos.ToolTips[MaterialIndex]);
 				}
-				MenuBuilder.AddMenuEntry(PerTextureCommands[CommandIndex], NAME_None, FText::FromString(Ite.Key()), FText::FromString(ToolTipOverride));
+				
+				MenuBuilder.AddMenuEntry(PerTextureCommands[CommandIndex], NAME_None, FText::FromString(MaterialInfos.DisplayName), FText::FromString(ToolTipOverride));
+				
 				ParamNameMap.Add(CommandIndex, Ite.Value());
 				++CommandIndex;
 			}
diff --git a/Engine/Source/Runtime/Renderer/Private/VelocityRendering.cpp b/Engine/Source/Runtime/Renderer/Private/VelocityRendering.cpp
--- a/Engine/Source/Runtime/Renderer/Private/VelocityRendering.cpp
+++ b/Engine/Source/Runtime/Renderer/Private/VelocityRendering.cpp
@@ -390,6 +390,11 @@ bool FVelocityMeshProcessor::PrimitiveHasVelocityForView(const FViewInfo& View,
 	{
 		return false;
 	}
+	// Velocity pass not rendered for debug views.
+	if (View.Family->UseDebugViewPS())
+	{
+		return false;
+	}
 
 	const FBoxSphereBounds& PrimitiveBounds = PrimitiveSceneProxy->GetBounds();
 	const float LODFactorDistanceSquared = (PrimitiveBounds.Origin - View.ViewMatrices.GetViewOrigin()).SizeSquared() * FMath::Square(View.LODDistanceFactor);
