Fix CreateAsset dialog cancel button not doing anything

diff --git a/Engine/Source/Editor/UnrealEd/Private/SCreateAssetFromObject.cpp b/Engine/Source/Editor/UnrealEd/Private/SCreateAssetFromObject.cpp
--- a/Engine/Source/Editor/UnrealEd/Private/SCreateAssetFromObject.cpp
+++ b/Engine/Source/Editor/UnrealEd/Private/SCreateAssetFromObject.cpp
@@ -122,11 +122,12 @@ void SCreateAssetFromObject::Construct(const FArguments& InArgs, TSharedPtr<SWin
 				.Text(CreateButtonText)
 			]
 			+ SHorizontalBox::Slot()
-			.Padding(0, 2, 0, 5)
+			.Padding(0, 2, 10, 5)
 			.AutoWidth()
 			[
 				SNew(SButton)
 				.Text(LOCTEXT("CancelButtonText", "Cancel"))
+				.OnClicked(this, &SCreateAssetFromObject::OnCancelCreateAssetFromActor)
 			]
 		]
 	];
