AssetRegistryGenerator - update syntax to current patterns used on the core team. No change in behavior.
#rnx

diff --git a/Engine/Source/Editor/UnrealEd/Private/Commandlets/AssetRegistryGenerator.cpp b/Engine/Source/Editor/UnrealEd/Private/Commandlets/AssetRegistryGenerator.cpp
--- a/Engine/Source/Editor/UnrealEd/Private/Commandlets/AssetRegistryGenerator.cpp
+++ b/Engine/Source/Editor/UnrealEd/Private/Commandlets/AssetRegistryGenerator.cpp
@@ -28,7 +28,9 @@
 #include "Templates/UniquePtr.h"
 #include "Engine/AssetManager.h"
 #include "Misc/DataDrivenPlatformInfoRegistry.h"
+#include "Algo/Sort.h"
 #include "Algo/StableSort.h"
+#include "ProfilingDebugging/ScopedTimers.h"
 
 #include "PakFileUtilities.h"
 
@@ -91,9 +93,8 @@ FAssetRegistryGenerator::FAssetRegistryGenerator(const ITargetPlatform* InPlatfo
 	, bGenerateChunks(false)
 	, bUseAssetManager(false)
 	, HighestChunkId(0)
+	, DependencyInfo(*GetMutableDefault<UChunkDependencyInfo>())
 {
-	DependencyInfo = GetMutableDefault<UChunkDependencyInfo>();
-
 	bool bOnlyHardReferences = false;
 	const UProjectPackagingSettings* const PackagingSettings = GetDefault<UProjectPackagingSettings>();
 	if (PackagingSettings)
@@ -119,16 +120,6 @@ FAssetRegistryGenerator::FAssetRegistryGenerator(const ITargetPlatform* InPlatfo
 
 FAssetRegistryGenerator::~FAssetRegistryGenerator()
 {
-	for (auto ChunkSet : ChunkManifests)
-	{
-		delete ChunkSet;
-	}
-	ChunkManifests.Empty();
-	for (auto ChunkSet : FinalChunkManifests)
-	{
-		delete ChunkSet;
-	}
-	FinalChunkManifests.Empty();
 }
 
 bool FAssetRegistryGenerator::CleanTempPackagingDirectory(const FString& Platform) const
@@ -171,10 +162,9 @@ bool FAssetRegistryGenerator::ShouldPlatformGenerateStreamingInstallManifest(con
 	return false;
 }
 
-
 int64 FAssetRegistryGenerator::GetMaxChunkSizePerPlatform(const ITargetPlatform* Platform) const
 {
-	if ( Platform )
+	if (Platform)
 	{
 		FConfigFile PlatformIniFile;
 		FConfigCacheIni::LoadLocalIniFile(PlatformIniFile, TEXT("Game"), true, *Platform->IniPlatformName());
@@ -188,6 +178,36 @@ int64 FAssetRegistryGenerator::GetMaxChunkSizePerPlatform(const ITargetPlatform*
 	return -1;
 }
 
+TArray<int32> FAssetRegistryGenerator::GetExistingPackageChunkAssignments(FName PackageFName)
+{
+	TArray<int32> ExistingChunkIDs;
+	for (uint32 ChunkIndex = 0, MaxChunk = ChunkManifests.Num(); ChunkIndex < MaxChunk; ++ChunkIndex)
+	{
+		if (ChunkManifests[ChunkIndex] && ChunkManifests[ChunkIndex]->Contains(PackageFName))
+		{
+			ExistingChunkIDs.AddUnique(ChunkIndex);
+		}
+	}
+
+	if (StartupPackages.Contains(PackageFName))
+	{
+		ExistingChunkIDs.AddUnique(0);
+	}
+
+	return ExistingChunkIDs;
+}
+
+TArray<int32> FAssetRegistryGenerator::GetAssetRegistryChunkAssignments(const FName& PackageFName)
+{
+	TArray<int32> PackageInputChunkIds;
+	const TArray<int32>* FoundIDs = PackageChunkIDMap.Find(PackageFName);
+	if (FoundIDs)
+	{
+		PackageInputChunkIds = *FoundIDs;
+	}
+	return PackageInputChunkIds;
+}
+
 class FPackageFileSizeVisitor : public IPlatformFile::FDirectoryStatVisitor
 {
 	TMap<FString, int64>& PackageFileSizes;
@@ -277,18 +297,17 @@ static void AssignLayerChunkDelegate(const FAssignLayerChunkMap* ChunkManifest,
 	}
 }
 
-void IChunkDataGenerator::GenerateChunkDataFiles(const int32 InChunkId, const TSet<FName>& InPackagesInChunk, const ITargetPlatform* TargetPlatform, FSandboxPlatformFile* InSandboxFile, TArray<FString>& OutChunkFilenames)
+void IChunkDataGenerator::GenerateChunkDataFiles(const int32 InChunkId, const TSet<FName>& InPackagesInChunk,
+	const ITargetPlatform* TargetPlatform, FSandboxPlatformFile* InSandboxFile, TArray<FString>& OutChunkFilenames)
 {
 	PRAGMA_DISABLE_DEPRECATION_WARNINGS
 	GenerateChunkDataFiles(InChunkId, InPackagesInChunk, TargetPlatform->PlatformName(), InSandboxFile, OutChunkFilenames);
 	PRAGMA_ENABLE_DEPRECATION_WARNINGS
 }
 
-bool FAssetRegistryGenerator::GenerateStreamingInstallManifest(int64 InExtraFlavorChunkSize, FSandboxPlatformFile* InSandboxFile)
+bool FAssetRegistryGenerator::GenerateStreamingInstallManifest(int64 InExtraFlavorChunkSize, FSandboxPlatformFile& InSandboxFile)
 {
 	const FString Platform = TargetPlatform->PlatformName();
-
-	// empty out the current paklist directory
 	FString TmpPackagingDir = GetTempPackagingDirectoryForPlatform(Platform);
 
 	int64 MaxChunkSize = GetMaxChunkSizePerPlatform( TargetPlatform );
@@ -299,38 +318,44 @@ bool FAssetRegistryGenerator::GenerateStreamingInstallManifest(int64 InExtraFlav
 		MaxChunkSize = InExtraFlavorChunkSize;
 	}
 
-	if (!IFileManager::Get().MakeDirectory(*TmpPackagingDir, true))
+	if (!IFileManager::Get().MakeDirectory(*TmpPackagingDir, true /* Tree */))
 	{
 		UE_LOG(LogAssetRegistryGenerator, Error, TEXT("Failed to create directory: %s"), *TmpPackagingDir);
 		return false;
 	}
 	
-	// open a file for writing the list of pak file lists that we've generated
 	FString PakChunkListFilename = TmpPackagingDir / TEXT("pakchunklist.txt");
+	FString PakChunkLayerInfoFilename = TmpPackagingDir / TEXT("pakchunklayers.txt");
+	// List of pak file lists
 	TUniquePtr<FArchive> PakChunkListFile(IFileManager::Get().CreateFileWriter(*PakChunkListFilename));
-
-	if (!PakChunkListFile)
+	// List of disc layer for each chunk
+	TUniquePtr<FArchive> ChunkLayerFile(IFileManager::Get().CreateFileWriter(*PakChunkLayerInfoFilename));
+	if (!PakChunkListFile || !ChunkLayerFile)
 	{
-		UE_LOG(LogAssetRegistryGenerator, Error, TEXT("Failed to open output pakchunklist file %s"), *PakChunkListFilename);
+		UE_LOG(LogAssetRegistryGenerator, Error, TEXT("Failed to open output pakchunklist file %s"),
+			!PakChunkListFile ? *PakChunkListFilename : *PakChunkLayerInfoFilename);
 		return false;
 	}
 
-	FString PakChunkLayerInfoFilename = FString::Printf(TEXT("%s/pakchunklayers.txt"), *TmpPackagingDir);
-	TUniquePtr<FArchive> ChunkLayerFile(IFileManager::Get().CreateFileWriter(*PakChunkLayerInfoFilename));
-
 	TArray<FString> CompressedChunkWildcards;
+	TOptional<FConfigFile> PlatformIniFile;
+	auto ConditionalLoadPlatformIniFile = [&PlatformIniFile, this]()
 	{
-		// never screw with server pak files.  This hack only cares about client platforms
-		if (!TargetPlatform->IsServerOnly())
+		if (!PlatformIniFile)
 		{
-			FConfigFile PlatformIniFile;
-			FConfigCacheIni::LoadLocalIniFile(PlatformIniFile, TEXT("Game"), true, *TargetPlatform->IniPlatformName());
-			FString ConfigString;
-			PlatformIniFile.GetArray(TEXT("/Script/UnrealEd.ProjectPackagingSettings"), TEXT("CompressedChunkWildcard"), CompressedChunkWildcards);
+			PlatformIniFile.Emplace();
+			FConfigCacheIni::LoadLocalIniFile(*PlatformIniFile, TEXT("Game"), true, *TargetPlatform->IniPlatformName());
 		}
+	};
+	if (!TargetPlatform->IsServerOnly())
+	{
+		// Load the list of wildcards to specify which pakfiles should be compressed, if the targetplatform supports it.
+		// This is only used in client platforms.
+		ConditionalLoadPlatformIniFile();
+		PlatformIniFile->GetArray(TEXT("/Script/UnrealEd.ProjectPackagingSettings"), TEXT("CompressedChunkWildcard"), CompressedChunkWildcards);
 	}
 
-	// Update manifests for any encryption groups that contain non-asset files files
+	// Update manifests for any encryption groups that contain non-asset files
 	if (bUseAssetManager && !TargetPlatform->HasSecurePackageFormat())
 	{
 		FContentEncryptionConfig ContentEncryptionConfig;
@@ -353,20 +378,20 @@ bool FAssetRegistryGenerator::GenerateStreamingInstallManifest(int64 InExtraFlav
 					// Extend the array until it is large enough to hold the requested index, filling it in with nulls on all the newly added indices.
 					// Note that this will temporarily break our contract that FinalChunkManifests does not contain null pointers; we fix up the contract
 					// by replacing any remaining null pointers in the loop over FinalChunkManifests at the bottom of this function.
-					FinalChunkManifests.AddZeroed(PakchunkIndex - FinalChunkManifests.Num() + 1);
+					FinalChunkManifests.AddDefaulted(PakchunkIndex - FinalChunkManifests.Num() + 1);
 				}
-				checkf(PakchunkIndex < FinalChunkManifests.Num(), TEXT("Chunk %i out of range. %i manifests available"), PakchunkIndex, FinalChunkManifests.Num() - 1);
 
-				FChunkPackageSet* Manifest = FinalChunkManifests[PakchunkIndex];
+				FChunkPackageSet* Manifest = FinalChunkManifests[PakchunkIndex].Get();
 				if (Manifest == nullptr)
 				{
 					Manifest = new FChunkPackageSet();
-					FinalChunkManifests[PakchunkIndex] = Manifest;
+					FinalChunkManifests[PakchunkIndex].Reset(Manifest);
 				}
 
 				for (const FString& NonAssetFile : EncryptionGroup.NonAssetFiles)
 				{
-					Manifest->Add(*NonAssetFile, FPaths::RootDir() / NonAssetFile); // Paths added as relative to the root. The staging code will need to map this onto the target path of all staged assets
+					// Paths added as relative to the root. The staging code will need to map this onto the target path of all staged assets
+					Manifest->Add(*NonAssetFile, FPaths::RootDir() / NonAssetFile);
 				}
 			}
 		}
@@ -376,7 +401,7 @@ bool FAssetRegistryGenerator::GenerateStreamingInstallManifest(int64 InExtraFlav
 	if (MaxChunkSize > 0)
 	{
 		// ZENTODO: Pull file size information from the AssetRegistryState, since the packages do not exist on disk in zen
-		FString SandboxPath = InSandboxFile->GetSandboxDirectory();
+		FString SandboxPath = InSandboxFile.GetSandboxDirectory();
 		SandboxPath.ReplaceInline(TEXT("[Platform]"), *Platform);
 		FPackageFileSizeVisitor PackageSearch(PackageFileSizes);
 		IFileManager::Get().IterateDirectoryStatRecursively(*SandboxPath, PackageSearch);
@@ -386,13 +411,10 @@ bool FAssetRegistryGenerator::GenerateStreamingInstallManifest(int64 InExtraFlav
 	bool bUseSecondaryOpenOrder = false;
 	TArray<FString> OrderFileSpecStrings;
 	{
-		FConfigFile PlatformIniFile;
-		FConfigCacheIni::LoadLocalIniFile(PlatformIniFile, TEXT("Game"), true, *TargetPlatform->IniPlatformName());
-		FString ConfigString;
-		PlatformIniFile.GetBool(TEXT("/Script/UnrealEd.ProjectPackagingSettings"), TEXT("bEnableAssetRegistryGameOpenOrderSort"), bEnableGameOpenOrderSort);
-		PlatformIniFile.GetBool(TEXT("/Script/UnrealEd.ProjectPackagingSettings"), TEXT("bPakUsesSecondaryOrder"), bUseSecondaryOpenOrder);
-
-		PlatformIniFile.GetArray(TEXT("/Script/UnrealEd.ProjectPackagingSettings"), TEXT("PakOrderFileSpecs"), OrderFileSpecStrings);
+		ConditionalLoadPlatformIniFile();
+		PlatformIniFile->GetBool(TEXT("/Script/UnrealEd.ProjectPackagingSettings"), TEXT("bEnableAssetRegistryGameOpenOrderSort"), bEnableGameOpenOrderSort);
+		PlatformIniFile->GetBool(TEXT("/Script/UnrealEd.ProjectPackagingSettings"), TEXT("bPakUsesSecondaryOrder"), bUseSecondaryOpenOrder);
+		PlatformIniFile->GetArray(TEXT("/Script/UnrealEd.ProjectPackagingSettings"), TEXT("PakOrderFileSpecs"), OrderFileSpecStrings);
 	}
 
 	// if a game open order can be found then use that to sort the filenames
@@ -415,7 +437,6 @@ bool FAssetRegistryGenerator::GenerateStreamingInstallManifest(int64 InExtraFlav
 		else
 		{
 			UScriptStruct* Struct = FPakOrderFileSpec::StaticStruct();
-			int32 Index = 0;
 			for (const FString& String : OrderFileSpecStrings)
 			{
 				FPakOrderFileSpec Spec;
@@ -517,7 +538,7 @@ bool FAssetRegistryGenerator::GenerateStreamingInstallManifest(int64 InExtraFlav
 	bool bSucceeded = true;
 	for (int32 PakchunkIndex = 0; PakchunkIndex < FinalChunkManifests.Num() && bSucceeded; ++PakchunkIndex)
 	{
-		const FChunkPackageSet* Manifest = FinalChunkManifests[PakchunkIndex];
+		const FChunkPackageSet* Manifest = FinalChunkManifests[PakchunkIndex].Get();
 
 		// Serialize chunk layers whether chunk is empty or not
 		int32 TargetLayer = 0;
@@ -530,8 +551,8 @@ bool FAssetRegistryGenerator::GenerateStreamingInstallManifest(int64 InExtraFlav
 		// fill it in with an empty FChunkPackageSet
 		if (!Manifest)
 		{
-			FinalChunkManifests[PakchunkIndex] = new FChunkPackageSet;
-			Manifest = FinalChunkManifests[PakchunkIndex];
+			FinalChunkManifests[PakchunkIndex].Reset(new FChunkPackageSet());
+			Manifest = FinalChunkManifests[PakchunkIndex].Get();
 		}
 
 		int32 FilenameIndex = 0;
@@ -579,7 +600,9 @@ bool FAssetRegistryGenerator::GenerateStreamingInstallManifest(int64 InExtraFlav
 						FName RegistryName = UAssetManager::Get().GetUniqueAssetRegistryName(PakchunkIndex);
 						if (RegistryName != NAME_None)
 						{
-							FString AssetRegistryFilename = FString::Printf(TEXT("%s%sAssetRegistry%s.bin"), *InSandboxFile->GetSandboxDirectory(), *InSandboxFile->GetGameSandboxDirectoryName(), *RegistryName.ToString());
+							FString AssetRegistryFilename = FString::Printf(TEXT("%s%sAssetRegistry%s.bin"),
+								*InSandboxFile.GetSandboxDirectory(), *InSandboxFile.GetGameSandboxDirectoryName(),
+								*RegistryName.ToString());
 							ChunkFilenames.Add(AssetRegistryFilename);
 						}
 					}
@@ -598,7 +621,7 @@ bool FAssetRegistryGenerator::GenerateStreamingInstallManifest(int64 InExtraFlav
 
 				for (const TSharedRef<IChunkDataGenerator>& ChunkDataGenerator : ChunkDataGenerators)
 				{
-					ChunkDataGenerator->GenerateChunkDataFiles(PakchunkIndex, PackagesInChunk, TargetPlatform, InSandboxFile, ChunkFilenames);
+					ChunkDataGenerator->GenerateChunkDataFiles(PakchunkIndex, PackagesInChunk, TargetPlatform, &InSandboxFile, ChunkFilenames);
 				}
 			}
 
@@ -630,7 +653,7 @@ bool FAssetRegistryGenerator::GenerateStreamingInstallManifest(int64 InExtraFlav
 							RelativeFilename = FPaths::SetExtension(RelativeFilename, TEXT("uasset")); // only use the uassets to decide which pak file these guys should live in
 						}
 						RelativeFilename.ToLowerInline();
-						uint64 FileOpenOrder = OrderMap.GetFileOrder(RelativeFilename, true);
+						uint64 FileOpenOrder = OrderMap.GetFileOrder(RelativeFilename, true /* bAllowUexpUBulkFallback */);
 						/*if (FileOpenOrder != MAX_uint64)
 						{
 							UE_LOG(LogAssetRegistryGenerator, Display, TEXT("Found file open order for %s, %ll"), *RelativeFilename, FileOpenOrder);
@@ -660,7 +683,7 @@ bool FAssetRegistryGenerator::GenerateStreamingInstallManifest(int64 InExtraFlav
 					ChunkFilenames.Empty(SortedFiles.Num());
 					for (int I = 0; I < SortedFiles.Num(); ++I)
 					{
-						ChunkFilenames.Add(SortedFiles[I].Filename);
+						ChunkFilenames.Add(MoveTemp(SortedFiles[I].Filename));
 					}
 				}
 				else
@@ -716,7 +739,7 @@ bool FAssetRegistryGenerator::GenerateStreamingInstallManifest(int64 InExtraFlav
 	{
 		FString ChunkManifestDirectory = FPaths::ProjectDir() / TEXT("Metadata") / TEXT("ChunkManifest");
 		ChunkManifestDirectory =
-			InSandboxFile->ConvertToAbsolutePathForExternalAppForWrite(*ChunkManifestDirectory)
+			InSandboxFile.ConvertToAbsolutePathForExternalAppForWrite(*ChunkManifestDirectory)
 				.Replace(TEXT("[Platform]"), *Platform);
 
 		IPlatformFile& PlatformFile = FPlatformFileManager::Get().GetPlatformFile();
@@ -730,7 +753,8 @@ bool FAssetRegistryGenerator::GenerateStreamingInstallManifest(int64 InExtraFlav
 	return bSucceeded;
 }
 
-void FAssetRegistryGenerator::GenerateChunkManifestForPackage(const FName& PackageFName, const FString& PackagePathName, const FString& SandboxFilename, const FString& LastLoadedMapName, FSandboxPlatformFile* InSandboxFile)
+void FAssetRegistryGenerator::GenerateChunkManifestForPackage(const FName& PackageFName, const FString& PackagePathName,
+	const FString& SandboxFilename, const FString& LastLoadedMapName, FSandboxPlatformFile& InSandboxFile)
 {
 	TArray<int32> TargetChunks;
 	TArray<int32> ExistingChunkIDs;
@@ -740,8 +764,7 @@ void FAssetRegistryGenerator::GenerateChunkManifestForPackage(const FName& Packa
 		TargetChunks.AddUnique(0);
 		ExistingChunkIDs.AddUnique(0);
 	}
-
-	if (bGenerateChunks)
+	else
 	{
 		// Collect all chunk IDs associated with this package from the asset registry
 		TArray<int32> RegistryChunkIDs = GetAssetRegistryChunkAssignments(PackageFName);
@@ -775,22 +798,20 @@ void FAssetRegistryGenerator::GenerateChunkManifestForPackage(const FName& Packa
 	}
 
 	// if the delegate requested a specific chunk assignment, add them package to it now.
-	for (const auto& PackageChunk : TargetChunks)
+	for (const int32 PackageChunk : TargetChunks)
 	{
 		AddPackageToManifest(SandboxFilename, PackageFName, PackageChunk);
 	}
 	// If the delegate requested to remove the package from any chunk, remove it now
-	for (const auto& PackageChunk : ExistingChunkIDs)
+	for (const int32 PackageChunk : ExistingChunkIDs)
 	{
 		if (!TargetChunks.Contains(PackageChunk))
 		{
 			RemovePackageFromManifest(PackageFName, PackageChunk);
 		}
 	}
-
 }
 
-
 void FAssetRegistryGenerator::CleanManifestDirectories()
 {
 	CleanTempPackagingDirectory(TargetPlatform->PlatformName());
@@ -886,53 +907,55 @@ void FAssetRegistryGenerator::InjectEncryptionData(FAssetRegistryState& TargetSt
 	}
 }
 
-bool FAssetRegistryGenerator::SaveManifests(FSandboxPlatformFile* InSandboxFile, int64 InExtraFlavorChunkSize)
+bool FAssetRegistryGenerator::SaveManifests(FSandboxPlatformFile& InSandboxFile, int64 InExtraFlavorChunkSize)
 {
 	// Always do package dependency work, is required to modify asset registry
 	FixupPackageDependenciesForChunks(InSandboxFile);
 
-	if (bGenerateChunks)
-	{	
-		if (!GenerateStreamingInstallManifest(InExtraFlavorChunkSize, InSandboxFile))
-		{
-			return false;
-		}
+	if (!bGenerateChunks)
+	{
+		return true;
+	}
 
-		// Generate map for the platform abstraction
-		TMultiMap<FString, int32> PakchunkMap;	// asset -> ChunkIDs map
-		TSet<int32> PakchunkIndicesInUse;
-		const FString PlatformName = TargetPlatform->PlatformName();
+	if (!GenerateStreamingInstallManifest(InExtraFlavorChunkSize, InSandboxFile))
+	{
+		return false;
+	}
 
-		// Collect all unique chunk indices and map all files to their chunks
-		for (int32 PakchunkIndex = 0; PakchunkIndex < FinalChunkManifests.Num(); ++PakchunkIndex)
+	// Generate map for the platform abstraction
+	TMultiMap<FString, int32> PakchunkMap;	// asset -> ChunkIDs map
+	TSet<int32> PakchunkIndicesInUse;
+	const FString PlatformName = TargetPlatform->PlatformName();
+
+	// Collect all unique chunk indices and map all files to their chunks
+	for (int32 PakchunkIndex = 0; PakchunkIndex < FinalChunkManifests.Num(); ++PakchunkIndex)
+	{
+		check(FinalChunkManifests[PakchunkIndex]);
+		if (FinalChunkManifests[PakchunkIndex]->Num())
 		{
-			check(FinalChunkManifests[PakchunkIndex]);
-			if (FinalChunkManifests[PakchunkIndex]->Num())
+			PakchunkIndicesInUse.Add(PakchunkIndex);
+			for (const TPair<FName,FString>& Pair: *FinalChunkManifests[PakchunkIndex])
 			{
-				PakchunkIndicesInUse.Add(PakchunkIndex);
-				for (auto& Filename : *FinalChunkManifests[PakchunkIndex])
-				{
-					FString PlatFilename = Filename.Value.Replace(TEXT("[Platform]"), *PlatformName);
-					PakchunkMap.Add(PlatFilename, PakchunkIndex);
-				}
+				FString PlatFilename = Pair.Value.Replace(TEXT("[Platform]"), *PlatformName);
+				PakchunkMap.Add(PlatFilename, PakchunkIndex);
 			}
 		}
+	}
 
-		// Sort our chunk IDs and file paths
-		PakchunkMap.KeySort(TLess<FString>());
-		PakchunkIndicesInUse.Sort(TLess<int32>());
+	// Sort our chunk IDs and file paths
+	PakchunkMap.KeySort(TLess<FString>());
+	PakchunkIndicesInUse.Sort(TLess<int32>());
 
-		// Platform abstraction will generate any required platform-specific files for the chunks
-		if (!TargetPlatform->GenerateStreamingInstallManifest(PakchunkMap, PakchunkIndicesInUse))
-		{
-			return false;
-		}
+	// Platform abstraction will generate any required platform-specific files for the chunks
+	if (!TargetPlatform->GenerateStreamingInstallManifest(PakchunkMap, PakchunkIndicesInUse))
+	{
+		return false;
+	}
 
-		if (!bUseAssetManager)
-		{
-			// In new flow, this is written later
-			GenerateAssetChunkInformationCSV(FPaths::Combine(*FPaths::ProjectLogDir(), TEXT("ChunkLists")), true);
-		}
+	if (!bUseAssetManager)
+	{
+		// In new flow, this is written later
+		GenerateAssetChunkInformationCSV(FPaths::Combine(*FPaths::ProjectLogDir(), TEXT("ChunkLists")), true);
 	}
 
 	return true;
@@ -1157,25 +1180,24 @@ void FAssetRegistryGenerator::ComputePackageRemovals(const FAssetRegistryState&
 	}
 }
 
-void FAssetRegistryGenerator::BuildChunkManifest(const TSet<FName>& InCookedPackages, const TSet<FName>& InDevelopmentOnlyPackages, FSandboxPlatformFile* InSandboxFile, bool bGenerateStreamingInstallManifest)
+void FAssetRegistryGenerator::BuildChunkManifest(const TSet<FName>& InCookedPackages,
+	const TSet<FName>& InDevelopmentOnlyPackages, FSandboxPlatformFile& InSandboxFile,
+	bool bGenerateStreamingInstallManifest)
 {
-	// If we were asked to generate a streaming install manifest explicitly and we did not have bGenerateNoChunks set, we will generate chunks.
-	// Otherwise, we will defer to the config settings for the platform.
+	// bGenerateNoChunks overrides bGenerateStreamingInstallManifest overrides ShouldPlatformGenerateStreamingInstallManifest
+	// bGenerateChunks means we allow chunks other than 0 based on package ChunkIds, AND we generate a manifest for each chunk
 	const UProjectPackagingSettings* PackagingSettings = Cast<UProjectPackagingSettings>(UProjectPackagingSettings::StaticClass()->GetDefaultObject());
 	if (PackagingSettings->bGenerateNoChunks)
 	{
 		bGenerateChunks = false;
 	}
+	else if (bGenerateStreamingInstallManifest)
+	{
+		bGenerateChunks = true;
+	}
 	else
 	{
-		if (bGenerateStreamingInstallManifest)
-		{
-			bGenerateChunks = true;
-		}
-		else
-		{
-			bGenerateChunks = ShouldPlatformGenerateStreamingInstallManifest(TargetPlatform);
-		}
+		bGenerateChunks = ShouldPlatformGenerateStreamingInstallManifest(TargetPlatform);
 	}
 
 	CookedPackages = InCookedPackages;
@@ -1198,31 +1220,23 @@ void FAssetRegistryGenerator::BuildChunkManifest(const TSet<FName>& InCookedPack
 	}
 
 	// initialize FoundIDList, PackageChunkIDMap
-	const TMap<FName, const FAssetData*>& ObjectToDataMap = State.GetObjectPathToAssetDataMap();
-
-	for (const TPair<FName, const FAssetData*>& Pair : ObjectToDataMap)
+	for (const TPair<FName, const FAssetData*>& Pair : State.GetObjectPathToAssetDataMap())
 	{
-		// Chunk Ids are safe to modify in place so do a const cast
-		FAssetData& AssetData = *const_cast<FAssetData*>(Pair.Value);
-		for (auto ChunkIt = AssetData.ChunkIDs.CreateConstIterator(); ChunkIt; ++ChunkIt)
+		const FAssetData& AssetData = *Pair.Value;
+		for (int32 ChunkID : AssetData.ChunkIDs)
 		{
-			int32 ChunkID = *ChunkIt;
 			if (ChunkID < 0)
 			{
 				UE_LOG(LogAssetRegistryGenerator, Warning, TEXT("Out of range ChunkID: %d"), ChunkID);
 				ChunkID = 0;
 			}
 			
-			auto* FoundIDList = PackageChunkIDMap.Find(AssetData.PackageName);
-			if (!FoundIDList)
-			{
-				FoundIDList = &PackageChunkIDMap.Add(AssetData.PackageName);
-			}
-			FoundIDList->AddUnique(ChunkID);
+			PackageChunkIDMap.FindOrAdd(AssetData.PackageName).AddUnique(ChunkID);
 		}
 
-			// Now clear the original chunk id list. We will fill it with real IDs when cooking.
-			AssetData.ChunkIDs.Empty();
+		// Now clear the original chunk id list. We will fill it with real IDs when cooking.
+		// Chunk Ids are safe to modify in place so do a const cast
+		const_cast<FAssetData&>(AssetData).ChunkIDs.Empty();
 
 		// Update whether the owner package contains a map
 		if (AssetData.GetClass()->IsChildOf(UWorld::StaticClass()) || AssetData.GetClass()->IsChildOf(ULevel::StaticClass()))
@@ -1234,7 +1248,7 @@ void FAssetRegistryGenerator::BuildChunkManifest(const TSet<FName>& InCookedPack
 	// add all the packages to the unassigned package list
 	for (FName CookedPackage : CookedPackages)
 	{
-		const FString SandboxPath = InSandboxFile->ConvertToAbsolutePathForExternalAppForWrite(*FPackageName::LongPackageNameToFilename(CookedPackage.ToString()));
+		const FString SandboxPath = InSandboxFile.ConvertToAbsolutePathForExternalAppForWrite(*FPackageName::LongPackageNameToFilename(CookedPackage.ToString()));
 
 		AllCookedPackageSet.Add(CookedPackage, SandboxPath);
 		UnassignedPackageSet.Add(CookedPackage, SandboxPath);
@@ -1249,7 +1263,7 @@ void FAssetRegistryGenerator::BuildChunkManifest(const TSet<FName>& InCookedPack
 		FString StartupPackageMapName(TEXT("None"));
 		for (FName CookedPackage : StartupPackages)
 		{
-			const FString SandboxPath = InSandboxFile->ConvertToAbsolutePathForExternalAppForWrite(*FPackageName::LongPackageNameToFilename(CookedPackage.ToString()));
+			const FString SandboxPath = InSandboxFile.ConvertToAbsolutePathForExternalAppForWrite(*FPackageName::LongPackageNameToFilename(CookedPackage.ToString()));
 			const FString PackagePathName = CookedPackage.ToString();
 			AllCookedPackageSet.Add(CookedPackage, SandboxPath);
 			GenerateChunkManifestForPackage(CookedPackage, PackagePathName, SandboxPath, StartupPackageMapName, InSandboxFile);
@@ -1284,7 +1298,7 @@ void FAssetRegistryGenerator::BuildChunkManifest(const TSet<FName>& InCookedPack
 				const FString* SandboxFilenamePtr = AllCookedPackageSet.Find(PackageFName);
 				if (!SandboxFilenamePtr)
 				{
-					const FString SandboxPath = InSandboxFile->ConvertToAbsolutePathForExternalAppForWrite(*FPackageName::LongPackageNameToFilename(PackagePathName));
+					const FString SandboxPath = InSandboxFile.ConvertToAbsolutePathForExternalAppForWrite(*FPackageName::LongPackageNameToFilename(PackagePathName));
 
 					AllCookedPackageSet.Add(PackageFName, SandboxPath);
 
@@ -1311,7 +1325,6 @@ void FAssetRegistryGenerator::BuildChunkManifest(const TSet<FName>& InCookedPack
 	}
 
 	// anything that remains in the UnAssignedPackageSet will be put in chunk0 when we save the asset registry
-
 }
 
 void FAssetRegistryGenerator::RegisterChunkDataGenerator(TSharedRef<IChunkDataGenerator> InChunkDataGenerator)
@@ -1362,7 +1375,7 @@ void FAssetRegistryGenerator::AddAssetToFileOrderRecursive(const FName& InPackag
 bool FAssetRegistryGenerator::SaveAssetRegistry(const FString& SandboxPath, bool bSerializeDevelopmentAssetRegistry, bool bForceNoFilter)
 {
 	UE_LOG(LogAssetRegistryGenerator, Display, TEXT("Saving asset registry v%d."), FAssetRegistryVersion::Type::LatestVersion);
-	double StartTime = FPlatformTime::Seconds();
+	FAutoScopedDurationTimer Timer;
 
 	const TMap<FName, const FAssetData*>& ObjectToDataMap = State.GetObjectPathToAssetDataMap();
 	
@@ -1383,7 +1396,7 @@ bool FAssetRegistryGenerator::SaveAssetRegistry(const FString& SandboxPath, bool
 	// First flush the asset registry and make sure the asset data is in sync, as it may have been updated during cook
 	AssetRegistry.Tick(-1.0f);
 	AssetRegistry.InitializeTemporaryAssetRegistryState(State, SaveOptions, true);
-	// Then possibly apply previous AssetData and AssetPackageData data for packages kept from a previous cook
+	// Then possibly apply previous AssetData and AssetPackageData for packages kept from a previous cook
 	UpdateKeptPackages();
 	UpdateCollectionAssetData();
 
@@ -1426,7 +1439,7 @@ bool FAssetRegistryGenerator::SaveAssetRegistry(const FString& SandboxPath, bool
 			// which will be written to the master asset registry in chunk 0
 			for (int32 PakchunkIndex = 0; PakchunkIndex < FinalChunkManifests.Num(); ++PakchunkIndex)
 			{
-				FChunkPackageSet* Manifest = FinalChunkManifests[PakchunkIndex];
+				FChunkPackageSet* Manifest = FinalChunkManifests[PakchunkIndex].Get();
 				if (Manifest == nullptr)
 				{
 					continue;
@@ -1482,8 +1495,7 @@ bool FAssetRegistryGenerator::SaveAssetRegistry(const FString& SandboxPath, bool
 				if (ChunkBucketElement.Key != GenericChunkBucket)
 				{
 					check(ChunkBucketElement.Key < FinalChunkManifests.Num());
-					FChunkPackageSet* ChunkPackageSet = FinalChunkManifests[ChunkBucketElement.Key];
-					check(ChunkPackageSet);
+					check(FinalChunkManifests[ChunkBucketElement.Key]);
 					FilenameForLog = FString::Printf(TEXT("[chunkbucket %i] "), ChunkBucketElement.Key);
 				}
 				UE_LOG(LogAssetRegistryGenerator, Display, TEXT("Generated asset registry %snum assets %d, size is %5.2fkb"), *FilenameForLog, NewState.GetNumAssets(), (float)SerializedAssetRegistry.Num() / 1024.f);
@@ -1508,21 +1520,20 @@ bool FAssetRegistryGenerator::SaveAssetRegistry(const FString& SandboxPath, bool
 		}
 	}
 	
-	double StopTime = FPlatformTime::Seconds();
-	UE_LOG(LogAssetRegistryGenerator, Display, TEXT("Done saving asset registry. Took %.1f seconds."), StopTime-StartTime);
+	UE_LOG(LogAssetRegistryGenerator, Display, TEXT("Done saving asset registry. Took %.1f seconds."), Timer.GetTime());
 
 	return true;
 }
 
 class FPackageCookerOpenOrderVisitor : public IPlatformFile::FDirectoryVisitor
 {
-	const FSandboxPlatformFile* SandboxFile;
+	const FSandboxPlatformFile& SandboxFile;
 	const FString& PlatformSandboxPath;
 	const TSet<FString>& ValidExtensions;
 	TMultiMap<FString, FString>& PackageExtensions;
 public:
 	FPackageCookerOpenOrderVisitor(
-		const FSandboxPlatformFile* InSandboxFile,
+		const FSandboxPlatformFile& InSandboxFile,
 		const FString& InPlatformSandboxPath,
 		const TSet<FString>& InValidExtensions,
 		TMultiMap<FString, FString>& OutPackageExtensions) :
@@ -1543,8 +1554,8 @@ public:
 		if (ValidExtensions.Contains(FileExtension))
 		{
 			FString PackageName;
-			Filename.ReplaceInline(*PlatformSandboxPath, *SandboxFile->GetSandboxDirectory());
-			FString AssetSourcePath = SandboxFile->ConvertFromSandboxPath(*Filename);
+			Filename.ReplaceInline(*PlatformSandboxPath, *SandboxFile.GetSandboxDirectory());
+			FString AssetSourcePath = SandboxFile.ConvertFromSandboxPath(*Filename);
 			FString StandardAssetSourcePath = FPaths::CreateStandardFilename(AssetSourcePath);
 			if (StandardAssetSourcePath.EndsWith(TEXT(".m.ubulk")))
 			{
@@ -1565,7 +1576,7 @@ public:
 	}
 };
 
-bool FAssetRegistryGenerator::WriteCookerOpenOrder(FSandboxPlatformFile* InSandboxFile)
+bool FAssetRegistryGenerator::WriteCookerOpenOrder(FSandboxPlatformFile& InSandboxFile)
 {
 	TSet<FName> PackageNameSet;
 	TSet<FName> MapList;
@@ -1645,10 +1656,11 @@ bool FAssetRegistryGenerator::WriteCookerOpenOrder(FSandboxPlatformFile* InSandb
 		};
 		const TSet<FString> ValidExtensionSet(ValidExtensions);
 
-		const FString SandboxPath = InSandboxFile->GetSandboxDirectory();
+		const FString SandboxPath = InSandboxFile.GetSandboxDirectory();
 		const FString Platform = TargetPlatform->PlatformName();
 		FString PlatformSandboxPath = SandboxPath.Replace(TEXT("[Platform]"), *Platform);
-
+		
+		// ZENTODO: Change this to work with Zen
 		TMultiMap<FString, FString> CookedPackageFilesMap;
 		FPackageCookerOpenOrderVisitor PackageSearch(InSandboxFile, PlatformSandboxPath, ValidExtensionSet, CookedPackageFilesMap);
 		IFileManager::Get().IterateDirectoryRecursively(*PlatformSandboxPath, PackageSearch);
@@ -1665,7 +1677,8 @@ bool FAssetRegistryGenerator::WriteCookerOpenOrder(FSandboxPlatformFile* InSandb
 
 			for (const FString& CookedFile : CookedFiles)
 			{
-				FString Line = FString::Printf(TEXT("\"%s\" %i\n"), *CookedFile, CurrentIndex++);
+				TStringBuilder<256> Line;
+				Line.Appendf(TEXT("\"%s\" %i\n"), *CookedFile, CurrentIndex++);
 				CookerFileOrderString.Append(Line);
 			}
 		}
@@ -1673,14 +1686,14 @@ bool FAssetRegistryGenerator::WriteCookerOpenOrder(FSandboxPlatformFile* InSandb
 
 	if (CookerFileOrderString.Len())
 	{
-		FString OpenOrderFilename;
+		TStringBuilder<256> OpenOrderFilename;
 		if (FDataDrivenPlatformInfoRegistry::GetPlatformInfo(TargetPlatform->PlatformName()).bIsConfidential)
 		{
-			OpenOrderFilename = FString::Printf(TEXT("%sPlatforms/%s/Build/FileOpenOrder/CookerOpenOrder.log"), *FPaths::ProjectDir(), *TargetPlatform->PlatformName());
+			OpenOrderFilename.Appendf(TEXT("%sPlatforms/%s/Build/FileOpenOrder/CookerOpenOrder.log"), *FPaths::ProjectDir(), *TargetPlatform->PlatformName());
 		}
 		else
 		{
-			OpenOrderFilename = FString::Printf(TEXT("%sBuild/%s/FileOpenOrder/CookerOpenOrder.log"), *FPaths::ProjectDir(), *TargetPlatform->PlatformName());
+			OpenOrderFilename.Appendf(TEXT("%sBuild/%s/FileOpenOrder/CookerOpenOrder.log"), *FPaths::ProjectDir(), *TargetPlatform->PlatformName());
 		}
 		FFileHelper::SaveStringToFile(CookerFileOrderString, *OpenOrderFilename);
 	}
@@ -1849,7 +1862,7 @@ bool FAssetRegistryGenerator::GenerateAssetChunkInformationCSV(const FString& Ou
 
 				// Build "other chunks" string or None if not part of
 				TmpStringChunks.Empty(64);
-				for (const auto& OtherChunk : AssetData.ChunkIDs)
+				for (int32 OtherChunk : AssetData.ChunkIDs)
 				{
 					if (OtherChunk != PakchunkIndex)
 					{
@@ -1889,13 +1902,13 @@ void FAssetRegistryGenerator::AddPackageToManifest(const FString& PackageSandbox
 	HighestChunkId = ChunkId > HighestChunkId ? ChunkId : HighestChunkId;
 	int32 PakchunkIndex = GetPakchunkIndex(ChunkId);
 
-	while (PakchunkIndex >= ChunkManifests.Num())
+	if (PakchunkIndex >= ChunkManifests.Num())
 	{
-		ChunkManifests.Add(nullptr);
+		ChunkManifests.AddDefaulted(PakchunkIndex - ChunkManifests.Num() + 1);
 	}
 	if (!ChunkManifests[PakchunkIndex])
 	{
-		ChunkManifests[PakchunkIndex] = new FChunkPackageSet();
+		ChunkManifests[PakchunkIndex].Reset(new FChunkPackageSet());
 	}
 	ChunkManifests[PakchunkIndex]->Add(PackageName, PackageSandboxPath);
 	//Safety check, it the package happens to exist in the unassigned list remove it now.
@@ -1913,46 +1926,54 @@ void FAssetRegistryGenerator::RemovePackageFromManifest(FName PackageName, int32
 	}
 }
 
-void FAssetRegistryGenerator::ResolveChunkDependencyGraph(const FChunkDependencyTreeNode& Node, const TSet<FName>& BaseAssetSet, TArray<TArray<FName>>& OutPackagesMovedBetweenChunks)
+void FAssetRegistryGenerator::ResolveChunkDependencyGraph(const FChunkDependencyTreeNode& Node,
+	const TSet<FName>& BaseAssetSet, TArray<TArray<FName>>& OutPackagesMovedBetweenChunks)
 {
-	if (FinalChunkManifests.Num() > Node.ChunkID && FinalChunkManifests[Node.ChunkID])
+	if (FinalChunkManifests.Num() <= Node.ChunkID || !FinalChunkManifests[Node.ChunkID])
 	{
-		for (auto It = BaseAssetSet.CreateConstIterator(); It; ++It)
+		return;
+	}
+	FChunkPackageSet& NodeManifest = *FinalChunkManifests[Node.ChunkID];
+	for (FName PackageName : BaseAssetSet)
+	{
+		// Remove any assets belonging to our parents.
+		if (NodeManifest.Remove(PackageName) > 0)
 		{
-			// Remove any assets belonging to our parents.			
-			if (FinalChunkManifests[Node.ChunkID]->Remove(*It) > 0)
-			{
-				OutPackagesMovedBetweenChunks[Node.ChunkID].Add(*It);
-				UE_LOG(LogAssetRegistryGenerator, Verbose, TEXT("Removed %s from chunk %i because it is duplicated in another chunk."), *It->ToString(), Node.ChunkID);
-			}
+			OutPackagesMovedBetweenChunks[Node.ChunkID].Add(PackageName);
+			UE_LOG(LogAssetRegistryGenerator, Verbose, TEXT("Removed %s from chunk %i because it is duplicated in another chunk."),
+				*PackageName.ToString(), Node.ChunkID);
 		}
+	}
 		
-		TSet<FName> ModifiedAssetSet;
-		
-		// Add the current Chunk's assets
-		for (auto It = FinalChunkManifests[Node.ChunkID]->CreateConstIterator(); It; ++It)//for (const auto It : *(FinalChunkManifests[Node.ChunkID]))
+	if (!Node.ChildNodes.Num())
+	{
+		return;
+	}
+
+	// Add the current Chunk's assets
+	TSet<FName> CumulativeAssetSet;
+	if (!NodeManifest.IsEmpty())
+	{
+		CumulativeAssetSet.Reserve(BaseAssetSet.Num() + NodeManifest.Num());
+		CumulativeAssetSet.Append(BaseAssetSet);
+		for (const TPair<FName, FString>& Pair : NodeManifest)
 		{
-			if (!ModifiedAssetSet.Num())
-			{
-				ModifiedAssetSet = BaseAssetSet;
-			}
-			
-			ModifiedAssetSet.Add(It.Key());
+			CumulativeAssetSet.Add(Pair.Key);
 		}
+	}
 		
-		const auto& AssetSet = ModifiedAssetSet.Num() ? ModifiedAssetSet : BaseAssetSet;
-		for (const auto& It : Node.ChildNodes)
-		{
-			ResolveChunkDependencyGraph(It, AssetSet, OutPackagesMovedBetweenChunks);
-		}
+	const TSet<FName>& NewRecursiveParentPackages = CumulativeAssetSet.Num() ? CumulativeAssetSet : BaseAssetSet;
+	for (const FChunkDependencyTreeNode& ChildNode : Node.ChildNodes)
+	{
+		ResolveChunkDependencyGraph(ChildNode, NewRecursiveParentPackages, OutPackagesMovedBetweenChunks);
 	}
 }
 
 bool FAssetRegistryGenerator::CheckChunkAssetsAreNotInChild(const FChunkDependencyTreeNode& Node)
 {
-	for (const auto& It : Node.ChildNodes)
+	for (const FChunkDependencyTreeNode& ChildNode : Node.ChildNodes)
 	{
-		if (!CheckChunkAssetsAreNotInChild(It))
+		if (!CheckChunkAssetsAreNotInChild(ChildNode))
 		{
 			return false;
 		}
@@ -1963,13 +1984,15 @@ bool FAssetRegistryGenerator::CheckChunkAssetsAreNotInChild(const FChunkDependen
 		return true;
 	}
 
-	for (const auto& ChildIt : Node.ChildNodes)
+	FChunkPackageSet& ParentManifest = *FinalChunkManifests[Node.ChunkID];
+	for (const FChunkDependencyTreeNode& ChildNode : Node.ChildNodes)
 	{
-		if(FinalChunkManifests.Num() > ChildIt.ChunkID && FinalChunkManifests[ChildIt.ChunkID])
+		if (FinalChunkManifests.Num() > ChildNode.ChunkID && FinalChunkManifests[ChildNode.ChunkID])
 		{
-			for (auto It = FinalChunkManifests[Node.ChunkID]->CreateConstIterator(); It; ++It)
+			FChunkPackageSet& ChildManifest = *FinalChunkManifests[ChildNode.ChunkID];
+			for (const TPair<FName,FString>& ParentPair : ParentManifest)
 			{
-				if (FinalChunkManifests[ChildIt.ChunkID]->Find(It.Key()))
+				if (ChildManifest.Find(ParentPair.Key))
 				{
 					return false;
 				}
@@ -1980,12 +2003,14 @@ bool FAssetRegistryGenerator::CheckChunkAssetsAreNotInChild(const FChunkDependen
 	return true;
 }
 
-void FAssetRegistryGenerator::AddPackageAndDependenciesToChunk(FChunkPackageSet* ThisPackageSet, FName InPkgName, const FString& InSandboxFile, int32 PakchunkIndex, FSandboxPlatformFile* SandboxPlatformFile)
+void FAssetRegistryGenerator::AddPackageAndDependenciesToChunk(FChunkPackageSet& ThisPackageSet, FName InPkgName,
+	const FString& InSandboxFile, int32 PakchunkIndex, FSandboxPlatformFile& SandboxPlatformFile)
 {
-	FChunkPackageSet* InitialPackageSetForThisChunk = ChunkManifests.IsValidIndex(PakchunkIndex) ? ChunkManifests[PakchunkIndex] : nullptr;
+	FChunkPackageSet* InitialPackageSetForThisChunk =
+		ChunkManifests.IsValidIndex(PakchunkIndex) ? ChunkManifests[PakchunkIndex].Get() : nullptr;
 
 	//Add this asset
-	ThisPackageSet->Add(InPkgName, InSandboxFile);
+	ThisPackageSet.Add(InPkgName, InSandboxFile);
 
 	// Only gather dependencies the slow way if we're chunking or not using asset manager
 	if (!bGenerateChunks || bUseAssetManager)
@@ -2012,8 +2037,8 @@ void FAssetRegistryGenerator::AddPackageAndDependenciesToChunk(FChunkPackageSet*
 				{
 					continue;
 				}
-				FString DependentSandboxFile = SandboxPlatformFile->ConvertToAbsolutePathForExternalAppForWrite(*FPackageName::LongPackageNameToFilename(*FilteredPackageName.ToString()));
-				if (!ThisPackageSet->Contains(FilteredPackageName))
+				FString DependentSandboxFile = SandboxPlatformFile.ConvertToAbsolutePathForExternalAppForWrite(*FPackageName::LongPackageNameToFilename(*FilteredPackageName.ToString()));
+				if (!ThisPackageSet.Contains(FilteredPackageName))
 				{
 					if ((InitialPackageSetForThisChunk != nullptr) && InitialPackageSetForThisChunk->Contains(PkgName))
 					{
@@ -2036,36 +2061,32 @@ void FAssetRegistryGenerator::AddPackageAndDependenciesToChunk(FChunkPackageSet*
 						}
 					}
 				}
-				ThisPackageSet->Add(FilteredPackageName, DependentSandboxFile);
+				ThisPackageSet.Add(FilteredPackageName, DependentSandboxFile);
 				UnassignedPackageSet.Remove(PkgName);
 			}
 		}
 	}
 }
 
-void FAssetRegistryGenerator::FixupPackageDependenciesForChunks(FSandboxPlatformFile* InSandboxFile)
+void FAssetRegistryGenerator::FixupPackageDependenciesForChunks(FSandboxPlatformFile& InSandboxFile)
 {
 	UE_LOG(LogAssetRegistryGenerator, Log, TEXT("Starting FixupPackageDependenciesForChunks..."));
 	SCOPE_LOG_TIME_IN_SECONDS(TEXT("... FixupPackageDependenciesForChunks complete."), nullptr);
 
 	// Clear any existing manifests from the final array
-	for (FChunkPackageSet* Manifest : FinalChunkManifests)
-	{
-		delete Manifest;
-	}
-	FinalChunkManifests.Empty();
+	FinalChunkManifests.Empty(ChunkManifests.Num());
 
 	for (int32 PakchunkIndex = 0, MaxPakchunk = ChunkManifests.Num(); PakchunkIndex < MaxPakchunk; ++PakchunkIndex)
 	{
-		FinalChunkManifests.Add(new FChunkPackageSet());
+		FChunkPackageSet& FinalManifest = *FinalChunkManifests.Emplace_GetRef(new FChunkPackageSet());
 		if (!ChunkManifests[PakchunkIndex])
 		{
 			continue;
 		}
 		
-		for (auto It = ChunkManifests[PakchunkIndex]->CreateConstIterator(); It; ++It)
+		for (const TPair<FName,FString>& Pair : *ChunkManifests[PakchunkIndex])
 		{
-			AddPackageAndDependenciesToChunk(FinalChunkManifests[PakchunkIndex], It.Key(), It.Value(), PakchunkIndex, InSandboxFile);
+			AddPackageAndDependenciesToChunk(FinalManifest, Pair.Key, Pair.Value, PakchunkIndex, InSandboxFile);
 		}
 	}
 
@@ -2074,20 +2095,20 @@ void FAssetRegistryGenerator::FixupPackageDependenciesForChunks(FSandboxPlatform
 	bool bSkipResolveChunkDependencyGraph = false;
 	PlatformIniFile.GetBool(TEXT("Script/UnrealEd.ChunkDependencyInfo"), TEXT("bSkipResolveChunkDependencyGraph"), bSkipResolveChunkDependencyGraph);
 
-	const FChunkDependencyTreeNode* ChunkDepGraph = DependencyInfo->GetOrBuildChunkDependencyGraph(!bSkipResolveChunkDependencyGraph ? HighestChunkId : 0);
+	const FChunkDependencyTreeNode* ChunkDepGraph = DependencyInfo.GetOrBuildChunkDependencyGraph(!bSkipResolveChunkDependencyGraph ? HighestChunkId : 0);
 
 	//Once complete, Add any remaining assets (that are not assigned to a chunk) to the first chunk.
 	if (FinalChunkManifests.Num() == 0)
 	{
-		FinalChunkManifests.Add(new FChunkPackageSet());
+		FinalChunkManifests.Emplace(new FChunkPackageSet());
 	}
-	check(FinalChunkManifests[0]);
+	FChunkPackageSet& Chunk0Manifest = *FinalChunkManifests[0];
 	
 	// Copy the remaining assets
-	auto RemainingAssets = UnassignedPackageSet;
-	for (auto It = RemainingAssets.CreateConstIterator(); It; ++It)
+	FChunkPackageSet RemainingAssets = UnassignedPackageSet; // Loop removes elements from UnassignedPackageSet
+	for (const TPair<FName,FString>& Pair : RemainingAssets)
 	{
-		AddPackageAndDependenciesToChunk(FinalChunkManifests[0], It.Key(), It.Value(), 0, InSandboxFile);
+		AddPackageAndDependenciesToChunk(Chunk0Manifest, Pair.Key, Pair.Value, 0, InSandboxFile);
 	}
 
 	if (!CheckChunkAssetsAreNotInChild(*ChunkDepGraph))
@@ -2101,9 +2122,9 @@ void FAssetRegistryGenerator::FixupPackageDependenciesForChunks(FSandboxPlatform
 	//Finally, if the previous step may added any extra packages to the 0 chunk. Pull them out of other chunks and save space
 	ResolveChunkDependencyGraph(*ChunkDepGraph, TSet<FName>(), PackagesRemovedFromChunks);
 
-	for (int32 PakchunkIndex = 0; PakchunkIndex < ChunkManifests.Num(); ++PakchunkIndex)
+	if (!bUseAssetManager)
 	{
-		if (!bUseAssetManager)
+		for (int32 PakchunkIndex = 0; PakchunkIndex < ChunkManifests.Num(); ++PakchunkIndex)
 		{
 			FName CollectionName(*FString::Printf(TEXT("PackagesRemovedFromChunk%i"), PakchunkIndex));
 			if (CreateOrEmptyCollection(CollectionName))
@@ -2116,17 +2137,15 @@ void FAssetRegistryGenerator::FixupPackageDependenciesForChunks(FSandboxPlatform
 	for (int32 PakchunkIndex = 0, MaxPakchunk = ChunkManifests.Num(); PakchunkIndex < MaxPakchunk; ++PakchunkIndex)
 	{
 		const int32 ChunkManifestNum = ChunkManifests[PakchunkIndex] ? ChunkManifests[PakchunkIndex]->Num() : 0;
+		check(PakchunkIndex < FinalChunkManifests.Num() && FinalChunkManifests[PakchunkIndex]);
 		const int32 FinalChunkManifestNum = FinalChunkManifests[PakchunkIndex]->Num();
-		UE_LOG(LogAssetRegistryGenerator, Log, TEXT("Chunk: %i, Started with %i packages, Final after dependency resolve: %i"), PakchunkIndex, ChunkManifestNum, FinalChunkManifestNum);
+		UE_LOG(LogAssetRegistryGenerator, Log, TEXT("Chunk: %i, Started with %i packages, Final after dependency resolve: %i"),
+			PakchunkIndex, ChunkManifestNum, FinalChunkManifestNum);
 	}
 	
 	// Fix up the asset registry to reflect this chunk layout
 	for (int32 PakchunkIndex = 0 ; PakchunkIndex < FinalChunkManifests.Num(); ++PakchunkIndex)
 	{
-		if (PakchunkIndex >= FinalChunkManifests.Num())
-		{
-			continue;
-		}
 		check(FinalChunkManifests[PakchunkIndex]);
 		for (const TPair<FName, FString>& Asset : *FinalChunkManifests[PakchunkIndex])
 		{
@@ -2249,11 +2268,11 @@ void FAssetRegistryGenerator::WriteCollection(FName CollectionName, const TArray
 
 int32 FAssetRegistryGenerator::GetPakchunkIndex(int32 ChunkId) const
 {
-	if (ChunkIdPakchunkIndexMapping.Contains(ChunkId))
+	const int32* PakChunkId = ChunkIdPakchunkIndexMapping.Find(ChunkId);
+	if (PakChunkId)
 	{
-		int32 NewChunkId = ChunkIdPakchunkIndexMapping[ChunkId];
-		check(NewChunkId >= 0);
-		return NewChunkId;
+		check(*PakChunkId >= 0);
+		return *PakChunkId;
 	}
 
 	return ChunkId;
@@ -2273,9 +2292,9 @@ void FAssetRegistryGenerator::GetChunkAssignments(TArray<TSet<FName>>& OutAssign
 
 		if (ChunkManifests[0])
 		{
-			for (FChunkPackageSet::TConstIterator It(*ChunkManifests[0]); It; ++It)
+			for (const TPair<FName,FString>& Pair: *ChunkManifests[0])
 			{
-				PackagesInChunk0.Add(It.Key());
+				PackagesInChunk0.Add(Pair.Key);
 			}
 		}
 		OutAssignments.Add(PackagesInChunk0);
@@ -2285,9 +2304,9 @@ void FAssetRegistryGenerator::GetChunkAssignments(TArray<TSet<FName>>& OutAssign
 			TSet<FName> PackagesInChunk;
 			if (ChunkManifests[ChunkIndex])
 			{
-				for (FChunkPackageSet::TConstIterator It(*ChunkManifests[ChunkIndex]); It; ++It)
+				for (const TPair<FName,FString>& Pair : *ChunkManifests[ChunkIndex])
 				{
-					PackagesInChunk.Add(It.Key());
+					PackagesInChunk.Add(Pair.Key);
 				}
 			}
 
@@ -2328,13 +2347,17 @@ void FAssetRegistryGenerator::UpdateAssetRegistryPackageData(const UPackage& Pac
 	const FName PackageName = Package.GetFName();
 	PreviousPackagesToUpdate.Remove(PackageName);
 
+	uint32 NewPackageFlags = 0;
+	FAssetPackageData* AssetPackageData = GetAssetPackageData(PackageName);
 	bool bSaveSucceeded = SavePackageResult.IsSuccessful();
 	if (bSaveSucceeded)
 	{
 		// Ensure all assets in the package are recorded in the registry
 		CreateOrFindAssetDatas(Package);
 
-		FAssetPackageData* AssetPackageData = GetAssetPackageData(PackageName);
+		// Set the PackageFlags to the recorded value from SavePackage
+		NewPackageFlags = SavePackageResult.SerializedPackageFlags;
+
 		AssetPackageData->DiskSize = SavePackageResult.TotalFileSize;
 		// If there is no hash (e.g.: when SavePackageResult == ESavePackageResult::ReplaceCompletely), don't attempt to setup a continuation to update
 		// the AssetRegistry entry with it later.  Just leave the asset registry entry with a default constructed FMD5Hash which is marked as invalid.
@@ -2344,16 +2367,22 @@ void FAssetRegistryGenerator::UpdateAssetRegistryPackageData(const UPackage& Pac
 				{
 					// Store the cooked hash in the Asset Registry when it is done computing in another thread.
 					// NOTE: For this to work, we rely on:
-					// 1) UPackage::WaitForAsyncFileWrites to have been called before any use of the CookedHash - it is called in CookByTheBookFinished before the registry does any work with the registries
-					// 2) AssetPackageData must continue to be a valid pointer - CreateOrFindAssetDatas allocates the FAssetPackageData individually and doesn't relocate or delete them until pruning, which happens after WaitForAsyncFileWrites
+					// 1) UPackage::WaitForAsyncFileWrites needs to be called before any use of the CookedHash.
+					//    It is currently called in CookByTheBookFinished before BuildChunkManifest.
+					// 2) AssetPackageData must continue to be a valid pointer until WaitForAsyncFileWrites is called.
+					//    This is currently true because CreateOrFindAssetDatas allocates FAssetPackageData individually
+					//    and doesn't modify them until pruning, which happens during BuildChunkManifest.
 					AssetPackageData->CookedHash = CookedHash;
 				});
 		}
 	}
+	else
+	{
+		// Set the package flags to zero to indicate that the package failed to save
+		NewPackageFlags = 0;
+	}
 
-	// Set the package flags to zero to indicate that the package failed to save
-	const uint32 PackageFlags = bSaveSucceeded ? SavePackageResult.SerializedPackageFlags : 0;
-	const bool bUpdated = UpdateAssetPackageFlags(PackageName, PackageFlags);
+	const bool bUpdated = UpdateAssetPackageFlags(PackageName, NewPackageFlags);
 	UE_CLOG(!bUpdated && bSaveSucceeded, LogAssetRegistryGenerator, Warning,
 		TEXT("Trying to update asset package flags in package '%s' that does not exist"), *PackageName.ToString());
 }
diff --git a/Engine/Source/Editor/UnrealEd/Private/CookOnTheFlyServer.cpp b/Engine/Source/Editor/UnrealEd/Private/CookOnTheFlyServer.cpp
--- a/Engine/Source/Editor/UnrealEd/Private/CookOnTheFlyServer.cpp
+++ b/Engine/Source/Editor/UnrealEd/Private/CookOnTheFlyServer.cpp
@@ -6888,12 +6888,12 @@ void UCookOnTheFlyServer::CookByTheBookFinished()
 				}
 				{
 					UE_SCOPED_HIERARCHICAL_COOKTIMER(BuildChunkManifest);
-					Generator.BuildChunkManifest(CookedPackageNames, IgnorePackageNames, SandboxFile.Get(), CookByTheBookOptions->bGenerateStreamingInstallManifests);
+					Generator.BuildChunkManifest(CookedPackageNames, IgnorePackageNames, *SandboxFile, CookByTheBookOptions->bGenerateStreamingInstallManifests);
 				}
 				{
 					UE_SCOPED_HIERARCHICAL_COOKTIMER(SaveManifests);
 					// Always try to save the manifests, this is required to make the asset registry work, but doesn't necessarily write a file
-					if (!Generator.SaveManifests(SandboxFile.Get()))
+					if (!Generator.SaveManifests(*SandboxFile))
 					{
 						UE_LOG(LogCook, Warning, TEXT("Failed to save chunk manifest"));
 					}
@@ -6903,7 +6903,7 @@ void UCookOnTheFlyServer::CookByTheBookFinished()
 					{
 						if (ExtraFlavorChunkSize > 0)
 						{
-							if (!Generator.SaveManifests(SandboxFile.Get(), ExtraFlavorChunkSize))
+							if (!Generator.SaveManifests(*SandboxFile, ExtraFlavorChunkSize))
 							{
 								UE_LOG(LogCook, Warning, TEXT("Failed to save chunk manifest"));
 							}
@@ -6921,7 +6921,7 @@ void UCookOnTheFlyServer::CookByTheBookFinished()
 					UE_SCOPED_HIERARCHICAL_COOKTIMER(WriteCookerOpenOrder);
 					if (!IsCookFlagSet(ECookInitializationFlags::Iterative))
 					{
-						Generator.WriteCookerOpenOrder(SandboxFile.Get());
+						Generator.WriteCookerOpenOrder(*SandboxFile);
 					}
 				}
 				// now that we have the asset registry and cooking open order, we have enough information to split the shader library
diff --git a/Engine/Source/Editor/UnrealEd/Public/Commandlets/AssetRegistryGenerator.h b/Engine/Source/Editor/UnrealEd/Public/Commandlets/AssetRegistryGenerator.h
--- a/Engine/Source/Editor/UnrealEd/Public/Commandlets/AssetRegistryGenerator.h
+++ b/Engine/Source/Editor/UnrealEd/Public/Commandlets/AssetRegistryGenerator.h
@@ -95,7 +95,8 @@ public:
 	 * @param InSandboxFile sandbox to load/save data
 	 * @param bGenerateStreamingInstallManifest should we build a streaming install manifest 
 	 */
-	void BuildChunkManifest(const TSet<FName>& CookedPackages, const TSet<FName>& DevelopmentOnlyPackages, FSandboxPlatformFile* InSandboxFile, bool bGenerateStreamingInstallManifest);
+	void BuildChunkManifest(const TSet<FName>& CookedPackages, const TSet<FName>& DevelopmentOnlyPackages,
+		FSandboxPlatformFile& InSandboxFile, bool bGenerateStreamingInstallManifest);
 
 	/**
 	 * Register a chunk data generator with this generator.
@@ -129,24 +130,6 @@ public:
 	 */
 	FAssetPackageData* GetAssetPackageData(const FName& PackageName);
 
-	/**
-	 * Adds a package to chunk manifest (just calls the other AddPackageToChunkManifestFunction with more parameters)
-	 *
-	 * @param Package Package to add to one of the manifests
-	 * @param SandboxFilename Cooked sandbox path of the package to add to a manifest
-	 * @param LastLoadedMapName Name of the last loaded map (can be empty)
-	 * @param the SandboxPlatformFile used during cook
-	 */
-	void AddPackageToChunkManifest(const FName& PackageFName, const FString& PackagePathName, const FString& SandboxFilename, const FString& LastLoadedMapName, FSandboxPlatformFile* InSandboxFile);
-	
-	/**
-	 * Add a package to the manifest but don't assign it to any chunk yet, packages which are not assigned by the end of the cook will be put into chunk 0
-	 * 
-	 * @param Package which is unassigned
-	 * @param The sandbox file path of the package
-	 */
-	void AddUnassignedPackageToManifest(UPackage* Package, const FString& PackageSandboxPath );
-
 	/**
 	 * Deletes temporary manifest directories.
 	 */
@@ -158,7 +141,7 @@ public:
 	 * @param InSandboxFile the InSandboxFile used during cook
 	 * @param InExtraFlavorChunkSize the ChunkSize used during cooking for InExtraFlavor, value greater than 0 will trigger a cook for extra flavor with specified chunk size
 	 */
-	bool SaveManifests(FSandboxPlatformFile* InSandboxFile, int64 InExtraFlavorChunkSize = 0);
+	bool SaveManifests(FSandboxPlatformFile& InSandboxFile, int64 InExtraFlavorChunkSize = 0);
 
 	/**
 	* Saves generated asset registry data for each platform.
@@ -168,7 +151,7 @@ public:
 	/** 
 	 * Writes out CookerOpenOrder.log file 
 	 */
-	bool WriteCookerOpenOrder(FSandboxPlatformFile* InSandboxFile);
+	bool WriteCookerOpenOrder(FSandboxPlatformFile& InSandboxFile);
 
 	/**
 	 * Follows an assets dependency chain to build up a list of package names in the same order as the runtime would attempt to load them
@@ -257,19 +240,19 @@ private:
 	/** Highest chunk id, being used for geneating dependency tree */
 	int32 HighestChunkId;
 	/** Array of Maps with chunks<->packages assignments */
-	TArray<FChunkPackageSet*>		ChunkManifests;
+	TArray<TUniquePtr<FChunkPackageSet>> ChunkManifests;
 	/** Map of packages that has not been assigned to chunks */
-	FChunkPackageSet				UnassignedPackageSet;
+	FChunkPackageSet UnassignedPackageSet;
 	/** Map of all cooked Packages */
-	FChunkPackageSet				AllCookedPackageSet;
+	FChunkPackageSet AllCookedPackageSet;
 	/** Array of Maps with chunks<->packages assignments. This version contains all dependent packages */
-	TArray<FChunkPackageSet*>		FinalChunkManifests;
+	TArray<TUniquePtr<FChunkPackageSet>> FinalChunkManifests;
 	/** Additional data generators used when creating chunks */
 	TArray<TSharedRef<IChunkDataGenerator>> ChunkDataGenerators;
 	/** Lookup table of used package names used when searching references. */
-	TSet<FName>						InspectedNames;
-	/** */
-	UChunkDependencyInfo*			DependencyInfo;
+	TSet<FName> InspectedNames;
+	/** Source of the config-driven parent-child relationships between chunks. */
+	UChunkDependencyInfo& DependencyInfo;
 
 	/** Required flags a dependency must have if it is to be followed when adding package dependencies to chunks.*/
 	UE::AssetRegistry::EDependencyQuery DependencyQuery;
@@ -323,14 +306,15 @@ private:
 	 * 
 	 * @param the InSandboxFile used during cook
 	 */
-	void FixupPackageDependenciesForChunks(FSandboxPlatformFile* InSandboxFile);
+	void FixupPackageDependenciesForChunks(FSandboxPlatformFile& InSandboxFile);
 
 	/**
 	 * Attaches encryption key guids into the registry data for encrypted primary assets
 	 */
 	void InjectEncryptionData(FAssetRegistryState& TargetState);
 
-	void AddPackageAndDependenciesToChunk(FChunkPackageSet* ThisPackageSet, FName InPkgName, const FString& InSandboxFile, int32 PakchunkIndex, FSandboxPlatformFile* SandboxPlatformFile);
+	void AddPackageAndDependenciesToChunk(FChunkPackageSet& ThisPackageSet, FName InPkgName,
+		const FString& InSandboxFile, int32 PakchunkIndex, FSandboxPlatformFile& SandboxPlatformFile);
 
 	/**
 	 * Returns the path of the temporary packaging directory for the specified platform.
@@ -340,49 +324,18 @@ private:
 		return FPaths::ProjectSavedDir() / TEXT("TmpPackaging") / Platform;
 	}
 
-	/**
-	 * 
-	 */
+	/** Returns the config-driven max size of a chunk for the given platform, or -1 for no limit. */
 	int64 GetMaxChunkSizePerPlatform( const ITargetPlatform* Platform ) const;
 
-	/**
-	* Returns an array of chunks ID for a package name that have been assigned during the cook process.
-	*/
-	FORCEINLINE TArray<int32> GetExistingPackageChunkAssignments(FName PackageFName)
-	{
-		TArray<int32> ExistingChunkIDs;
-		for (uint32 ChunkIndex = 0, MaxChunk = ChunkManifests.Num(); ChunkIndex < MaxChunk; ++ChunkIndex)
-		{
-			if (ChunkManifests[ChunkIndex] && ChunkManifests[ChunkIndex]->Contains(PackageFName))
-			{
-				ExistingChunkIDs.AddUnique(ChunkIndex);
-			}
-		}
-
-		if ( StartupPackages.Contains(PackageFName ))
-		{
-			ExistingChunkIDs.AddUnique(0);
-		}
-
-		return ExistingChunkIDs;
-	}
+	/** Returns an array of chunks ID for a package name that have been assigned during the cook process. */
+	TArray<int32> GetExistingPackageChunkAssignments(FName PackageFName);
 
-	/**
-	* Returns an array of chunks IDs for a package that have been assigned in the editor.
-	*/
-	FORCEINLINE TArray<int32> GetAssetRegistryChunkAssignments(const FName& PackageFName)
-	{
-		TArray<int32> RegistryChunkIDs;
-		auto* FoundIDs = PackageChunkIDMap.Find(PackageFName);
-		if (FoundIDs)
-		{
-			RegistryChunkIDs = *FoundIDs;
-		}
-		return RegistryChunkIDs;
-	}
+	/** Returns an array of chunks IDs for a package that have been assigned in the editor. */
+	TArray<int32> GetAssetRegistryChunkAssignments(const FName& PackageFName);
 
 	/** Generate manifest for a single package */
-	void GenerateChunkManifestForPackage(const FName& PackageFName, const FString& PackagePathName, const FString& SandboxFilename, const FString& LastLoadedMapName, FSandboxPlatformFile* InSandboxFile);
+	void GenerateChunkManifestForPackage(const FName& PackageFName, const FString& PackagePathName,
+		const FString& SandboxFilename, const FString& LastLoadedMapName, FSandboxPlatformFile& InSandboxFile);
 
 	/** Deletes the temporary packaging directory for the specified platform */
 	bool CleanTempPackagingDirectory(const FString& Platform) const;
@@ -391,7 +344,7 @@ private:
 	bool ShouldPlatformGenerateStreamingInstallManifest(const ITargetPlatform* Platform) const;
 
 	/** Generates and saves streaming install chunk manifest */
-	bool GenerateStreamingInstallManifest(int64 InExtraFlavorChunkSize, FSandboxPlatformFile* InSandboxFile);
+	bool GenerateStreamingInstallManifest(int64 InExtraFlavorChunkSize, FSandboxPlatformFile& InSandboxFile);
 
 	/** Gather a list of dependencies required by to completely load this package */
 	bool GatherAllPackageDependencies(FName PackageName, TArray<FName>& DependentPackageNames);
@@ -412,7 +365,8 @@ private:
 	FString	GetShortestReferenceChain(FName PackageName, int32 ChunkID);
 
 	/** Deprecated method to remove redundant chunks */
-	void ResolveChunkDependencyGraph(const FChunkDependencyTreeNode& Node, const TSet<FName>& BaseAssetSet, TArray<TArray<FName>>& OutPackagesMovedBetweenChunks);
+	void ResolveChunkDependencyGraph(const FChunkDependencyTreeNode& Node,
+		const TSet<FName>& BaseAssetSet, TArray<TArray<FName>>& OutPackagesMovedBetweenChunks);
 
 	/** Helper function to verify Chunk asset assignment is valid */
 	bool CheckChunkAssetsAreNotInChild(const FChunkDependencyTreeNode& Node);
