Restore Breakpoints on Blueprint Reload and preserve debug info on Blueprint Duplicate

diff --git a/Engine/Source/Editor/UnrealEd/Private/Kismet2/BlueprintEditorUtils.cpp b/Engine/Source/Editor/UnrealEd/Private/Kismet2/BlueprintEditorUtils.cpp
--- a/Engine/Source/Editor/UnrealEd/Private/Kismet2/BlueprintEditorUtils.cpp
+++ b/Engine/Source/Editor/UnrealEd/Private/Kismet2/BlueprintEditorUtils.cpp
@@ -1812,6 +1812,7 @@ void FBlueprintEditorUtils::PostDuplicateBlueprint(UBlueprint* Blueprint, bool b
 		{
 			// Grab the old CDO, which contains the class defaults
 			UClass* OldBPGCAsClass = Blueprint->GeneratedClass;
+			UBlueprint* OldBlueprint = Cast<UBlueprint>(OldBPGCAsClass->ClassGeneratedBy);
 			UBlueprintGeneratedClass* OldBPGC = (UBlueprintGeneratedClass*)(OldBPGCAsClass);
 			UObject* OldCDO = OldBPGC->GetDefaultObject();
 			check(OldCDO != nullptr);
@@ -1938,9 +1939,13 @@ void FBlueprintEditorUtils::PostDuplicateBlueprint(UBlueprint* Blueprint, bool b
 				Var.VarGuid = NewVarGuids.Emplace(Var.VarGuid, FGuid::NewGuid());
 			}
 
-			// Give all nodes a new Guid
 			TArray< UEdGraphNode* > AllGraphNodes;
 			GetAllNodesOfClass(Blueprint, AllGraphNodes);
+
+			// Before we update Guids, we can use them to dupe breakpoints, watchpins to the new BP
+			FKismetDebugUtilities::PostDuplicateBlueprint(OldBlueprint, Blueprint, AllGraphNodes);
+
+			// Give all nodes a new Guid
 			for(UEdGraphNode* Node : AllGraphNodes)
 			{
 				if (!FBlueprintDuplicationScopeFlags::HasAnyFlag(FBlueprintDuplicationScopeFlags::TheSameNodeGuid))
diff --git a/Engine/Source/Editor/UnrealEd/Private/Kismet2/KismetDebugUtilities.cpp b/Engine/Source/Editor/UnrealEd/Private/Kismet2/KismetDebugUtilities.cpp
--- a/Engine/Source/Editor/UnrealEd/Private/Kismet2/KismetDebugUtilities.cpp
+++ b/Engine/Source/Editor/UnrealEd/Private/Kismet2/KismetDebugUtilities.cpp
@@ -1194,15 +1194,17 @@ void FKismetDebugUtilities::RestoreBreakpointsOnLoad(const UBlueprint* Blueprint
 		return;
 	}
 
+#if WITH_EDITORONLY_DATA
 	// Remove stale breakpoints
 	RemoveBreakpointsByPredicate(
 		Blueprint,
 		[Blueprint](const FBlueprintBreakpoint& Breakpoint)
 		{
 			const UEdGraphNode* const Location = Breakpoint.GetLocation();
-			return !Location || !Location->IsIn(Blueprint);
+			return !Location || Location->GetTypedOuter<UPackage>()->GetPersistentGuid() != Blueprint->GetTypedOuter<UPackage>()->GetPersistentGuid();
 		}
 	);
+#endif
 
 	// Restore breakpoints based on preferred method
 	const UBlueprintEditorSettings* BlueprintEditorSettings = GetDefault<UBlueprintEditorSettings>();
@@ -1236,6 +1238,55 @@ bool FKismetDebugUtilities::HasDebuggingData(const UBlueprint* Blueprint)
 //////////////////////////////////////////////////////////////////////////
 // Blueprint utils
 
+void FKismetDebugUtilities::PostDuplicateBlueprint(UBlueprint* SrcBlueprint, UBlueprint* DupBlueprint, const TArray<UEdGraphNode*>& DupNodes)
+{
+	// Duplicate Breakpoints from the source blueprint
+	if (TArray<FBlueprintBreakpoint>* Breakpoints = GetBreakpoints(SrcBlueprint))
+	{
+		for (FBlueprintBreakpoint& Breakpoint : *Breakpoints)
+		{
+			if (UEdGraphNode* Location = Breakpoint.GetLocation())
+			{
+				UEdGraphNode* const* NewLocation = DupNodes.FindByPredicate(
+					[Location](const UEdGraphNode* Node) -> bool
+					{
+						return Node->NodeGuid == Location->NodeGuid;
+					}
+				);
+				if (NewLocation)
+				{
+					CreateBreakpoint(DupBlueprint, *NewLocation, Breakpoint.IsEnabled());
+				}
+			}
+		}
+	}
+
+	// Duplicate Watched Pins
+	if (TArray<FBlueprintWatchedPin>* WatchedPins = GetWatchedPins(SrcBlueprint))
+	{
+		for (FBlueprintWatchedPin& WatchedPin : *WatchedPins)
+		{
+			if (UEdGraphPin* Pin = WatchedPin.Get())
+			{
+				UEdGraphNode* const* NewOwningNode = DupNodes.FindByPredicate(
+					[Pin](const UEdGraphNode* Node) -> bool
+					{
+						return Node->NodeGuid == Pin->GetOwningNode()->NodeGuid;
+					}
+				);
+
+				if (NewOwningNode)
+				{
+					if (const UEdGraphPin* NewPin = (*NewOwningNode)->FindPin(Pin->PinName))
+					{
+						AddPinWatch(DupBlueprint, NewPin);
+					}
+				}
+			}
+		}
+	}
+}
+
 bool FKismetDebugUtilities::BlueprintHasBreakpoints(const UBlueprint* Blueprint)
 {
 	return GetBreakpoints(Blueprint) != nullptr;
diff --git a/Engine/Source/Editor/UnrealEd/Public/Kismet2/KismetDebugUtilities.h b/Engine/Source/Editor/UnrealEd/Public/Kismet2/KismetDebugUtilities.h
--- a/Engine/Source/Editor/UnrealEd/Public/Kismet2/KismetDebugUtilities.h
+++ b/Engine/Source/Editor/UnrealEd/Public/Kismet2/KismetDebugUtilities.h
@@ -287,6 +287,9 @@ public:
 	
 	// Blueprint utils 
 
+	/** Duplicates debug data from original blueprint to new duplicated blueprint */
+	static void PostDuplicateBlueprint(UBlueprint* SrcBlueprint, UBlueprint* DupBlueprint, const TArray<UEdGraphNode*>& DupNodes);
+
 	// Looks thru the debugging data for any class variables associated with the pin
 	static class FProperty* FindClassPropertyForPin(UBlueprint* Blueprint, const UEdGraphPin* Pin);
 
