Adjustments to mode tool palette style

diff --git a/Engine/Source/Editor/EditorStyle/Private/StarshipStyle.cpp b/Engine/Source/Editor/EditorStyle/Private/StarshipStyle.cpp
--- a/Engine/Source/Editor/EditorStyle/Private/StarshipStyle.cpp
+++ b/Engine/Source/Editor/EditorStyle/Private/StarshipStyle.cpp
@@ -1647,6 +1647,9 @@ void FStarshipEditorStyle::FStyle::SetupGeneralStyles()
 		PaletteToolBarStyle.SetComboButtonPadding(  FMargin(0.0f, 0.0f));
 		PaletteToolBarStyle.SetIndentedBlockPadding(FMargin(0.0f, 0.0f));
 		PaletteToolBarStyle.SetBlockPadding(        FMargin(0.0f, 0.0f));
+		PaletteToolBarStyle.ToggleButton.SetPadding(FMargin(2.0f, 6.0f));
+		PaletteToolBarStyle.ButtonStyle.SetNormalPadding(FMargin(2.0f, 6.0f));
+		PaletteToolBarStyle.ButtonStyle.SetPressedPadding(FMargin(2.0f, 6.0f));
 
 		Set( "PaletteToolBar.Tab",  FCheckBoxStyle()
 			.SetCheckBoxType(            ESlateCheckBoxType::ToggleButton)
diff --git a/Engine/Source/Editor/UnrealEd/Private/Toolkits/BaseToolkit.cpp b/Engine/Source/Editor/UnrealEd/Private/Toolkits/BaseToolkit.cpp
--- a/Engine/Source/Editor/UnrealEd/Private/Toolkits/BaseToolkit.cpp
+++ b/Engine/Source/Editor/UnrealEd/Private/Toolkits/BaseToolkit.cpp
@@ -19,6 +19,8 @@
 #include "Widgets/Docking/SDockTab.h"
 #include "Widgets/Input/SCheckBox.h"
 #include "Widgets/Images/SImage.h"
+#include "Widgets/Input/SSegmentedControl.h"
+
 #define LOCTEXT_NAMESPACE "BaseToolkit"
 
 FBaseToolkit::FBaseToolkit()
@@ -478,14 +480,15 @@ void FModeToolkit::UpdatePrimaryModePanel()
 
 		if (HasIntegratedToolPalettes())
 		{
-			TSharedRef< SUniformWrapPanel> PaletteTabBox = SNew(SUniformWrapPanel)
-				.SlotPadding(FMargin(1.f, 2.f))
-				.HAlign(HAlign_Center);
+			TSharedRef<SSegmentedControl<FName>> PaletteTabBox = SNew(SSegmentedControl<FName>)
+				.UniformPadding(FMargin(8.f, 3.f))
+				.Value_Lambda([this]() { return GetCurrentPalette(); } ) 
+				.OnValueChanged_Lambda([this](const FName& Palette) { SetCurrentPalette(Palette); } );
 
 			// Only show if there is more than one child in the switcher
 			PaletteTabBox->SetVisibility(TAttribute<EVisibility>::Create(TAttribute<EVisibility>::FGetter::CreateLambda([PaletteTabBox]() -> EVisibility
 				{
-					return PaletteTabBox->GetChildren()->Num() > 1 ? EVisibility::Visible : EVisibility::Collapsed;
+					return PaletteTabBox->NumSlots() > 1 ? EVisibility::Visible : EVisibility::Collapsed;
 				})));
 
 			// Also build the toolkit here 
@@ -496,63 +499,47 @@ void FModeToolkit::UpdatePrimaryModePanel()
 			CommandList = GetToolkitCommands();
 
 			TSharedRef< SWidgetSwitcher > PaletteSwitcher = SNew(SWidgetSwitcher)
-				.WidgetIndex_Lambda([this, PaletteNames]() -> int32 {
-				int32 FoundIndex;
-				if (PaletteNames.Find(GetCurrentPalette(), FoundIndex))
+				.WidgetIndex_Lambda(
+				[this, PaletteNames]() -> int32
 				{
-					return FoundIndex;
-				}
-				return 0;
-					});
-
+					int32 FoundIndex;
+					if (PaletteNames.Find(GetCurrentPalette(), FoundIndex))
+					{
+						return FoundIndex;
+					}
+					return 0;
+				});
 
+			
 			for (auto Palette : PaletteNames)
 			{
 				FName ToolbarCustomizationName = GetEditorMode() ?  GetEditorMode()->GetModeInfo().ToolbarCustomizationName : GetScriptableEditorMode()->GetModeInfo().ToolbarCustomizationName;
 				FUniformToolBarBuilder ModeToolbarBuilder(CommandList, FMultiBoxCustomization(ToolbarCustomizationName));
 				ModeToolbarBuilder.SetStyle(&FEditorStyle::Get(), "PaletteToolBar");
 
-				 BuildToolPalette(Palette, ModeToolbarBuilder);
+				BuildToolPalette(Palette, ModeToolbarBuilder);
 
 				TSharedRef<SWidget> PaletteWidget = ModeToolbarBuilder.MakeWidget();
 
-				PaletteTabBox->AddSlot()
-					[
-						SNew(SCheckBox)
-						.Padding(FMargin(8.f, 4.f, 8.f, 5.f))
-						.Style(FEditorStyle::Get(), "PaletteToolBar.Tab")
-						.OnCheckStateChanged_Lambda([this, Palette](const ECheckBoxState) {
-							SetCurrentPalette(Palette);
-						}
-					)
-					.IsChecked_Lambda([this, Palette]() -> ECheckBoxState {
-							if (GetCurrentPalette() == Palette)
-							{
-								return ECheckBoxState::Checked;
-							}
-							return ECheckBoxState::Unchecked;
-						})
-							[
-								SNew(STextBlock)
-								.TextStyle(FAppStyle::Get(), "NormalText")
-								.Text(GetToolPaletteDisplayName(Palette))
-								.Justification(ETextJustify::Center)
-							]
-					];
-
+				const bool bRebuildChildren = false;
+				PaletteTabBox->AddSlot(Palette, false)
+				.Text(GetToolPaletteDisplayName(Palette));
 
 				PaletteSwitcher->AddSlot()
-					[
-						PaletteWidget
-					];
+				[
+					PaletteWidget
+				];
 			}
+
+			PaletteTabBox->RebuildChildren();
+
 			if (ModeToolHeader)
 			{
 				ModeToolHeader->SetContent(
 					SNew(SVerticalBox)
-
 					+ SVerticalBox::Slot()
-					.Padding(8.f, 0.f, 0.f, 8.f)
+					.Padding(0.f, 0.f, 0.f, 8.f)
+					.HAlign(HAlign_Center)
 					.AutoHeight()
 					[
 						PaletteTabBox
diff --git a/Engine/Source/Runtime/Slate/Public/Widgets/Input/SSegmentedControl.h b/Engine/Source/Runtime/Slate/Public/Widgets/Input/SSegmentedControl.h
--- a/Engine/Source/Runtime/Slate/Public/Widgets/Input/SSegmentedControl.h
+++ b/Engine/Source/Runtime/Slate/Public/Widgets/Input/SSegmentedControl.h
@@ -213,12 +213,10 @@ public:
 
 				+SHorizontalBox::Slot()
 				.VAlign(VAlign_Center)
-				.Padding(0.0f, 0.5f, 0.f, 0.f)  // Compensate down for the baseline - helps when using all caps
 				.AutoWidth()
 				[
 					SNew(STextBlock)
 					.TextStyle(TextStyle)
-					.Justification(ETextJustify::Center)
 					.Text(ChildSlotPtr->_Text) 
 				];
 			}
