End Debuging mode if PIE Window is closed

diff --git a/Engine/Source/Editor/UnrealEd/Private/PlayLevel.cpp b/Engine/Source/Editor/UnrealEd/Private/PlayLevel.cpp
--- a/Engine/Source/Editor/UnrealEd/Private/PlayLevel.cpp
+++ b/Engine/Source/Editor/UnrealEd/Private/PlayLevel.cpp
@@ -3215,6 +3215,19 @@ TSharedRef<SPIEViewport> UEditorEngine::GeneratePIEViewportWindow(const FRequest
 	{
 		struct FLocal
 		{
+			static void RequestDestroyPIEWindowOverride(const TSharedRef<SWindow>& WindowBeingClosed, TWeakObjectPtr<UEditorEngine> OwningEditorEngine)
+			{
+				if (OwningEditorEngine.IsValid())
+				{
+					OwningEditorEngine->RequestEndPlayMap();
+					FSlateApplication::Get().LeaveDebuggingMode();
+				}
+				else
+				{
+					FSlateApplication::Get().RequestDestroyWindow(WindowBeingClosed);
+				}
+			}
+
 			static void OnPIEWindowClosed(const TSharedRef<SWindow>& WindowBeingClosed, TWeakPtr<SViewport> PIEViewportWidget, TWeakObjectPtr<UEditorEngine> OwningEditorEngine, int32 ViewportIndex, bool bRestoreRootWindow)
 			{
 				// Save off the window position
@@ -3243,6 +3256,7 @@ TSharedRef<SPIEViewport> UEditorEngine::GeneratePIEViewportWindow(const FRequest
 			}
 		};
 
+		PieWindow->SetRequestDestroyWindowOverride(FRequestDestroyWindowOverride::CreateStatic(&FLocal::RequestDestroyPIEWindowOverride, TWeakObjectPtr<UEditorEngine>(this)));
 		PieWindow->SetOnWindowClosed(FOnWindowClosed::CreateStatic(&FLocal::OnPIEWindowClosed, TWeakPtr<SViewport>(PieViewportWidget), TWeakObjectPtr<UEditorEngine>(this),
 			InViewportIndex, bShouldMinimizeRootWindow));
 	}
