Fix local rotation bug that occurs on static mesh instances within a hierarchy of transformations (e.g. attached to an actor).

diff --git a/Engine/Source/Editor/UnrealEd/Private/EditorModeManager.cpp b/Engine/Source/Editor/UnrealEd/Private/EditorModeManager.cpp
--- a/Engine/Source/Editor/UnrealEd/Private/EditorModeManager.cpp
+++ b/Engine/Source/Editor/UnrealEd/Private/EditorModeManager.cpp
@@ -936,9 +936,9 @@ FMatrix FEditorModeTools::GetLocalCoordinateSystem()
 	{
 		if (TTypedElement<ITypedElementWorldInterface> BottomSelected = GetEditorSelectionSet()->GetBottomSelectedElement<ITypedElementWorldInterface>())
 		{
-			FTransform RelativeTransform;
-			BottomSelected.GetRelativeTransform(RelativeTransform);
-			Matrix = FQuatRotationMatrix(RelativeTransform.GetRotation());
+			FTransform LocalToWorldTransform;
+			BottomSelected.GetWorldTransform(LocalToWorldTransform);
+			Matrix = FQuatRotationMatrix(LocalToWorldTransform.GetRotation());
 		}
 	}
 
