Fix expanding Set Properties crashes the Blueprint Debugger by giving Set Elements valid Display Names

diff --git a/Engine/Source/Editor/UnrealEd/Private/Kismet2/KismetDebugUtilities.cpp b/Engine/Source/Editor/UnrealEd/Private/Kismet2/KismetDebugUtilities.cpp
--- a/Engine/Source/Editor/UnrealEd/Private/Kismet2/KismetDebugUtilities.cpp
+++ b/Engine/Source/Editor/UnrealEd/Private/Kismet2/KismetDebugUtilities.cpp
@@ -2061,8 +2061,7 @@ void FPropertyInstanceInfo::PopulateChildren(FPropertyInstance PropertyInstance,
 				};
 				const TSharedPtr<FPropertyInstanceInfo> ChildInfo = FindOrMake(ChildProperty, VisitedNodes);
 				
-				// members of sets don't have their own names
-				ChildInfo->DisplayName = FText::GetEmpty();
+				ChildInfo->DisplayName = FText::Format(LOCTEXT("SetIndexName", "[{0}]"), FText::AsNumber(Index));
 
 				Children.Add(ChildInfo);
 
