Fixed a issue that blocked the duplication of actors that are not from the current level.

diff --git a/Engine/Source/Editor/UnrealEd/Private/EditorActor.cpp b/Engine/Source/Editor/UnrealEd/Private/EditorActor.cpp
--- a/Engine/Source/Editor/UnrealEd/Private/EditorActor.cpp
+++ b/Engine/Source/Editor/UnrealEd/Private/EditorActor.cpp
@@ -448,6 +448,7 @@ public:
 		{
 			// Cache the current source level
 			ULevel* CurrentSrcLevel = SrcLevel->OwningWorld->GetCurrentLevel();
+			SrcLevel->OwningWorld->SetCurrentLevel(SrcLevel);
 			GUnrealEd->CopyActors(Actors, SrcLevel->OwningWorld, &ScratchData);
 			SrcLevel->OwningWorld->SetCurrentLevel(CurrentSrcLevel);
 		}
