Basic Zen-based DDC2 executor implementation. Matches current Zen server implementation, and is absolutely not in final form (it is synchronous and unoptimised and needs better error handling). Currently tested with ShooterGame only.

#rnx

diff --git a/Engine/Source/Editor/UnrealEd/Private/DerivedDataBuildZenExecutor.cpp b/Engine/Source/Editor/UnrealEd/Private/DerivedDataBuildZenExecutor.cpp
new file mode 100644
--- /dev/null
+++ b/Engine/Source/Editor/UnrealEd/Private/DerivedDataBuildZenExecutor.cpp
@@ -0,0 +1,596 @@
+// Copyright Epic Games, Inc. All Rights Reserved.
+
+#include "DerivedDataBuildZenExecutor.h"
+
+#include "DerivedDataBuild.h"
+#include "DerivedDataBuildAction.h"
+#include "DerivedDataBuildInputs.h"
+#include "DerivedDataBuildOutput.h"
+#include "DerivedDataBuildWorker.h"
+#include "DerivedDataPayload.h"
+#include "DerivedDataRequest.h"
+#include "DerivedDataRequestOwner.h"
+#include "Features/IModularFeatures.h"
+#include "HAL/Event.h"
+#include "HAL/PlatformProcess.h"
+#include "HAL/FileManager.h"
+#include "Misc/Guid.h"
+#include "Misc/Optional.h"
+#include "Misc/Paths.h"
+#include "Misc/PathViews.h"
+#include "Modules/ModuleManager.h"
+#include "Serialization/LargeMemoryWriter.h"
+#include "Serialization/CompactBinaryPackage.h"
+#include "Serialization/CompactBinarySerialization.h"
+#include "Serialization/CompactBinaryValidation.h"
+#include "Serialization/CompactBinaryWriter.h"
+#include "ZenServerHttp.h"
+
+#include <atomic>
+
+#if UE_WITH_ZEN
+
+namespace UE::DerivedData
+{
+
+DEFINE_LOG_CATEGORY_STATIC(LogDerivedDataBuildZenExecutor, Log, All);
+
+/**
+ * This implements a simple Zen executor which passes build requests to
+ * a local Zen instance for execution. This is intentionally as simple as 
+ * possible and everything is synchronous. This is not meant to be used 
+ * in production at this point.
+ */
+
+class FZenBuildWorkerExecutor final : public IBuildWorkerExecutor
+{
+public:
+	FZenBuildWorkerExecutor()
+	{
+		// We don't share request pool with any other service for now, but this should
+		// probably change later once we have a better asynchronous HTTP client implementation
+
+		RequestPool = MakeUnique<UE::Zen::FZenHttpRequestPool>(TEXT("http://localhost:1337/"));
+
+		// Clean out any leftovers from a previous run
+
+		UE_LOG(LogDerivedDataBuildZenExecutor, Display, TEXT("Deleting existing local execution state from '%s'"), *SandboxRootDir);
+
+		const bool RequireExists = false;
+		const bool Tree = true;
+		IFileManager::Get().DeleteDirectory(*SandboxRootDir, RequireExists, Tree);
+
+		IModularFeatures::Get().RegisterModularFeature(IBuildWorkerExecutor::GetFeatureName(), this);
+	}
+
+	virtual ~FZenBuildWorkerExecutor()
+	{
+	}
+
+	void BuildAction(const FBuildAction&			Action,
+					 const FOptionalBuildInputs&	Inputs,
+					 const FBuildWorker&			Worker,
+					 IBuild&						BuildSystem,
+					 EBuildPolicy					Policy,
+					 IRequestOwner&					Owner,
+					 FOnBuildWorkerActionComplete&& OnComplete) final
+	{
+		// Review build action inputs to determine if they need to be materialized/propagated 
+		// (right now, they always will be)
+
+		TArray<FString> MissingInputs;
+		TArray<FStringView> MissingInputViews;
+
+		Action.IterateInputs([&](FStringView Key, const FIoHash& RawHash, uint64 RawSize)
+			{
+				if (Inputs.IsNull() || Inputs.Get().FindInput(Key).IsNull())
+				{
+					MissingInputs.Emplace(Key);
+					MissingInputViews.Add(MissingInputs.Last());
+				}
+			});
+
+		if (!MissingInputViews.IsEmpty())
+		{
+			// Report missing inputs
+			return OnComplete({ Action.GetKey(), {}, MissingInputViews, EStatus::Ok });
+		}
+
+		// Describe worker
+
+		FCbWriter WorkerDescriptor;
+		WorkerDescriptor.BeginObject();
+
+		WorkerDescriptor.AddString("name"_ASV, Worker.GetName());
+		WorkerDescriptor.AddString("path"_ASV, Worker.GetPath());
+		WorkerDescriptor.AddString("host"_ASV, Worker.GetHostPlatform());
+		WorkerDescriptor.AddUuid("buildsystem_version"_ASV, Worker.GetBuildSystemVersion());
+
+		WorkerDescriptor.BeginArray("environment"_ASV);
+		Worker.IterateEnvironment([&](FStringView Name, FStringView Value) {
+			WorkerDescriptor.AddString(WriteToString<256>(Name, "=", Value));
+		});
+		WorkerDescriptor.EndArray();
+
+		WorkerDescriptor.BeginArray("executables"_ASV);
+		Worker.IterateExecutables([&](FStringView Key, const FIoHash& RawHash, uint64 RawSize) 
+		{
+			WorkerDescriptor.BeginObject();
+			WorkerDescriptor.AddString("name"_ASV, Key);
+			WorkerDescriptor.AddBinaryAttachment("hash"_ASV, RawHash);
+			WorkerDescriptor.AddInteger("size"_ASV, RawSize);
+			WorkerDescriptor.EndObject();
+		});
+		WorkerDescriptor.EndArray();
+
+		WorkerDescriptor.BeginArray("files");
+		Worker.IterateFiles([&](FStringView Key, const FIoHash& RawHash, uint64 RawSize) 
+		{
+			WorkerDescriptor.BeginObject();
+			WorkerDescriptor.AddString("name"_ASV, Key);
+			WorkerDescriptor.AddBinaryAttachment("hash"_ASV, RawHash);
+			WorkerDescriptor.AddInteger("size"_ASV, RawSize);
+			WorkerDescriptor.EndObject();
+		});
+		WorkerDescriptor.EndArray();
+
+		WorkerDescriptor.BeginArray("dirs"_ASV);
+		WorkerDescriptor.AddString("Engine/Binaries/Win64"_ASV);
+		WorkerDescriptor.EndArray();
+
+		WorkerDescriptor.BeginArray("functions"_ASV);
+		Worker.IterateFunctions([&](FStringView Name, const FGuid& Version) 
+		{
+			WorkerDescriptor.BeginObject();
+			WorkerDescriptor.AddString("name"_ASV, Name);
+			WorkerDescriptor.AddUuid("version"_ASV, Version);
+			WorkerDescriptor.EndObject();
+		});
+		WorkerDescriptor.EndArray();
+
+		WorkerDescriptor.EndObject();
+		FCbObject WorkDesc = WorkerDescriptor.Save().AsObject();
+
+		// Text representation for debugging / visualization
+
+#if 1
+		TUtf8StringBuilder<1024> TextDesc;
+		CompactBinaryToJson(WorkDesc, TextDesc);
+#endif
+
+		// For now, use the hashed descriptor as the identifier
+
+		FIoHash WorkerId = FIoHash::HashBuffer(WorkDesc.GetBuffer());
+
+		UE::Zen::FZenScopedRequestPtr WorkerRequest(RequestPool.Get());
+
+		// Prepare function (worker)
+		
+		{
+			TStringBuilder<128> WorkerUri;
+			WorkerUri.AppendAnsi("/apply/workers/");
+			WorkerUri << WorkerId;
+
+			// Check if the server knows this worker
+
+			UE::Zen::FZenHttpRequest::Result WorkerGetResult = WorkerRequest->PerformBlockingDownload(WorkerUri, nullptr);
+
+			if (WorkerGetResult == Zen::FZenHttpRequest::Result::Failed)
+			{
+				// TODO: log failure!
+				return OnComplete({ Action.GetKey(), {}, {}, EStatus::Error });
+			}
+			else if (WorkerRequest->GetResponseCode() == 404)
+			{
+				// Worker is new, register it and transmit all the pieces we need
+
+				WorkerRequest->Reset();
+
+				UE::Zen::FZenHttpRequest::Result WorkerPrepResult = WorkerRequest->PerformBlockingPost(WorkerUri, WorkDesc);
+
+				if (WorkerPrepResult == Zen::FZenHttpRequest::Result::Failed)
+				{
+					// TODO: log failure
+					return OnComplete({ Action.GetKey(), {}, {}, EStatus::Error });
+				}
+				else if (WorkerRequest->GetResponseCode() == 404)
+				{
+					FCbObjectView Response = WorkerRequest->GetResponseAsObject();
+					FCbArrayView NeedArray = Response["need"].AsArrayView();
+
+					TSet<FIoHash> NeedHashes;
+
+					for (auto& It : NeedArray)
+					{
+						NeedHashes.Add(It.AsHash());
+					}
+
+					FCbPackage Package;
+
+					{
+						TArray<FIoHash> WorkerFileHashes;
+						TArray<TTuple<FStringView, bool>> WorkerFileMeta;
+						Worker.IterateExecutables([&](FStringView Path, const FIoHash& RawHash, uint64 RawSize) {
+							if (NeedHashes.Contains(RawHash))
+							{
+								WorkerFileHashes.Emplace(RawHash);
+								WorkerFileMeta.Emplace(Path, true);
+							}
+						});
+
+						Worker.IterateFiles([&](FStringView Path, const FIoHash& RawHash, uint64 RawSize) {
+							if (NeedHashes.Contains(RawHash))
+							{
+								WorkerFileHashes.Emplace(RawHash);
+								WorkerFileMeta.Emplace(Path, false);
+							}
+						});
+
+						FRequestOwner BlockingOwner(EPriority::Blocking);
+						Worker.FindFileData(WorkerFileHashes, BlockingOwner, [&](FBuildWorkerFileDataCompleteParams&& Params) {
+							uint32 MetaIndex = 0;
+							for (const FCompressedBuffer& Buffer : Params.Files)
+							{
+								const TTuple<FStringView, bool>& Meta = WorkerFileMeta[MetaIndex];
+
+								FCbAttachment Attachment{ Buffer };
+								Package.AddAttachment(Attachment);
+
+								++MetaIndex;
+							}
+						});
+						BlockingOwner.Wait();
+					}
+
+					Package.SetObject(WorkDesc);
+
+					WorkerRequest->Reset();
+					UE::Zen::FZenHttpRequest::Result WorkerTransmitResult = WorkerRequest->PerformBlockingPostPackage(WorkerUri, Package);
+
+					if (WorkerTransmitResult == Zen::FZenHttpRequest::Result::Failed)
+					{
+						// TODO: log failure
+						return OnComplete({ Action.GetKey(), {}, {}, EStatus::Error });
+					}
+				}
+			}
+
+			if (!UE::Zen::IsSuccessCode(WorkerRequest->GetResponseCode()))
+			{
+				// TODO: log failure
+				return OnComplete({ Action.GetKey(), {}, {}, EStatus::Error });
+			}
+		}
+
+		// Apply function to inputs (i.e execute the job)
+
+		TStringBuilder<128> JobPrepUri;
+		JobPrepUri.AppendAnsi("/apply/jobs/");
+		JobPrepUri << WorkerId;
+
+		FCbWriter BuildActionWriter;
+		Action.Save(BuildActionWriter);
+		FCbObject ActionDesc = BuildActionWriter.Save().AsObject();
+
+		TUtf8StringBuilder<1024> ActionTextDesc;
+		CompactBinaryToJson(ActionDesc, ActionTextDesc);
+
+		WorkerRequest->Reset();
+		UE::Zen::FZenHttpRequest::Result JobPrepResult = WorkerRequest->PerformBlockingPost(JobPrepUri, ActionDesc);
+			
+		FCbPackage ResultPackage;
+
+		if (JobPrepResult == Zen::FZenHttpRequest::Result::Failed)
+		{
+			// TODO: log failure!
+			return OnComplete({ Action.GetKey(), {}, {}, EStatus::Error });
+		}
+		else if (WorkerRequest->GetResponseCode() == 404)
+		{
+			// Marshal payloads
+
+			FCbObjectView Response = WorkerRequest->GetResponseAsObject();
+			FCbArrayView NeedArray = Response["need"].AsArrayView();
+
+			TSet<FIoHash> NeedHashes;
+
+			for (auto& It : NeedArray)
+			{
+				NeedHashes.Add(It.AsHash());
+			}
+
+			FCbPackage ActionPackage;
+
+			Inputs.Get().IterateInputs([&](FStringView Key, const FCompressedBuffer& Buffer) {
+				FIoHash Hash = Buffer.GetRawHash();
+
+				if (NeedHashes.Contains(Hash))
+				{
+					FCbAttachment Attachment{ Buffer };
+					ActionPackage.AddAttachment(Attachment);
+				}
+			});
+
+			ActionPackage.SetObject(ActionDesc);
+
+			WorkerRequest->Reset();
+			UE::Zen::FZenHttpRequest::Result WorkerTransmitResult = WorkerRequest->PerformBlockingPostPackage(JobPrepUri, ActionPackage);
+
+			if (WorkerTransmitResult == Zen::FZenHttpRequest::Result::Failed)
+			{
+				// TODO: log failure!
+				return OnComplete({ Action.GetKey(), {}, {}, EStatus::Error });
+			}
+		}
+
+		ResultPackage = WorkerRequest->GetResponseAsPackage();
+
+		FOptionalBuildOutput RemoteBuildOutput;
+
+#if 1
+		TUtf8StringBuilder<1024> TextBuildOutputDesc;
+		CompactBinaryToJson(ResultPackage.GetObject(), TextBuildOutputDesc);
+#endif
+
+		RemoteBuildOutput = FBuildOutput::Load(Action.GetName(), Action.GetFunction(), ResultPackage.GetObject());
+
+		if (RemoteBuildOutput.IsNull())
+		{
+			UE_LOG(LogDerivedDataBuildZenExecutor, Warning, TEXT("Remote execution system error: build output blob missing!"));
+
+			return OnComplete({ Action.GetKey(), {}, {}, EStatus::Error });
+		}
+
+		FBuildOutputBuilder OutputBuilder = BuildSystem.CreateOutput(Action.GetName(), Action.GetFunction());
+
+		for (const FPayload& Payload : RemoteBuildOutput.Get().GetPayloads())
+		{
+			FCompressedBuffer BufferForPayload;
+					
+			if (const FCbAttachment* Attachment = ResultPackage.FindAttachment(Payload.GetRawHash()))
+			{
+				BufferForPayload = Attachment->AsCompressedBinary();
+			}
+
+			if (BufferForPayload.IsNull())
+			{
+				UE_LOG(LogDerivedDataBuildZenExecutor, Warning, TEXT("Zen execution system error: payload blob missing!"));
+
+				return OnComplete({ Action.GetKey(), {}, {}, EStatus::Error });
+			}
+
+			OutputBuilder.AddPayload(FPayload(Payload.GetId(), BufferForPayload));
+		}
+
+		FBuildOutput BuildOutput = OutputBuilder.Build();
+
+		return OnComplete({ Action.GetKey(), BuildOutput, {}, EStatus::Ok });
+
+#if 0
+		// This path will execute the build action synchronously in a scratch directory
+		// at this stage, all inputs are available in process
+		//
+		// Currently no cleanup whatsoever is performed, so inputs/outputs can be inspected
+		// this could be problematic for large runs so we should probably add support for
+		// configurable cleanup policies
+
+		static std::atomic<int32> SerialNo = 0;
+
+		FString SandboxRoot = SandboxRootDir / TEXT("Scratch");
+		SandboxRoot.AppendInt(++SerialNo);
+
+		// Manifest worker in scratch area
+
+		{
+			TArray<FIoHash> MissingWorkerData;
+
+			TArray<FIoHash> WorkerFileHashes;
+			TArray<TTuple<FStringView, bool>> WorkerFileMeta;
+			Worker.IterateExecutables([&WorkerFileHashes, &WorkerFileMeta](FStringView Path, const FIoHash& RawHash, uint64 RawSize) {
+				WorkerFileHashes.Emplace(RawHash);
+				WorkerFileMeta.Emplace(Path, true);
+			});
+
+			Worker.IterateFiles([&WorkerFileHashes, &WorkerFileMeta](FStringView Path, const FIoHash& RawHash, uint64 RawSize) {
+				WorkerFileHashes.Emplace(RawHash);
+				WorkerFileMeta.Emplace(Path, false);
+			});
+
+			FRequestOwner BlockingOwner(EPriority::Blocking);
+			Worker.FindFileData(WorkerFileHashes, BlockingOwner, [&](FBuildWorkerFileDataCompleteParams&& Params) {
+				uint32 MetaIndex = 0;
+				for (const FCompressedBuffer& Buffer : Params.Files)
+				{
+					const TTuple<FStringView, bool>& Meta = WorkerFileMeta[MetaIndex];
+
+					FString Path { SandboxRoot / FString(Meta.Key) };
+
+					if (TUniquePtr<FArchive> Ar{ IFileManager::Get().CreateFileWriter(*Path, FILEWRITE_Silent) })
+					{
+						FCompositeBuffer DecompressedComposite = Buffer.DecompressToComposite();
+
+						for (auto& Segment : DecompressedComposite.GetSegments())
+						{
+							Ar->Serialize((void*)Segment.GetData(), Segment.GetSize());
+						}
+					}
+
+					++MetaIndex;
+				}
+			});
+			BlockingOwner.Wait();
+
+			// This directory must exist in order for the builder to run correctly
+			FString BinPath{ SandboxRoot / TEXT("Engine/Binaries/Win64") };
+			IFileManager::Get().MakeDirectory(*BinPath);
+		}
+
+		// Manifest inputs in scratch area
+
+		if (!Inputs.IsNull())
+		{
+			Inputs.Get().IterateInputs([&](FStringView Key, const FCompressedBuffer& Buffer) {
+				TStringBuilder<128> InputPath;
+				InputPath << TEXT("Inputs/") << FIoHash(Buffer.GetRawHash());
+
+				FString Path{ SandboxRoot / *InputPath };
+
+				if (TUniquePtr<FArchive> Ar{ IFileManager::Get().CreateFileWriter(*Path, FILEWRITE_Silent) })
+				{
+					FCompressedBuffer Comp = FCompressedBuffer::FromCompressed(Buffer.GetCompressed());
+					*Ar << Comp;
+				}
+			});
+		}
+
+		// Serialize action specification
+
+		{
+			FString ActionPath{ SandboxRoot / TEXT("Build.action") };
+
+			if (TUniquePtr<FArchive> Ar{ IFileManager::Get().CreateFileWriter(*ActionPath, FILEWRITE_Silent) })
+			{
+				FCbWriter BuildActionWriter;
+				Action.Save(BuildActionWriter);
+				BuildActionWriter.Save(*Ar);
+			}
+		}
+
+		const bool bLaunchDetached = false;
+		const bool bLaunchHidden = false;
+		const bool bLaunchReallyHidden = false;
+		uint32 ProcessID = 0;
+		const int PriorityModifier = 0;
+		const TCHAR* WorkingDirectory = *SandboxRoot;
+		const FString WorkerPath { SandboxRoot / FString{Worker.GetPath()} };
+
+		FProcHandle ProcHandle = FPlatformProcess::CreateProc(
+			*WorkerPath, 
+			TEXT("-Build=build.action"), 
+			bLaunchDetached,
+			bLaunchHidden, 
+			bLaunchReallyHidden, 
+			&ProcessID, 
+			PriorityModifier, 
+			WorkingDirectory, 
+			nullptr, 
+			nullptr);
+
+		FPlatformProcess::WaitForProc(ProcHandle);
+		
+		int32 ExitCode = -1;
+		if (!FPlatformProcess::GetProcReturnCode(ProcHandle, &ExitCode) || ExitCode != 0)
+		{
+			UE_LOG(LogDerivedDataBuildZenExecutor, Warning, TEXT("Worker process exit code = %d!"), ExitCode);
+			return OnComplete({ Action.GetKey(), {}, {}, EStatus::Error });
+		}
+
+		// Gather results
+
+		FOptionalBuildOutput RemoteBuildOutput;
+
+		{
+			FString BuildOutputPath{ SandboxRoot / TEXT("Build.output") };
+
+			if (TUniquePtr<FArchive> Ar{ IFileManager::Get().CreateFileReader(*BuildOutputPath, FILEREAD_Silent) })
+			{
+				FCbObject BuildOutput = LoadCompactBinary(*Ar).AsObject();
+
+				if (Ar->IsError())
+				{
+					UE_LOG(LogDerivedDataBuildZenExecutor, Warning, TEXT("Worker error: build output structure not valid!"));
+
+					return OnComplete({ Action.GetKey(), {}, {}, EStatus::Error });
+				}
+
+				RemoteBuildOutput = FBuildOutput::Load(	Action.GetName(), 
+														Action.GetFunction(), 
+														BuildOutput);
+			}
+		}
+
+		if (RemoteBuildOutput.IsNull())
+		{
+			UE_LOG(LogDerivedDataBuildZenExecutor, Warning, TEXT("Remote execution system error: build output blob missing!"));
+
+			return OnComplete({ Action.GetKey(), {}, {}, EStatus::Error });
+		}
+
+		FBuildOutputBuilder OutputBuilder = BuildSystem.CreateOutput(Action.GetName(), Action.GetFunction());
+
+		for (const FPayload& Payload : RemoteBuildOutput.Get().GetPayloads())
+		{
+			TStringBuilder<128> OutputPath;
+			OutputPath << TEXT("Outputs/") << Payload.GetRawHash();
+
+			FString Path{ SandboxRoot / *OutputPath };
+
+			FCompressedBuffer BufferForPayload;
+
+			if (TUniquePtr<FArchive> Ar{ IFileManager::Get().CreateFileReader(*Path, FILEREAD_Silent) })
+			{
+				BufferForPayload = FCompressedBuffer::FromCompressed(*Ar);
+			}
+
+			if (BufferForPayload.IsNull())
+			{
+				UE_LOG(LogDerivedDataBuildZenExecutor, Warning, TEXT("Zen execution system error: payload blob missing!"));
+
+				return OnComplete({ Action.GetKey(), {}, {}, EStatus::Error });
+			}
+
+			OutputBuilder.AddPayload(FPayload(Payload.GetId(), BufferForPayload));
+		}
+
+		FBuildOutput BuildOutput = OutputBuilder.Build();
+
+		return OnComplete({ Action.GetKey(), BuildOutput, {}, EStatus::Ok });
+#endif
+	}
+
+	TConstArrayView<FStringView> GetHostPlatforms() const final
+	{
+		static constexpr FStringView HostPlatforms[]{TEXT("Win64"_SV)};
+		return HostPlatforms;
+	}
+
+	void DumpStats() {}
+
+private:
+	TUniquePtr<UE::Zen::FZenHttpRequestPool> RequestPool;
+	FString SandboxRootDir = FPaths::EngineSavedDir() / TEXT("LocalExec");
+};
+
+}  // namespace UE::DerivedData
+
+TOptional<UE::DerivedData::FZenBuildWorkerExecutor> GZenBuildWorkerExecutor;
+
+void InitDerivedDataBuildZenExecutor()
+{
+	if (!GZenBuildWorkerExecutor.IsSet())
+	{
+		GZenBuildWorkerExecutor.Emplace();
+	}
+}
+
+void DumpDerivedDataBuildZenExecutorStats()
+{
+	static bool bHasRun = false;
+	if (GZenBuildWorkerExecutor.IsSet() && !bHasRun)
+	{
+		bHasRun = true;
+		GZenBuildWorkerExecutor->DumpStats();
+	}
+}
+
+#else
+
+void InitDerivedDataBuildZenExecutor()
+{
+}
+
+void DumpDerivedDataBuildZenExecutorStats()
+{
+}
+
+#endif
\ No newline at end of file
diff --git a/Engine/Source/Editor/UnrealEd/Private/DerivedDataBuildZenExecutor.h b/Engine/Source/Editor/UnrealEd/Private/DerivedDataBuildZenExecutor.h
new file mode 100644
--- /dev/null
+++ b/Engine/Source/Editor/UnrealEd/Private/DerivedDataBuildZenExecutor.h
@@ -0,0 +1,6 @@
+// Copyright Epic Games, Inc. All Rights Reserved.
+
+#pragma once
+
+void InitDerivedDataBuildZenExecutor();
+void DumpDerivedDataBuildZenExecutorStats();
