AssetRegistryGenerator - update function names and comments for clarity.
Split ExtraFlavor features in two to make them clearer.
Move FixupPackageDependenciesForChunks from SaveManifests into FinalizeChunkIDs since it needs to happen even when not saving manifests.
Set DiskSize to -1 for failed package saves in UpdateAssetRegistryPackageData, since that is part of updating assetdata for the package.

diff --git a/Engine/Source/Editor/UnrealEd/Private/Commandlets/AssetRegistryGenerator.cpp b/Engine/Source/Editor/UnrealEd/Private/Commandlets/AssetRegistryGenerator.cpp
--- a/Engine/Source/Editor/UnrealEd/Private/Commandlets/AssetRegistryGenerator.cpp
+++ b/Engine/Source/Editor/UnrealEd/Private/Commandlets/AssetRegistryGenerator.cpp
@@ -197,10 +197,10 @@ TArray<int32> FAssetRegistryGenerator::GetExistingPackageChunkAssignments(FName
 	return ExistingChunkIDs;
 }
 
-TArray<int32> FAssetRegistryGenerator::GetAssetRegistryChunkAssignments(const FName& PackageFName)
+TArray<int32> FAssetRegistryGenerator::GetExplicitChunkIDs(const FName& PackageFName)
 {
 	TArray<int32> PackageInputChunkIds;
-	const TArray<int32>* FoundIDs = PackageChunkIDMap.Find(PackageFName);
+	const TArray<int32>* FoundIDs = ExplicitChunkIDs.Find(PackageFName);
 	if (FoundIDs)
 	{
 		PackageInputChunkIds = *FoundIDs;
@@ -305,18 +305,16 @@ void IChunkDataGenerator::GenerateChunkDataFiles(const int32 InChunkId, const TS
 	PRAGMA_ENABLE_DEPRECATION_WARNINGS
 }
 
-bool FAssetRegistryGenerator::GenerateStreamingInstallManifest(int64 InExtraFlavorChunkSize, FSandboxPlatformFile& InSandboxFile)
+bool FAssetRegistryGenerator::GenerateStreamingInstallManifest(int64 InOverrideChunkSize, const TCHAR* InManifestSubDir,
+	FSandboxPlatformFile& InSandboxFile)
 {
 	const FString Platform = TargetPlatform->PlatformName();
 	FString TmpPackagingDir = GetTempPackagingDirectoryForPlatform(Platform);
-
-	int64 MaxChunkSize = GetMaxChunkSizePerPlatform( TargetPlatform );
-
-	if (InExtraFlavorChunkSize > 0)
+	if (InManifestSubDir)
 	{
-		TmpPackagingDir /= TEXT("ExtraFlavor");
-		MaxChunkSize = InExtraFlavorChunkSize;
+		TmpPackagingDir /= InManifestSubDir;
 	}
+	int64 MaxChunkSize = InOverrideChunkSize > 0 ? InOverrideChunkSize : GetMaxChunkSizePerPlatform(TargetPlatform);
 
 	if (!IFileManager::Get().MakeDirectory(*TmpPackagingDir, true /* Tree */))
 	{
@@ -753,7 +751,7 @@ bool FAssetRegistryGenerator::GenerateStreamingInstallManifest(int64 InExtraFlav
 	return bSucceeded;
 }
 
-void FAssetRegistryGenerator::GenerateChunkManifestForPackage(const FName& PackageFName, const FString& PackagePathName,
+void FAssetRegistryGenerator::CalculateChunkIdsAndAssignToManifest(const FName& PackageFName, const FString& PackagePathName,
 	const FString& SandboxFilename, const FString& LastLoadedMapName, FSandboxPlatformFile& InSandboxFile)
 {
 	TArray<int32> TargetChunks;
@@ -766,16 +764,14 @@ void FAssetRegistryGenerator::GenerateChunkManifestForPackage(const FName& Packa
 	}
 	else
 	{
-		// Collect all chunk IDs associated with this package from the asset registry
-		TArray<int32> RegistryChunkIDs = GetAssetRegistryChunkAssignments(PackageFName);
-
+		TArray<int32> PackageChunkIDs = GetExplicitChunkIDs(PackageFName);
 		ExistingChunkIDs = GetExistingPackageChunkAssignments(PackageFName);
 		if (bUseAssetManager)
 		{
 			// No distinction between source of existing chunks for new flow
-			RegistryChunkIDs.Append(ExistingChunkIDs);
+			PackageChunkIDs.Append(ExistingChunkIDs);
 
-			UAssetManager::Get().GetPackageChunkIds(PackageFName, TargetPlatform, RegistryChunkIDs, TargetChunks);
+			UAssetManager::Get().GetPackageChunkIds(PackageFName, TargetPlatform, PackageChunkIDs, TargetChunks);
 		}
 		else
 		{
@@ -786,23 +782,23 @@ void FAssetRegistryGenerator::GenerateChunkManifestForPackage(const FName& Packa
 			PRAGMA_ENABLE_DEPRECATION_WARNINGS
 			if (AssignStreamingChunkDelegate.IsBound())
 			{
-				AssignStreamingChunkDelegate.ExecuteIfBound(PackagePathName, LastLoadedMapName, RegistryChunkIDs, ExistingChunkIDs, TargetChunks);
+				AssignStreamingChunkDelegate.ExecuteIfBound(PackagePathName, LastLoadedMapName, PackageChunkIDs, ExistingChunkIDs, TargetChunks);
 			}
 			else
 			{
-				//Take asset registry assignments and existing assignments
-				TargetChunks.Append(RegistryChunkIDs);
+				//Take explicit assignments and existing assignments
+				TargetChunks.Append(PackageChunkIDs);
 				TargetChunks.Append(ExistingChunkIDs);
 			}
 		}
 	}
 
-	// if the delegate requested a specific chunk assignment, add them package to it now.
+	// Add the package to the manifest for every chunk the AssetManager found it should belong to
 	for (const int32 PackageChunk : TargetChunks)
 	{
 		AddPackageToManifest(SandboxFilename, PackageFName, PackageChunk);
 	}
-	// If the delegate requested to remove the package from any chunk, remove it now
+	// Remove the package from the manifest for every chunk the AssetManager rejected from the existing chunks
 	for (const int32 PackageChunk : ExistingChunkIDs)
 	{
 		if (!TargetChunks.Contains(PackageChunk))
@@ -907,17 +903,15 @@ void FAssetRegistryGenerator::InjectEncryptionData(FAssetRegistryState& TargetSt
 	}
 }
 
-bool FAssetRegistryGenerator::SaveManifests(FSandboxPlatformFile& InSandboxFile, int64 InExtraFlavorChunkSize)
+bool FAssetRegistryGenerator::SaveManifests(FSandboxPlatformFile& InSandboxFile, int64 InOverrideChunkSize,
+	const TCHAR* InManifestSubDir)
 {
-	// Always do package dependency work, is required to modify asset registry
-	FixupPackageDependenciesForChunks(InSandboxFile);
-
 	if (!bGenerateChunks)
 	{
 		return true;
 	}
 
-	if (!GenerateStreamingInstallManifest(InExtraFlavorChunkSize, InSandboxFile))
+	if (!GenerateStreamingInstallManifest(InOverrideChunkSize, InManifestSubDir, InSandboxFile))
 	{
 		return false;
 	}
@@ -1180,7 +1174,7 @@ void FAssetRegistryGenerator::ComputePackageRemovals(const FAssetRegistryState&
 	}
 }
 
-void FAssetRegistryGenerator::BuildChunkManifest(const TSet<FName>& InCookedPackages,
+void FAssetRegistryGenerator::FinalizeChunkIDs(const TSet<FName>& InCookedPackages,
 	const TSet<FName>& InDevelopmentOnlyPackages, FSandboxPlatformFile& InSandboxFile,
 	bool bGenerateStreamingInstallManifest)
 {
@@ -1219,7 +1213,7 @@ void FAssetRegistryGenerator::BuildChunkManifest(const TSet<FName>& InCookedPack
 		PackageData->DiskSize = -1;
 	}
 
-	// initialize FoundIDList, PackageChunkIDMap
+	// Copy ExplicitChunkIDs and other data from the AssetRegistry into the maps we use during finalization
 	for (const TPair<FName, const FAssetData*>& Pair : State.GetObjectPathToAssetDataMap())
 	{
 		const FAssetData& AssetData = *Pair.Value;
@@ -1231,10 +1225,10 @@ void FAssetRegistryGenerator::BuildChunkManifest(const TSet<FName>& InCookedPack
 				ChunkID = 0;
 			}
 			
-			PackageChunkIDMap.FindOrAdd(AssetData.PackageName).AddUnique(ChunkID);
+			ExplicitChunkIDs.FindOrAdd(AssetData.PackageName).AddUnique(ChunkID);
 		}
 
-		// Now clear the original chunk id list. We will fill it with real IDs when cooking.
+		// Clear the Asset's chunk id list. We will fill it with the final IDs to use later on.
 		// Chunk Ids are safe to modify in place so do a const cast
 		const_cast<FAssetData&>(AssetData).ChunkIDs.Empty();
 
@@ -1266,10 +1260,10 @@ void FAssetRegistryGenerator::BuildChunkManifest(const TSet<FName>& InCookedPack
 			const FString SandboxPath = InSandboxFile.ConvertToAbsolutePathForExternalAppForWrite(*FPackageName::LongPackageNameToFilename(CookedPackage.ToString()));
 			const FString PackagePathName = CookedPackage.ToString();
 			AllCookedPackageSet.Add(CookedPackage, SandboxPath);
-			GenerateChunkManifestForPackage(CookedPackage, PackagePathName, SandboxPath, StartupPackageMapName, InSandboxFile);
+			CalculateChunkIdsAndAssignToManifest(CookedPackage, PackagePathName, SandboxPath, StartupPackageMapName, InSandboxFile);
 		}
 
-		// Capture list at start as it may change during iteration
+		// Capture list at start as elements will be removed during iteration
 		UnassignedPackageSet.GenerateKeyArray(UnassignedPackageList);
 
 		// assign chunks for all the map packages
@@ -1307,12 +1301,12 @@ void FAssetRegistryGenerator::BuildChunkManifest(const TSet<FName>& InCookedPack
 				}
 				const FString& SandboxFilename = *SandboxFilenamePtr;
 
-				GenerateChunkManifestForPackage(PackageFName, PackagePathName, SandboxFilename, MapName, InSandboxFile);
+				CalculateChunkIdsAndAssignToManifest(PackageFName, PackagePathName, SandboxFilename, MapName, InSandboxFile);
 			}
 		}
 	}
 
-	// Capture list at start as it may change during iteration
+	// Capture list at start as elements will be removed during iteration
 	UnassignedPackageSet.GenerateKeyArray(UnassignedPackageList);
 
 	// process the remaining unassigned packages
@@ -1321,10 +1315,11 @@ void FAssetRegistryGenerator::BuildChunkManifest(const TSet<FName>& InCookedPack
 		const FString& SandboxFilename = AllCookedPackageSet.FindChecked(PackageFName);
 		const FString PackagePathName = PackageFName.ToString();
 
-		GenerateChunkManifestForPackage(PackageFName, PackagePathName, SandboxFilename, FString(), InSandboxFile);
+		CalculateChunkIdsAndAssignToManifest(PackageFName, PackagePathName, SandboxFilename, FString(), InSandboxFile);
 	}
+	// anything that remains in the UnAssignedPackageSet is put in chunk0 by FixupPackageDependenciesForChunks
 
-	// anything that remains in the UnAssignedPackageSet will be put in chunk0 when we save the asset registry
+	FixupPackageDependenciesForChunks(InSandboxFile);
 }
 
 void FAssetRegistryGenerator::RegisterChunkDataGenerator(TSharedRef<IChunkDataGenerator> InChunkDataGenerator)
@@ -1402,7 +1397,7 @@ bool FAssetRegistryGenerator::SaveAssetRegistry(const FString& SandboxPath, bool
 
 	if (DevelopmentSaveOptions.bSerializeAssetRegistry && bSerializeDevelopmentAssetRegistry)
 	{
-		// Create development registry data, used for incremental cook and editor viewing
+		// Create development registry data, used for DLC cooks, iterative cooks, and editor viewing
 		FArrayWriter SerializedAssetRegistry;
 
 		State.Save(SerializedAssetRegistry, DevelopmentSaveOptions);
@@ -1659,7 +1654,7 @@ bool FAssetRegistryGenerator::WriteCookerOpenOrder(FSandboxPlatformFile& InSandb
 		const FString SandboxPath = InSandboxFile.GetSandboxDirectory();
 		const FString Platform = TargetPlatform->PlatformName();
 		FString PlatformSandboxPath = SandboxPath.Replace(TEXT("[Platform]"), *Platform);
-		
+
 		// ZENTODO: Change this to work with Zen
 		TMultiMap<FString, FString> CookedPackageFilesMap;
 		FPackageCookerOpenOrderVisitor PackageSearch(InSandboxFile, PlatformSandboxPath, ValidExtensionSet, CookedPackageFilesMap);
@@ -1911,7 +1906,7 @@ void FAssetRegistryGenerator::AddPackageToManifest(const FString& PackageSandbox
 		ChunkManifests[PakchunkIndex].Reset(new FChunkPackageSet());
 	}
 	ChunkManifests[PakchunkIndex]->Add(PackageName, PackageSandboxPath);
-	//Safety check, it the package happens to exist in the unassigned list remove it now.
+	// Now that the package has been assigned to a chunk, remove it from the unassigned set.
 	UnassignedPackageSet.Remove(PackageName);
 }
 
@@ -1926,15 +1921,15 @@ void FAssetRegistryGenerator::RemovePackageFromManifest(FName PackageName, int32
 	}
 }
 
-void FAssetRegistryGenerator::ResolveChunkDependencyGraph(const FChunkDependencyTreeNode& Node,
-	const TSet<FName>& BaseAssetSet, TArray<TArray<FName>>& OutPackagesMovedBetweenChunks)
+void FAssetRegistryGenerator::SubtractParentChunkPackagesFromChildChunks(const FChunkDependencyTreeNode& Node,
+	const TSet<FName>& CumulativeParentPackages, TArray<TArray<FName>>& OutPackagesMovedBetweenChunks)
 {
 	if (FinalChunkManifests.Num() <= Node.ChunkID || !FinalChunkManifests[Node.ChunkID])
 	{
 		return;
 	}
 	FChunkPackageSet& NodeManifest = *FinalChunkManifests[Node.ChunkID];
-	for (FName PackageName : BaseAssetSet)
+	for (FName PackageName : CumulativeParentPackages)
 	{
 		// Remove any assets belonging to our parents.
 		if (NodeManifest.Remove(PackageName) > 0)
@@ -1944,28 +1939,28 @@ void FAssetRegistryGenerator::ResolveChunkDependencyGraph(const FChunkDependency
 				*PackageName.ToString(), Node.ChunkID);
 		}
 	}
-		
+
 	if (!Node.ChildNodes.Num())
 	{
 		return;
 	}
 
 	// Add the current Chunk's assets
-	TSet<FName> CumulativeAssetSet;
+	TSet<FName> CumulativePackages;
 	if (!NodeManifest.IsEmpty())
 	{
-		CumulativeAssetSet.Reserve(BaseAssetSet.Num() + NodeManifest.Num());
-		CumulativeAssetSet.Append(BaseAssetSet);
+		CumulativePackages.Reserve(CumulativeParentPackages.Num() + NodeManifest.Num());
+		CumulativePackages.Append(CumulativeParentPackages);
 		for (const TPair<FName, FString>& Pair : NodeManifest)
 		{
-			CumulativeAssetSet.Add(Pair.Key);
+			CumulativePackages.Add(Pair.Key);
 		}
 	}
-		
-	const TSet<FName>& NewRecursiveParentPackages = CumulativeAssetSet.Num() ? CumulativeAssetSet : BaseAssetSet;
+
+	const TSet<FName>& NewRecursiveParentPackages = CumulativePackages.Num() ? CumulativePackages : CumulativeParentPackages;
 	for (const FChunkDependencyTreeNode& ChildNode : Node.ChildNodes)
 	{
-		ResolveChunkDependencyGraph(ChildNode, NewRecursiveParentPackages, OutPackagesMovedBetweenChunks);
+		SubtractParentChunkPackagesFromChildChunks(ChildNode, NewRecursiveParentPackages, OutPackagesMovedBetweenChunks);
 	}
 }
 
@@ -2012,7 +2007,7 @@ void FAssetRegistryGenerator::AddPackageAndDependenciesToChunk(FChunkPackageSet&
 	//Add this asset
 	ThisPackageSet.Add(InPkgName, InSandboxFile);
 
-	// Only gather dependencies the slow way if we're chunking or not using asset manager
+	// Only gather dependencies the slow way if we're chunking and not using asset manager
 	if (!bGenerateChunks || bUseAssetManager)
 	{
 		return;
@@ -2119,8 +2114,10 @@ void FAssetRegistryGenerator::FixupPackageDependenciesForChunks(FSandboxPlatform
 	TArray<TArray<FName>> PackagesRemovedFromChunks;
 	PackagesRemovedFromChunks.AddDefaulted(ChunkManifests.Num());
 
-	//Finally, if the previous step may added any extra packages to the 0 chunk. Pull them out of other chunks and save space
-	ResolveChunkDependencyGraph(*ChunkDepGraph, TSet<FName>(), PackagesRemovedFromChunks);
+	// Recursively remove child chunk's redundant copies of parent chunk's packages.
+	// This has to be done after all remaining assets were added to chunk 0 above, since all chunks
+	// are children of chunk 0.
+	SubtractParentChunkPackagesFromChildChunks(*ChunkDepGraph, TSet<FName>(), PackagesRemovedFromChunks);
 
 	if (!bUseAssetManager)
 	{
@@ -2143,7 +2140,7 @@ void FAssetRegistryGenerator::FixupPackageDependenciesForChunks(FSandboxPlatform
 			PakchunkIndex, ChunkManifestNum, FinalChunkManifestNum);
 	}
 	
-	// Fix up the asset registry to reflect this chunk layout
+	// Fix up the data in the FAssetRegistryState to reflect this chunk layout
 	for (int32 PakchunkIndex = 0 ; PakchunkIndex < FinalChunkManifests.Num(); ++PakchunkIndex)
 	{
 		check(FinalChunkManifests[PakchunkIndex]);
@@ -2380,6 +2377,9 @@ void FAssetRegistryGenerator::UpdateAssetRegistryPackageData(const UPackage& Pac
 	{
 		// Set the package flags to zero to indicate that the package failed to save
 		NewPackageFlags = 0;
+
+		// Set DiskSize (previous value was disksize in the WorkspaceDomain) to -1 to indicate the cooked file does not exist
+		AssetPackageData->DiskSize = -1;
 	}
 
 	const bool bUpdated = UpdateAssetPackageFlags(PackageName, NewPackageFlags);
diff --git a/Engine/Source/Editor/UnrealEd/Private/CookOnTheFlyServer.cpp b/Engine/Source/Editor/UnrealEd/Private/CookOnTheFlyServer.cpp
--- a/Engine/Source/Editor/UnrealEd/Private/CookOnTheFlyServer.cpp
+++ b/Engine/Source/Editor/UnrealEd/Private/CookOnTheFlyServer.cpp
@@ -6888,25 +6888,25 @@ void UCookOnTheFlyServer::CookByTheBookFinished()
 				}
 				{
 					UE_SCOPED_HIERARCHICAL_COOKTIMER(BuildChunkManifest);
-					Generator.BuildChunkManifest(CookedPackageNames, IgnorePackageNames, *SandboxFile, CookByTheBookOptions->bGenerateStreamingInstallManifests);
+					Generator.FinalizeChunkIDs(CookedPackageNames, IgnorePackageNames, *SandboxFile,
+						CookByTheBookOptions->bGenerateStreamingInstallManifests);
 				}
 				{
 					UE_SCOPED_HIERARCHICAL_COOKTIMER(SaveManifests);
-					// Always try to save the manifests, this is required to make the asset registry work, but doesn't necessarily write a file
 					if (!Generator.SaveManifests(*SandboxFile))
 					{
 						UE_LOG(LogCook, Warning, TEXT("Failed to save chunk manifest"));
 					}
 
 					int64 ExtraFlavorChunkSize;
-					if (FParse::Value(FCommandLine::Get(), TEXT("ExtraFlavorChunkSize="), ExtraFlavorChunkSize))
+					if (FParse::Value(FCommandLine::Get(), TEXT("ExtraFlavorChunkSize="), ExtraFlavorChunkSize) && ExtraFlavorChunkSize > 0)
 					{
-						if (ExtraFlavorChunkSize > 0)
+						// ExtraFlavor is a legacy term for this override; etymology unknown. Override the chunksize specified by the platform,
+						// and write the manifest files created with that chunksize into a separate subdirectory.
+						const TCHAR* ManifestSubDir = TEXT("ExtraFlavor");
+						if (!Generator.SaveManifests(*SandboxFile, ExtraFlavorChunkSize, ManifestSubDir))
 						{
-							if (!Generator.SaveManifests(*SandboxFile, ExtraFlavorChunkSize))
-							{
-								UE_LOG(LogCook, Warning, TEXT("Failed to save chunk manifest"));
-							}
+							UE_LOG(LogCook, Warning, TEXT("Failed to save chunk manifest"));
 						}
 					}
 				}
diff --git a/Engine/Source/Editor/UnrealEd/Public/Commandlets/AssetRegistryGenerator.h b/Engine/Source/Editor/UnrealEd/Public/Commandlets/AssetRegistryGenerator.h
--- a/Engine/Source/Editor/UnrealEd/Public/Commandlets/AssetRegistryGenerator.h
+++ b/Engine/Source/Editor/UnrealEd/Public/Commandlets/AssetRegistryGenerator.h
@@ -87,15 +87,21 @@ public:
 	void ComputePackageRemovals(const FAssetRegistryState& PreviousState, TArray<FName>& RemovedPackages);
 
 	/**
-	 * GenerateChunkManifest 
-	 * generate chunk manifest for the packages passed in using the asset registry to determine dependencies
+	 * FinalizeChunkIDs 
+	 * Create the list of packages to store in each Chunk; each Chunk corresponds to a pak file, or group of pak files if the chunk is split.
+	 * Each package may be in multiple chunks. The selection of chunks for each package is based on
+	 * ChunkIds explicitly assigned to the package in the editor (UPackage::GetChunkIds()) and
+	 * on assignment rules defined by the AssetManager
 	 *
 	 * @param CookedPackages list of packages which were cooked
 	 * @param DevelopmentOnlyPackages list of packages that were specifically not cooked, but to add to the development asset registry
 	 * @param InSandboxFile sandbox to load/save data
-	 * @param bGenerateStreamingInstallManifest should we build a streaming install manifest 
+	 * @param bGenerateStreamingInstallManifest should we build a streaming install manifest
+	 *        If false, no manifest is written and all packages are implicitly assigned to chunk 0 by the automation tool.
+	 *        If true, packages are assigned to chunks based on settings (possibly all in chunk 0 if settings are empty)
+	 *        and a manifest of packagenames is written for each chunk
 	 */
-	void BuildChunkManifest(const TSet<FName>& CookedPackages, const TSet<FName>& DevelopmentOnlyPackages,
+	void FinalizeChunkIDs(const TSet<FName>& CookedPackages, const TSet<FName>& DevelopmentOnlyPackages,
 		FSandboxPlatformFile& InSandboxFile, bool bGenerateStreamingInstallManifest);
 
 	/**
@@ -139,9 +145,13 @@ public:
 	 * Saves all generated manifests for each target platform.
 	 * 
 	 * @param InSandboxFile the InSandboxFile used during cook
-	 * @param InExtraFlavorChunkSize the ChunkSize used during cooking for InExtraFlavor, value greater than 0 will trigger a cook for extra flavor with specified chunk size
+	 * @param InOverrideChunkSize the ChunkSize used during chunk division.
+	 *        If greater than 0, this overrides the default chunksize derived from the platform.
+	 * @param InManifestSubDir If non-null, the manifests are written into this subpath
+	 *        of the usual location.
 	 */
-	bool SaveManifests(FSandboxPlatformFile& InSandboxFile, int64 InExtraFlavorChunkSize = 0);
+	bool SaveManifests(FSandboxPlatformFile& InSandboxFile, int64 InOverrideChunkSize = 0,
+		const TCHAR* InManifestSubDir = nullptr);
 
 	/**
 	* Saves generated asset registry data for each platform.
@@ -229,8 +239,8 @@ private:
 	const ITargetPlatform* TargetPlatform;
 	/** List of all asset packages that were created while loading the last package in the cooker. */
 	TSet<FName> AssetsLoadedWithLastPackage;
-	/** Lookup for the original ChunkID Mappings */
-	TMap<FName, TArray<int32> > PackageChunkIDMap;
+	/** Lookup for the ChunkIDs that were explicitly set by the user in the editor */
+	TMap<FName, TArray<int32>> ExplicitChunkIDs;
 	/** Set of packages containing a map */
 	TSet<FName> PackagesContainingMaps;
 	/** Should the chunks be generated or only asset registry */
@@ -330,21 +340,25 @@ private:
 	/** Returns an array of chunks ID for a package name that have been assigned during the cook process. */
 	TArray<int32> GetExistingPackageChunkAssignments(FName PackageFName);
 
-	/** Returns an array of chunks IDs for a package that have been assigned in the editor. */
-	TArray<int32> GetAssetRegistryChunkAssignments(const FName& PackageFName);
+	/**
+	 * Get chunks IDs for a package that were assigned to the package in the editor from AssetFileContextMenu.
+	 * These explicit chunkids are unioned with the chunkids calculated by the AssetManager.
+	 */
+	TArray<int32> GetExplicitChunkIDs(const FName& PackageFName);
 
-	/** Generate manifest for a single package */
-	void GenerateChunkManifestForPackage(const FName& PackageFName, const FString& PackagePathName,
+	/** Calculate the final ChunkIds used by the package and store the package in the manifest for each of those chunks. */
+	void CalculateChunkIdsAndAssignToManifest(const FName& PackageFName, const FString& PackagePathName,
 		const FString& SandboxFilename, const FString& LastLoadedMapName, FSandboxPlatformFile& InSandboxFile);
 
 	/** Deletes the temporary packaging directory for the specified platform */
 	bool CleanTempPackagingDirectory(const FString& Platform) const;
 
-	/** Returns true if the specific platform desires a chunk manifest */
+	/** Returns true if the specific platform desires multiple chunks suitable for streaming install */
 	bool ShouldPlatformGenerateStreamingInstallManifest(const ITargetPlatform* Platform) const;
 
 	/** Generates and saves streaming install chunk manifest */
-	bool GenerateStreamingInstallManifest(int64 InExtraFlavorChunkSize, FSandboxPlatformFile& InSandboxFile);
+	bool GenerateStreamingInstallManifest(int64 InOverrideChunkSize, const TCHAR* InManifestSubDir,
+		FSandboxPlatformFile& InSandboxFile);
 
 	/** Gather a list of dependencies required by to completely load this package */
 	bool GatherAllPackageDependencies(FName PackageName, TArray<FName>& DependentPackageNames);
@@ -364,9 +378,9 @@ private:
 	/** Helper function for FindShortestReferenceChain */
 	FString	GetShortestReferenceChain(FName PackageName, int32 ChunkID);
 
-	/** Deprecated method to remove redundant chunks */
-	void ResolveChunkDependencyGraph(const FChunkDependencyTreeNode& Node,
-		const TSet<FName>& BaseAssetSet, TArray<TArray<FName>>& OutPackagesMovedBetweenChunks);
+	/** Recursively remove redundant packages from child chunks based on the chunk dependency tree. */
+	void SubtractParentChunkPackagesFromChildChunks(const FChunkDependencyTreeNode& Node,
+		const TSet<FName>& CumulativeParentPackages, TArray<TArray<FName>>& OutPackagesMovedBetweenChunks);
 
 	/** Helper function to verify Chunk asset assignment is valid */
 	bool CheckChunkAssetsAreNotInChild(const FChunkDependencyTreeNode& Node);
