[Subobject Editor] Remove old binding to OnComponentCreated for basic shapes and add the correct binding for OnSubobjectCreated, allowing the component type registry to correctly set the material params.

#rnx

diff --git a/Engine/Source/Editor/UnrealEd/Private/ComponentTypeRegistry.cpp b/Engine/Source/Editor/UnrealEd/Private/ComponentTypeRegistry.cpp
--- a/Engine/Source/Editor/UnrealEd/Private/ComponentTypeRegistry.cpp
+++ b/Engine/Source/Editor/UnrealEd/Private/ComponentTypeRegistry.cpp
@@ -82,9 +82,11 @@ void FComponentTypeRegistryData::AddBasicShapeComponents(TArray<FComponentClassC
 {
 	FString BasicShapesHeading = LOCTEXT("BasicShapesHeading", "Basic Shapes").ToString();
 
-	const auto OnBasicShapeCreated = [](UActorComponent* Component)
+	const auto OnBasicShapeCreated = [](FSubobjectDataHandle ComponentHandle)
 	{
-		UStaticMeshComponent* SMC = Cast<UStaticMeshComponent>(Component);
+		FSubobjectData* Data = ComponentHandle.GetData();
+		// TODO const cast is bad practice, but until the subobject refactor it is only way for internal stuff to get mutable components
+		UStaticMeshComponent* SMC = const_cast<UStaticMeshComponent*>(Cast<UStaticMeshComponent>(Data->GetComponentTemplate()));
 		if (SMC)
 		{
 			const FString MaterialName = TEXT("/Engine/BasicShapes/BasicShapeMaterial.BasicShapeMaterial");
@@ -108,7 +110,7 @@ void FComponentTypeRegistryData::AddBasicShapeComponents(TArray<FComponentClassC
 	{
 		FComponentEntryCustomizationArgs Args;
 		Args.AssetOverride = FindOrLoadObject<UStaticMesh>(UActorFactoryBasicShape::BasicCube.ToString());
-		Args.OnComponentCreated = FOnComponentCreated::CreateStatic(OnBasicShapeCreated);
+		Args.OnSubobjectCreated = FOnSubobjectCreated::CreateStatic(OnBasicShapeCreated);
 		Args.ComponentNameOverride = LOCTEXT("BasicCubeShapeDisplayName", "Cube").ToString();
 		Args.IconOverrideBrushName = FName("ClassIcon.Cube");
 		Args.SortPriority = 2;
@@ -128,7 +130,7 @@ void FComponentTypeRegistryData::AddBasicShapeComponents(TArray<FComponentClassC
 	{
 		FComponentEntryCustomizationArgs Args;
 		Args.AssetOverride = FindOrLoadObject<UStaticMesh>(UActorFactoryBasicShape::BasicPlane.ToString());
-		Args.OnComponentCreated = FOnComponentCreated::CreateStatic(OnBasicShapeCreated);
+		Args.OnSubobjectCreated = FOnSubobjectCreated::CreateStatic(OnBasicShapeCreated);
 		Args.ComponentNameOverride = LOCTEXT("BasicPlaneShapeDisplayName", "Plane").ToString();
 		Args.IconOverrideBrushName = FName("ClassIcon.Plane");
 		Args.SortPriority = 2;
@@ -148,7 +150,7 @@ void FComponentTypeRegistryData::AddBasicShapeComponents(TArray<FComponentClassC
 	{
 		FComponentEntryCustomizationArgs Args;
 		Args.AssetOverride = FindOrLoadObject<UStaticMesh>(UActorFactoryBasicShape::BasicSphere.ToString());
-		Args.OnComponentCreated = FOnComponentCreated::CreateStatic(OnBasicShapeCreated);
+		Args.OnSubobjectCreated = FOnSubobjectCreated::CreateStatic(OnBasicShapeCreated);
 		Args.ComponentNameOverride = LOCTEXT("BasicSphereShapeDisplayName", "Sphere").ToString();
 		Args.IconOverrideBrushName = FName("ClassIcon.Sphere");
 		Args.SortPriority = 2;
@@ -167,7 +169,7 @@ void FComponentTypeRegistryData::AddBasicShapeComponents(TArray<FComponentClassC
 	{
 		FComponentEntryCustomizationArgs Args;
 		Args.AssetOverride = FindOrLoadObject<UStaticMesh>(UActorFactoryBasicShape::BasicCylinder.ToString());
-		Args.OnComponentCreated = FOnComponentCreated::CreateStatic(OnBasicShapeCreated);
+		Args.OnSubobjectCreated = FOnSubobjectCreated::CreateStatic(OnBasicShapeCreated);
 		Args.ComponentNameOverride = LOCTEXT("BasicCylinderShapeDisplayName", "Cylinder").ToString();
 		Args.IconOverrideBrushName = FName("ClassIcon.Cylinder");
 		Args.SortPriority = 3;
@@ -178,7 +180,7 @@ void FComponentTypeRegistryData::AddBasicShapeComponents(TArray<FComponentClassC
 	{
 		FComponentEntryCustomizationArgs Args;
 		Args.AssetOverride = FindOrLoadObject<UStaticMesh>(UActorFactoryBasicShape::BasicCone.ToString());
-		Args.OnComponentCreated = FOnComponentCreated::CreateStatic(OnBasicShapeCreated);
+		Args.OnSubobjectCreated = FOnSubobjectCreated::CreateStatic(OnBasicShapeCreated);
 		Args.ComponentNameOverride = LOCTEXT("BasicConeShapeDisplayName", "Cone").ToString();
 		Args.IconOverrideBrushName = FName("ClassIcon.Cone");
 		Args.SortPriority = 4;
diff --git a/Engine/Source/Editor/UnrealEd/Private/SComponentClassCombo.cpp b/Engine/Source/Editor/UnrealEd/Private/SComponentClassCombo.cpp
--- a/Engine/Source/Editor/UnrealEd/Private/SComponentClassCombo.cpp
+++ b/Engine/Source/Editor/UnrealEd/Private/SComponentClassCombo.cpp
@@ -253,22 +253,12 @@ void SComponentClassCombo::OnAddComponentSelectionChanged( FComponentClassComboE
 					UBlueprint* LoadedObject = LoadObject<UBlueprint>(nullptr, *InItem->GetComponentPath(), nullptr, LoadFlags, nullptr);
 					ComponentClass = GetAuthoritativeBlueprintClass(LoadedObject);
 				}
-
-				UActorComponent* NewActorComponent =
-					OnComponentClassSelected.IsBound() ?
-					OnComponentClassSelected.Execute(ComponentClass, InItem->GetComponentCreateAction(), InItem->GetAssetOverride())
-					: nullptr;
 				
 				FSubobjectDataHandle NewActorCompHandle =
 					OnSubobjectClassSelected.IsBound() ?
 					OnSubobjectClassSelected.Execute(ComponentClass, InItem->GetComponentCreateAction(), InItem->GetAssetOverride())
 					: FSubobjectDataHandle::InvalidHandle;
-				
-				if(NewActorComponent)
-				{
-					InItem->GetOnComponentCreated().ExecuteIfBound(NewActorComponent);
-				}
-				
+
 				if(NewActorCompHandle.IsValid())
 				{
 					InItem->GetOnSubobjectCreated().ExecuteIfBound(NewActorCompHandle);
diff --git a/Engine/Source/Editor/UnrealEd/Public/SComponentClassCombo.h b/Engine/Source/Editor/UnrealEd/Public/SComponentClassCombo.h
--- a/Engine/Source/Editor/UnrealEd/Public/SComponentClassCombo.h
+++ b/Engine/Source/Editor/UnrealEd/Public/SComponentClassCombo.h
@@ -34,7 +34,6 @@ namespace EComponentCreateAction
 }
 
 
-DECLARE_DELEGATE_OneParam(FOnComponentCreated, UActorComponent*);
 DECLARE_DELEGATE_OneParam(FOnSubobjectCreated, FSubobjectDataHandle);
 
 DECLARE_DELEGATE_RetVal_ThreeParams( UActorComponent*, FComponentClassSelected, TSubclassOf<UActorComponent>, EComponentCreateAction::Type, UObject*);
@@ -46,10 +45,9 @@ struct FComponentEntryCustomizationArgs
 	TWeakObjectPtr<UObject> AssetOverride;
 	/** Custom name to display */
 	FString ComponentNameOverride;
-	/** Callback when a new component is created */
-	FOnComponentCreated OnComponentCreated;
+	
 	/** Callback when a new subobject is created */
-	FOnSubobjectCreated OnOnSubobjectCreated;
+	FOnSubobjectCreated OnSubobjectCreated;
 	/** Brush icon to use instead of the class icon */
 	FName IconOverrideBrushName;
 	/** Custom sort priority to use (smaller means sorted first) */
@@ -58,7 +56,6 @@ struct FComponentEntryCustomizationArgs
 	FComponentEntryCustomizationArgs()
 		: AssetOverride( nullptr )
 		, ComponentNameOverride()
-		, OnComponentCreated()
 		, IconOverrideBrushName( NAME_None )
 		, SortPriority(0)
 	{
@@ -138,15 +135,10 @@ public:
 	{
 		return ComponentCreateAction;
 	}
-
-	FOnComponentCreated& GetOnComponentCreated()
-	{
-		return CustomizationArgs.OnComponentCreated;
-	}
 	
 	FOnSubobjectCreated& GetOnSubobjectCreated()
     {
-    	return CustomizationArgs.OnOnSubobjectCreated;
+    	return CustomizationArgs.OnSubobjectCreated;
     }
 	
 	FString GetClassName() const;
