Add Insight trace for AttemptModifiedPackageNotification

diff --git a/Engine/Source/Editor/UnrealEd/Private/UnrealEdEngine.cpp b/Engine/Source/Editor/UnrealEd/Private/UnrealEdEngine.cpp
--- a/Engine/Source/Editor/UnrealEd/Private/UnrealEdEngine.cpp
+++ b/Engine/Source/Editor/UnrealEd/Private/UnrealEdEngine.cpp
@@ -524,6 +524,8 @@ void UUnrealEdEngine::OnPackageDirtyStateUpdated( UPackage* Pkg)
 
 void UUnrealEdEngine::AttemptModifiedPackageNotification()
 {
+	TRACE_CPUPROFILER_EVENT_SCOPE(UUnrealEdEngine::AttemptModifiedPackageNotification);
+
 	bool bIsCooking = CookServer && CookServer->IsCookingInEditor() && CookServer->IsCookByTheBookRunning();
 
 	if (bShowPackageNotification && !bIsCooking)
