CookOnTheFlyServer GetAdditionalCurrentIniVersionStrings: Remove static keyword added by accident; it is blocking the read of console variables after the first.
#rnx

diff --git a/Engine/Source/Editor/UnrealEd/Private/CookOnTheFlyServer.cpp b/Engine/Source/Editor/UnrealEd/Private/CookOnTheFlyServer.cpp
--- a/Engine/Source/Editor/UnrealEd/Private/CookOnTheFlyServer.cpp
+++ b/Engine/Source/Editor/UnrealEd/Private/CookOnTheFlyServer.cpp
@@ -4986,7 +4986,7 @@ void GetAdditionalCurrentIniVersionStrings( const ITargetPlatform* TargetPlatfor
 
 	for (const FString& RValue : VersionedRValues)
 	{
-		static const auto CVar = IConsoleManager::Get().FindTConsoleVariableDataInt(*RValue);
+		const auto CVar = IConsoleManager::Get().FindTConsoleVariableDataInt(*RValue);
 		if (CVar)
 		{
 			IniVersionMap.Add(*RValue, FString::Printf(TEXT("%d"), CVar->GetValueOnGameThread()));
