RemoteExecute: Add missing results to completed list as failed when remote execution fails due to a server error or if cancelled due to timeout, so they can be logged.

#rnx

diff --git a/Engine/Source/Editor/Experimental/HordeExecutor/Private/Execution.cpp b/Engine/Source/Editor/Experimental/HordeExecutor/Private/Execution.cpp
--- a/Engine/Source/Editor/Experimental/HordeExecutor/Private/Execution.cpp
+++ b/Engine/Source/Editor/Experimental/HordeExecutor/Private/Execution.cpp
@@ -135,6 +135,17 @@ namespace UE::RemoteExecution
 
 			if (WaitTask->Complete.Num() == WaitTask->TaskHashes.Num() || Response.Key != EStatusCode::Ok)
 			{
+				for (const FIoHash& TaskHash : WaitTask->TaskHashes)
+				{
+					if (!WaitTask->Complete.Contains(TaskHash))
+					{
+						FGetTaskUpdateResponse Update;
+						Update.TaskHash = TaskHash;
+						Update.Outcome = EComputeTaskOutcome::NoResult;
+						WaitTask->Complete.Add(TaskHash, MoveTemp(Update));
+					}
+				}
+
 				FGetTaskUpdatesResponse Result;
 				WaitTask->Complete.GenerateValueArray(Result.Updates);
 				WaitTask->Promise.EmplaceValue(TPair<EStatusCode, FGetTaskUpdatesResponse>(MoveTemp(Response.Key), MoveTemp(Result)));
diff --git a/Engine/Source/Editor/UnrealEd/Private/DerivedDataBuildRemoteExecutor.cpp b/Engine/Source/Editor/UnrealEd/Private/DerivedDataBuildRemoteExecutor.cpp
--- a/Engine/Source/Editor/UnrealEd/Private/DerivedDataBuildRemoteExecutor.cpp
+++ b/Engine/Source/Editor/UnrealEd/Private/DerivedDataBuildRemoteExecutor.cpp
@@ -1146,6 +1146,13 @@ void FRemoteBuildExecutionRequest::OnExecutionCompleted(UE::RemoteExecution::ESt
 
 	if (!IsStatusOk(Status, TEXT("OnExecutionCompleted")))
 	{
+		if (!Result.Updates.IsEmpty())
+		{
+			UE_LOG(LogDerivedDataBuildRemoteExecutor, Warning, TEXT("Remote execution system error: Task %s Outcome %s: %s"),
+				*FString::FromHexBlob(Result.Updates[0].TaskHash.GetBytes(), sizeof(FIoHash::ByteArray)),
+				*UE::RemoteExecution::ComputeTaskOutcomeString(Result.Updates[0].Outcome),
+				*Result.Updates[0].Detail);
+		}
 		return;
 	}
 
@@ -1162,7 +1169,8 @@ void FRemoteBuildExecutionRequest::OnExecutionCompleted(UE::RemoteExecution::ESt
 
 	if (Result.Updates[0].Outcome != UE::RemoteExecution::EComputeTaskOutcome::Success)
 	{
-		UE_LOG(LogDerivedDataBuildRemoteExecutor, Warning, TEXT("Remote execution system error: Outcome %s: %s"),
+		UE_LOG(LogDerivedDataBuildRemoteExecutor, Warning, TEXT("Remote execution system error: Task %s Outcome %s: %s"),
+			*FString::FromHexBlob(Result.Updates[0].TaskHash.GetBytes(), sizeof(FIoHash::ByteArray)),
 			*UE::RemoteExecution::ComputeTaskOutcomeString(Result.Updates[0].Outcome),
 			*Result.Updates[0].Detail);
 		State.Owner.End(this, [this]
