Push soft references onto the queue of packages to be processed rather than dealing with them immediately. This vastly reduces peak memory usage e.g. in levels with soft references to many streamed sublevels.

diff --git a/Engine/Source/Editor/UnrealEd/Classes/Commandlets/DerivedDataCacheCommandlet.h b/Engine/Source/Editor/UnrealEd/Classes/Commandlets/DerivedDataCacheCommandlet.h
--- a/Engine/Source/Editor/UnrealEd/Classes/Commandlets/DerivedDataCacheCommandlet.h
+++ b/Engine/Source/Editor/UnrealEd/Classes/Commandlets/DerivedDataCacheCommandlet.h
@@ -37,8 +37,8 @@ class UDerivedDataCacheCommandlet : public UCommandlet
 
 	void MaybeMarkPackageAsAlreadyLoaded(UPackage *Package);
 
-	void CacheLoadedPackages(UPackage* CurrentPackage, uint8 PackageFilter, const TArray<ITargetPlatform*>& Platforms);
-	void CacheWorldPackages(UWorld* World, uint8 PackageFilter, const TArray<ITargetPlatform*>& Platforms);
+	void CacheLoadedPackages(UPackage* CurrentPackage, uint8 PackageFilter, const TArray<ITargetPlatform*>& Platforms, TSet<FName>& OutNewProcessedPackages);
+	void CacheWorldPackages(UWorld* World, uint8 PackageFilter, const TArray<ITargetPlatform*>& Platforms, TSet<FName>& OutNewProcessedPackages);
 	bool ProcessCachingObjects(const TArray<ITargetPlatform*>& Platforms);
 	void FinishCachingObjects(const TArray<ITargetPlatform*>& Platforms);
 };
diff --git a/Engine/Source/Editor/UnrealEd/Private/Commandlets/DerivedDataCacheCommandlet.cpp b/Engine/Source/Editor/UnrealEd/Private/Commandlets/DerivedDataCacheCommandlet.cpp
--- a/Engine/Source/Editor/UnrealEd/Private/Commandlets/DerivedDataCacheCommandlet.cpp
+++ b/Engine/Source/Editor/UnrealEd/Private/Commandlets/DerivedDataCacheCommandlet.cpp
@@ -187,7 +187,7 @@ static void PumpAsync(bool* bInOutHadActivity = nullptr)
 	}
 }
 
-void UDerivedDataCacheCommandlet::CacheLoadedPackages(UPackage* CurrentPackage, uint8 PackageFilter, const TArray<ITargetPlatform*>& Platforms)
+void UDerivedDataCacheCommandlet::CacheLoadedPackages(UPackage* CurrentPackage, uint8 PackageFilter, const TArray<ITargetPlatform*>& Platforms, TSet<FName>& OutNewProcessedPackages)
 {
 	TRACE_CPUPROFILER_EVENT_SCOPE(UDerivedDataCacheCommandlet::CacheLoadedPackages);
 
@@ -216,6 +216,7 @@ void UDerivedDataCacheCommandlet::CacheLoadedPackages(UPackage* CurrentPackage,
 				UE_LOG(LogDerivedDataCacheCommandlet, Display, TEXT("Processing %s"), *NewPackageName.ToString());
 
 				ProcessedPackages.Add(NewPackageName);
+				OutNewProcessedPackages.Add(NewPackageName);
 				NewPackageIt.RemoveCurrent();
 
 				ObjectsWithOuter.Reset();
@@ -331,7 +332,7 @@ void UDerivedDataCacheCommandlet::FinishCachingObjects(const TArray<ITargetPlatf
 	FinishCacheTime += FPlatformTime::Seconds() - FinishCacheTimeStart;
 }
 
-void UDerivedDataCacheCommandlet::CacheWorldPackages(UWorld* World, uint8 PackageFilter, const TArray<ITargetPlatform*>& Platforms)
+void UDerivedDataCacheCommandlet::CacheWorldPackages(UWorld* World, uint8 PackageFilter, const TArray<ITargetPlatform*>& Platforms, TSet<FName>& OutNewProcessedPackages)
 {
 	check(World);
 
@@ -360,12 +361,12 @@ void UDerivedDataCacheCommandlet::CacheWorldPackages(UWorld* World, uint8 Packag
 		UWorldPartition* WorldPartition = World->GetWorldPartition();
 		check(WorldPartition);
 
-		FWorldPartitionHelpers::ForEachActorWithLoading(WorldPartition, [this, PackageFilter, &Platforms](const FWorldPartitionActorDesc* ActorDesc)
+		FWorldPartitionHelpers::ForEachActorWithLoading(WorldPartition, [this, PackageFilter, &Platforms, &OutNewProcessedPackages](const FWorldPartitionActorDesc* ActorDesc)
 		{
 			if (AActor* Actor = ActorDesc->GetActor())
 			{
 				UE_LOG(LogDerivedDataCacheCommandlet, Display, TEXT("Loaded actor %s"), *Actor->GetName());
-				CacheLoadedPackages(Actor->GetPackage(), PackageFilter, Platforms);
+				CacheLoadedPackages(Actor->GetPackage(), PackageFilter, Platforms, OutNewProcessedPackages);
 			}
 			return true;
 		});
@@ -541,6 +542,20 @@ int32 UDerivedDataCacheCommandlet::Main( const FString& Params )
 			PackagePaths.SetNum(NewNum);
 		}
 
+		if (PackagePaths.Num() == 0)
+		{
+			UE_LOG(LogDerivedDataCacheCommandlet, Display, TEXT("No packages found to load from command line arguments."));
+		}
+		else
+		{
+			UE_LOG(LogDerivedDataCacheCommandlet, Display, TEXT("%d packages to load from command line arguments"), PackagePaths.Num());
+			for( int32 Index=0; Index < PackagePaths.Num(); ++Index)
+			{
+				const TPair<FString, FName>& Pair = PackagePaths[Index];
+				UE_LOG(LogDerivedDataCacheCommandlet, Display, TEXT(" %d) %s"), Index + 1, *Pair.Get<1>().ToString());
+			}
+		}
+
 		ITargetPlatformManagerModule* TPM = GetTargetPlatformManager();
 		const TArray<ITargetPlatform*>& Platforms = TPM->GetActiveTargetPlatforms();
 
@@ -561,41 +576,92 @@ int32 UDerivedDataCacheCommandlet::Main( const FString& Params )
 		int32 NumProcessedSinceLastGC = 0;
 		bool bLastPackageWasMap = false;
 
-		if (PackagePaths.Num() == 0)
-		{
-			UE_LOG(LogDerivedDataCacheCommandlet, Display, TEXT("No packages found to load."));
-		}
-		else
-		{
-			UE_LOG(LogDerivedDataCacheCommandlet, Display, TEXT("%d packages to load..."), PackagePaths.Num());
-		}
-
-		// Gather the list of packages to process
+		// Mark command-line packages as already discovered so we don't double-add from soft refs and can avoid loading packages on other shards
 		PackagesToProcess.Empty(PackagePaths.Num());
 		for (int32 PackageIndex = PackagePaths.Num() - 1; PackageIndex >= 0; PackageIndex--)
 		{
 			PackagesToProcess.Add(PackagePaths[PackageIndex].Get<1>());
 		}
 
-		// Process each package
-		for (int32 PackageIndex = PackagePaths.Num() - 1; PackageIndex >= 0; PackageIndex-- )
+		// Add all soft object references from no asset in particular to the packages to be processed, before filtering in the case of distributed work
 		{
-			TPair<FString, FName>& PackagePath = PackagePaths[PackageIndex];
-			const FString& Filename = PackagePath.Get<0>();
-			FName PackageFName = PackagePath.Get<1>();
-			check(!ProcessedPackages.Contains(PackageFName));
+			int32 StartingPackageCount = PackagePaths.Num();
+			TSet<FName> SoftReferencedPackages;
+			GRedirectCollector.ProcessSoftObjectPathPackageList(NAME_None, false, SoftReferencedPackages);
+			for (FName SoftRefName : SoftReferencedPackages)
+			{
+				if (PackagesToProcess.Contains(SoftRefName))
+				{
+					continue;
+				}
+
+				FString SoftRefFilename;
+				if (FPackageName::DoesPackageExist(SoftRefName.ToString(), nullptr, &SoftRefFilename))
+				{
+					PackagePaths.Push(TPair<FString, FName>(SoftRefFilename, SoftRefName));
+					PackagesToProcess.Add(SoftRefName);
+				}
+			}
 
-			// If work is distributed, skip packages that are meant to be process by other machines
-			if (bDoSubset)
+			if (StartingPackageCount == PackagePaths.Num())
+			{
+				UE_LOG(LogDerivedDataCacheCommandlet, Display, TEXT("No packages found to load from startup soft references."));
+			}
+			else
 			{
+				UE_LOG(LogDerivedDataCacheCommandlet, Display, TEXT("%d packages to load from startup soft references"), PackagePaths.Num() - StartingPackageCount);
+				for (int32 Index = StartingPackageCount; Index < PackagePaths.Num(); ++Index)
+				{
+					const TPair<FString, FName>& Pair = PackagePaths[Index];
+					UE_LOG(LogDerivedDataCacheCommandlet, Display, TEXT(" %d) %s"), Index + 1 - StartingPackageCount, *Pair.Get<1>().ToString());
+				}
+			}
+		}
+
+		// If work is distributed, skip packages that are meant to be process by other machines
+		// Do this before the main loop so that we don't filter soft refs that we enqueue 
+		if (bDoSubset)
+		{
+			PackagePaths.RemoveAll([SubsetMod, SubsetTarget](const TPair<FString, FName>& P) {
+				FName PackageFName = P.Value;
+
 				FString PackageName = PackageFName.ToString();
 				if (FCrc::StrCrc_DEPRECATED(*PackageName.ToUpper()) % SubsetMod != SubsetTarget)
 				{
-					continue;
+					return true;
+				}
+				return false;
+			});
+
+			if (PackagePaths.Num() == 0)
+			{
+				UE_LOG(LogDerivedDataCacheCommandlet, Display, TEXT("No packages to process after subset split!"));
+			}
+			else
+			{
+				UE_LOG(LogDerivedDataCacheCommandlet, Display, TEXT("%d packages to load after subset split"), PackagePaths.Num());
+				for (int32 Index = 0; Index < PackagePaths.Num(); ++Index)
+				{
+					const TPair<FString, FName>& Pair = PackagePaths[Index];
+					UE_LOG(LogDerivedDataCacheCommandlet, Display, TEXT(" %d) %s"), Index + 1, *Pair.Get<1>().ToString());
 				}
 			}
+		}
+
+		// Process each package
+		int32 PackageOrder = 0;
+		while( PackagePaths.Num())
+		{
+			TTuple<FString, FName> PackagePath = PackagePaths.Pop();
+			const FString& Filename = PackagePath.Get<0>();
+			FName PackageFName = PackagePath.Get<1>();
+			if (ProcessedPackages.Contains(PackageFName))
+			{
+				// Soft refs may be queued, then processed as a hard ref from something else.
+				continue;
+			}
 
-			UE_LOG(LogDerivedDataCacheCommandlet, Display, TEXT("Loading (%d) %s"), FilesInPath.Num() - PackageIndex, *Filename);
+			UE_LOG(LogDerivedDataCacheCommandlet, Display, TEXT("Loading (%d) %s"), ++PackageOrder, *Filename);
 
 			UPackage* Package = LoadPackage(NULL, *Filename, LOAD_None);
 			if (Package == NULL)
@@ -609,26 +675,57 @@ int32 UDerivedDataCacheCommandlet::Main( const FString& Params )
 				NumProcessedSinceLastGC++;
 			}
 
-			// even if the load failed this could be the first time through the loop so it might have all the startup packages to resolve
-			GRedirectCollector.ResolveAllSoftObjectPaths();
-
 			// Find any new packages and cache all the objects in each package
-			CacheLoadedPackages(Package, PackageFilter, Platforms);
+			TSet<FName> NewProcessedPackages;
+			CacheLoadedPackages(Package, PackageFilter, Platforms, NewProcessedPackages);
 
 			// Ensure we load maps to process all their referenced packages in case they are using world partition.
 			if (bLastPackageWasMap)
 			{
 				if (UWorld* World = UWorld::FindWorldInPackage(Package))
 				{
-					CacheWorldPackages(World, PackageFilter, Platforms);
+					CacheWorldPackages(World, PackageFilter, Platforms, NewProcessedPackages);
+				}
+			}
+
+			// Queue up soft references of each package we just processed
+			NewProcessedPackages.Add(NAME_None); // Always check for more references from non-asset systems each step
+			for( FName NewProcessedPackage : NewProcessedPackages)
+			{
+				TSet<FName> SoftReferencedPackages;
+				GRedirectCollector.ProcessSoftObjectPathPackageList(NewProcessedPackage, false, SoftReferencedPackages);
+				for (FName SoftRefName : SoftReferencedPackages)
+				{
+					// Packages may already be enqueued on this or another machine 
+					if (!PackagesToProcess.Contains(SoftRefName) && !ProcessedPackages.Contains(SoftRefName))
+					{
+						PackagesToProcess.Add(SoftRefName);
+						FString SoftRefFilename;
+						if (FPackageName::DoesPackageExist(SoftRefName.ToString(), nullptr, &SoftRefFilename))
+						{
+							UE_LOG(LogDerivedDataCacheCommandlet, Log, TEXT("Queueing soft reference '%s' for later processing"), *SoftRefName.ToString());
+							PackagePaths.Push(TPair<FString, FName>(SoftRefFilename, SoftRefName));
+						}
+						else
+						{
+							UE_LOG(LogDerivedDataCacheCommandlet, Warning, TEXT("Failed to find soft reference '%s'"), *SoftRefName.ToString());
+						}
+					}
+					else
+					{
+						UE_LOG(LogDerivedDataCacheCommandlet, Verbose, TEXT("Skipping soft reference '%s': %s, %s "), 
+							PackagesToProcess.Contains(SoftRefName) ? TEXT("ALREADY QUEUED") : TEXT("NOT QUEUED"),
+							ProcessedPackages.Contains(SoftRefName) ? TEXT("ALREADY PROCESSED") : TEXT("NOT PROCESSED")
+						);
+					}
 				}
 			}
 
 			// Perform a GC if conditions are met
-			if (NumProcessedSinceLastGC >= GCInterval || PackageIndex < 0 || bLastPackageWasMap)
+			if (NumProcessedSinceLastGC >= GCInterval || PackagePaths.IsEmpty() || bLastPackageWasMap)
 			{
 				const double StartGCTime = FPlatformTime::Seconds();
-				if (NumProcessedSinceLastGC >= GCInterval || PackageIndex < 0)
+				if (NumProcessedSinceLastGC >= GCInterval || PackagePaths.IsEmpty())
 				{
 					UE_LOG(LogDerivedDataCacheCommandlet, Display, TEXT("GC (Full)..."));
 					CollectGarbage(RF_NoFlags);
