EditorFramework: shut down Mode-level ToolsContext before clearing the Toolkit, so that if a Tool is still active, it has a chance to clean up
#rnx

diff --git a/Engine/Source/Editor/UnrealEd/Private/Tools/UEdMode.cpp b/Engine/Source/Editor/UnrealEd/Private/Tools/UEdMode.cpp
--- a/Engine/Source/Editor/UnrealEd/Private/Tools/UEdMode.cpp
+++ b/Engine/Source/Editor/UnrealEd/Private/Tools/UEdMode.cpp
@@ -141,6 +141,14 @@ void UEdMode::Exit()
 		SettingsObject->SaveConfig();
 	}
 
+	// Shutdown the Mode-scope ToolsContext, and notify the EditorToolsContext so it can release the reference it holds.
+	// Do this before shutting down the Toolkit as if a Mode-scope Tool is still active it will need to clean up
+	if (UObjectInitialized() && ModeToolsContext != nullptr)
+	{
+		ModeToolsContext->ShutdownContext();
+		EditorToolsContext->OnChildEdModeToolsContextShutdown(ModeToolsContext);
+	}
+
 	if (Toolkit.IsValid())
 	{
 		const TSharedRef<FUICommandList>& CommandList = Toolkit->GetToolkitCommands();
@@ -154,13 +162,6 @@ void UEdMode::Exit()
 	}
 	RegisteredEditorTools.SetNum(0);
 
-	// shutdown the Mode-scope ToolsContext, and notify the EditorToolsContext so it can release the reference it holds
-	if (UObjectInitialized() && ModeToolsContext != nullptr)
-	{
-		ModeToolsContext->ShutdownContext();
-		EditorToolsContext->OnChildEdModeToolsContextShutdown(ModeToolsContext);
-	}
-
 	// disconnect from the Mode Manager's shared ToolsContext
 	EditorToolsContext->ToolManager->OnToolStarted.RemoveAll(this);
 	EditorToolsContext->ToolManager->OnToolEnded.RemoveAll(this);
