Lightmass: Fixed hardcoded dependency paths to include DerivedDataCache and SSL

#rnx

diff --git a/Engine/Source/Editor/UnrealEd/Private/Lightmass/Lightmass.cpp b/Engine/Source/Editor/UnrealEd/Private/Lightmass/Lightmass.cpp
--- a/Engine/Source/Editor/UnrealEd/Private/Lightmass/Lightmass.cpp
+++ b/Engine/Source/Editor/UnrealEd/Private/Lightmass/Lightmass.cpp
@@ -2925,8 +2925,10 @@ bool FLightmassProcessor::BeginRun()
 		TEXT("Binaries/Win32/UnrealLightmass-ApplicationCore.dll"),
 		TEXT("Binaries/Win32/UnrealLightmass-Core.dll"),
 		TEXT("Binaries/Win32/UnrealLightmass-CoreUObject.dll"),
+		TEXT("Binaries/Win32/UnrealLightmass-DerivedDataCache.dll"),
 		TEXT("Binaries/Win32/UnrealLightmass-Projects.dll"),
 		TEXT("Binaries/Win32/UnrealLightmass-Json.dll"),
+		TEXT("Binaries/Win32/UnrealLightmass-SSL.dll")
 		TEXT("Binaries/Win32/UnrealLightmass-BuildSettings.dll")
 	};
 
@@ -2941,8 +2943,10 @@ bool FLightmassProcessor::BeginRun()
 		TEXT("Binaries/Win64/UnrealLightmass-ApplicationCore.dll"),
 		TEXT("Binaries/Win64/UnrealLightmass-Core.dll"),
 		TEXT("Binaries/Win64/UnrealLightmass-CoreUObject.dll"),
+		TEXT("Binaries/Win64/UnrealLightmass-DerivedDataCache.dll"),
 		TEXT("Binaries/Win64/UnrealLightmass-Projects.dll"),
 		TEXT("Binaries/Win64/UnrealLightmass-Json.dll"),
+		TEXT("Binaries/Win64/UnrealLightmass-SSL.dll"),
 		TEXT("Binaries/Win64/UnrealLightmass-BuildSettings.dll"),
 		TEXT("Binaries/Win64/embree.dll"),
 		TEXT("Binaries/Win64/tbb.dll"),
@@ -2956,8 +2960,10 @@ bool FLightmassProcessor::BeginRun()
 		TEXT("Binaries/Mac/UnrealLightmass-ApplicationCore.dylib"),
 		TEXT("Binaries/Mac/UnrealLightmass-Core.dylib"),
 		TEXT("Binaries/Mac/UnrealLightmass-CoreUObject.dylib"),
+		TEXT("Binaries/Mac/UnrealLightmass-DerivedDataCache.dylib"),
 		TEXT("Binaries/Mac/UnrealLightmass-Json.dylib"),
 		TEXT("Binaries/Mac/UnrealLightmass-Projects.dylib"),
+		TEXT("Binaries/Mac/UnrealLightmass-SSL.dylib"),
 		TEXT("Binaries/Mac/UnrealLightmass-SwarmInterface.dylib"),
 		TEXT("Binaries/Mac/UnrealLightmass-BuildSettings.dylib"),
 		TEXT("Binaries/Mac/libembree.2.dylib"),
@@ -2972,8 +2978,10 @@ bool FLightmassProcessor::BeginRun()
 		TEXT("Binaries/Linux/libUnrealLightmass-ApplicationCore.so"),
 		TEXT("Binaries/Linux/libUnrealLightmass-Core.so"),
 		TEXT("Binaries/Linux/libUnrealLightmass-CoreUObject.so"),
+		TEXT("Binaries/Linux/libUnrealLightmass-DerivedDataCache.so"),
 		TEXT("Binaries/Linux/libUnrealLightmass-Json.so"),
 		TEXT("Binaries/Linux/libUnrealLightmass-Projects.so"),
+		TEXT("Binaries/Linux/libUnrealLightmass-SSL.so"),
 		TEXT("Binaries/Linux/libUnrealLightmass-SwarmInterface.so"),
 		TEXT("Binaries/Linux/libUnrealLightmass-Networking.so"),
 		TEXT("Binaries/Linux/libUnrealLightmass-Messaging.so"),
