[World Partition Navmesh] Use iterative cell loading for building world partition static navmesh data
Include a fix in UWorldPartitionBuilder for IterativeCellOverlapSize not being applied to loaded bounds

diff --git a/Engine/Source/Editor/UnrealEd/Private/WorldPartition/WorldPartitionBuilder.cpp b/Engine/Source/Editor/UnrealEd/Private/WorldPartition/WorldPartitionBuilder.cpp
--- a/Engine/Source/Editor/UnrealEd/Private/WorldPartition/WorldPartitionBuilder.cpp
+++ b/Engine/Source/Editor/UnrealEd/Private/WorldPartition/WorldPartitionBuilder.cpp
@@ -227,7 +227,7 @@ bool UWorldPartitionBuilder::Run(UWorld* World, FPackageSourceControlHelper& Pac
 					const FVector Min(x * IterativeCellSize, y * IterativeCellSize, z * IterativeCellSize);
 					const FVector Max = Min + FVector(IterativeCellSize);
 					FBox BoundsToLoad(Min, Max);
-					BoundsToLoad.ExpandBy(IterativeCellOverlapSize);
+					BoundsToLoad = BoundsToLoad.ExpandBy(IterativeCellOverlapSize);
 
 					UE_LOG(LogWorldPartitionBuilder, Verbose, TEXT("Loading Bounds: Min %s, Max %s"), *BoundsToLoad.Min.ToString(), *BoundsToLoad.Max.ToString());
 					WorldPartition->LoadEditorCells(BoundsToLoad, false);
diff --git a/Engine/Source/Editor/UnrealEd/Private/WorldPartition/WorldPartitionNavigationDataBuilder.cpp b/Engine/Source/Editor/UnrealEd/Private/WorldPartition/WorldPartitionNavigationDataBuilder.cpp
--- a/Engine/Source/Editor/UnrealEd/Private/WorldPartition/WorldPartitionNavigationDataBuilder.cpp
+++ b/Engine/Source/Editor/UnrealEd/Private/WorldPartition/WorldPartitionNavigationDataBuilder.cpp
@@ -17,9 +17,16 @@ DEFINE_LOG_CATEGORY_STATIC(LogWorldPartitionNavigationDataBuilder, Log, All);
 UWorldPartitionNavigationDataBuilder::UWorldPartitionNavigationDataBuilder(const FObjectInitializer& ObjectInitializer)
 	: Super(ObjectInitializer)
 {
+	// Size of loaded cell. Set as big as your hardware can afford.
+	// @todo: move to a config file.
+	IterativeCellSize = 204800;
+	
+	// Extra padding around loaded cell.
+	// @todo: set value programatically.
+	IterativeCellOverlapSize = 2000;
 }
 
-bool UWorldPartitionNavigationDataBuilder::RunInternal(UWorld* World, const FBox& Bounds, FPackageSourceControlHelper& PackageHelper)
+bool UWorldPartitionNavigationDataBuilder::RunInternal(UWorld* World, const FBox& LoadedBounds, FPackageSourceControlHelper& PackageHelper)
 {
 	UWorldPartitionSubsystem* WorldPartitionSubsystem = World->GetSubsystem<UWorldPartitionSubsystem>();
 	check(WorldPartitionSubsystem);
@@ -35,10 +42,23 @@ bool UWorldPartitionNavigationDataBuilder::RunInternal(UWorld* World, const FBox
 		NavigationDataChunkActorPackages.Add(ItActor->GetPackage());
 	}
 
+	// Destroy any existing navigation data chunk actors within bounds we are generating, we will make new ones.
+	const FBox GeneratingBounds = LoadedBounds.ExpandBy(-IterativeCellOverlapSize);
+	for (TActorIterator<ANavigationDataChunkActor> It(World); It; ++It)
+	{
+		ANavigationDataChunkActor* Actor = *It;
+		const FVector Location = Actor->GetActorLocation();
+		if (GeneratingBounds.IsInside(Location))
+		{
+			World->DestroyActor(Actor);	
+		}
+	}
+
+	// Make sure static meshes have compiled before generating navigation data
 	FStaticMeshCompilingManager::Get().FinishAllCompilation();
 	
-	// Rebuild ANavigationDataChunkActor for the whole world
-	WorldPartition->GenerateNavigationData();
+	// Rebuild ANavigationDataChunkActor in loaded bounds
+	WorldPartition->GenerateNavigationData(LoadedBounds);
 
 	// Gather all packages again to include newly created ANavigationDataChunkActor actors
 	for (TActorIterator<ANavigationDataChunkActor> ItActor(World); ItActor; ++ItActor)
@@ -51,6 +71,7 @@ bool UWorldPartitionNavigationDataBuilder::RunInternal(UWorld* World, const FBox
 
 	for (UPackage* ActorPackage : NavigationDataChunkActorPackages)
 	{
+		// Only change package that have been dirtied
 		if (ActorPackage && ActorPackage->IsDirty())
 		{
 			if (UPackage::IsEmptyPackage(ActorPackage))
@@ -68,6 +89,11 @@ bool UWorldPartitionNavigationDataBuilder::RunInternal(UWorld* World, const FBox
 	if (!PackagesToDelete.IsEmpty())
 	{
 		UE_LOG(LogWorldPartitionNavigationDataBuilder, Log, TEXT("Deleting %d packages."), PackagesToDelete.Num());
+		for (const UPackage* Package : PackagesToDelete)
+		{
+			UE_LOG(LogWorldPartitionNavigationDataBuilder, Verbose, TEXT("   Deleting package  %s."), *Package->GetName());	
+		}
+		
 		if (!PackageHelper.Delete(PackagesToDelete))
 		{
 			UE_LOG(LogWorldPartitionNavigationDataBuilder, Error, TEXT("Error deleting packages."));
@@ -78,49 +104,62 @@ bool UWorldPartitionNavigationDataBuilder::RunInternal(UWorld* World, const FBox
 	// Save packages
 	if (!PackagesToSave.IsEmpty())
 	{
-		// Checkout packages to save
-		UE_LOG(LogWorldPartitionNavigationDataBuilder, Log, TEXT("Checking out %d actor packages."), PackagesToSave.Num());
-		if (PackageHelper.UseSourceControl())
 		{
-			FEditorFileUtils::CheckoutPackages(PackagesToSave, /*OutPackagesCheckedOut*/nullptr, /*bErrorIfAlreadyCheckedOut*/false);
-		}
-		else
-		{
-			for (const UPackage* Package : PackagesToSave)
+			// Checkout packages to save
+			TRACE_CPUPROFILER_EVENT_SCOPE(CheckoutPackages);
+			UE_LOG(LogWorldPartitionNavigationDataBuilder, Log, TEXT("Checking out %d packages."), PackagesToSave.Num());
+
+			if (PackageHelper.UseSourceControl())
 			{
-				const FString PackageFilename = SourceControlHelpers::PackageFilename(Package);
-				if (IPlatformFile::GetPlatformPhysical().FileExists(*PackageFilename))
+				FEditorFileUtils::CheckoutPackages(PackagesToSave, /*OutPackagesCheckedOut*/nullptr, /*bErrorIfAlreadyCheckedOut*/false);
+			}
+			else
+			{
+				// Remove read-only
+				for (const UPackage* Package : PackagesToSave)
 				{
-					if (!IPlatformFile::GetPlatformPhysical().SetReadOnly(*PackageFilename, /*bNewReadOnlyValue*/false))
+					const FString PackageFilename = SourceControlHelpers::PackageFilename(Package);
+					if (IPlatformFile::GetPlatformPhysical().FileExists(*PackageFilename))
 					{
-						UE_LOG(LogWorldPartitionNavigationDataBuilder, Error, TEXT("Error setting %s writable"), *PackageFilename);
-						return 1;
+						if (!IPlatformFile::GetPlatformPhysical().SetReadOnly(*PackageFilename, /*bNewReadOnlyValue*/false))
+						{
+							UE_LOG(LogWorldPartitionNavigationDataBuilder, Error, TEXT("Error setting %s writable"), *PackageFilename);
+							return 1;
+						}
 					}
 				}
 			}
 		}
-
-		// Save packages
-		UE_LOG(LogWorldPartitionNavigationDataBuilder, Log, TEXT("Saving %d packages."), PackagesToSave.Num());
-		for (UPackage* Package : PackagesToSave)
+		
 		{
-			UE_LOG(LogWorldPartitionNavigationDataBuilder, Verbose, TEXT("Saving package  %s."), *Package->GetName());
-			FString PackageFileName = SourceControlHelpers::PackageFilename(Package);
-			if (!UPackage::SavePackage(Package, nullptr, RF_Standalone, *PackageFileName, GError, nullptr, false, true, SAVE_Async))
+			// Save packages
+			TRACE_CPUPROFILER_EVENT_SCOPE(SavingPackages);
+			UE_LOG(LogWorldPartitionNavigationDataBuilder, Log, TEXT("Saving %d packages."), PackagesToSave.Num());
+
+			for (UPackage* Package : PackagesToSave)
 			{
-				UE_LOG(LogWorldPartitionNavigationDataBuilder, Error, TEXT("Error saving package %s."), *Package->GetName());
-				return 1;
+				UE_LOG(LogWorldPartitionNavigationDataBuilder, Verbose, TEXT("   Saving package  %s."), *Package->GetName());
+				FString PackageFileName = SourceControlHelpers::PackageFilename(Package);
+				if (!UPackage::SavePackage(Package, nullptr, RF_Standalone, *PackageFileName, GError, nullptr, false, true, SAVE_Async))
+				{
+					UE_LOG(LogWorldPartitionNavigationDataBuilder, Error, TEXT("Error saving package %s."), *Package->GetName());
+					return 1;
+				}
 			}
 		}
-
-		// Add new packages to source control
-		UE_LOG(LogWorldPartitionNavigationDataBuilder, Log, TEXT("Adding packages to source control."));
-		for (UPackage* Package : PackagesToSave)
+		
 		{
-			if (!PackageHelper.AddToSourceControl(Package))
+			// Add new packages to source control
+			TRACE_CPUPROFILER_EVENT_SCOPE(AddingToSourceControl);
+			UE_LOG(LogWorldPartitionNavigationDataBuilder, Log, TEXT("Adding packages to source control."));
+
+			for (UPackage* Package : PackagesToSave)
 			{
-				UE_LOG(LogWorldPartitionNavigationDataBuilder, Error, TEXT("Error adding package %s to source control."), *Package->GetName());
-				return 1;
+				if (!PackageHelper.AddToSourceControl(Package))
+				{
+					UE_LOG(LogWorldPartitionNavigationDataBuilder, Error, TEXT("Error adding package %s to source control."), *Package->GetName());
+					return 1;
+				}
 			}
 		}
 
diff --git a/Engine/Source/Editor/UnrealEd/Public/WorldPartition/WorldPartitionNavigationDataBuilder.h b/Engine/Source/Editor/UnrealEd/Public/WorldPartition/WorldPartitionNavigationDataBuilder.h
--- a/Engine/Source/Editor/UnrealEd/Public/WorldPartition/WorldPartitionNavigationDataBuilder.h
+++ b/Engine/Source/Editor/UnrealEd/Public/WorldPartition/WorldPartitionNavigationDataBuilder.h
@@ -13,9 +13,9 @@ class UWorldPartitionNavigationDataBuilder : public UWorldPartitionBuilder
 public:
 	// UWorldPartitionBuilder interface begin
 	virtual bool RequiresCommandletRendering() const override { return false; }
-	virtual ELoadingMode GetLoadingMode() const override { return ELoadingMode::EntireWorld; }
+	virtual ELoadingMode GetLoadingMode() const override { return ELoadingMode::IterativeCells; }
 
 protected:
-	virtual bool RunInternal(UWorld* World, const FBox& Bounds, FPackageSourceControlHelper& PackageHelper) override;
+	virtual bool RunInternal(UWorld* World, const FBox& LoadedBounds, FPackageSourceControlHelper& PackageHelper) override;
 	// UWorldPartitionBuilder interface end
 };
diff --git a/Engine/Source/Runtime/Engine/Classes/AI/NavigationSystemBase.h b/Engine/Source/Runtime/Engine/Classes/AI/NavigationSystemBase.h
--- a/Engine/Source/Runtime/Engine/Classes/AI/NavigationSystemBase.h
+++ b/Engine/Source/Runtime/Engine/Classes/AI/NavigationSystemBase.h
@@ -240,7 +240,13 @@ public:
 	UE_DEPRECATED(4.20, "GetMainNavData is deprecated. Use FNavigationSystem::GetCurrent<UNavigationSystemV1>()->GetDefaultNavDataInstance instead")
 	INavigationDataInterface* GetMainNavData(int) { return nullptr; }
 
-	virtual bool ContainsNavData(const FBox& Bounds) const { return false; }
+	virtual void SetBuildBounds(const FBox& Bounds) PURE_VIRTUAL(UNavigationSystemBase::SetBuildBounds, );
+
+	virtual FBox GetNavigableWorldBounds() const PURE_VIRTUAL(UNavigationSystemBase::GetNavigableWorldBounds, return FBox(ForceInit););
+	
+	virtual bool ContainsNavData(const FBox& Bounds) const PURE_VIRTUAL(UNavigationSystemBase::ContainsNavData, return false;);
+	virtual FBox ComputeNavDataBounds() const PURE_VIRTUAL(UNavigationSystemBase::GetNavigableWorldBounds, return FBox(ForceInit););
+	
 	virtual void AddNavigationDataChunk(class ANavigationDataChunkActor& DataChunkActor) {}
 	virtual void RemoveNavigationDataChunk(class ANavigationDataChunkActor& DataChunkActor) {}
 	virtual void FillNavigationDataChunkActor(const FBox& QueryBounds, class ANavigationDataChunkActor& DataChunkActor, FBox& OutTilesBounds) {}
diff --git a/Engine/Source/Runtime/Engine/Private/WorldPartition/RuntimeSpatialHash/RuntimeSpatialHashNavigation.cpp b/Engine/Source/Runtime/Engine/Private/WorldPartition/RuntimeSpatialHash/RuntimeSpatialHashNavigation.cpp
--- a/Engine/Source/Runtime/Engine/Private/WorldPartition/RuntimeSpatialHash/RuntimeSpatialHashNavigation.cpp
+++ b/Engine/Source/Runtime/Engine/Private/WorldPartition/RuntimeSpatialHash/RuntimeSpatialHashNavigation.cpp
@@ -13,47 +13,70 @@
 
 DEFINE_LOG_CATEGORY_STATIC(LogWorldPartitionRuntimeSpatialHashNav, Log, All);
 
-bool UWorldPartitionRuntimeSpatialHash::GenerateNavigationData()
+bool UWorldPartitionRuntimeSpatialHash::GenerateNavigationData(const FBox& LoadedBounds)
 {
 	TRACE_CPUPROFILER_EVENT_SCOPE(UWorldPartitionRuntimeSpatialHash::GenerateNavigationData);
 
-	UE_LOG(LogWorldPartitionRuntimeSpatialHashNav, Log, TEXT("%s"), ANSI_TO_TCHAR(__FUNCTION__));
+	static int32 CallCount = 0;
+	CallCount++;
+	UE_LOG(LogWorldPartitionRuntimeSpatialHashNav, Log, TEXT("%i. GenerateNavigationData for LoadedBounds %s"), CallCount, *LoadedBounds.ToString());
 
-	UWorldPartition* WorldPartition = GetOuterUWorldPartition();
+	const UWorldPartition* WorldPartition = GetOuterUWorldPartition();
 	UWorld* World = WorldPartition->World;
 
 	// Generate navmesh
 	// Make sure navigation is added and initialized in EditorMode
 	FNavigationSystem::AddNavigationSystemToWorld(*World, FNavigationSystemRunMode::EditorMode);
 
+	UNavigationSystemBase* NavSystem = World->GetNavigationSystem();
+	if (NavSystem == nullptr)
+	{
+		UE_LOG(LogWorldPartitionRuntimeSpatialHashNav, Verbose, TEXT("No navigation system to generate navigation data."));
+		return false;
+	}
+
+	// First check if there is any intersection with the nav bounds.
+	const FBox NavBounds = NavSystem->GetNavigableWorldBounds();
+	if (!LoadedBounds.Intersect(NavBounds))
+	{
+		// No intersections, early out
+		UE_LOG(LogWorldPartitionRuntimeSpatialHashNav, Log, TEXT("GenerateNavigationData finished (no intersection with nav bounds)."));
+		return true;
+	}
+
 	// Invoke navigation data generator
+	NavSystem->SetBuildBounds(LoadedBounds);
 	FNavigationSystem::Build(*World);
 
+	// Compute navdata bounds from tiles
+	const FBox NavDataBounds = NavSystem->ComputeNavDataBounds();
+	if (!LoadedBounds.Intersect(NavDataBounds))
+	{
+		// No intersections, early out
+		UE_LOG(LogWorldPartitionRuntimeSpatialHashNav, Log, TEXT("GenerateNavigationData finished (no intersection with generated navigation data)."));
+		return true;
+	}
+
 	// For each cell, gather navmesh and generate a datachunk actor
 	const FBox WorldBounds = WorldPartition->GetWorldBounds();
 
-	const int32 GridIndex = 0; //Todo At: only work with grid 0 for now.
-	const uint32 GridLevel = 3; //Todo AT, only generate for level 3 for now
-	UE_LOG(LogWorldPartitionRuntimeSpatialHashNav, Verbose, TEXT("Generate NavDataChunk actor for GridIndex %i."), GridIndex);
+	constexpr int32 GridIndex = 0; //Todo At: only work with grid 0 for now.
+	constexpr uint32 GridLevel = 3; //Todo AT, only generate for level 3 for now
+	UE_LOG(LogWorldPartitionRuntimeSpatialHashNav, Verbose, TEXT("Generate NavDataChunk actors for GridIndex %i, GridLevel %i."), GridIndex, GridLevel);
 
 	const FSpatialHashRuntimeGrid& RuntimeGrid = Grids[GridIndex];
 	const FSquare2DGridHelper GridHelper = GetGridHelper(WorldBounds, RuntimeGrid.CellSize);
 	int32 ActorCount = 0;
 
-	const UNavigationSystemBase* NavSystem = World->GetNavigationSystem();
-	if (NavSystem == nullptr)
-	{
-		UE_LOG(LogWorldPartitionRuntimeSpatialHashNav, Verbose, TEXT("No navigation system to generate navigation data."));
-		return false;
-	}
-
 	// Keep track of all valid navigation data chunk actors
 	TSet<ANavigationDataChunkActor*> ValidNavigationDataChunkActors;
 
 	const FSquare2DGridHelper::FGridLevel& GridLevelHelper = GridHelper.Levels[GridLevel];
 
-	GridLevelHelper.ForEachCells([GridLevel, &GridHelper, &GridLevelHelper, RuntimeGrid, WorldPartition, &ActorCount, World, NavSystem, &ValidNavigationDataChunkActors, this](const FIntVector2& CellCoord)
+	GridLevelHelper.ForEachIntersectingCells(LoadedBounds, [GridLevel, &GridHelper, &GridLevelHelper, RuntimeGrid, &ActorCount, World, &ValidNavigationDataChunkActors, &NavDataBounds, this](const FIntVector2& CellCoord)
 		{
+			TRACE_CPUPROFILER_EVENT_SCOPE(MakeNavigationDataChunkActorForGridCell);
+		
 			FBox2D CellBounds;
 			GridLevelHelper.GetCellBounds(CellCoord, CellBounds);
 			if (CellBounds.GetExtent().X < 1.f || CellBounds.GetExtent().Y < 1.f)
@@ -63,16 +86,16 @@ bool UWorldPartitionRuntimeSpatialHash::GenerateNavigationData()
 				return;
 			}
 
-			const float HalfHeight = HALF_WORLD_MAX;
+			constexpr float HalfHeight = HALF_WORLD_MAX;
 			const FBox QueryBounds(FVector(CellBounds.Min.X, CellBounds.Min.Y, -HalfHeight), FVector(CellBounds.Max.X, CellBounds.Max.Y, HalfHeight));
 
-			if (NavSystem->ContainsNavData(QueryBounds) == false)
+			if (!NavDataBounds.IsValid || !NavDataBounds.Intersect(QueryBounds))
 			{
 				// Skip if there is no navdata for this cell
 				return;
 			}
 
-			 //@todo_ow: Properly handle data layers
+			//@todo_ow: Properly handle data layers
 			FActorSpawnParameters SpawnParams;
 			SpawnParams.bDeferConstruction = true;
 			SpawnParams.bCreateActorPackage = true;
@@ -102,15 +125,7 @@ bool UWorldPartitionRuntimeSpatialHash::GenerateNavigationData()
 			UE_LOG(LogWorldPartitionRuntimeSpatialHashNav, Verbose, TEXT("%i) %s added."), ActorCount, *DataChunkActor->GetName());
 		});
 
-	// Destroy all invalid navigation data chunk actors
-	for (TActorIterator<ANavigationDataChunkActor> It(GetWorld()); It; ++It)
-	{
-		if (!ValidNavigationDataChunkActors.Contains(*It))
-		{
-			GetWorld()->DestroyActor(*It);
-		}
-	}
-
+	UE_LOG(LogWorldPartitionRuntimeSpatialHashNav, Log, TEXT("GenerateNavigationData finished (%i actors added)."), ActorCount);
 	return true;
 }
 
diff --git a/Engine/Source/Runtime/Engine/Private/WorldPartition/WorldPartition.cpp b/Engine/Source/Runtime/Engine/Private/WorldPartition/WorldPartition.cpp
--- a/Engine/Source/Runtime/Engine/Private/WorldPartition/WorldPartition.cpp
+++ b/Engine/Source/Runtime/Engine/Private/WorldPartition/WorldPartition.cpp
@@ -1318,9 +1318,9 @@ void UWorldPartition::GenerateHLOD(ISourceControlHelper* SourceControlHelper, bo
 	RuntimeHash->GenerateHLOD(SourceControlHelper, bCreateActorsOnly);
 }
 
-void UWorldPartition::GenerateNavigationData()
+void UWorldPartition::GenerateNavigationData(const FBox& LoadedBounds)
 {
-	RuntimeHash->GenerateNavigationData();
+	RuntimeHash->GenerateNavigationData(LoadedBounds);
 }
 
 const UActorDescContainer* UWorldPartition::RegisterActorDescContainer(FName PackageName)
diff --git a/Engine/Source/Runtime/Engine/Public/WorldPartition/WorldPartition.h b/Engine/Source/Runtime/Engine/Public/WorldPartition/WorldPartition.h
--- a/Engine/Source/Runtime/Engine/Public/WorldPartition/WorldPartition.h
+++ b/Engine/Source/Runtime/Engine/Public/WorldPartition/WorldPartition.h
@@ -139,7 +139,7 @@ public:
 	FBox GetWorldBounds() const;
 	FBox GetEditorWorldBounds() const;
 	void GenerateHLOD(ISourceControlHelper* SourceControlHelper, bool bCreateActorsOnly);
-	void GenerateNavigationData();
+	void GenerateNavigationData(const FBox& LoadedBounds);
 
 	const UActorDescContainer* RegisterActorDescContainer(FName PackageName);
 
diff --git a/Engine/Source/Runtime/Engine/Public/WorldPartition/WorldPartitionRuntimeHash.h b/Engine/Source/Runtime/Engine/Public/WorldPartition/WorldPartitionRuntimeHash.h
--- a/Engine/Source/Runtime/Engine/Public/WorldPartition/WorldPartitionRuntimeHash.h
+++ b/Engine/Source/Runtime/Engine/Public/WorldPartition/WorldPartitionRuntimeHash.h
@@ -39,7 +39,7 @@ class ENGINE_API UWorldPartitionRuntimeHash : public UObject
 	virtual bool FinalizeGeneratorPackageForCook(const TArray<ICookPackageSplitter::FGeneratedPackageForPreSave>& InGeneratedPackages) { return false; }
 	virtual void FlushStreaming() {}
 	virtual bool GenerateHLOD(ISourceControlHelper* SourceControlHelper, bool bCreateActorsOnly) { return false; }
-	virtual bool GenerateNavigationData() { return false; }
+	virtual bool GenerateNavigationData(const FBox& LoadedBounds) { return false; }
 	virtual FName GetActorRuntimeGrid(const AActor* Actor) const { return NAME_None; }
 	virtual void DrawPreview() const {}
 
diff --git a/Engine/Source/Runtime/Engine/Public/WorldPartition/WorldPartitionRuntimeSpatialHash.h b/Engine/Source/Runtime/Engine/Public/WorldPartition/WorldPartitionRuntimeSpatialHash.h
--- a/Engine/Source/Runtime/Engine/Public/WorldPartition/WorldPartitionRuntimeSpatialHash.h
+++ b/Engine/Source/Runtime/Engine/Public/WorldPartition/WorldPartitionRuntimeSpatialHash.h
@@ -233,7 +233,7 @@ public:
 	virtual bool FinalizeGeneratorPackageForCook(const TArray<ICookPackageSplitter::FGeneratedPackageForPreSave>& InGeneratedPackages) override;
 	virtual void FlushStreaming() override;
 	virtual bool GenerateHLOD(ISourceControlHelper* SourceControlHelper, bool bCreateActorsOnly) override;
-	virtual bool GenerateNavigationData() override;
+	virtual bool GenerateNavigationData(const FBox& LoadedBounds) override;
 	virtual FName GetActorRuntimeGrid(const AActor* Actor) const override;
 	virtual void DrawPreview() const override;
 
diff --git a/Engine/Source/Runtime/NavigationSystem/Private/NavMesh/RecastNavMesh.cpp b/Engine/Source/Runtime/NavigationSystem/Private/NavMesh/RecastNavMesh.cpp
--- a/Engine/Source/Runtime/NavigationSystem/Private/NavMesh/RecastNavMesh.cpp
+++ b/Engine/Source/Runtime/NavigationSystem/Private/NavMesh/RecastNavMesh.cpp
@@ -2911,6 +2911,61 @@ void ARecastNavMesh::RebuildTile(const TArray<FIntPoint>& Tiles)
 	}
 }
 
+void ARecastNavMesh::DirtyTilesInBounds(const FBox& Bounds)
+{
+	if (HasValidNavmesh() == false)
+	{
+		return;
+	}
+	
+	FRecastNavMeshGenerator* MyGenerator = static_cast<FRecastNavMeshGenerator*>(GetGenerator());
+	if (MyGenerator)
+	{
+		MyGenerator->DiscardCurrentBuildingTasks();
+	}
+	
+	const dtNavMesh* DetourNavMesh = RecastNavMeshImpl->GetRecastMesh();
+
+	// Remove all tiles
+	const int32 TileCount = GetNavMeshTilesCount();
+	TArray<FIntPoint> TilesToRemove;
+	TilesToRemove.Reserve(TileCount);
+	for (int32 TileIndex = 0; TileIndex < TileCount; ++TileIndex)
+	{
+		const dtMeshTile* Tile = DetourNavMesh->getTile(TileIndex);
+		dtMeshHeader* Header = Tile != nullptr ? Tile->header : nullptr;
+		if (Header)
+		{
+			TilesToRemove.Add(FIntPoint(Header->x, Header->y));
+		}
+	}
+	RemoveTiles(TilesToRemove);
+
+	const UNavigationSystemV1* NavSys = FNavigationSystem::GetCurrent<UNavigationSystemV1>(GetWorld());
+	const FBox OverlappingBounds = Bounds.Overlap(NavSys->GetWorldBounds());
+
+	if (OverlappingBounds.IsValid)
+	{
+		// Add tiles within the overlapping bounds
+		TArray<FIntPoint> Points;
+		const FVector RcNavMeshOrigin = Unreal2RecastPoint(NavMeshOriginOffset);
+		const float RcTileSize = FMath::TruncToInt(TileSizeUU / CellSize);
+		const float TileSizeInWorldUnits = RcTileSize * CellSize;
+		const FRcTileBox TileBox(OverlappingBounds, RcNavMeshOrigin, TileSizeInWorldUnits);
+
+		UE_LOG(LogNavigation, VeryVerbose, TEXT("RebuildTilesFromBounds %i tiles: (%i,%i) to (%i,%i)"), (TileBox.XMax-TileBox.XMin)*(TileBox.YMax-TileBox.YMin), TileBox.XMin, TileBox.YMin, TileBox.XMax, TileBox.YMax);
+
+		for (int32 TileY = TileBox.YMin; TileY <= TileBox.YMax; ++TileY)
+		{
+			for (int32 TileX = TileBox.XMin; TileX <= TileBox.XMax; ++TileX)
+			{
+				Points.Add(FIntPoint(TileX, TileY));
+			}
+		}
+		RebuildTile(Points);
+	}
+}
+
 #if RECAST_INTERNAL_DEBUG_DATA
 const TMap<FIntPoint, struct FRecastInternalDebugData>* ARecastNavMesh::GetDebugDataMap() const
 {
diff --git a/Engine/Source/Runtime/NavigationSystem/Private/NavMesh/RecastNavMeshDataChunk.cpp b/Engine/Source/Runtime/NavigationSystem/Private/NavMesh/RecastNavMeshDataChunk.cpp
--- a/Engine/Source/Runtime/NavigationSystem/Private/NavMesh/RecastNavMeshDataChunk.cpp
+++ b/Engine/Source/Runtime/NavigationSystem/Private/NavMesh/RecastNavMeshDataChunk.cpp
@@ -254,7 +254,7 @@ TArray<uint32> URecastNavMeshDataChunk::AttachTiles(FPImplRecastNavMesh& NavMesh
 	}
 #endif// WITH_RECAST
 
-	UE_LOG(LogNavigation, Log, TEXT("Attached %d tiles to NavMesh - %s"), Result.Num(), *NavigationDataName.ToString());
+	UE_LOG(LogNavigation, Verbose, TEXT("Attached %d tiles to NavMesh - %s"), Result.Num(), *NavigationDataName.ToString());
 	return Result;
 }
 
@@ -326,14 +326,14 @@ TArray<uint32> URecastNavMeshDataChunk::DetachTiles(FPImplRecastNavMesh& NavMesh
 	}
 #endif// WITH_RECAST
 
-	UE_LOG(LogNavigation, Log, TEXT("Detached %d tiles from NavMesh - %s"), Result.Num(), *NavigationDataName.ToString());
+	UE_LOG(LogNavigation, Verbose, TEXT("Detached %d tiles from NavMesh - %s"), Result.Num(), *NavigationDataName.ToString());
 	return Result;
 }
 
 void URecastNavMeshDataChunk::MoveTiles(FPImplRecastNavMesh& NavMeshImpl, const FIntPoint& Offset, const float RotationDeg, const FVector2D& RotationCenter)
 {
 #if WITH_RECAST	
-	UE_LOG(LogNavigation, Log, TEXT("%s Moving %i tiles on navmesh %s."), ANSI_TO_TCHAR(__FUNCTION__), Tiles.Num(), *NavigationDataName.ToString());
+	UE_LOG(LogNavigation, Verbose, TEXT("%s Moving %i tiles on navmesh %s."), ANSI_TO_TCHAR(__FUNCTION__), Tiles.Num(), *NavigationDataName.ToString());
 
 	dtNavMesh* NavMesh = NavMeshImpl.DetourNavMesh;
 	if (NavMesh != nullptr)
@@ -369,12 +369,12 @@ void URecastNavMeshDataChunk::MoveTiles(FPImplRecastNavMesh& NavMeshImpl, const
 				const int OffsetWithRotX = Offset.X + DeltaX;
 				const int OffsetWithRotY = Offset.Y + DeltaY;
 				const bool bSuccess = dtTransformTileData(TileData.TileRawData->RawData, TileData.TileDataSize, OffsetWithRotX, OffsetWithRotY, TileWidth, TileHeight, RotationDeg);
-				UE_CLOG(bSuccess, LogNavigation, Log, TEXT("   Moved tile from (%i,%i) to (%i,%i)."), TileData.OriginalX, TileData.OriginalY, (TileData.OriginalX + OffsetWithRotX), (TileData.OriginalY + OffsetWithRotY));
+				UE_CLOG(bSuccess, LogNavigation, Verbose, TEXT("   Moved tile from (%i,%i) to (%i,%i)."), TileData.OriginalX, TileData.OriginalY, (TileData.OriginalX + OffsetWithRotX), (TileData.OriginalY + OffsetWithRotY));
 			}
 		}
 	}
 
-	UE_LOG(LogNavigation, Log, TEXT("%s Moving done."), ANSI_TO_TCHAR(__FUNCTION__));
+	UE_LOG(LogNavigation, Verbose, TEXT("%s Moving done."), ANSI_TO_TCHAR(__FUNCTION__));
 #endif// WITH_RECAST
 }
 
diff --git a/Engine/Source/Runtime/NavigationSystem/Private/NavMesh/RecastNavMeshGenerator.cpp b/Engine/Source/Runtime/NavigationSystem/Private/NavMesh/RecastNavMeshGenerator.cpp
--- a/Engine/Source/Runtime/NavigationSystem/Private/NavMesh/RecastNavMeshGenerator.cpp
+++ b/Engine/Source/Runtime/NavigationSystem/Private/NavMesh/RecastNavMeshGenerator.cpp
@@ -5166,7 +5166,7 @@ TArray<uint32> FRecastNavMeshGenerator::RemoveTileLayers(const int32 TileX, cons
 				dtPolyRef TileRef = DetourMesh->getTileRef(Tiles[i]);
 
 				NumActiveTiles--;
-				UE_LOG(LogNavigation, Log, TEXT("%s> Tile (%d,%d:%d), removing TileRef: 0x%X (active:%d)"),
+				UE_LOG(LogNavigation, Verbose, TEXT("%s> Tile (%d,%d:%d), removing TileRef: 0x%X (active:%d)"),
 					*DestNavMesh->GetName(), TileX, TileY, LayerIndex, TileRef, NumActiveTiles);
 
 				DetourMesh->removeTile(TileRef, nullptr, nullptr);
@@ -5235,7 +5235,7 @@ void FRecastNavMeshGenerator::AddGeneratedTileLayer(int32 LayerIndex, FRecastTil
 		if (OldTileRef)
 		{
 			NumActiveTiles--;
-			UE_LOG(LogNavigation, Log, TEXT("%s> Tile (%d,%d:%d), removing TileRef: 0x%X (active:%d)"),
+			UE_LOG(LogNavigation, Verbose, TEXT("%s> Tile (%d,%d:%d), removing TileRef: 0x%X (active:%d)"),
 				*DestNavMesh->GetName(), TileX, TileY, LayerIndex, OldTileRef, NumActiveTiles);
 
 			DetourMesh->removeTile(OldTileRef, nullptr, nullptr);
@@ -5280,7 +5280,7 @@ void FRecastNavMeshGenerator::AddGeneratedTileLayer(int32 LayerIndex, FRecastTil
 				OutResultTileIndices.AddUnique(DetourMesh->decodePolyIdTile(ResultTileRef));
 				NumActiveTiles++;
 
-				UE_LOG(LogNavigation, Log, TEXT("%s> Tile (%d,%d:%d), added TileRef: 0x%X (active:%d)"),
+				UE_LOG(LogNavigation, Verbose, TEXT("%s> Tile (%d,%d:%d), added TileRef: 0x%X (active:%d)"),
 					*DestNavMesh->GetName(), TileX, TileY, LayerIndex, ResultTileRef, NumActiveTiles);
 
 				{
@@ -5815,7 +5815,11 @@ TArray<uint32> FRecastNavMeshGenerator::ProcessTileTasksAsync(const int32 NumTas
 				{
 					RunningElement.AsyncTask->StartSynchronousTask();
 				}
-			
+
+				static int32 Count = 1;
+				UE_LOG(LogNavigationDataBuild, VeryVerbose, TEXT("   Tile generation task #%i)"), Count);
+				Count++;
+				
 				RunningDirtyTiles.Add(RunningElement);
 			}
 			else if (!bGameStaticNavMesh)
diff --git a/Engine/Source/Runtime/NavigationSystem/Private/NavigationSystem.cpp b/Engine/Source/Runtime/NavigationSystem/Private/NavigationSystem.cpp
--- a/Engine/Source/Runtime/NavigationSystem/Private/NavigationSystem.cpp
+++ b/Engine/Source/Runtime/NavigationSystem/Private/NavigationSystem.cpp
@@ -375,6 +375,7 @@ UNavigationSystemV1::UNavigationSystemV1(const FObjectInitializer& ObjectInitial
 	, bSkipAgentHeightCheckWhenPickingNavData(false)
 	, DirtyAreaWarningSizeThreshold(-1.0f)
 	, GatheringNavModifiersWarningLimitTime(-1.0f)
+	, BuildBounds(EForceInit::ForceInit)
 	, OperationMode(FNavigationSystemRunMode::InvalidMode)
 	, bAbortAsyncQueriesRequested(false)
 	, NavBuildingLockFlags(0)
@@ -1831,6 +1832,16 @@ ANavigationData* UNavigationSystemV1::GetNavDataForAgentName(const FName AgentNa
 	return Result;
 }
 
+FBox UNavigationSystemV1::GetNavigableWorldBounds() const
+{
+	return GetWorldBounds();
+}
+
+void UNavigationSystemV1::SetBuildBounds(const FBox& Bounds)
+{
+	BuildBounds = Bounds;
+}
+
 bool UNavigationSystemV1::ContainsNavData(const FBox& Bounds) const
 {
 	for (const ANavigationData* NavData : NavDataSet)
@@ -1843,6 +1854,21 @@ bool UNavigationSystemV1::ContainsNavData(const FBox& Bounds) const
 	return false;
 }
 
+FBox UNavigationSystemV1::ComputeNavDataBounds() const
+{
+	TRACE_CPUPROFILER_EVENT_SCOPE(UNavigationSystemV1::ComputeNavDataBounds);
+	
+	FBox Bounds(ForceInit);
+	for (const ANavigationData* NavData : NavDataSet)
+	{
+		if (NavData)
+		{
+			Bounds += NavData->GetBounds();
+		}
+	}
+	return Bounds;
+}
+
 void UNavigationSystemV1::AddNavigationDataChunk(ANavigationDataChunkActor& DataChunkActor)
 {
 	for (ANavigationData* NavData : NavDataSet)
@@ -3316,6 +3342,8 @@ void UNavigationSystemV1::GatherNavigationBounds()
 
 void UNavigationSystemV1::Build()
 {
+	TRACE_CPUPROFILER_EVENT_SCOPE(UNavigationSystemV1::Build);
+	
 	UE_LOG(LogNavigationDataBuild, Display, TEXT("UNavigationSystemV1::Build started..."));
 #if PHYSICS_INTERFACE_PHYSX
 	UE_LOG(LogNavigationDataBuild, Display, TEXT("   Building navigation data using PHYSICS_INTERFACE_PHYSX."));
@@ -3364,6 +3392,12 @@ void UNavigationSystemV1::Build()
 		UpdateInvokers();
 	}
 
+	if (BuildBounds.IsValid)
+	{
+		// Prepare to build tiles overlapping the bounds
+		DirtyTilesInBuildBounds();
+	}
+
 	// and now iterate through all registered and just start building them
 	RebuildAll();
 
@@ -3692,9 +3726,11 @@ void UNavigationSystemV1::RebuildAll(bool bIsLoadTime)
 	{
 		ANavigationData* NavData = NavDataSet[NavDataIndex];
 				
-		if (NavData && (!bIsLoadTime || NavData->NeedsRebuildOnLoad()) && (!bIsInGame || NavData->SupportsRuntimeGeneration()))
+		if (NavData && (!bIsLoadTime || NavData->NeedsRebuildOnLoad()) && (!bIsInGame || NavData->SupportsRuntimeGeneration()) && (BuildBounds.IsValid == 0))
 		{
-			UE_LOG(LogNavigationDataBuild, Display, TEXT("   Building NavData:  %s."), *NavData->GetConfig().GetDescription());
+			UE_LOG(LogNavigationDataBuild, Display, TEXT("   RebuildAll building NavData:  %s."), *NavData->GetConfig().GetDescription());
+			UE_LOG(LogNavigationDataBuild, Verbose, TEXT("   RebuildAll bIsLoadTime=%s, NavData->NeedsRebuildOnLoad()=%s, bIsInGame=%s, NavData->SupportsRuntimeGeneration()=%s, BuildBounds.IsValid=%s"),
+				*LexToString(bIsLoadTime), *LexToString(NavData->NeedsRebuildOnLoad()), *LexToString(bIsInGame), *LexToString(NavData->SupportsRuntimeGeneration()), *LexToString(BuildBounds.IsValid));
 
 #if	WITH_EDITOR
 			NavData->SetIsBuildingOnLoad(bIsLoadTime);
@@ -4541,6 +4577,17 @@ void UNavigationSystemV1::UpdateInvokers()
 	}
 }
 
+void UNavigationSystemV1::DirtyTilesInBuildBounds()
+{
+#if WITH_RECAST
+	UE_VLOG(this, LogNavigation, Log, TEXT("SetupTilesFromBuildBounds"));
+	for (TActorIterator<ARecastNavMesh> It(GetWorld()); It; ++It)
+	{
+		It->DirtyTilesInBounds(BuildBounds);
+	}
+#endif // WITH_RECAST
+}
+
 void UNavigationSystemV1::RegisterNavigationInvoker(AActor* Invoker, float TileGenerationRadius, float TileRemovalRadius)
 {
 	if (Invoker != nullptr)
diff --git a/Engine/Source/Runtime/NavigationSystem/Public/NavMesh/RecastNavMesh.h b/Engine/Source/Runtime/NavigationSystem/Public/NavMesh/RecastNavMesh.h
--- a/Engine/Source/Runtime/NavigationSystem/Public/NavMesh/RecastNavMesh.h
+++ b/Engine/Source/Runtime/NavigationSystem/Public/NavMesh/RecastNavMesh.h
@@ -1077,6 +1077,7 @@ public:
 	virtual void UpdateActiveTiles(const TArray<FNavigationInvokerRaw>& InvokerLocations);
 	virtual void RemoveTiles(const TArray<FIntPoint>& Tiles);
 	void RebuildTile(const TArray<FIntPoint>& Tiles);
+	void DirtyTilesInBounds(const FBox& Bounds);
 
 #if RECAST_INTERNAL_DEBUG_DATA
 	const TMap<FIntPoint, struct FRecastInternalDebugData>* GetDebugDataMap() const;
diff --git a/Engine/Source/Runtime/NavigationSystem/Public/NavigationSystem.h b/Engine/Source/Runtime/NavigationSystem/Public/NavigationSystem.h
--- a/Engine/Source/Runtime/NavigationSystem/Public/NavigationSystem.h
+++ b/Engine/Source/Runtime/NavigationSystem/Public/NavigationSystem.h
@@ -309,6 +309,9 @@ protected:
 	FNavAgentSelector SupportedAgentsMask;
 
 public:
+	/** Bounds of tiles to be built */
+	UPROPERTY(Transient)
+	FBox BuildBounds;
 
 	UPROPERTY(Transient)
 	TArray<TObjectPtr<ANavigationData>> NavDataSet;
@@ -345,6 +348,8 @@ private:
 	float NextInvokersUpdateTime;
 	void UpdateInvokers();
 
+	void DirtyTilesInBuildBounds();
+
 public:
 	//----------------------------------------------------------------------//
 	// Blueprint functions
@@ -591,7 +596,13 @@ public:
 	virtual INavigationDataInterface* GetMainNavData() const override { return Cast<INavigationDataInterface>(GetDefaultNavDataInstance()); }
 	ANavigationData& GetMainNavDataChecked() const { check(MainNavData); return *MainNavData; }
 
+	/** Set limiting bounds to be used when building navigation data. */
+	virtual void SetBuildBounds(const FBox& Bounds) override;
+
+	virtual FBox GetNavigableWorldBounds() const override;
+
 	virtual bool ContainsNavData(const FBox& Bounds) const override;
+	virtual FBox ComputeNavDataBounds() const override;
 	virtual void AddNavigationDataChunk(class ANavigationDataChunkActor& DataChunkActor) override;
 	virtual void RemoveNavigationDataChunk(class ANavigationDataChunkActor& DataChunkActor) override;
 	virtual void FillNavigationDataChunkActor(const FBox& QueryBounds, class ANavigationDataChunkActor& DataChunkActor, FBox& OutTilesBounds) override;
