Fix Remove Unused Nodes not working properly. It would simply clean up the entire graph each time except for the nodes directly attached to the material master node.

#rnx

diff --git a/Engine/Source/Editor/UnrealEd/Private/MaterialGraph.cpp b/Engine/Source/Editor/UnrealEd/Private/MaterialGraph.cpp
--- a/Engine/Source/Editor/UnrealEd/Private/MaterialGraph.cpp
+++ b/Engine/Source/Editor/UnrealEd/Private/MaterialGraph.cpp
@@ -658,11 +658,10 @@ void UMaterialGraph::GetUnusedExpressions(TArray<UEdGraphNode*>& UnusedNodes) co
 				UsedNodes.Add(GraphNode);
 
 				// Iterate over the expression's inputs and add them to the pending stack.
-				for (UEdGraphPin* Pin : RootNode->Pins)
+				for (UEdGraphPin* Pin : GraphNode->Pins)
 				{
 					if (Pin->Direction == EGPD_Input &&
 						Pin->PinType.PinCategory != UMaterialGraphSchema::PC_Exec &&
-						MaterialInputs[Pin->SourceIndex].IsVisiblePin(Material) &&
 						Pin->LinkedTo.Num() > 0 &&
 						Pin->LinkedTo[0])
 					{
