misc non-unity compilation fixes

diff --git a/Engine/Plugins/FX/Niagara/Source/Niagara/Private/DataInterface/NiagaraDataInterfaceStaticMesh.cpp b/Engine/Plugins/FX/Niagara/Source/Niagara/Private/DataInterface/NiagaraDataInterfaceStaticMesh.cpp
--- a/Engine/Plugins/FX/Niagara/Source/Niagara/Private/DataInterface/NiagaraDataInterfaceStaticMesh.cpp
+++ b/Engine/Plugins/FX/Niagara/Source/Niagara/Private/DataInterface/NiagaraDataInterfaceStaticMesh.cpp
@@ -15,6 +15,11 @@
 #include "GPUSkinCache.h"
 #include "ShaderCompilerCore.h"
 
+#if WITH_EDITOR
+#include "Subsystems/ImportSubsystem.h"
+#include "Editor.h"
+#endif
+
 #define LOCTEXT_NAMESPACE "NiagaraDataInterfaceStaticMesh"
 
 //-OPT: We should be able to share data between data interfaces where possible
diff --git a/Engine/Source/Editor/DetailCustomizations/Private/CustomPrimitiveDataCustomization.h b/Engine/Source/Editor/DetailCustomizations/Private/CustomPrimitiveDataCustomization.h
--- a/Engine/Source/Editor/DetailCustomizations/Private/CustomPrimitiveDataCustomization.h
+++ b/Engine/Source/Editor/DetailCustomizations/Private/CustomPrimitiveDataCustomization.h
@@ -5,6 +5,7 @@
 #include "CoreMinimal.h"
 #include "IPropertyTypeCustomization.h"
 #include "IDetailChildrenBuilder.h"
+#include "Materials/Material.h"
 #include "Materials/MaterialLayersFunctions.h"
 
 class UMaterialExpressionVectorParameter;
diff --git a/Engine/Source/Editor/UnrealEd/Private/WorldPartition/WorldPartitionNavigationDataBuilder.cpp b/Engine/Source/Editor/UnrealEd/Private/WorldPartition/WorldPartitionNavigationDataBuilder.cpp
--- a/Engine/Source/Editor/UnrealEd/Private/WorldPartition/WorldPartitionNavigationDataBuilder.cpp
+++ b/Engine/Source/Editor/UnrealEd/Private/WorldPartition/WorldPartitionNavigationDataBuilder.cpp
@@ -4,7 +4,7 @@
 
 #include "CoreMinimal.h"
 #include "EngineUtils.h"
-#include "Editor.h"
+#include "FileHelpers.h"
 #include "StaticMeshCompiler.h"
 #include "Logging/LogMacros.h"
 
@@ -97,7 +97,7 @@ bool UWorldPartitionNavigationDataBuilder::RunInternal(UWorld* World, const FBox
 		if (!PackageHelper.Delete(PackagesToDelete))
 		{
 			UE_LOG(LogWorldPartitionNavigationDataBuilder, Error, TEXT("Error deleting packages."));
-			return 1;
+			return true;
 		}
 	}
 
@@ -124,7 +124,7 @@ bool UWorldPartitionNavigationDataBuilder::RunInternal(UWorld* World, const FBox
 						if (!IPlatformFile::GetPlatformPhysical().SetReadOnly(*PackageFilename, /*bNewReadOnlyValue*/false))
 						{
 							UE_LOG(LogWorldPartitionNavigationDataBuilder, Error, TEXT("Error setting %s writable"), *PackageFilename);
-							return 1;
+							return true;
 						}
 					}
 				}
@@ -143,7 +143,7 @@ bool UWorldPartitionNavigationDataBuilder::RunInternal(UWorld* World, const FBox
 				if (!UPackage::SavePackage(Package, nullptr, RF_Standalone, *PackageFileName, GError, nullptr, false, true, SAVE_Async))
 				{
 					UE_LOG(LogWorldPartitionNavigationDataBuilder, Error, TEXT("Error saving package %s."), *Package->GetName());
-					return 1;
+					return true;
 				}
 			}
 		}
@@ -158,7 +158,7 @@ bool UWorldPartitionNavigationDataBuilder::RunInternal(UWorld* World, const FBox
 				if (!PackageHelper.AddToSourceControl(Package))
 				{
 					UE_LOG(LogWorldPartitionNavigationDataBuilder, Error, TEXT("Error adding package %s to source control."), *Package->GetName());
-					return 1;
+					return true;
 				}
 			}
 		}
diff --git a/Engine/Source/Runtime/AutomationTest/Private/AutomationTestExcludelist.cpp b/Engine/Source/Runtime/AutomationTest/Private/AutomationTestExcludelist.cpp
--- a/Engine/Source/Runtime/AutomationTest/Private/AutomationTestExcludelist.cpp
+++ b/Engine/Source/Runtime/AutomationTest/Private/AutomationTestExcludelist.cpp
@@ -8,6 +8,7 @@
 	#include "SourceControlOperations.h"
 	#include "ISourceControlProvider.h"
 	#include "ISourceControlModule.h"
+	#include "GenericPlatform/GenericPlatformFile.h"
 #endif
 
 DEFINE_LOG_CATEGORY_STATIC(LogAutomationTestExcludelist, Log, All);
diff --git a/Engine/Source/Runtime/Experimental/Chaos/Private/RewindData.cpp b/Engine/Source/Runtime/Experimental/Chaos/Private/RewindData.cpp
--- a/Engine/Source/Runtime/Experimental/Chaos/Private/RewindData.cpp
+++ b/Engine/Source/Runtime/Experimental/Chaos/Private/RewindData.cpp
@@ -6,6 +6,7 @@
 #include "HAL/IConsoleManager.h"
 #include "Chaos/PBDJointConstraints.h"
 #include "Misc/FileHelper.h"
+#include "Misc/Paths.h"
 
 namespace Chaos
 {
