BSP: Fix placement new for FPoly::Vertices using outdated vector type firing check

diff --git a/Engine/Plugins/Editor/GeometryMode/Source/GeometryMode/Private/GeometryModifiers.cpp b/Engine/Plugins/Editor/GeometryMode/Source/GeometryMode/Private/GeometryModifiers.cpp
--- a/Engine/Plugins/Editor/GeometryMode/Source/GeometryMode/Private/GeometryModifiers.cpp
+++ b/Engine/Plugins/Editor/GeometryMode/Source/GeometryMode/Private/GeometryModifiers.cpp
@@ -1383,7 +1383,7 @@ void UGeomModifier_Pen::Apply()
 
 		for( int32 v = 0 ; v < ShapeVertices.Num() ; ++v )
 		{
-			new(Poly.Vertices) FVector( ShapeVertices[v] - BaseLocation );
+			new(Poly.Vertices) FVector3f( ShapeVertices[v] - BaseLocation );
 		}
 
 		if( Poly.Finalize( ResultingBrush, 1 ) == 0 )
@@ -2609,7 +2609,7 @@ bool UGeomModifier_Create::OnApply()
 			{
 				FGeomVertex* gv = Verts[x];
 
-				new(NewPoly->Vertices) FVector(*gv);
+				new(NewPoly->Vertices) FVector3f(*gv);
 			}
 
 			NewPoly->Normal = FVector::ZeroVector;
@@ -3326,9 +3326,9 @@ bool UGeomModifier_Turn::OnApply()
 				NewPoly = new( gp->GetParentObject()->GetActualBrush()->Brush->Polys->Element )FPoly();
 
 				NewPoly->Init();
-				new(NewPoly->Vertices) FVector(Quad[2]);
-				new(NewPoly->Vertices) FVector(Quad[1]);
-				new(NewPoly->Vertices) FVector(Quad[3]);
+				new(NewPoly->Vertices) FVector3f(Quad[2]);
+				new(NewPoly->Vertices) FVector3f(Quad[1]);
+				new(NewPoly->Vertices) FVector3f(Quad[3]);
 
 				NewPoly->Base = SavePoly0.Base;
 				NewPoly->Material = SavePoly0.Material;
@@ -3341,9 +3341,9 @@ bool UGeomModifier_Turn::OnApply()
 				NewPoly = new( gp->GetParentObject()->GetActualBrush()->Brush->Polys->Element )FPoly();
 
 				NewPoly->Init();
-				new(NewPoly->Vertices) FVector(Quad[3]);
-				new(NewPoly->Vertices) FVector(Quad[1]);
-				new(NewPoly->Vertices) FVector(Quad[0]);
+				new(NewPoly->Vertices) FVector3f(Quad[3]);
+				new(NewPoly->Vertices) FVector3f(Quad[1]);
+				new(NewPoly->Vertices) FVector3f(Quad[0]);
 
 				NewPoly->Base = SavePoly1.Base;
 				NewPoly->Material = SavePoly1.Material;
diff --git a/Engine/Source/Editor/UnrealEd/Private/Factories/EditorFactories.cpp b/Engine/Source/Editor/UnrealEd/Private/Factories/EditorFactories.cpp
--- a/Engine/Source/Editor/UnrealEd/Private/Factories/EditorFactories.cpp
+++ b/Engine/Source/Editor/UnrealEd/Private/Factories/EditorFactories.cpp
@@ -1687,7 +1687,7 @@ UObject* UPolysFactory::FactoryCreateText
 					{
 						int32 iPoint = FMath::Abs(FCString::Atoi(*ExtraLine));
 						if( iPoint>0 && iPoint<=NumPoints )
-							new(NewPoly.Vertices) FVector(PointPool[iPoint-1]);
+							new(NewPoly.Vertices) FVector3f(PointPool[iPoint-1]);
 						else UE_LOG(LogEditorFactories, Warning, TEXT("DXF: Invalid point index %i/%i"), iPoint, NumPoints );
 					}
 				}
@@ -1722,9 +1722,9 @@ UObject* UPolysFactory::FactoryCreateText
 				else if( FCString::Strstr(Str,*FaceText) )
 				{
 					Poly.Init();
-					new(Poly.Vertices)FVector(PointPool[FCString::Atoi(FCString::Strstr(Str,TEXT("A:"))+2)]);
-					new(Poly.Vertices)FVector(PointPool[FCString::Atoi(FCString::Strstr(Str,TEXT("B:"))+2)]);
-					new(Poly.Vertices)FVector(PointPool[FCString::Atoi(FCString::Strstr(Str,TEXT("C:"))+2)]);
+					new(Poly.Vertices)FVector3f(PointPool[FCString::Atoi(FCString::Strstr(Str,TEXT("A:"))+2)]);
+					new(Poly.Vertices)FVector3f(PointPool[FCString::Atoi(FCString::Strstr(Str,TEXT("B:"))+2)]);
+					new(Poly.Vertices)FVector3f(PointPool[FCString::Atoi(FCString::Strstr(Str,TEXT("C:"))+2)]);
 					Poly.Base = Poly.Vertices[0];
 					Poly.Finalize(nullptr,0);
 					new(Polys->Element)FPoly(Poly);
diff --git a/Engine/Source/Editor/UnrealEd/Private/UnrealEdSrv.cpp b/Engine/Source/Editor/UnrealEd/Private/UnrealEdSrv.cpp
--- a/Engine/Source/Editor/UnrealEd/Private/UnrealEdSrv.cpp
+++ b/Engine/Source/Editor/UnrealEd/Private/UnrealEdSrv.cpp
@@ -1707,9 +1707,9 @@ TArray<FPoly*> GetSelectedPolygons()
 							Polygon->Init();
 							Polygon->PolyFlags = PF_DefaultFlags;
 
-							new(Polygon->Vertices) FVector(ActorToWorld.TransformPosition( PositionVertexBuffer.VertexPosition(Idx2) ));
-							new(Polygon->Vertices) FVector(ActorToWorld.TransformPosition( PositionVertexBuffer.VertexPosition(Idx1) ));
-							new(Polygon->Vertices) FVector(ActorToWorld.TransformPosition( PositionVertexBuffer.VertexPosition(Idx0) ));
+							new(Polygon->Vertices) FVector3f(ActorToWorld.TransformPosition( PositionVertexBuffer.VertexPosition(Idx2) ));
+							new(Polygon->Vertices) FVector3f(ActorToWorld.TransformPosition( PositionVertexBuffer.VertexPosition(Idx1) ));
+							new(Polygon->Vertices) FVector3f(ActorToWorld.TransformPosition( PositionVertexBuffer.VertexPosition(Idx0) ));
 
 							Polygon->CalcNormal(1);
 							Polygon->Fix();
@@ -1726,9 +1726,9 @@ TArray<FPoly*> GetSelectedPolygons()
 							Polygon->Init();
 							Polygon->PolyFlags = PF_DefaultFlags;
 
-							new(Polygon->Vertices) FVector(ActorToWorld.TransformPosition( PositionVertexBuffer.VertexPosition(Idx2) ));
-							new(Polygon->Vertices) FVector(ActorToWorld.TransformPosition( PositionVertexBuffer.VertexPosition(Idx0) ));
-							new(Polygon->Vertices) FVector(ActorToWorld.TransformPosition( PositionVertexBuffer.VertexPosition(Idx1) ));
+							new(Polygon->Vertices) FVector3f(ActorToWorld.TransformPosition( PositionVertexBuffer.VertexPosition(Idx2) ));
+							new(Polygon->Vertices) FVector3f(ActorToWorld.TransformPosition( PositionVertexBuffer.VertexPosition(Idx0) ));
+							new(Polygon->Vertices) FVector3f(ActorToWorld.TransformPosition( PositionVertexBuffer.VertexPosition(Idx1) ));
 							Polygon->CalcNormal(1);
 							Polygon->Fix();
 							if( Polygon->Vertices.Num() > 2 )
@@ -1834,9 +1834,9 @@ FPoly* CreateHugeTrianglePolygonOnPlane( const FPlane* InPlane )
 	Triangle->Init();
 	Triangle->PolyFlags = PF_DefaultFlags;
 
-	new(Triangle->Vertices) FVector( V0 );
-	new(Triangle->Vertices) FVector( V2 );
-	new(Triangle->Vertices) FVector( V1 );
+	new(Triangle->Vertices) FVector3f( V0 );
+	new(Triangle->Vertices) FVector3f( V2 );
+	new(Triangle->Vertices) FVector3f( V1 );
 
 	Triangle->CalcNormal(1);
 	Triangle->Fix();
