- Fix a precision issues that could result in 24 fps anim sequences being imported as 120 fps.
- Fix precision issue in some of the sampling by changing some calculation and passing in an exact frames per second amount.

diff --git a/Engine/Source/Editor/UnrealEd/Private/SkeletalMeshEdit.cpp b/Engine/Source/Editor/UnrealEd/Private/SkeletalMeshEdit.cpp
--- a/Engine/Source/Editor/UnrealEd/Private/SkeletalMeshEdit.cpp
+++ b/Engine/Source/Editor/UnrealEd/Private/SkeletalMeshEdit.cpp
@@ -725,21 +725,21 @@ int32 GetAnimationCurveRate(FbxAnimCurve* CurrentCurve)
 			//    least common multiplier to get a sample rate that go through all keys.
 
 			SampleRate = 1;
-			float OldKeyTime = 0.0f;
+			double OldKeyTime = 0.0f;
 			TSet<int32> DeltaComputed;
 			//Reserve some space
 			DeltaComputed.Reserve(30);
-			const float KeyMultiplier = (1.0f / KINDA_SMALL_NUMBER);
+			const double KeyMultiplier = (1.0f / KINDA_SMALL_NUMBER);
 			//Find also the smallest delta time between keys
 			for (int32 KeyIndex = 0; KeyIndex < KeyCount; ++KeyIndex)
 			{
-				float KeyTime = (float)(CurrentCurve->KeyGet(KeyIndex).GetTime().GetSecondDouble());
+				double KeyTime = (CurrentCurve->KeyGet(KeyIndex).GetTime().GetSecondDouble());
 				//Collect the smallest delta time, there is no delta in case the first animation key time is negative
-				float Delta = (KeyTime < 0 && KeyIndex == 0) ? 0.0f : KeyTime - OldKeyTime;
+				double Delta = (KeyTime < 0 && KeyIndex == 0) ? 0.0 : KeyTime - OldKeyTime;
 				//use the fractional part of the delta to have the delta between 0.0f and 1.0f
 				Delta = FPlatformMath::Fractional(Delta);
 				int32 DeltaKey = FPlatformMath::RoundToInt(Delta*KeyMultiplier);
-				if (!FMath::IsNearlyZero(Delta, KINDA_SMALL_NUMBER) && !DeltaComputed.Contains(DeltaKey))
+				if (!FMath::IsNearlyZero((float)Delta, KINDA_SMALL_NUMBER) && !DeltaComputed.Contains(DeltaKey))
 				{
 					int32 ComputeSampleRate = GetTimeSampleRate(Delta);
 					DeltaComputed.Add(DeltaKey);
diff --git a/Engine/Source/Runtime/Engine/Private/Animation/AnimSequenceHelpers.cpp b/Engine/Source/Runtime/Engine/Private/Animation/AnimSequenceHelpers.cpp
--- a/Engine/Source/Runtime/Engine/Private/Animation/AnimSequenceHelpers.cpp
+++ b/Engine/Source/Runtime/Engine/Private/Animation/AnimSequenceHelpers.cpp
@@ -192,11 +192,12 @@ void BuildPoseFromModel(const UAnimDataModel* Model, FCompactPose& OutPose, cons
 
 	const int32 NumberOfKeys = Model->GetNumberOfKeys();
 	const float SequenceLength = Model->GetPlayLength();
+	const double FramesPerSecond = Model->GetFrameRate().AsDecimal();
 
 	// Generate keys to interpolate between
 	int32 KeyIndex1, KeyIndex2;
 	float Alpha;
-	FAnimationRuntime::GetKeyIndicesFromTime(KeyIndex1, KeyIndex2, Alpha, Time, NumberOfKeys, SequenceLength);
+	FAnimationRuntime::GetKeyIndicesFromTime(KeyIndex1, KeyIndex2, Alpha, Time, NumberOfKeys, SequenceLength, FramesPerSecond);
 
 	if (InterpolationType == EAnimInterpolationType::Step)
 	{
diff --git a/Engine/Source/Runtime/Engine/Private/Animation/AnimationRuntime.cpp b/Engine/Source/Runtime/Engine/Private/Animation/AnimationRuntime.cpp
--- a/Engine/Source/Runtime/Engine/Private/Animation/AnimationRuntime.cpp
+++ b/Engine/Source/Runtime/Engine/Private/Animation/AnimationRuntime.cpp
@@ -1389,7 +1389,7 @@ void FAnimationRuntime::ApplyWeightToTransform(const FBoneContainer& RequiredBon
 }
 
 /* from % from OutKeyIndex1, meaning (CurrentKeyIndex(double)-OutKeyIndex1)/(OutKeyIndex2-OutKeyIndex1) */
-void FAnimationRuntime::GetKeyIndicesFromTime(int32& OutKeyIndex1, int32& OutKeyIndex2, float& OutAlpha, const double Time, const int32 NumKeys, const double SequenceLength)
+void FAnimationRuntime::GetKeyIndicesFromTime(int32& OutKeyIndex1, int32& OutKeyIndex2, float& OutAlpha, const double Time, const int32 NumKeys, const double SequenceLength, double FramesPerSecond)
 {
 	// Check for 1-frame, before-first-frame and after-last-frame cases.
 	if (Time <= 0.0 || NumKeys == 1)
@@ -1409,8 +1409,14 @@ void FAnimationRuntime::GetKeyIndicesFromTime(int32& OutKeyIndex1, int32& OutKey
 		return;
 	}
 
-	const int32 NumFrames = NumKeys - 1;
-	const double KeyPos = (Time / SequenceLength) * (double)NumFrames;
+	// Calulate the frames per second if we didn't provide any.
+	if (FramesPerSecond <= 0.0)
+	{
+		const int32 NumFrames = NumKeys - 1;
+		FramesPerSecond = NumFrames / SequenceLength;
+	}
+
+	const double KeyPos = Time * FramesPerSecond;
 
 	// Find the integer part (ensuring within range) and that gives us the 'starting' key index.
 	const int32 KeyIndex1 = FMath::Clamp<int32>( FMath::FloorToInt(KeyPos), 0, NumKeys - 1 );  // @todo should be changed to FMath::TruncToInt
diff --git a/Engine/Source/Runtime/Engine/Public/AnimationRuntime.h b/Engine/Source/Runtime/Engine/Public/AnimationRuntime.h
--- a/Engine/Source/Runtime/Engine/Public/AnimationRuntime.h
+++ b/Engine/Source/Runtime/Engine/Public/AnimationRuntime.h
@@ -521,8 +521,11 @@ public:
 	 * Get Key Indices (start/end with alpha from start) with input parameter Time, NumKeys
 	 * from % from StartKeyIndex, meaning (CurrentKeyIndex(float)-StartKeyIndex)/(EndKeyIndex-StartKeyIndex)
 	 * by this Start-End, it will be between 0-(NumKeys-1), not number of Pos/Rot key tracks 
+	 * The FramesPerSecond parameter must be the sample rate of the animation data, for example 30.
+	 * If the FramesPerSecond parameter is set to 0 or negative, it will automatically calculate the FramesPerSecond based on the sequence length and number of frames.
+	 * The reason why you can provide a FramesPerSecond value is because this can be slightly more accurate than calculating it, in case super high precision is needed.
 	 **/
-	static void GetKeyIndicesFromTime(int32& OutKeyIndex1, int32& OutKeyIndex2, float& OutAlpha, const double Time, const int32 NumKeys, const double SequenceLength);
+	static void GetKeyIndicesFromTime(int32& OutKeyIndex1, int32& OutKeyIndex2, float& OutAlpha, const double Time, const int32 NumKeys, const double SequenceLength, double FramesPerSecond=-1.0);
 
 	/** 
 	 *	Utility for taking an array of bone indices and ensuring that all parents are present 
