Fix thumbnail refresh for generated textures (non-asset).

#rnx

diff --git a/Engine/Source/Editor/UnrealEd/Private/AssetThumbnail.cpp b/Engine/Source/Editor/UnrealEd/Private/AssetThumbnail.cpp
--- a/Engine/Source/Editor/UnrealEd/Private/AssetThumbnail.cpp
+++ b/Engine/Source/Editor/UnrealEd/Private/AssetThumbnail.cpp
@@ -1545,7 +1545,7 @@ void FAssetThumbnailPool::DirtyThumbnailForObject(UObject* ObjectBeingModified)
 		ObjectBeingModified = ActorBeingModified->GetWorld();
 	}
 
-	if (ObjectBeingModified && ObjectBeingModified->IsAsset())
+	if (ObjectBeingModified)
 	{
 		// An object in memory was modified.  We'll mark it's thumbnail as dirty so that it'll be
 		// regenerated on demand later. (Before being displayed in the browser, or package saves, etc.)
