Editor: Speculative fix for an unset TOptional access for what looked like trying to package for Android while running PIE, and having that package attempt fail to deploy to device.

diff --git a/Engine/Source/Editor/UnrealEd/Private/PlayLevel.cpp b/Engine/Source/Editor/UnrealEd/Private/PlayLevel.cpp
--- a/Engine/Source/Editor/UnrealEd/Private/PlayLevel.cpp
+++ b/Engine/Source/Editor/UnrealEd/Private/PlayLevel.cpp
@@ -603,7 +603,7 @@ void UEditorEngine::TeardownPlaySession(FWorldContext& PieWorldContext)
 			{
 				TSharedPtr<IAssetViewport> Viewport = SlatePlayInEditorSession->DestinationSlateViewport.Pin();
 
-				if(PlayInEditorSessionInfo->OriginalRequestParams.WorldType == EPlaySessionWorldType::PlayInEditor)
+				if(PlayInEditorSessionInfo.IsSet() && PlayInEditorSessionInfo->OriginalRequestParams.WorldType == EPlaySessionWorldType::PlayInEditor)
 				{
 					// Set the editor viewport location to match that of Play in Viewport if we aren't simulating in the editor, we have a valid player to get the location from (unless we're going back to VR Editor, in which case we won't teleport the user.)
 					if (bLastViewAndLocationValid == true && !GEngine->IsStereoscopic3D( Viewport->GetActiveViewport() ) )
