Fix bHaveGameOpenOrder to be set to whether we have any order files not named CookerOpenOrder

diff --git a/Engine/Source/Editor/UnrealEd/Private/Commandlets/AssetRegistryGenerator.cpp b/Engine/Source/Editor/UnrealEd/Private/Commandlets/AssetRegistryGenerator.cpp
--- a/Engine/Source/Editor/UnrealEd/Private/Commandlets/AssetRegistryGenerator.cpp
+++ b/Engine/Source/Editor/UnrealEd/Private/Commandlets/AssetRegistryGenerator.cpp
@@ -446,7 +446,7 @@ bool FAssetRegistryGenerator::GenerateStreamingInstallManifest(int64 InExtraFlav
 
 			for (const FString& Found : FoundFiles)
 			{
-				bHaveGameOpenOrder = Found.Contains(TEXT("CookerOpenOrder")) == false;
+				bHaveGameOpenOrder = bHaveGameOpenOrder || Found.Contains(TEXT("CookerOpenOrder")) == false;
 				UE_LOG(LogAssetRegistryGenerator, Display, TEXT("Found order file %s"), *Found);
 			}
 
