[NO NON-PII COMMIT MESSAGE]

diff --git a/Engine/Plugins/Experimental/AlembicImporter/Source/AlembicLibrary/Private/AbcImporter.cpp b/Engine/Plugins/Experimental/AlembicImporter/Source/AlembicLibrary/Private/AbcImporter.cpp
--- a/Engine/Plugins/Experimental/AlembicImporter/Source/AlembicLibrary/Private/AbcImporter.cpp
+++ b/Engine/Plugins/Experimental/AlembicImporter/Source/AlembicLibrary/Private/AbcImporter.cpp
@@ -29,7 +29,6 @@ THIRD_PARTY_INCLUDES_END
 
 #include "Engine/StaticMesh.h"
 #include "Engine/SkeletalMesh.h"
-#include "SkelImport.h"
 #include "Animation/AnimSequence.h"
 #include "Rendering/SkeletalMeshModel.h"
 
diff --git a/Engine/Plugins/Importers/USDImporter/Source/USDUtilities/Private/USDSkeletalDataConversion.cpp b/Engine/Plugins/Importers/USDImporter/Source/USDUtilities/Private/USDSkeletalDataConversion.cpp
--- a/Engine/Plugins/Importers/USDImporter/Source/USDUtilities/Private/USDSkeletalDataConversion.cpp
+++ b/Engine/Plugins/Importers/USDImporter/Source/USDUtilities/Private/USDSkeletalDataConversion.cpp
@@ -29,6 +29,7 @@
 #if WITH_EDITOR
 #include "Animation/DebugSkelMeshComponent.h"
 #include "Factories/FbxSkeletalMeshImportData.h"
+#include "ImportUtils/SkeletalMeshImportUtils.h"
 #include "MeshUtilities.h"
 #endif // WITH_EDITOR
 
@@ -2067,7 +2068,7 @@ USkeletalMesh* UsdToUnreal::GetSkeletalMeshFromImportData(
 	int32 SkeletalDepth = 0;
 	FSkeletalMeshImportData DummyData;
 	DummyData.RefBonesBinary = InSkeletonBones;
-	if ( !SkeletalMeshHelper::ProcessImportMeshSkeleton( SkeletalMesh->GetSkeleton(), SkeletalMesh->GetRefSkeleton(), SkeletalDepth, DummyData ) )
+	if ( !SkeletalMeshImportUtils::ProcessImportMeshSkeleton( SkeletalMesh->GetSkeleton(), SkeletalMesh->GetRefSkeleton(), SkeletalDepth, DummyData ) )
 	{
 		return nullptr;
 	}
@@ -2112,7 +2113,7 @@ USkeletalMesh* UsdToUnreal::GetSkeletalMeshFromImportData(
 		FSkeletalMeshLODModel& LODModel = ImportedResource->LODModels.Last();
 
 		// Process bones influence (normalization and optimization) (optional)
-		SkeletalMeshHelper::ProcessImportMeshInfluences(LODImportData, SkeletalMesh->GetPathName());
+		SkeletalMeshImportUtils::ProcessImportMeshInfluences(LODImportData, SkeletalMesh->GetPathName());
 
 		FSkeletalMeshLODInfo& NewLODInfo = SkeletalMesh->AddLODInfo();
 		NewLODInfo.ReductionSettings.NumOfTrianglesPercentage = 1.0f;
diff --git a/Engine/Plugins/Runtime/ApexDestruction/Source/ApexDestruction/Private/ApexDestructibleAssetImport.cpp b/Engine/Plugins/Runtime/ApexDestruction/Source/ApexDestruction/Private/ApexDestructibleAssetImport.cpp
--- a/Engine/Plugins/Runtime/ApexDestruction/Source/ApexDestruction/Private/ApexDestructibleAssetImport.cpp
+++ b/Engine/Plugins/Runtime/ApexDestruction/Source/ApexDestruction/Private/ApexDestructibleAssetImport.cpp
@@ -24,7 +24,7 @@
 DEFINE_LOG_CATEGORY_STATIC(LogApexDestructibleAssetImport, Log, All);
 
 #include "PhysicsPublic.h"
-#include "SkelImport.h"
+#include "ImportUtils/SkeletalMeshImportUtils.h"
 #include "PhysXIncludes.h"
 #include "MeshUtilities.h"
 #include "UObject/UObjectHash.h"
@@ -242,7 +242,7 @@ TSharedPtr<FExistingDestMeshData> SaveExistingDestMeshData(UDestructibleMesh* Ex
 		ExistingDestMeshDataPtr = MakeShared<FExistingDestMeshData>();
 
 		// Only save off SkelMeshData if it's been created
-		ExistingDestMeshDataPtr->SkelMeshData = SkeletalMeshHelper::SaveExistingSkelMeshData(ExistingDestructibleMesh, true, INDEX_NONE);
+		ExistingDestMeshDataPtr->SkelMeshData = SkeletalMeshImportUtils::SaveExistingSkelMeshData(ExistingDestructibleMesh, true, INDEX_NONE);
 		ExistingDestMeshDataPtr->BodySetup = ExistingDestructibleMesh->BodySetup;
 		ExistingDestMeshDataPtr->FractureEffects = ExistingDestructibleMesh->FractureEffects;
 	}
@@ -257,7 +257,7 @@ static void RestoreExistingDestMeshData(const TSharedPtr<FExistingDestMeshData>&
 		// Restore old settings, but resize arrays to make sense with the new NxDestructibleAsset
 		if (MeshData->SkelMeshData)
 		{
-			SkeletalMeshHelper::RestoreExistingSkelMeshData(MeshData->SkelMeshData, DestructibleMesh, INDEX_NONE, false, false, false);
+			SkeletalMeshImportUtils::RestoreExistingSkelMeshData(MeshData->SkelMeshData, DestructibleMesh, INDEX_NONE, false, false, false);
 		}
 		DestructibleMesh->BodySetup =  MeshData->BodySetup;
 		DestructibleMesh->FractureEffects = MeshData->FractureEffects;
@@ -765,7 +765,7 @@ apex::DestructibleAsset* CreateApexDestructibleAssetFromFile(const FString& File
 
 bool SetApexDestructibleAsset(UDestructibleMesh& DestructibleMesh, apex::DestructibleAsset& ApexDestructibleAsset, FSkeletalMeshImportData* OutData, EDestructibleImportOptions::Type Options)
 {
-	using namespace SkeletalMeshHelper;
+	using namespace SkeletalMeshImportUtils;
 	DestructibleMesh.PreEditChange(NULL);
 
 	DestructibleMesh.InvalidateDeriveDataCacheGUID();
diff --git a/Engine/Source/Developer/MeshUtilities/Private/SkeletalMeshTools.h b/Engine/Source/Developer/MeshUtilities/Private/SkeletalMeshTools.h
--- a/Engine/Source/Developer/MeshUtilities/Private/SkeletalMeshTools.h
+++ b/Engine/Source/Developer/MeshUtilities/Private/SkeletalMeshTools.h
@@ -2,13 +2,13 @@
 
 #pragma once
 
-#include "CoreMinimal.h"
+#include "BoneIndices.h"
+#include "Components.h"
 #include "Containers/IndirectArray.h"
+#include "CoreMinimal.h"
 #include "GPUSkinPublicDefs.h"
-#include "Components.h"
-#include "BoneIndices.h"
-#include "SkelImport.h"
 #include "MeshBuild.h"
+#include "Rendering/SkeletalMeshLODImporterData.h"
 
 class USkeletalMesh;
 struct FSoftSkinVertex;
diff --git a/Engine/Source/Developer/MeshUtilities/Public/MeshUtilities.h b/Engine/Source/Developer/MeshUtilities/Public/MeshUtilities.h
--- a/Engine/Source/Developer/MeshUtilities/Public/MeshUtilities.h
+++ b/Engine/Source/Developer/MeshUtilities/Public/MeshUtilities.h
@@ -7,7 +7,6 @@
 #include "Modules/ModuleInterface.h"
 #include "Components.h"
 #include "Engine/MeshMerging.h"
-#include "SkelImport.h"
 #include "MeshBuild.h"
 
 #include "IMeshMergeUtilities.h"
@@ -28,6 +27,13 @@ struct FStaticMeshLODResources;
 class UMeshComponent;
 class UStaticMesh;
 
+namespace SkeletalMeshImportData
+{
+	struct FMeshFace;
+	struct FMeshWedge;
+	struct FVertInfluence;
+};
+
 namespace ETangentOptions
 {
 	enum Type
diff --git a/Engine/Source/Editor/UnrealEd/Classes/Factories/FbxSkeletalMeshImportData.h b/Engine/Source/Editor/UnrealEd/Classes/Factories/FbxSkeletalMeshImportData.h
--- a/Engine/Source/Editor/UnrealEd/Classes/Factories/FbxSkeletalMeshImportData.h
+++ b/Engine/Source/Editor/UnrealEd/Classes/Factories/FbxSkeletalMeshImportData.h
@@ -109,18 +109,4 @@ public:
 
 	/** This function add the last import content type to the asset registry which is use by the thumbnail overlay of the skeletal mesh */
 	virtual void AppendAssetRegistryTags(TArray<FAssetRegistryTag>& OutTags);
-};
-
-namespace SkeletalMeshHelper
-{
-	/** Backups the given SkeletalMesh into a FExistingSkelMeshData */
-	UNREALED_API TSharedPtr<FExistingSkelMeshData> SaveExistingSkelMeshData(USkeletalMesh* SourceSkeletalMesh, bool bSaveMaterials, int32 ReimportLODIndex);
-
-	/** Restore a backed up FExistingSkelMeshData into a SkeletalMesh asset */
-	UNREALED_API void RestoreExistingSkelMeshData(TSharedPtr<const FExistingSkelMeshData> MeshData, USkeletalMesh* SkeletalMesh, int32 ReimportLODIndex, bool bCanShowDialog, bool bImportSkinningOnly, bool bForceMaterialReset);
-
-	UNREALED_API void ProcessImportMeshInfluences(FSkeletalMeshImportData& ImportData, const FString& SkeletalMeshName);
-	UNREALED_API void ProcessImportMeshMaterials(TArray<FSkeletalMaterial>& Materials, FSkeletalMeshImportData& ImportData);
-	UNREALED_API bool ProcessImportMeshSkeleton(const USkeleton* SkeletonAsset, FReferenceSkeleton& RefSkeleton, int32& SkeletalDepth, FSkeletalMeshImportData& ImportData);
-	UNREALED_API void ApplySkinning(USkeletalMesh* SkeletalMesh, FSkeletalMeshLODModel& SrcLODModel, FSkeletalMeshLODModel& DestLODModel);
-}
\ No newline at end of file
+};
\ No newline at end of file
diff --git a/Engine/Source/Editor/UnrealEd/Private/Factories/SkeletalMeshImport.cpp b/Engine/Source/Editor/UnrealEd/Private/Factories/SkeletalMeshImport.cpp
deleted file mode 100644
diff --git a/Engine/Source/Editor/UnrealEd/Private/Fbx/FbxFactory.cpp b/Engine/Source/Editor/UnrealEd/Private/Fbx/FbxFactory.cpp
--- a/Engine/Source/Editor/UnrealEd/Private/Fbx/FbxFactory.cpp
+++ b/Engine/Source/Editor/UnrealEd/Private/Fbx/FbxFactory.cpp
@@ -34,7 +34,6 @@
 #include "Misc/Paths.h"
 #include "Misc/ScopedSlowTask.h"
 #include "ObjectTools.h"
-#include "SkelImport.h"
 #include "UObject/UnrealType.h"
 
 #define LOCTEXT_NAMESPACE "FBXFactory"
diff --git a/Engine/Source/Editor/UnrealEd/Private/Fbx/FbxMainImport.cpp b/Engine/Source/Editor/UnrealEd/Private/Fbx/FbxMainImport.cpp
--- a/Engine/Source/Editor/UnrealEd/Private/Fbx/FbxMainImport.cpp
+++ b/Engine/Source/Editor/UnrealEd/Private/Fbx/FbxMainImport.cpp
@@ -16,7 +16,7 @@
 #include "Factories/FbxTextureImportData.h"
 
 #include "Materials/MaterialInterface.h"
-#include "SkelImport.h"
+#include "Rendering/SkeletalMeshLODImporterData.h"
 #include "Logging/TokenizedMessage.h"
 #include "Misc/FbxErrors.h"
 #include "FbxImporter.h"
diff --git a/Engine/Source/Editor/UnrealEd/Private/Fbx/FbxReimportDialogs.cpp b/Engine/Source/Editor/UnrealEd/Private/Fbx/FbxReimportDialogs.cpp
--- a/Engine/Source/Editor/UnrealEd/Private/Fbx/FbxReimportDialogs.cpp
+++ b/Engine/Source/Editor/UnrealEd/Private/Fbx/FbxReimportDialogs.cpp
@@ -46,7 +46,6 @@
 #include "Factories/FbxStaticMeshImportData.h"
 
 //Skeletal mesh includes
-#include "SkelImport.h"
 #include "Animation/Skeleton.h"
 #include "Engine/SkeletalMesh.h"
 #include "Components/SkinnedMeshComponent.h"
diff --git a/Engine/Source/Editor/UnrealEd/Private/Fbx/FbxSceneImportFactory.cpp b/Engine/Source/Editor/UnrealEd/Private/Fbx/FbxSceneImportFactory.cpp
--- a/Engine/Source/Editor/UnrealEd/Private/Fbx/FbxSceneImportFactory.cpp
+++ b/Engine/Source/Editor/UnrealEd/Private/Fbx/FbxSceneImportFactory.cpp
@@ -39,7 +39,6 @@
 #include "Editor.h"
 #include "FileHelpers.h"
 #include "CineCameraComponent.h"
-#include "SkelImport.h"
 #include "Rendering/SkeletalMeshModel.h"
 
 #include "AssetSelection.h"
diff --git a/Engine/Source/Editor/UnrealEd/Private/Fbx/FbxSkeletalMeshImport.cpp b/Engine/Source/Editor/UnrealEd/Private/Fbx/FbxSkeletalMeshImport.cpp
--- a/Engine/Source/Editor/UnrealEd/Private/Fbx/FbxSkeletalMeshImport.cpp
+++ b/Engine/Source/Editor/UnrealEd/Private/Fbx/FbxSkeletalMeshImport.cpp
@@ -32,7 +32,8 @@
 #include "Rendering/SkeletalMeshLODRenderData.h"
 #include "Misc/ScopedSlowTask.h"
 
-#include "SkelImport.h"
+#include "ImportUtils/SkeletalMeshImportUtils.h"
+#include "ImportUtils/SkelImport.h"
 #include "Logging/TokenizedMessage.h"
 #include "FbxImporter.h"
 
@@ -1777,7 +1778,7 @@ USkeletalMesh* UnFbx::FFbxImporter::ImportSkeletalMesh(FImportSkeletalMeshArgs &
 		}
 		//The backup of the skeletal mesh data empty the LOD array in the ImportedResource of the skeletal mesh
 		//If the import fail after this step the editor can crash when updating the bone later since the LODModel will not exist anymore
-		ExistSkelMeshDataPtr = SkeletalMeshHelper::SaveExistingSkelMeshData(ExistingSkelMesh, !ImportOptions->bImportMaterials, ImportSkeletalMeshArgs.LodIndex);
+		ExistSkelMeshDataPtr = SkeletalMeshImportUtils::SaveExistingSkelMeshData(ExistingSkelMesh, !ImportOptions->bImportMaterials, ImportSkeletalMeshArgs.LodIndex);
 	}
 
 	if (SkeletalMesh == nullptr)
@@ -1803,12 +1804,12 @@ USkeletalMesh* UnFbx::FFbxImporter::ImportSkeletalMesh(FImportSkeletalMeshArgs &
 	FSkeletalMeshLODModel& LODModel = ImportedResource->LODModels[ImportLODModelIndex];
 
 	// process materials from import data
-	SkeletalMeshHelper::ProcessImportMeshMaterials(SkeletalMesh->GetMaterials(), *SkelMeshImportDataPtr);
+	SkeletalMeshImportUtils::ProcessImportMeshMaterials(SkeletalMesh->GetMaterials(), *SkelMeshImportDataPtr);
 
 	// process reference skeleton from import data
 	int32 SkeletalDepth = 0;
 	USkeleton* ExistingSkeleton = ExistSkelMeshDataPtr ? ExistSkelMeshDataPtr->ExistingSkeleton : ImportOptions->SkeletonForAnimation;
-	if (!SkeletalMeshHelper::ProcessImportMeshSkeleton(ExistingSkeleton, SkeletalMesh->GetRefSkeleton(), SkeletalDepth, *SkelMeshImportDataPtr))
+	if (!SkeletalMeshImportUtils::ProcessImportMeshSkeleton(ExistingSkeleton, SkeletalMesh->GetRefSkeleton(), SkeletalDepth, *SkelMeshImportDataPtr))
 	{
 		EARLY_RETURN_ON_CANCEL(true, FailureCleanup);
 	}
@@ -1819,7 +1820,7 @@ USkeletalMesh* UnFbx::FFbxImporter::ImportSkeletalMesh(FImportSkeletalMeshArgs &
 		UE_LOG(LogFbx, Warning, TEXT("Bones digested - %i  Depth of hierarchy - %i"), SkeletalMesh->GetRefSkeleton().GetNum(), SkeletalDepth);
 
 	// process bone influences from import data
-	SkeletalMeshHelper::ProcessImportMeshInfluences(*SkelMeshImportDataPtr, SkeletalMesh->GetPathName());
+	SkeletalMeshImportUtils::ProcessImportMeshInfluences(*SkelMeshImportDataPtr, SkeletalMesh->GetPathName());
 
 	//Store the original fbx import data the SkelMeshImportDataPtr should not be modified after this
 	SkeletalMesh->SaveLODImportedData(ImportLODModelIndex, *SkelMeshImportDataPtr);
@@ -1961,7 +1962,7 @@ USkeletalMesh* UnFbx::FFbxImporter::ImportSkeletalMesh(FImportSkeletalMeshArgs &
 
 		if (ExistSkelMeshDataPtr)
 		{
-			SkeletalMeshHelper::RestoreExistingSkelMeshData(ExistSkelMeshDataPtr, SkeletalMesh, ImportSkeletalMeshArgs.LodIndex, ImportOptions->bCanShowDialog, ImportOptions->bImportAsSkeletalSkinning, ImportOptions->bResetToFbxOnMaterialConflict);
+			SkeletalMeshImportUtils::RestoreExistingSkelMeshData(ExistSkelMeshDataPtr, SkeletalMesh, ImportSkeletalMeshArgs.LodIndex, ImportOptions->bCanShowDialog, ImportOptions->bImportAsSkeletalSkinning, ImportOptions->bResetToFbxOnMaterialConflict);
 		}
 
 		SkeletalMesh->CalculateInvRefMatrices();
@@ -4006,7 +4007,7 @@ bool UnFbx::FFbxImporter::ImportSkeletalMeshLOD(USkeletalMesh* InSkeletalMesh, U
 		check(ImportOptions->bImportAsSkeletalGeometry);
 		//Apply the base LOD skinning to the new LOD geometry, DestImportedResource here is the base lod and is the source for the apply skinning.
 		//The Dest here mean we merge the just import LOD to the base LodModel.
-		SkeletalMeshHelper::ApplySkinning(InSkeletalMesh, DestImportedResource->LODModels[0], ImportedResource->LODModels[0]);
+		SkeletalMeshImportUtils::ApplySkinning(InSkeletalMesh, DestImportedResource->LODModels[0], ImportedResource->LODModels[0]);
 
 		//Excluded bones will be remove when we rebuild the asset
 		//IMeshUtilities& MeshUtilities = FModuleManager::Get().LoadModuleChecked<IMeshUtilities>("MeshUtilities");
diff --git a/Engine/Source/Editor/UnrealEd/Private/Fbx/FbxStaticMeshImport.cpp b/Engine/Source/Editor/UnrealEd/Private/Fbx/FbxStaticMeshImport.cpp
--- a/Engine/Source/Editor/UnrealEd/Private/Fbx/FbxStaticMeshImport.cpp
+++ b/Engine/Source/Editor/UnrealEd/Private/Fbx/FbxStaticMeshImport.cpp
@@ -30,6 +30,7 @@
 #include "Logging/TokenizedMessage.h"
 #include "FbxImporter.h"
 #include "GeomFitUtils.h"
+#include "ImportUtils/StaticMeshImportUtils.h"
 #include "Interfaces/ITargetPlatform.h"
 #include "Interfaces/ITargetPlatformManagerModule.h"
 #include "Misc/FbxErrors.h"
@@ -45,11 +46,6 @@ static const int32 LARGE_MESH_MATERIAL_INDEX_THRESHOLD = 64;
 
 using namespace UnFbx;
 
-struct FExistingStaticMeshData;
-extern TSharedPtr<FExistingStaticMeshData> SaveExistingStaticMeshData(UStaticMesh* ExistingMesh, FBXImportOptions* ImportOptions, int32 LodIndex);
-extern void RestoreExistingMeshSettings(const FExistingStaticMeshData* ExistingMesh, UStaticMesh* NewMesh, int32 LODIndex);
-extern void RestoreExistingMeshData(TSharedPtr<const FExistingStaticMeshData> ExistingMeshDataPtr, UStaticMesh* NewMesh, int32 LodLevel, bool bCanShowDialog, bool bForceConflictingMaterialReset);
-extern void UpdateSomeLodsImportMeshData(UStaticMesh* NewMesh, TArray<int32> *ReimportLodList);
 
 struct FRestoreReimportData
 {
@@ -1245,7 +1241,7 @@ UStaticMesh* UnFbx::FFbxImporter::ReimportSceneStaticMesh(uint64 FbxNodeUniqueId
 		return Mesh;
 	}
 
-	TSharedPtr<FExistingStaticMeshData> ExistMeshDataPtr = SaveExistingStaticMeshData(Mesh, ImportOptions, INDEX_NONE);
+	TSharedPtr<FExistingStaticMeshData> ExistMeshDataPtr = StaticMeshImportUtils::SaveExistingStaticMeshData(Mesh, ImportOptions, INDEX_NONE);
 
 	if (Node)
 	{
@@ -1321,7 +1317,7 @@ UStaticMesh* UnFbx::FFbxImporter::ReimportSceneStaticMesh(uint64 FbxNodeUniqueId
 	if(FirstBaseMesh)
 	{
 	//Don't restore materials when reimporting scene
-		RestoreExistingMeshData(ExistMeshDataPtr, FirstBaseMesh, INDEX_NONE, false, ImportOptions->bResetToFbxOnMaterialConflict);
+		StaticMeshImportUtils::RestoreExistingMeshData(ExistMeshDataPtr, FirstBaseMesh, INDEX_NONE, false, ImportOptions->bResetToFbxOnMaterialConflict);
 		RestoreData.CleanupDuplicateMesh();
 	}
 	else
@@ -1466,7 +1462,7 @@ UStaticMesh* UnFbx::FFbxImporter::ReimportStaticMesh(UStaticMesh* Mesh, UFbxStat
 	ImportOptions->bImportMaterials = false;
 	ImportOptions->bImportTextures = false;
 
-	TSharedPtr<FExistingStaticMeshData> ExistMeshDataPtr = SaveExistingStaticMeshData(Mesh, ImportOptions, INDEX_NONE);
+	TSharedPtr<FExistingStaticMeshData> ExistMeshDataPtr = StaticMeshImportUtils::SaveExistingStaticMeshData(Mesh, ImportOptions, INDEX_NONE);
 
 	TArray<int32> ReimportLodList;
 	if (bCombineMeshesLOD)
@@ -1597,8 +1593,8 @@ UStaticMesh* UnFbx::FFbxImporter::ReimportStaticMesh(UStaticMesh* Mesh, UFbxStat
 	}
 	if (NewMesh != nullptr)
 	{
-		UpdateSomeLodsImportMeshData(NewMesh, &ReimportLodList);
-		RestoreExistingMeshData(ExistMeshDataPtr, NewMesh, INDEX_NONE, ImportOptions->bCanShowDialog, ImportOptions->bResetToFbxOnMaterialConflict);
+		StaticMeshImportUtils::UpdateSomeLodsImportMeshData(NewMesh, &ReimportLodList);
+		StaticMeshImportUtils::RestoreExistingMeshData(ExistMeshDataPtr, NewMesh, INDEX_NONE, ImportOptions->bCanShowDialog, ImportOptions->bResetToFbxOnMaterialConflict);
 		RestoreData.CleanupDuplicateMesh();
 	}
 	else
@@ -2006,7 +2002,7 @@ UStaticMesh* UnFbx::FFbxImporter::ImportStaticMeshAsSingle(UObject* InParent, TA
 
 		if (ExistMeshDataPtr && InStaticMesh)
 		{
-			RestoreExistingMeshSettings(ExistMeshDataPtr, InStaticMesh, StaticMesh->LODGroup != NAME_None ? INDEX_NONE : LODIndex);
+			StaticMeshImportUtils::RestoreExistingMeshSettings(ExistMeshDataPtr, InStaticMesh, StaticMesh->LODGroup != NAME_None ? INDEX_NONE : LODIndex);
 		}
 
 		// The code to check for bad lightmap UVs doesn't scale well with number of triangles.
@@ -2665,11 +2661,6 @@ bool UnFbx::FFbxImporter::FillCollisionModelList(FbxNode* Node)
 	return false;
 }
 
-extern bool AddConvexGeomFromVertices( const TArray<FVector3f>& Verts, FKAggregateGeom* AggGeom, const TCHAR* ObjName );
-extern bool AddSphereGeomFromVerts(const TArray<FVector3f>& Verts, FKAggregateGeom* AggGeom, const TCHAR* ObjName);
-extern bool AddCapsuleGeomFromVerts(const TArray<FVector3f>& Verts, FKAggregateGeom* AggGeom, const TCHAR* ObjName);
-extern bool AddBoxGeomFromTris(const TArray<FPoly>& Tris, FKAggregateGeom* AggGeom, const TCHAR* ObjName);
-extern bool DecomposeUCXMesh( const TArray<FVector3f>& CollisionVertices, const TArray<int32>& CollisionFaceIdx, UBodySetup* BodySetup );
 
 bool UnFbx::FFbxImporter::ImportCollisionModels(UStaticMesh* StaticMesh, const FbxString& InNodeName)
 {
@@ -2782,7 +2773,7 @@ bool UnFbx::FFbxImporter::ImportCollisionModels(UStaticMesh* StaticMesh, const F
 		{
 			if( !ImportOptions->bOneConvexHullPerUCX )
 			{
-				if (DecomposeUCXMesh(CollisionVertices, CollisionFaceIdx, StaticMesh->GetBodySetup()))
+				if (StaticMeshImportUtils::DecomposeUCXMesh(CollisionVertices, CollisionFaceIdx, StaticMesh->GetBodySetup()))
 				{
 					bAtLeastOneCollisionMeshImported = true;
 				}
@@ -2798,7 +2789,7 @@ bool UnFbx::FFbxImporter::ImportCollisionModels(UStaticMesh* StaticMesh, const F
 
 				// This function cooks the given data, so we cannot test for duplicates based on the position data
 				// before we call it
-				if (AddConvexGeomFromVertices(CollisionVertices, &AggGeo, ANSI_TO_TCHAR(Node->GetName())))
+				if (StaticMeshImportUtils::AddConvexGeomFromVertices(CollisionVertices, &AggGeo, ANSI_TO_TCHAR(Node->GetName())))
 				{
 					bAtLeastOneCollisionMeshImported = true;
 					FKConvexElem& NewElem = AggGeo.ConvexElems.Last();
@@ -2835,7 +2826,7 @@ bool UnFbx::FFbxImporter::ImportCollisionModels(UStaticMesh* StaticMesh, const F
 		{
 			FKAggregateGeom& AggGeo = StaticMesh->GetBodySetup()->AggGeom;
 
-			if(AddBoxGeomFromTris(CollisionTriangles, &AggGeo, ANSI_TO_TCHAR(Node->GetName())))
+			if(StaticMeshImportUtils::AddBoxGeomFromTris(CollisionTriangles, &AggGeo, ANSI_TO_TCHAR(Node->GetName())))
 			{
 				bAtLeastOneCollisionMeshImported = true;
 				FKBoxElem& NewElem = AggGeo.BoxElems.Last();
@@ -2858,7 +2849,7 @@ bool UnFbx::FFbxImporter::ImportCollisionModels(UStaticMesh* StaticMesh, const F
 		{
 			FKAggregateGeom& AggGeo = StaticMesh->GetBodySetup()->AggGeom;
 
-			if(AddSphereGeomFromVerts(CollisionVertices, &AggGeo, ANSI_TO_TCHAR(Node->GetName())))
+			if (StaticMeshImportUtils::AddSphereGeomFromVerts(CollisionVertices, &AggGeo, ANSI_TO_TCHAR(Node->GetName())))
 			{
 				bAtLeastOneCollisionMeshImported = true;
 				FKSphereElem& NewElem = AggGeo.SphereElems.Last();
@@ -2881,7 +2872,7 @@ bool UnFbx::FFbxImporter::ImportCollisionModels(UStaticMesh* StaticMesh, const F
 		{
 			FKAggregateGeom& AggGeo = StaticMesh->GetBodySetup()->AggGeom;
 
-			if (AddCapsuleGeomFromVerts(CollisionVertices, &AggGeo, ANSI_TO_TCHAR(Node->GetName())))
+			if (StaticMeshImportUtils::AddCapsuleGeomFromVerts(CollisionVertices, &AggGeo, ANSI_TO_TCHAR(Node->GetName())))
 			{
 				bAtLeastOneCollisionMeshImported = true;
 				FKSphylElem& NewElem = AggGeo.SphylElems.Last();
diff --git a/Engine/Source/Editor/UnrealEd/Private/FbxMeshUtils.cpp b/Engine/Source/Editor/UnrealEd/Private/FbxMeshUtils.cpp
--- a/Engine/Source/Editor/UnrealEd/Private/FbxMeshUtils.cpp
+++ b/Engine/Source/Editor/UnrealEd/Private/FbxMeshUtils.cpp
@@ -18,11 +18,12 @@
 #include "FbxImporter.h"
 #include "StaticMeshResources.h"
 #include "Components/SkeletalMeshComponent.h"
-#include "SkelImport.h"
 #include "Rendering/SkeletalMeshModel.h"
 #include "EditorFramework/AssetImportData.h"
 #include "DesktopPlatformModule.h"
 #include "Editor.h"
+#include "ImportUtils/SkeletalMeshImportUtils.h"
+#include "ImportUtils/StaticMeshImportUtils.h"
 
 #if WITH_APEX_CLOTHING
 	#include "ApexClothingUtils.h"
@@ -41,12 +42,6 @@ DEFINE_LOG_CATEGORY_STATIC(LogExportMeshUtils, Log, All);
 
 #define LOCTEXT_NAMESPACE "FbxMeshUtil"
 
-struct FExistingStaticMeshData;
-extern TSharedPtr<FExistingStaticMeshData> SaveExistingStaticMeshData(UStaticMesh* ExistingMesh, UnFbx::FBXImportOptions* ImportOptions, int32 LodIndex);
-extern void RestoreExistingMeshSettings(const FExistingStaticMeshData* ExistingMesh, UStaticMesh* NewMesh, int32 LODIndex);
-extern void RestoreExistingMeshData(TSharedPtr<const FExistingStaticMeshData> ExistingMeshDataPtr, UStaticMesh* NewMesh, int32 LodLevel, bool bCanShowDialog, bool bForceConflictingMaterialReset);
-extern void UpdateSomeLodsImportMeshData(UStaticMesh* NewMesh, TArray<int32> *ReimportLodList);
-
 
 namespace FbxMeshUtils
 {
@@ -167,7 +162,7 @@ namespace FbxMeshUtils
 			TSharedPtr<FExistingStaticMeshData> ExistMeshDataPtr;
 			if (IsReimport)
 			{
-				ExistMeshDataPtr = SaveExistingStaticMeshData(BaseStaticMesh, FFbxImporter->ImportOptions, LODLevel);
+				ExistMeshDataPtr = StaticMeshImportUtils::SaveExistingStaticMeshData(BaseStaticMesh, FFbxImporter->ImportOptions, LODLevel);
 			}
 
 			// Display the LOD selection dialog
@@ -200,10 +195,10 @@ namespace FbxMeshUtils
 					FFbxImporter->PostImportStaticMesh(TempStaticMesh, *(LODNodeList[bUseLODs ? LODLevel : 0]), LODLevel);
 					TArray<int32> ReimportLodList;
 					ReimportLodList.Add(LODLevel);
-					UpdateSomeLodsImportMeshData(BaseStaticMesh, &ReimportLodList);
+					StaticMeshImportUtils::UpdateSomeLodsImportMeshData(BaseStaticMesh, &ReimportLodList);
 					if(IsReimport)
 					{
-						RestoreExistingMeshData(ExistMeshDataPtr, BaseStaticMesh, LODLevel, false, ImportOptions->bResetToFbxOnMaterialConflict);
+						StaticMeshImportUtils::RestoreExistingMeshData(ExistMeshDataPtr, BaseStaticMesh, LODLevel, false, ImportOptions->bResetToFbxOnMaterialConflict);
 					}
 
 					// Update mesh component
@@ -487,7 +482,7 @@ namespace FbxMeshUtils
 				if (SelectedSkelMesh->GetLODNum() > SelectedLOD)
 				{
 					SelectedSkelMesh->PreEditChange(NULL);
-					SkelMeshDataPtr = SkeletalMeshHelper::SaveExistingSkelMeshData(SelectedSkelMesh, true, SelectedLOD);
+					SkelMeshDataPtr = SkeletalMeshImportUtils::SaveExistingSkelMeshData(SelectedSkelMesh, true, SelectedLOD);
 				}
 
 				//Original fbx data storage
@@ -517,7 +512,7 @@ namespace FbxMeshUtils
 
 					if (SkelMeshDataPtr)
 					{
-						SkeletalMeshHelper::RestoreExistingSkelMeshData(SkelMeshDataPtr, SelectedSkelMesh, SelectedLOD, false, ImportOptions->bImportAsSkeletalSkinning, ImportOptions->bResetToFbxOnMaterialConflict);
+						SkeletalMeshImportUtils::RestoreExistingSkelMeshData(SkelMeshDataPtr, SelectedSkelMesh, SelectedLOD, false, ImportOptions->bImportAsSkeletalSkinning, ImportOptions->bResetToFbxOnMaterialConflict);
 					}
 
 					if (ImportOptions->bImportMorph)
diff --git a/Engine/Source/Editor/UnrealEd/Private/ImportUtils/SkeletalMeshImportUtils.cpp b/Engine/Source/Editor/UnrealEd/Private/ImportUtils/SkeletalMeshImportUtils.cpp
new file mode 100644
--- /dev/null
+++ b/Engine/Source/Editor/UnrealEd/Private/ImportUtils/SkeletalMeshImportUtils.cpp
@@ -0,0 +1,1158 @@
+// Copyright Epic Games, Inc. All Rights Reserved.
+
+/*=============================================================================
+	SkeletalMeshImportUtils.cpp: Skeletal mesh import code.
+=============================================================================*/
+
+#include "ImportUtils/SkeletalMeshImportUtils.h"
+
+#include "ClothingAssetBase.h"
+#include "CoreMinimal.h"
+#include "EditorFramework/ThumbnailInfo.h"
+#include "Engine/AssetUserData.h"
+#include "Engine/SkeletalMesh.h"
+#include "Engine/SkeletalMeshSocket.h"
+#include "Factories/FbxSkeletalMeshImportData.h"
+#include "FbxImporter.h"
+#include "ImportUtils/SkelImport.h"
+#include "Interfaces/ITargetPlatform.h"
+#include "Interfaces/ITargetPlatformManagerModule.h"
+#include "LODUtilities.h"
+#include "Materials/MaterialInterface.h"
+#include "Misc/CoreMisc.h"
+#include "Misc/FbxErrors.h"
+#include "PhysicsEngine/PhysicsAsset.h"
+#include "ReferenceSkeleton.h"
+#include "Rendering/SkeletalMeshLODImporterData.h"
+#include "Rendering/SkeletalMeshLODModel.h"
+#include "Rendering/SkeletalMeshModel.h"
+#include "UObject/MetaData.h"
+#include "UObject/UObjectIterator.h"
+
+DEFINE_LOG_CATEGORY_STATIC(LogSkeletalMeshImport, Log, All);
+
+#define LOCTEXT_NAMESPACE "SkeletalMeshImport"
+
+namespace SkeletalMesUtilsImpl
+{
+	/** Check that root bone is the same, and that any bones that are common have the correct parent. */
+	bool SkeletonsAreCompatible(const FReferenceSkeleton& NewSkel, const FReferenceSkeleton& ExistSkel, bool bFailNoError);
+
+	bool SkeletalMeshIsUsingMaterialSlotNameWorkflow(UAssetImportData* AssetImportData);
+
+	void SaveSkeletalMeshLODModelSections(USkeletalMesh* SourceSkeletalMesh, TSharedPtr<FExistingSkelMeshData>& ExistingMeshDataPtr, int32 LodIndex, bool bSaveNonReducedMeshData);
+
+	void SaveSkeletalMeshMaterialNameWorkflowData(TSharedPtr<FExistingSkelMeshData>& ExistingMeshDataPtr, const USkeletalMesh* SourceSkeletalMesh);
+
+	void SaveSkeletalMeshAssetUserData(TSharedPtr<FExistingSkelMeshData>& ExistingMeshDataPtr, const TArray<UAssetUserData*>* UserData);
+
+	void RestoreDependentLODs(const TSharedPtr<const FExistingSkelMeshData>& MeshData, USkeletalMesh* SkeletalMesh);
+
+	void RestoreLODInfo(const TSharedPtr<const FExistingSkelMeshData>& MeshData, USkeletalMesh* SkeletalMesh, int32 LodIndex);
+
+	void RestoreMaterialNameWorkflowSection(const TSharedPtr<const FExistingSkelMeshData>& MeshData, USkeletalMesh* SkeletalMesh, int32 LodIndex, TArray<int32>& RemapMaterial, bool bMaterialReset);
+}
+
+bool SkeletalMesUtilsImpl::SkeletonsAreCompatible(const FReferenceSkeleton& NewSkel, const FReferenceSkeleton& ExistSkel, bool bFailNoError)
+{
+	if (NewSkel.GetBoneName(0) != ExistSkel.GetBoneName(0))
+	{
+		if (!bFailNoError)
+		{
+			UnFbx::FFbxImporter* FFbxImporter = UnFbx::FFbxImporter::GetInstance();
+			FFbxImporter->AddTokenizedErrorMessage(FTokenizedMessage::Create(EMessageSeverity::Error, FText::Format(LOCTEXT("MeshHasDifferentRoot", "Root Bone is '{0}' instead of '{1}'.\nDiscarding existing LODs."),
+				FText::FromName(NewSkel.GetBoneName(0)), FText::FromName(ExistSkel.GetBoneName(0)))), FFbxErrors::SkeletalMesh_DifferentRoots);
+		}
+		return false;
+	}
+
+	for (int32 i = 1; i < NewSkel.GetRawBoneNum(); i++)
+	{
+		// See if bone is in both skeletons.
+		int32 NewBoneIndex = i;
+		FName NewBoneName = NewSkel.GetBoneName(NewBoneIndex);
+		int32 BBoneIndex = ExistSkel.FindBoneIndex(NewBoneName);
+
+		// If it is, check parents are the same.
+		if (BBoneIndex != INDEX_NONE)
+		{
+			FName NewParentName = NewSkel.GetBoneName(NewSkel.GetParentIndex(NewBoneIndex));
+			FName ExistParentName = ExistSkel.GetBoneName(ExistSkel.GetParentIndex(BBoneIndex));
+
+			if (NewParentName != ExistParentName)
+			{
+				if (!bFailNoError)
+				{
+					UnFbx::FFbxImporter* FFbxImporter = UnFbx::FFbxImporter::GetInstance();
+					FFbxImporter->AddTokenizedErrorMessage(FTokenizedMessage::Create(EMessageSeverity::Error, FText::Format(LOCTEXT("MeshHasDifferentRoot", "Root Bone is '{0}' instead of '{1}'.\nDiscarding existing LODs."),
+						FText::FromName(NewBoneName), FText::FromName(NewParentName))), FFbxErrors::SkeletalMesh_DifferentRoots);
+				}
+				return false;
+			}
+		}
+	}
+
+	return true;
+}
+
+/**
+* Process and fill in the mesh Materials using the raw binary import data
+*
+* @param Materials - [out] array of materials to update
+* @param ImportData - raw binary import data to process
+*/
+void SkeletalMeshImportUtils::ProcessImportMeshMaterials(TArray<FSkeletalMaterial>& Materials, FSkeletalMeshImportData& ImportData)
+{
+	TArray <SkeletalMeshImportData::FMaterial>&	ImportedMaterials = ImportData.Materials;
+
+	// If direct linkup of materials is requested, try to find them here - to get a texture name from a 
+	// material name, cut off anything in front of the dot (beyond are special flags).
+	Materials.Empty();
+	int32 SkinOffset = INDEX_NONE;
+	for (int32 MatIndex = 0; MatIndex < ImportedMaterials.Num(); ++MatIndex)
+	{
+		const SkeletalMeshImportData::FMaterial& ImportedMaterial = ImportedMaterials[MatIndex];
+
+		UMaterialInterface* Material = NULL;
+		FString MaterialNameNoSkin = ImportedMaterial.MaterialImportName;
+		if (ImportedMaterial.Material.IsValid())
+		{
+			Material = ImportedMaterial.Material.Get();
+		}
+		else
+		{
+			const FString& MaterialName = ImportedMaterial.MaterialImportName;
+			MaterialNameNoSkin = MaterialName;
+			Material = FindObject<UMaterialInterface>(ANY_PACKAGE, *MaterialName);
+			if (Material == nullptr)
+			{
+				SkinOffset = MaterialName.Find(TEXT("_skin"), ESearchCase::IgnoreCase, ESearchDir::FromEnd);
+				if (SkinOffset != INDEX_NONE)
+				{
+					FString SkinXXNumber = MaterialName.Right(MaterialName.Len() - (SkinOffset + 1)).RightChop(4);
+					if (SkinXXNumber.IsNumeric())
+					{
+						MaterialNameNoSkin = MaterialName.LeftChop(MaterialName.Len() - SkinOffset);
+						Material = FindObject<UMaterialInterface>(ANY_PACKAGE, *MaterialNameNoSkin);
+					}
+				}
+			}
+		}
+
+		const bool bEnableShadowCasting = true;
+		Materials.Add(FSkeletalMaterial(Material, bEnableShadowCasting, false, Material != nullptr ? Material->GetFName() : FName(*MaterialNameNoSkin), FName(*(ImportedMaterial.MaterialImportName))));
+	}
+
+	int32 NumMaterialsToAdd = FMath::Max<int32>(ImportedMaterials.Num(), ImportData.MaxMaterialIndex + 1);
+
+	// Pad the material pointers
+	while (NumMaterialsToAdd > Materials.Num())
+	{
+		Materials.Add(FSkeletalMaterial(NULL, true, false, NAME_None, NAME_None));
+	}
+}
+
+/**
+* Process and fill in the mesh ref skeleton bone hierarchy using the raw binary import data
+*
+* @param RefSkeleton - [out] reference skeleton hierarchy to update
+* @param SkeletalDepth - [out] depth of the reference skeleton hierarchy
+* @param ImportData - raw binary import data to process
+* @return true if the operation completed successfully
+*/
+bool SkeletalMeshImportUtils::ProcessImportMeshSkeleton(const USkeleton* SkeletonAsset, FReferenceSkeleton& RefSkeleton, int32& SkeletalDepth, FSkeletalMeshImportData& ImportData)
+{
+	TArray <SkeletalMeshImportData::FBone>&	RefBonesBinary = ImportData.RefBonesBinary;
+
+	// Setup skeletal hierarchy + names structure.
+	RefSkeleton.Empty();
+
+	FReferenceSkeletonModifier RefSkelModifier(RefSkeleton, SkeletonAsset);
+
+	// Digest bones to the serializable format.
+	for (int32 b = 0; b < RefBonesBinary.Num(); b++)
+	{
+		const SkeletalMeshImportData::FBone & BinaryBone = RefBonesBinary[b];
+		const FString BoneName = FSkeletalMeshImportData::FixupBoneName(BinaryBone.Name);
+		const FMeshBoneInfo BoneInfo(FName(*BoneName, FNAME_Add), BinaryBone.Name, BinaryBone.ParentIndex);
+		const FTransform BoneTransform(BinaryBone.BonePos.Transform);
+
+		if (RefSkeleton.FindRawBoneIndex(BoneInfo.Name) != INDEX_NONE)
+		{
+			UnFbx::FFbxImporter* FFbxImporter = UnFbx::FFbxImporter::GetInstance();
+			FFbxImporter->AddTokenizedErrorMessage(FTokenizedMessage::Create(EMessageSeverity::Error, FText::Format(LOCTEXT("SkeletonHasDuplicateBones", "Skeleton has non-unique bone names.\nBone named '{0}' encountered more than once."), FText::FromName(BoneInfo.Name))), FFbxErrors::SkeletalMesh_DuplicateBones);
+			return false;
+		}
+
+		RefSkelModifier.Add(BoneInfo, BoneTransform);
+	}
+
+	// Add hierarchy index to each bone and detect max depth.
+	SkeletalDepth = 0;
+
+	TArray<int32> SkeletalDepths;
+	SkeletalDepths.Empty(RefBonesBinary.Num());
+	SkeletalDepths.AddZeroed(RefBonesBinary.Num());
+	for (int32 b = 0; b < RefSkeleton.GetRawBoneNum(); b++)
+	{
+		int32 Parent = RefSkeleton.GetRawParentIndex(b);
+		int32 Depth = 1.0f;
+
+		SkeletalDepths[b] = 1.0f;
+		if (Parent != INDEX_NONE)
+		{
+			Depth += SkeletalDepths[Parent];
+		}
+		if (SkeletalDepth < Depth)
+		{
+			SkeletalDepth = Depth;
+		}
+		SkeletalDepths[b] = Depth;
+	}
+
+	return true;
+}
+
+/**
+* Process and update the vertex Influences using the raw binary import data
+*
+* @param ImportData - raw binary import data to process
+*/
+void SkeletalMeshImportUtils::ProcessImportMeshInfluences(FSkeletalMeshImportData& ImportData, const FString& SkeletalMeshName)
+{
+	FLODUtilities::ProcessImportMeshInfluences(ImportData.Wedges.Num(), ImportData.Influences, SkeletalMeshName);
+}
+
+bool SkeletalMesUtilsImpl::SkeletalMeshIsUsingMaterialSlotNameWorkflow(UAssetImportData* AssetImportData)
+{
+	UFbxSkeletalMeshImportData* ImportData = Cast<UFbxSkeletalMeshImportData>(AssetImportData);
+	if (ImportData == nullptr || ImportData->ImportMaterialOriginalNameData.Num() <= 0)
+	{
+		return false;
+	}
+	bool AllNameAreNone = true;
+	for (FName ImportMaterialName : ImportData->ImportMaterialOriginalNameData)
+	{
+		if (ImportMaterialName != NAME_None)
+		{
+			AllNameAreNone = false;
+			break;
+		}
+	}
+	return !AllNameAreNone;
+}
+
+void SkeletalMesUtilsImpl::SaveSkeletalMeshLODModelSections(USkeletalMesh* SourceSkeletalMesh, TSharedPtr<FExistingSkelMeshData>& ExistingMeshDataPtr, int32 LodIndex, bool bSaveNonReducedMeshData)
+{
+	const FSkeletalMeshModel* SourceMeshModel = SourceSkeletalMesh->GetImportedModel();
+	const FSkeletalMeshLODModel* SourceLODModel = &SourceMeshModel->LODModels[LodIndex];
+	FSkeletalMeshLODModel OriginalLODModel;
+
+	if (bSaveNonReducedMeshData && (SourceMeshModel->OriginalReductionSourceMeshData.IsValidIndex(LodIndex) && !SourceMeshModel->OriginalReductionSourceMeshData[LodIndex]->IsEmpty()))
+	{
+		TMap<FString, TArray<FMorphTargetDelta>> TempLODMorphTargetData;
+		//Get the before reduce LODModel, this lod model contain all the possible sections
+		SourceMeshModel->OriginalReductionSourceMeshData[LodIndex]->LoadReductionData(OriginalLODModel, TempLODMorphTargetData, SourceSkeletalMesh);
+		//If there was section that was remove by the reduction (Disabled in the original data, zero triangle after reduction, GenerateUpTo settings...),
+		//we have to use the original section data and apply the section data that was modified after the reduction
+		if (OriginalLODModel.Sections.Num() > SourceLODModel->Sections.Num())
+		{
+			TArray<bool> OriginalMatched;
+			OriginalMatched.AddZeroed(OriginalLODModel.Sections.Num());
+			//Now apply the after reduce settings change, but we need to match the section since there can be reduced one
+			for (int32 ReduceSectionIndex = 0; ReduceSectionIndex < SourceLODModel->Sections.Num(); ++ReduceSectionIndex)
+			{
+				const FSkelMeshSection& ReduceSection = SourceLODModel->Sections[ReduceSectionIndex];
+				for (int32 OriginalSectionIndex = 0; OriginalSectionIndex < OriginalLODModel.Sections.Num(); ++OriginalSectionIndex)
+				{
+					if (OriginalMatched[OriginalSectionIndex])
+					{
+						continue;
+					}
+					FSkelMeshSection& OriginalSection = OriginalLODModel.Sections[OriginalSectionIndex];
+					if ((OriginalSection.bDisabled) || (OriginalSection.GenerateUpToLodIndex != INDEX_NONE && OriginalSection.GenerateUpToLodIndex < LodIndex))
+					{
+						continue;
+					}
+
+					if (ReduceSection.MaterialIndex == OriginalSection.MaterialIndex)
+					{
+						OriginalMatched[OriginalSectionIndex] = true;
+						OriginalSection.bDisabled = ReduceSection.bDisabled;
+						OriginalSection.bCastShadow = ReduceSection.bCastShadow;
+						OriginalSection.bRecomputeTangent = ReduceSection.bRecomputeTangent;
+						OriginalSection.RecomputeTangentsVertexMaskChannel = ReduceSection.RecomputeTangentsVertexMaskChannel;
+						OriginalSection.GenerateUpToLodIndex = ReduceSection.GenerateUpToLodIndex;
+						break;
+					}
+				}
+			}
+			//Set the unmatched original section data using the current UserSectionsData so we keep the user changes
+			for (int32 OriginalSectionIndex = 0; OriginalSectionIndex < OriginalLODModel.Sections.Num(); ++OriginalSectionIndex)
+			{
+				if (OriginalMatched[OriginalSectionIndex])
+				{
+					continue;
+				}
+				FSkelMeshSection& OriginalSection = OriginalLODModel.Sections[OriginalSectionIndex];
+				if (const FSkelMeshSourceSectionUserData* ReduceUserSectionData = SourceLODModel->UserSectionsData.Find(OriginalSection.OriginalDataSectionIndex))
+				{
+					OriginalSection.bDisabled = ReduceUserSectionData->bDisabled;
+					OriginalSection.bCastShadow = ReduceUserSectionData->bCastShadow;
+					OriginalSection.bRecomputeTangent = ReduceUserSectionData->bRecomputeTangent;
+					OriginalSection.RecomputeTangentsVertexMaskChannel = ReduceUserSectionData->RecomputeTangentsVertexMaskChannel;
+					OriginalSection.GenerateUpToLodIndex = ReduceUserSectionData->GenerateUpToLodIndex;
+				}
+			}
+			//Use the OriginalLODModel
+			SourceLODModel = &OriginalLODModel;
+		}
+	}
+	ExistingMeshDataPtr->ExistingImportMeshLodSectionMaterialData.AddZeroed();
+	check(ExistingMeshDataPtr->ExistingImportMeshLodSectionMaterialData.IsValidIndex(LodIndex));
+
+	for (const FSkelMeshSection& CurrentSection : SourceLODModel->Sections)
+	{
+		int32 SectionMaterialIndex = CurrentSection.MaterialIndex;
+		bool SectionCastShadow = CurrentSection.bCastShadow;
+		bool SectionRecomputeTangents = CurrentSection.bRecomputeTangent;
+		ESkinVertexColorChannel RecomputeTangentsVertexMaskChannel = CurrentSection.RecomputeTangentsVertexMaskChannel;
+		int32 GenerateUpTo = CurrentSection.GenerateUpToLodIndex;
+		bool bDisabled = CurrentSection.bDisabled;
+		bool bBoneChunkedSection = CurrentSection.ChunkedParentSectionIndex != INDEX_NONE;
+		//Save all the sections, even the chunked sections
+		if (ExistingMeshDataPtr->ExistingImportMaterialOriginalNameData.IsValidIndex(SectionMaterialIndex))
+		{
+			ExistingMeshDataPtr->ExistingImportMeshLodSectionMaterialData[LodIndex].Emplace(ExistingMeshDataPtr->ExistingImportMaterialOriginalNameData[SectionMaterialIndex], SectionCastShadow, SectionRecomputeTangents, RecomputeTangentsVertexMaskChannel, GenerateUpTo, bDisabled);
+		}
+	}
+}
+
+void SkeletalMesUtilsImpl::SaveSkeletalMeshMaterialNameWorkflowData(TSharedPtr<FExistingSkelMeshData>& ExistingMeshDataPtr, const USkeletalMesh* SourceSkeletalMesh)
+{
+	const UFbxSkeletalMeshImportData* ImportData = Cast<UFbxSkeletalMeshImportData>(SourceSkeletalMesh->GetAssetImportData());
+	if (!ImportData)
+	{
+		return;
+	}
+
+	for (int32 ImportMaterialOriginalNameDataIndex = 0; ImportMaterialOriginalNameDataIndex < ImportData->ImportMaterialOriginalNameData.Num(); ++ImportMaterialOriginalNameDataIndex)
+	{
+		FName MaterialName = ImportData->ImportMaterialOriginalNameData[ImportMaterialOriginalNameDataIndex];
+		ExistingMeshDataPtr->LastImportMaterialOriginalNameData.Add(MaterialName);
+	}
+
+	for (int32 LodIndex = 0; LodIndex < ImportData->ImportMeshLodData.Num(); ++LodIndex)
+	{
+		ExistingMeshDataPtr->LastImportMeshLodSectionMaterialData.AddZeroed();
+		const FImportMeshLodSectionsData &ImportMeshLodSectionsData = ImportData->ImportMeshLodData[LodIndex];
+		for (int32 SectionIndex = 0; SectionIndex < ImportMeshLodSectionsData.SectionOriginalMaterialName.Num(); ++SectionIndex)
+		{
+			FName MaterialName = ImportMeshLodSectionsData.SectionOriginalMaterialName[SectionIndex];
+			ExistingMeshDataPtr->LastImportMeshLodSectionMaterialData[LodIndex].Add(MaterialName);
+		}
+	}
+}
+
+void SkeletalMesUtilsImpl::SaveSkeletalMeshAssetUserData(TSharedPtr<FExistingSkelMeshData>& ExistingMeshDataPtr, const TArray<UAssetUserData*>* UserData)
+{
+	if (!UserData)
+	{
+		return;
+	}
+
+	for (int32 Idx = 0; Idx < UserData->Num(); Idx++)
+	{
+		if ((*UserData)[Idx] != nullptr)
+		{
+			UAssetUserData* DupObject = (UAssetUserData*)StaticDuplicateObject((*UserData)[Idx], GetTransientPackage());
+			bool bAddDupToRoot = !(DupObject->IsRooted());
+			if (bAddDupToRoot)
+			{
+				DupObject->AddToRoot();
+			}
+			ExistingMeshDataPtr->ExistingAssetUserData.Add(DupObject, bAddDupToRoot);
+		}
+	}
+}
+
+TSharedPtr<FExistingSkelMeshData> SkeletalMeshImportUtils::SaveExistingSkelMeshData(USkeletalMesh* SourceSkeletalMesh, bool bSaveMaterials, int32 ReimportLODIndex)
+{
+	using namespace SkeletalMesUtilsImpl;
+
+	if (!SourceSkeletalMesh)
+	{
+		return TSharedPtr<FExistingSkelMeshData>();
+	}
+
+	const int32 SafeReimportLODIndex = ReimportLODIndex < 0 ? 0 : ReimportLODIndex;
+	TSharedPtr<FExistingSkelMeshData> ExistingMeshDataPtr(MakeShared<FExistingSkelMeshData>());
+
+	//Save the package UMetaData
+	TMap<FName, FString>* MetaDataTagValues = UMetaData::GetMapForObject(SourceSkeletalMesh);
+	if (MetaDataTagValues && MetaDataTagValues->Num() > 0)
+	{
+		ExistingMeshDataPtr->ExistingUMetaDataTagValues = *MetaDataTagValues;
+	}
+	ExistingMeshDataPtr->UseMaterialNameSlotWorkflow = SkeletalMeshIsUsingMaterialSlotNameWorkflow(SourceSkeletalMesh->GetAssetImportData());
+	ExistingMeshDataPtr->MinLOD = SourceSkeletalMesh->GetMinLod();
+	ExistingMeshDataPtr->DisableBelowMinLodStripping = SourceSkeletalMesh->GetDisableBelowMinLodStripping();
+	ExistingMeshDataPtr->bOverrideLODStreamingSettings = SourceSkeletalMesh->GetOverrideLODStreamingSettings();
+	ExistingMeshDataPtr->bSupportLODStreaming = SourceSkeletalMesh->GetSupportLODStreaming();
+	ExistingMeshDataPtr->MaxNumStreamedLODs = SourceSkeletalMesh->GetMaxNumStreamedLODs();
+	ExistingMeshDataPtr->MaxNumOptionalLODs = SourceSkeletalMesh->GetMaxNumOptionalLODs();
+
+	const TArray<FSkeletalMaterial>& SourceMaterials = SourceSkeletalMesh->GetMaterials();
+	//Add the existing Material slot name data
+	for (int32 MaterialIndex = 0; MaterialIndex < SourceMaterials.Num(); ++MaterialIndex)
+	{
+		ExistingMeshDataPtr->ExistingImportMaterialOriginalNameData.Add(SourceMaterials[MaterialIndex].ImportedMaterialSlotName);
+	}
+
+	FSkeletalMeshModel* SourceMeshModel = SourceSkeletalMesh->GetImportedModel();
+	for (int32 LodIndex = 0; LodIndex < SourceMeshModel->LODModels.Num(); ++LodIndex)
+	{
+		const bool bImportNonReducedData = LodIndex == SafeReimportLODIndex;
+		SaveSkeletalMeshLODModelSections(SourceSkeletalMesh, ExistingMeshDataPtr, LodIndex, bImportNonReducedData);
+	}
+
+	ExistingMeshDataPtr->ExistingSockets = SourceSkeletalMesh->GetMeshOnlySocketList();
+	ExistingMeshDataPtr->bSaveRestoreMaterials = bSaveMaterials;
+	if (ExistingMeshDataPtr->bSaveRestoreMaterials)
+	{
+		ExistingMeshDataPtr->ExistingMaterials = SourceMaterials;
+	}
+	ExistingMeshDataPtr->ExistingRetargetBasePose = SourceSkeletalMesh->GetRetargetBasePose();
+
+	if (SourceMeshModel->LODModels.Num() > 0 &&
+		SourceSkeletalMesh->GetLODNum() == SourceMeshModel->LODModels.Num())
+	{
+		// Copy LOD models and LOD Infos.
+		check(SourceMeshModel->LODModels.Num() == SourceSkeletalMesh->GetLODInfoArray().Num());
+		ExistingMeshDataPtr->ExistingLODModels.Empty(SourceMeshModel->LODModels.Num());
+		for ( int32 LODIndex = 0; LODIndex < SourceMeshModel->LODModels.Num() ; ++LODIndex)
+		{
+			//const int32 ReductionLODIndex = LODModelIndex + OffsetReductionLODIndex;
+			TSharedPtr<FReductionBaseSkeletalMeshBulkData> ReductionLODData;
+			if (SourceMeshModel->OriginalReductionSourceMeshData.IsValidIndex(LODIndex) && !SourceMeshModel->OriginalReductionSourceMeshData[LODIndex]->IsEmpty())
+			{
+				FSkeletalMeshLODModel BaseLODModel;
+				TMap<FString, TArray<FMorphTargetDelta>> BaseLODMorphTargetData;
+				SourceMeshModel->OriginalReductionSourceMeshData[LODIndex]->LoadReductionData(BaseLODModel, BaseLODMorphTargetData, SourceSkeletalMesh);
+				ReductionLODData = MakeShared<FReductionBaseSkeletalMeshBulkData>();
+				ReductionLODData->SaveReductionData(BaseLODModel, BaseLODMorphTargetData, SourceSkeletalMesh);
+			}
+			//Add the reduction source mesh data if it exist, otherwise an empty sharedPtr.
+			ExistingMeshDataPtr->ExistingOriginalReductionSourceMeshData.Add(MoveTemp(ReductionLODData));
+			
+			//Add a new LOD Model to the existing LODModels data
+			const FSkeletalMeshLODModel& LODModel = SourceMeshModel->LODModels[LODIndex];
+			ExistingMeshDataPtr->ExistingLODModels.Add(FSkeletalMeshLODModel::CreateCopy(&LODModel));
+		}
+		check(ExistingMeshDataPtr->ExistingLODModels.Num() == SourceMeshModel->LODModels.Num());
+
+		ExistingMeshDataPtr->ExistingLODInfo = SourceSkeletalMesh->GetLODInfoArray();
+		ExistingMeshDataPtr->ExistingRefSkeleton = SourceSkeletalMesh->GetRefSkeleton();
+	}
+
+	// First asset should be the one that the skeletal mesh should point too
+	ExistingMeshDataPtr->ExistingPhysicsAssets.Empty();
+	ExistingMeshDataPtr->ExistingPhysicsAssets.Add(SourceSkeletalMesh->GetPhysicsAsset());
+	for (TObjectIterator<UPhysicsAsset> It; It; ++It)
+	{
+		UPhysicsAsset* PhysicsAsset = *It;
+		if (PhysicsAsset->PreviewSkeletalMesh == SourceSkeletalMesh && SourceSkeletalMesh->GetPhysicsAsset() != PhysicsAsset)
+		{
+			ExistingMeshDataPtr->ExistingPhysicsAssets.Add(PhysicsAsset);
+		}
+	}
+
+	ExistingMeshDataPtr->ExistingShadowPhysicsAsset = SourceSkeletalMesh->GetShadowPhysicsAsset();
+	ExistingMeshDataPtr->ExistingSkeleton = SourceSkeletalMesh->GetSkeleton();
+	// since copying back original skeleton, this should be safe to do
+	ExistingMeshDataPtr->ExistingPostProcessAnimBlueprint = SourceSkeletalMesh->GetPostProcessAnimBlueprint();
+	ExistingMeshDataPtr->ExistingLODSettings = SourceSkeletalMesh->GetLODSettings();
+
+	ExistingMeshDataPtr->ExistingMorphTargets = SourceSkeletalMesh->GetMorphTargets();
+	ExistingMeshDataPtr->ExistingAssetImportData = SourceSkeletalMesh->GetAssetImportData();
+	ExistingMeshDataPtr->ExistingThumbnailInfo = SourceSkeletalMesh->GetThumbnailInfo();
+	ExistingMeshDataPtr->ExistingClothingAssets = SourceSkeletalMesh->GetMeshClothingAssets();
+	ExistingMeshDataPtr->ExistingSamplingInfo = SourceSkeletalMesh->GetSamplingInfo();
+	ExistingMeshDataPtr->ExistingDefaultAnimatingRig = SourceSkeletalMesh->GetDefaultAnimatingRig();
+
+	if (ExistingMeshDataPtr->UseMaterialNameSlotWorkflow)
+	{
+		//Add the last fbx import data
+		SaveSkeletalMeshMaterialNameWorkflowData(ExistingMeshDataPtr, SourceSkeletalMesh);
+	}
+
+	//Store the user asset data
+	SaveSkeletalMeshAssetUserData(ExistingMeshDataPtr, SourceSkeletalMesh->GetAssetUserDataArray());
+	
+	//Store mesh changed delegate data
+	ExistingMeshDataPtr->ExistingOnMeshChanged = SourceSkeletalMesh->GetOnMeshChanged();
+
+	//Store mesh bounds extensions
+	ExistingMeshDataPtr->PositiveBoundsExtension = SourceSkeletalMesh->GetPositiveBoundsExtension();
+	ExistingMeshDataPtr->NegativeBoundsExtension = SourceSkeletalMesh->GetNegativeBoundsExtension();
+
+	return ExistingMeshDataPtr;
+}
+
+void SkeletalMesUtilsImpl::RestoreDependentLODs(const TSharedPtr<const FExistingSkelMeshData>& MeshData, USkeletalMesh* SkeletalMesh)
+{
+	check(SkeletalMesh != nullptr);
+	const int32 TotalLOD = MeshData->ExistingLODModels.Num();
+	FSkeletalMeshModel* SkeletalMeshImportedModel = SkeletalMesh->GetImportedModel();
+
+	for (int32 LODIndex = 1; LODIndex < TotalLOD; ++LODIndex)
+	{
+		if (LODIndex >= SkeletalMesh->GetLODInfoArray().Num())
+		{
+			// Create a copy of LODInfo and reset material maps, it won't work anyway. 
+			FSkeletalMeshLODInfo ExistLODInfo = MeshData->ExistingLODInfo[LODIndex];
+			ExistLODInfo.LODMaterialMap.Empty();
+			// add LOD info back
+			SkeletalMesh->AddLODInfo(MoveTemp(ExistLODInfo));
+			check(LODIndex < SkeletalMesh->GetLODInfoArray().Num());
+
+			const FSkeletalMeshLODModel& ExistLODModel = MeshData->ExistingLODModels[LODIndex];
+			SkeletalMeshImportedModel->LODModels.Add(FSkeletalMeshLODModel::CreateCopy(&ExistLODModel));
+		}
+	}
+}
+
+void SkeletalMesUtilsImpl::RestoreLODInfo(const TSharedPtr<const FExistingSkelMeshData>& MeshData, USkeletalMesh* SkeletalMesh, int32 LodIndex)
+{
+	FSkeletalMeshLODInfo& ImportedLODInfo = SkeletalMesh->GetLODInfoArray()[LodIndex];
+	if (!MeshData->ExistingLODInfo.IsValidIndex(LodIndex))
+	{
+		return;
+	}
+
+	const FSkeletalMeshLODInfo& ExistingLODInfo = MeshData->ExistingLODInfo[LodIndex];
+
+	ImportedLODInfo.ScreenSize = ExistingLODInfo.ScreenSize;
+	ImportedLODInfo.LODHysteresis = ExistingLODInfo.LODHysteresis;
+	ImportedLODInfo.BuildSettings = ExistingLODInfo.BuildSettings;
+	//Old assets may have non-applied reduction settings, so only restore the reduction settings if the LOD was effectively reduced.
+	if (ExistingLODInfo.bHasBeenSimplified)
+	{
+		ImportedLODInfo.ReductionSettings = ExistingLODInfo.ReductionSettings;
+	}
+	ImportedLODInfo.BonesToRemove = ExistingLODInfo.BonesToRemove;
+	ImportedLODInfo.BonesToPrioritize = ExistingLODInfo.BonesToPrioritize;
+	ImportedLODInfo.WeightOfPrioritization = ExistingLODInfo.WeightOfPrioritization;
+	ImportedLODInfo.BakePose = ExistingLODInfo.BakePose;
+	ImportedLODInfo.BakePoseOverride = ExistingLODInfo.BakePoseOverride;
+	ImportedLODInfo.SourceImportFilename = ExistingLODInfo.SourceImportFilename;
+	ImportedLODInfo.SkinCacheUsage = ExistingLODInfo.SkinCacheUsage;
+	ImportedLODInfo.bAllowCPUAccess = ExistingLODInfo.bAllowCPUAccess;
+	ImportedLODInfo.bSupportUniformlyDistributedSampling = ExistingLODInfo.bSupportUniformlyDistributedSampling;
+}
+
+void SkeletalMeshImportUtils::ApplySkinning(USkeletalMesh* SkeletalMesh, FSkeletalMeshLODModel& SrcLODModel, FSkeletalMeshLODModel& DestLODModel)
+{
+	TArray<FSoftSkinVertex> SrcVertices;
+	SrcLODModel.GetVertices(SrcVertices);
+
+	FBox OldBounds(EForceInit::ForceInit);
+	for (int32 SrcIndex = 0; SrcIndex < SrcVertices.Num(); ++SrcIndex)
+	{
+		const FSoftSkinVertex& SrcVertex = SrcVertices[SrcIndex];
+		OldBounds += SrcVertex.Position;
+	}
+
+	TWedgeInfoPosOctree SrcWedgePosOctree(OldBounds.GetCenter(), OldBounds.GetExtent().GetMax());
+	// Add each old vertex to the octree
+	for (int32 SrcIndex = 0; SrcIndex < SrcVertices.Num(); ++SrcIndex)
+	{
+		FWedgeInfo WedgeInfo;
+		WedgeInfo.WedgeIndex = SrcIndex;
+		WedgeInfo.Position = SrcVertices[SrcIndex].Position;
+		SrcWedgePosOctree.AddElement(WedgeInfo);
+	}
+
+	FOctreeQueryHelper OctreeQueryHelper(&SrcWedgePosOctree);
+
+	TArray<FBoneIndexType> RequiredActiveBones;
+
+	bool bUseBone = false;
+	for (int32 SectionIndex = 0; SectionIndex < DestLODModel.Sections.Num(); SectionIndex++)
+	{
+		FSkelMeshSection& Section = DestLODModel.Sections[SectionIndex];
+		Section.BoneMap.Reset();
+		for (FSoftSkinVertex& DestVertex : Section.SoftVertices)
+		{
+			//Find the nearest wedges in the src model
+			TArray<FWedgeInfo> NearestSrcWedges;
+			OctreeQueryHelper.FindNearestWedgeIndexes(DestVertex.Position, NearestSrcWedges);
+			if (NearestSrcWedges.Num() < 1)
+			{
+				//Should we check???
+				continue;
+			}
+			//Find the matching wedges in the src model
+			int32 MatchingSrcWedge = INDEX_NONE;
+			for (FWedgeInfo& SrcWedgeInfo : NearestSrcWedges)
+			{
+				int32 SrcIndex = SrcWedgeInfo.WedgeIndex;
+				const FSoftSkinVertex& SrcVertex = SrcVertices[SrcIndex];
+				if (SrcVertex.Position.Equals(DestVertex.Position, THRESH_POINTS_ARE_SAME) &&
+					SrcVertex.UVs[0].Equals(DestVertex.UVs[0], THRESH_UVS_ARE_SAME) &&
+					(SrcVertex.TangentX == DestVertex.TangentX) &&
+					(SrcVertex.TangentY == DestVertex.TangentY) &&
+					(SrcVertex.TangentZ == DestVertex.TangentZ))
+				{
+					MatchingSrcWedge = SrcIndex;
+					break;
+				}
+			}
+			if (MatchingSrcWedge == INDEX_NONE)
+			{
+				//We have to find the nearest wedges, then find the most similar normal
+				float MinDistance = MAX_FLT;
+				float MinNormalAngle = MAX_FLT;
+				for (FWedgeInfo& SrcWedgeInfo : NearestSrcWedges)
+				{
+					int32 SrcIndex = SrcWedgeInfo.WedgeIndex;
+					const FSoftSkinVertex& SrcVertex = SrcVertices[SrcIndex];
+					float VectorDelta = FVector::DistSquared(SrcVertex.Position, DestVertex.Position);
+					if (VectorDelta <= (MinDistance + KINDA_SMALL_NUMBER))
+					{
+						if (VectorDelta < MinDistance - KINDA_SMALL_NUMBER)
+						{
+							MinDistance = VectorDelta;
+							MinNormalAngle = MAX_FLT;
+						}
+						FVector DestTangentZ = DestVertex.TangentZ;
+						DestTangentZ.Normalize();
+						FVector SrcTangentZ = SrcVertex.TangentZ;
+						SrcTangentZ.Normalize();
+						float AngleDiff = FMath::Abs(FMath::Acos(FVector::DotProduct(DestTangentZ, SrcTangentZ)));
+						if (AngleDiff < MinNormalAngle)
+						{
+							MinNormalAngle = AngleDiff;
+							MatchingSrcWedge = SrcIndex;
+						}
+					}
+				}
+			}
+			check(SrcVertices.IsValidIndex(MatchingSrcWedge));
+			const FSoftSkinVertex& SrcVertex = SrcVertices[MatchingSrcWedge];
+
+			//Find the src section to assign the correct remapped bone
+			int32 SrcSectionIndex = INDEX_NONE;
+			int32 SrcSectionWedgeIndex = INDEX_NONE;
+			SrcLODModel.GetSectionFromVertexIndex(MatchingSrcWedge, SrcSectionIndex, SrcSectionWedgeIndex);
+			check(SrcSectionIndex != INDEX_NONE);
+
+			for (int32 InfluenceIndex = 0; InfluenceIndex < MAX_TOTAL_INFLUENCES; ++InfluenceIndex)
+			{
+				if (SrcVertex.InfluenceWeights[InfluenceIndex] > 0.0f)
+				{
+					Section.MaxBoneInfluences = FMath::Max(Section.MaxBoneInfluences, InfluenceIndex + 1);
+					//Copy the weight
+					DestVertex.InfluenceWeights[InfluenceIndex] = SrcVertex.InfluenceWeights[InfluenceIndex];
+					//Copy the bone ID
+					FBoneIndexType OriginalBoneIndex = SrcLODModel.Sections[SrcSectionIndex].BoneMap[SrcVertex.InfluenceBones[InfluenceIndex]];
+					int32 OverrideIndex;
+					if (Section.BoneMap.Find(OriginalBoneIndex, OverrideIndex))
+					{
+						DestVertex.InfluenceBones[InfluenceIndex] = OverrideIndex;
+					}
+					else
+					{
+						DestVertex.InfluenceBones[InfluenceIndex] = Section.BoneMap.Add(OriginalBoneIndex);
+						DestLODModel.ActiveBoneIndices.AddUnique(OriginalBoneIndex);
+					}
+					bUseBone = true;
+				}
+			}
+		}
+	}
+
+	if (bUseBone)
+	{
+		//Set the required/active bones
+		DestLODModel.RequiredBones = SrcLODModel.RequiredBones;
+		DestLODModel.RequiredBones.Sort();
+		SkeletalMesh->GetRefSkeleton().EnsureParentsExistAndSort(DestLODModel.ActiveBoneIndices);
+	}
+}
+
+
+void SkeletalMeshImportUtils::RestoreExistingSkelMeshData(const TSharedPtr<const FExistingSkelMeshData>& MeshData, USkeletalMesh* SkeletalMesh, int32 ReimportLODIndex, bool bCanShowDialog, bool bImportSkinningOnly, bool bForceMaterialReset)
+{
+	using namespace SkeletalMesUtilsImpl;
+	if (!MeshData || !SkeletalMesh)
+	{
+		return;
+	}
+
+	//Restore the package metadata
+	if (MeshData->ExistingUMetaDataTagValues.Num() > 0)
+	{
+		UMetaData* PackageMetaData = SkeletalMesh->GetOutermost()->GetMetaData();
+		checkSlow(PackageMetaData);
+		PackageMetaData->SetObjectValues(SkeletalMesh, MeshData->ExistingUMetaDataTagValues);
+	}
+
+	int32 SafeReimportLODIndex = ReimportLODIndex < 0 ? 0 : ReimportLODIndex;
+	SkeletalMesh->SetMinLod(MeshData->MinLOD);
+	SkeletalMesh->SetDisableBelowMinLodStripping(MeshData->DisableBelowMinLodStripping);
+	SkeletalMesh->SetOverrideLODStreamingSettings(MeshData->bOverrideLODStreamingSettings);
+	SkeletalMesh->SetSupportLODStreaming(MeshData->bSupportLODStreaming);
+	SkeletalMesh->SetMaxNumStreamedLODs(MeshData->MaxNumStreamedLODs);
+	SkeletalMesh->SetMaxNumOptionalLODs(MeshData->MaxNumOptionalLODs);
+
+	FSkeletalMeshModel* SkeletalMeshImportedModel = SkeletalMesh->GetImportedModel();
+
+	//Create a remap material Index use to find the matching section later
+	TArray<int32> RemapMaterial;
+	RemapMaterial.AddZeroed(SkeletalMesh->GetMaterials().Num());
+	TArray<FName> RemapMaterialName;
+	RemapMaterialName.AddZeroed(SkeletalMesh->GetMaterials().Num());
+
+	bool bMaterialReset = false;
+	if (MeshData->bSaveRestoreMaterials)
+	{
+		UnFbx::EFBXReimportDialogReturnOption ReturnOption;
+		//Ask the user to match the materials conflict
+		UnFbx::FFbxImporter::PrepareAndShowMaterialConflictDialog<FSkeletalMaterial>(MeshData->ExistingMaterials, SkeletalMesh->GetMaterials(), RemapMaterial, RemapMaterialName, bCanShowDialog, false, bForceMaterialReset, ReturnOption);
+
+		if (ReturnOption != UnFbx::EFBXReimportDialogReturnOption::FBXRDRO_ResetToFbx)
+		{
+			//Build a ordered material list that try to keep intact the existing material list
+			TArray<FSkeletalMaterial> MaterialOrdered;
+			TArray<bool> MatchedNewMaterial;
+			MatchedNewMaterial.AddZeroed(SkeletalMesh->GetMaterials().Num());
+			for (int32 ExistMaterialIndex = 0; ExistMaterialIndex < MeshData->ExistingMaterials.Num(); ++ExistMaterialIndex)
+			{
+				int32 MaterialIndexOrdered = MaterialOrdered.Add(MeshData->ExistingMaterials[ExistMaterialIndex]);
+				FSkeletalMaterial& OrderedMaterial = MaterialOrdered[MaterialIndexOrdered];
+				int32 NewMaterialIndex = INDEX_NONE;
+				if (RemapMaterial.Find(ExistMaterialIndex, NewMaterialIndex))
+				{
+					MatchedNewMaterial[NewMaterialIndex] = true;
+					RemapMaterial[NewMaterialIndex] = MaterialIndexOrdered;
+					OrderedMaterial.ImportedMaterialSlotName = SkeletalMesh->GetMaterials()[NewMaterialIndex].ImportedMaterialSlotName;
+				}
+				else
+				{
+					//Unmatched material must be conserve
+				}
+			}
+
+			//Add the new material entries (the one that do not match with any existing material)
+			for (int32 NewMaterialIndex = 0; NewMaterialIndex < MatchedNewMaterial.Num(); ++NewMaterialIndex)
+			{
+				if (MatchedNewMaterial[NewMaterialIndex] == false)
+				{
+					int32 NewMeshIndex = MaterialOrdered.Add(SkeletalMesh->GetMaterials()[NewMaterialIndex]);
+					RemapMaterial[NewMaterialIndex] = NewMeshIndex;
+				}
+			}
+
+			//Set the RemapMaterialName array helper
+			for (int32 MaterialIndex = 0; MaterialIndex < RemapMaterial.Num(); ++MaterialIndex)
+			{
+				int32 SourceMaterialMatch = RemapMaterial[MaterialIndex];
+				if (MeshData->ExistingMaterials.IsValidIndex(SourceMaterialMatch))
+				{
+					RemapMaterialName[MaterialIndex] = MeshData->ExistingMaterials[SourceMaterialMatch].ImportedMaterialSlotName;
+				}
+			}
+
+			//Copy the re ordered materials (this ensure the material array do not change when we re-import)
+			SkeletalMesh->SetMaterials(MaterialOrdered);
+		}
+		else
+		{
+			bMaterialReset = true;
+		}
+	}
+
+	SkeletalMesh->SetLODSettings(MeshData->ExistingLODSettings);
+	// ensure LOD 0 contains correct setting 
+	if (SkeletalMesh->GetLODSettings() && SkeletalMesh->GetLODInfoArray().Num() > 0)
+	{
+		SkeletalMesh->GetLODSettings()->SetLODSettingsToMesh(SkeletalMesh, 0);
+	}
+
+	//Do everything we need for base LOD re-import
+	if (SafeReimportLODIndex == 0)
+	{
+		// this is not ideal. Ideally we'll have to save only diff with indicating which joints, 
+		// but for now, we allow them to keep the previous pose IF the element count is same
+		if (MeshData->ExistingRetargetBasePose.Num() == SkeletalMesh->GetRefSkeleton().GetRawBoneNum())
+		{
+			SkeletalMesh->SetRetargetBasePose(MeshData->ExistingRetargetBasePose);
+		}
+
+		// Assign sockets from old version of this SkeletalMesh.
+		// Only copy ones for bones that exist in the new mesh.
+		for (int32 i = 0; i < MeshData->ExistingSockets.Num(); i++)
+		{
+			const int32 BoneIndex = SkeletalMesh->GetRefSkeleton().FindBoneIndex(MeshData->ExistingSockets[i]->BoneName);
+			if (BoneIndex != INDEX_NONE)
+			{
+				SkeletalMesh->GetMeshOnlySocketList().Add(MeshData->ExistingSockets[i]);
+			}
+		}
+
+		// We copy back and fix-up the LODs that still work with this skeleton.
+		if (MeshData->ExistingLODModels.Num() > 1)
+		{
+			auto RestoreReductionSourceData = [&SkeletalMesh, &SkeletalMeshImportedModel, &MeshData](int32 ExistingIndex, int32 NewIndex)
+			{
+				if (!MeshData->ExistingOriginalReductionSourceMeshData[ExistingIndex].IsValid() || MeshData->ExistingOriginalReductionSourceMeshData[ExistingIndex]->IsEmpty())
+				{
+					return;
+				}
+				//Restore the original reduction source mesh data
+				FSkeletalMeshLODModel BaseLODModel;
+				TMap<FString, TArray<FMorphTargetDelta>> BaseLODMorphTargetData;
+				MeshData->ExistingOriginalReductionSourceMeshData[ExistingIndex]->LoadReductionData(BaseLODModel, BaseLODMorphTargetData, SkeletalMesh);
+				FReductionBaseSkeletalMeshBulkData* ReductionLODData = new FReductionBaseSkeletalMeshBulkData();
+				ReductionLODData->SaveReductionData(BaseLODModel, BaseLODMorphTargetData, SkeletalMesh);
+				//Add necessary empty slot
+				while (SkeletalMeshImportedModel->OriginalReductionSourceMeshData.Num() < NewIndex)
+				{
+					FReductionBaseSkeletalMeshBulkData* EmptyReductionLODData = new FReductionBaseSkeletalMeshBulkData();
+					SkeletalMeshImportedModel->OriginalReductionSourceMeshData.Add(EmptyReductionLODData);
+				}
+				SkeletalMeshImportedModel->OriginalReductionSourceMeshData.Add(ReductionLODData);
+			};
+
+			if (SkeletonsAreCompatible(SkeletalMesh->GetRefSkeleton(), MeshData->ExistingRefSkeleton, bImportSkinningOnly))
+			{
+				// First create mapping table from old skeleton to new skeleton.
+				TArray<int32> OldToNewMap;
+				OldToNewMap.AddUninitialized(MeshData->ExistingRefSkeleton.GetRawBoneNum());
+				for (int32 i = 0; i < MeshData->ExistingRefSkeleton.GetRawBoneNum(); i++)
+				{
+					OldToNewMap[i] = SkeletalMesh->GetRefSkeleton().FindBoneIndex(MeshData->ExistingRefSkeleton.GetBoneName(i));
+				}
+
+				//Starting at index 1 because we only need to add LOD models of LOD 1 and higher.
+				for (int32 LODIndex = 1; LODIndex < MeshData->ExistingLODModels.Num(); ++LODIndex)
+				{
+					FSkeletalMeshLODModel* LODModelCopy = FSkeletalMeshLODModel::CreateCopy(&MeshData->ExistingLODModels[LODIndex]);
+					const FSkeletalMeshLODInfo& LODInfo = MeshData->ExistingLODInfo[LODIndex];
+
+					// Fix ActiveBoneIndices array.
+					bool bMissingBone = false;
+					FName MissingBoneName = NAME_None;
+					for (int32 j = 0; j < LODModelCopy->ActiveBoneIndices.Num() && !bMissingBone; j++)
+					{
+						int32 OldActiveBoneIndex = LODModelCopy->ActiveBoneIndices[j];
+						if (OldToNewMap.IsValidIndex(OldActiveBoneIndex))
+						{
+							int32 NewBoneIndex = OldToNewMap[OldActiveBoneIndex];
+							if (NewBoneIndex == INDEX_NONE)
+							{
+								bMissingBone = true;
+								MissingBoneName = MeshData->ExistingRefSkeleton.GetBoneName(LODModelCopy->ActiveBoneIndices[j]);
+							}
+							else
+							{
+								LODModelCopy->ActiveBoneIndices[j] = NewBoneIndex;
+							}
+						}
+						else
+						{
+							LODModelCopy->ActiveBoneIndices.RemoveAt(j, 1, false);
+							--j;
+						}
+					}
+
+					// Fix RequiredBones array.
+					for (int32 j = 0; j < LODModelCopy->RequiredBones.Num() && !bMissingBone; j++)
+					{
+						const int32 OldBoneIndex = LODModelCopy->RequiredBones[j];
+
+						if (OldToNewMap.IsValidIndex(OldBoneIndex))	//Previously virtual bones could end up in this array
+																	// Must validate against this
+						{
+							const int32 NewBoneIndex = OldToNewMap[OldBoneIndex];
+							if (NewBoneIndex == INDEX_NONE)
+							{
+								bMissingBone = true;
+								MissingBoneName = MeshData->ExistingRefSkeleton.GetBoneName(OldBoneIndex);
+							}
+							else
+							{
+								LODModelCopy->RequiredBones[j] = NewBoneIndex;
+							}
+						}
+						else
+						{
+							//Bone didn't exist in our required bones, clean up. 
+							LODModelCopy->RequiredBones.RemoveAt(j, 1, false);
+							--j;
+						}
+					}
+
+					// Sort ascending for parent child relationship
+					LODModelCopy->RequiredBones.Sort();
+					SkeletalMesh->GetRefSkeleton().EnsureParentsExistAndSort(LODModelCopy->ActiveBoneIndices);
+
+					// Fix the sections' BoneMaps.
+					for (int32 SectionIndex = 0; SectionIndex < LODModelCopy->Sections.Num(); SectionIndex++)
+					{
+						FSkelMeshSection& Section = LODModelCopy->Sections[SectionIndex];
+						for (int32 BoneIndex = 0; BoneIndex < Section.BoneMap.Num(); BoneIndex++)
+						{
+							int32 NewBoneIndex = OldToNewMap[Section.BoneMap[BoneIndex]];
+							if (NewBoneIndex == INDEX_NONE)
+							{
+								bMissingBone = true;
+								MissingBoneName = MeshData->ExistingRefSkeleton.GetBoneName(Section.BoneMap[BoneIndex]);
+								break;
+							}
+							else
+							{
+								Section.BoneMap[BoneIndex] = NewBoneIndex;
+							}
+						}
+						if (bMissingBone)
+						{
+							break;
+						}
+					}
+
+					if (bMissingBone)
+					{
+						UnFbx::FFbxImporter* FFbxImporter = UnFbx::FFbxImporter::GetInstance();
+						FFbxImporter->AddTokenizedErrorMessage(FTokenizedMessage::Create(EMessageSeverity::Warning, FText::Format(LOCTEXT("NewMeshMissingBoneFromLOD", "New mesh is missing bone '{0}' required by an LOD."), FText::FromName(MissingBoneName))), FFbxErrors::SkeletalMesh_LOD_MissingBone);
+						break;
+					}
+					else
+					{
+						//We need to add LODInfo
+						SkeletalMeshImportedModel->LODModels.Add(LODModelCopy);
+						SkeletalMesh->AddLODInfo(LODInfo);
+						RestoreReductionSourceData(LODIndex, SkeletalMesh->GetLODNum() - 1);
+					}
+				}
+			}
+			//We just need to restore the LOD model and LOD info the build should regenerate the LODs
+			RestoreDependentLODs(MeshData, SkeletalMesh);
+			
+			
+			//Old asset cannot use the new build system, we need to regenerate dependent LODs
+			if (SkeletalMesh->IsLODImportedDataBuildAvailable(SafeReimportLODIndex) == false)
+			{
+				FLODUtilities::RegenerateDependentLODs(SkeletalMesh, SafeReimportLODIndex, GetTargetPlatformManagerRef().GetRunningTargetPlatform());
+			}
+		}
+
+		for (int32 AssetIndex = 0; AssetIndex < MeshData->ExistingPhysicsAssets.Num(); ++AssetIndex)
+		{
+			UPhysicsAsset* PhysicsAsset = MeshData->ExistingPhysicsAssets[AssetIndex];
+			if (AssetIndex == 0)
+			{
+				// First asset is the one that the skeletal mesh should point too
+				SkeletalMesh->SetPhysicsAsset(PhysicsAsset);
+			}
+			// No need to mark as modified here, because the asset hasn't actually changed
+			if (PhysicsAsset)
+			{
+				PhysicsAsset->PreviewSkeletalMesh = SkeletalMesh;
+			}
+		}
+
+		SkeletalMesh->SetShadowPhysicsAsset(MeshData->ExistingShadowPhysicsAsset);
+		SkeletalMesh->SetSkeleton(MeshData->ExistingSkeleton);
+		SkeletalMesh->SetPostProcessAnimBlueprint(MeshData->ExistingPostProcessAnimBlueprint);
+		SkeletalMesh->SetDefaultAnimatingRig(MeshData->ExistingDefaultAnimatingRig);
+
+		SkeletalMesh->GetMorphTargets().Empty(MeshData->ExistingMorphTargets.Num());
+		SkeletalMesh->GetMorphTargets().Append(MeshData->ExistingMorphTargets);
+		SkeletalMesh->InitMorphTargets();
+		SkeletalMesh->SetAssetImportData(MeshData->ExistingAssetImportData.Get());
+		SkeletalMesh->SetThumbnailInfo(MeshData->ExistingThumbnailInfo.Get());
+		SkeletalMesh->SetMeshClothingAssets(MeshData->ExistingClothingAssets);
+
+		for (UClothingAssetBase* ClothingAsset : SkeletalMesh->GetMeshClothingAssets())
+		{
+			if (ClothingAsset)
+			{
+				ClothingAsset->RefreshBoneMapping(SkeletalMesh);
+			}
+		}
+
+		SkeletalMesh->SetSamplingInfo(MeshData->ExistingSamplingInfo);
+	}
+
+	//Restore the section change only for the reimport LOD, other LOD are not affected since the material array can only grow.
+	if (MeshData->UseMaterialNameSlotWorkflow)
+	{
+		RestoreMaterialNameWorkflowSection(MeshData, SkeletalMesh, SafeReimportLODIndex, RemapMaterial, bMaterialReset);
+	}
+
+	//Copy back the reimported LOD's specific data
+	if (SkeletalMesh->GetLODInfoArray().IsValidIndex(SafeReimportLODIndex))
+	{
+		RestoreLODInfo(MeshData, SkeletalMesh, SafeReimportLODIndex);
+	}
+
+	// Copy user data to newly created mesh
+	for (auto Kvp : MeshData->ExistingAssetUserData)
+	{
+		UAssetUserData* UserDataObject = Kvp.Key;
+		if (Kvp.Value)
+		{
+			//if the duplicated temporary UObject was add to root, we must remove it from the root
+			UserDataObject->RemoveFromRoot();
+		}
+		UserDataObject->Rename(nullptr, SkeletalMesh, REN_DontCreateRedirectors | REN_DoNotDirty);
+		SkeletalMesh->AddAssetUserData(UserDataObject);
+	}
+
+	if (!bImportSkinningOnly && (!MeshData->ExistingLODInfo.IsValidIndex(SafeReimportLODIndex) || !MeshData->ExistingLODInfo[SafeReimportLODIndex].bHasBeenSimplified))
+	{
+		if (SkeletalMeshImportedModel->OriginalReductionSourceMeshData.IsValidIndex(SafeReimportLODIndex))
+		{
+			SkeletalMeshImportedModel->OriginalReductionSourceMeshData[SafeReimportLODIndex]->EmptyBulkData();
+		}
+	}
+
+	//Copy mesh changed delegate data
+	SkeletalMesh->GetOnMeshChanged() = MeshData->ExistingOnMeshChanged;
+
+	//Copy mesh bounds extensions
+	SkeletalMesh->SetPositiveBoundsExtension(MeshData->PositiveBoundsExtension);
+	SkeletalMesh->SetNegativeBoundsExtension(MeshData->NegativeBoundsExtension);
+}
+
+void SkeletalMesUtilsImpl::RestoreMaterialNameWorkflowSection(const TSharedPtr<const FExistingSkelMeshData>& MeshData, USkeletalMesh* SkeletalMesh, int32 LodIndex, TArray<int32>& RemapMaterial, bool bMaterialReset)
+{
+	FSkeletalMeshModel* SkeletalMeshImportedModel = SkeletalMesh->GetImportedModel();
+	FSkeletalMeshLODModel &SkeletalMeshLodModel = SkeletalMeshImportedModel->LODModels[LodIndex];
+
+	//Restore the base LOD materialMap the LODs LODMaterialMap are restore differently
+	if (LodIndex == 0 && SkeletalMesh->GetLODInfoArray().IsValidIndex(LodIndex))
+	{
+		FSkeletalMeshLODInfo& BaseLODInfo = SkeletalMesh->GetLODInfoArray()[LodIndex];
+		if (bMaterialReset)
+		{
+			//If we reset the material array there is no point keeping the user changes
+			BaseLODInfo.LODMaterialMap.Empty();
+		}
+		else if (SkeletalMeshImportedModel->LODModels.IsValidIndex(LodIndex))
+		{
+			//Restore the Base MaterialMap
+			for (int32 SectionIndex = 0; SectionIndex < SkeletalMeshLodModel.Sections.Num(); ++SectionIndex)
+			{
+				int32 MaterialIndex = SkeletalMeshLodModel.Sections[SectionIndex].MaterialIndex;
+				if (MeshData->ExistingLODInfo[LodIndex].LODMaterialMap.IsValidIndex(SectionIndex))
+				{
+					int32 ExistingLODMaterialIndex = MeshData->ExistingLODInfo[LodIndex].LODMaterialMap[SectionIndex];
+					while (BaseLODInfo.LODMaterialMap.Num() <= SectionIndex)
+					{
+						BaseLODInfo.LODMaterialMap.Add(INDEX_NONE);
+					}
+					BaseLODInfo.LODMaterialMap[SectionIndex] = ExistingLODMaterialIndex;
+				}
+			}
+		}
+	}
+
+	const bool bIsValidSavedSectionMaterialData = MeshData->ExistingImportMeshLodSectionMaterialData.IsValidIndex(LodIndex) && MeshData->LastImportMeshLodSectionMaterialData.IsValidIndex(LodIndex);
+	const int32 MaxExistSectionNumber = bIsValidSavedSectionMaterialData ? FMath::Max(MeshData->ExistingImportMeshLodSectionMaterialData[LodIndex].Num(), MeshData->LastImportMeshLodSectionMaterialData[LodIndex].Num()) : 0;
+	TBitArray<> MatchedExistSectionIndex;
+	MatchedExistSectionIndex.Init(false, MaxExistSectionNumber);
+
+	const TArray<FSkeletalMaterial>& SkeletalMeshMaterials = SkeletalMesh->GetMaterials();
+	//Restore the section changes from the old import data
+	for (int32 SectionIndex = 0; SectionIndex < SkeletalMeshLodModel.Sections.Num(); SectionIndex++)
+	{
+		//Find the import section material index by using the RemapMaterial array. Fallback on the imported index if the remap entry is not valid
+		FSkelMeshSection& NewSection = SkeletalMeshLodModel.Sections[SectionIndex];
+		int32 RemapMaterialIndex = RemapMaterial.IsValidIndex(NewSection.MaterialIndex) ? RemapMaterial[NewSection.MaterialIndex] : NewSection.MaterialIndex;
+		if (!SkeletalMeshMaterials.IsValidIndex(RemapMaterialIndex))
+		{
+			//We have an invalid material section, in this case we set the material index to 0
+			NewSection.MaterialIndex = 0;
+			UE_LOG(LogSkeletalMeshImport, Display, TEXT("Reimport material match issue: Invalid RemapMaterialIndex [%d], will make it point to material index [0]"), RemapMaterialIndex);
+			continue;
+		}
+		NewSection.MaterialIndex = RemapMaterialIndex;
+
+		//skip the rest of the loop if we do not have valid saved data
+		if (!bIsValidSavedSectionMaterialData)
+		{
+			continue;
+		}
+		//Get the RemapMaterial section Imported material slot name. We need it to match the saved existing section, so we can put back the saved existing section data
+		FName CurrentSectionImportedMaterialName = SkeletalMeshMaterials[RemapMaterialIndex].ImportedMaterialSlotName;
+		for (int32 ExistSectionIndex = 0; ExistSectionIndex < MaxExistSectionNumber; ++ExistSectionIndex)
+		{
+			//Skip already matched exist section
+			if (MatchedExistSectionIndex[ExistSectionIndex])
+			{
+				continue;
+			}
+			//Verify we have valid existing section data, if not break from the loop higher index wont be valid
+			if (!MeshData->LastImportMeshLodSectionMaterialData[LodIndex].IsValidIndex(ExistSectionIndex) || !MeshData->ExistingImportMeshLodSectionMaterialData[LodIndex].IsValidIndex(ExistSectionIndex))
+			{
+				break;
+			}
+
+			//Get the Last imported skelmesh section slot import name
+			FName OriginalImportMeshSectionSlotName = MeshData->LastImportMeshLodSectionMaterialData[LodIndex][ExistSectionIndex];
+			if (OriginalImportMeshSectionSlotName != CurrentSectionImportedMaterialName)
+			{
+				//Skip until we found a match between the last import
+				continue;
+			}
+
+			//We have a match put back the data
+			NewSection.bCastShadow = MeshData->ExistingImportMeshLodSectionMaterialData[LodIndex][ExistSectionIndex].bCastShadow;
+			NewSection.bRecomputeTangent = MeshData->ExistingImportMeshLodSectionMaterialData[LodIndex][ExistSectionIndex].bRecomputeTangents;
+			NewSection.RecomputeTangentsVertexMaskChannel = MeshData->ExistingImportMeshLodSectionMaterialData[LodIndex][ExistSectionIndex].RecomputeTangentsVertexMaskChannel;
+			NewSection.GenerateUpToLodIndex = MeshData->ExistingImportMeshLodSectionMaterialData[LodIndex][ExistSectionIndex].GenerateUpTo;
+			NewSection.bDisabled = MeshData->ExistingImportMeshLodSectionMaterialData[LodIndex][ExistSectionIndex].bDisabled;
+			bool bBoneChunkedSection = NewSection.ChunkedParentSectionIndex >= 0;
+			int32 ParentOriginalSectionIndex = NewSection.OriginalDataSectionIndex;
+			if (!bBoneChunkedSection)
+			{
+				//Set the new Parent Index
+				FSkelMeshSourceSectionUserData& UserSectionData = SkeletalMeshLodModel.UserSectionsData.FindOrAdd(ParentOriginalSectionIndex);
+				UserSectionData.bDisabled = NewSection.bDisabled;
+				UserSectionData.bCastShadow = NewSection.bCastShadow;
+				UserSectionData.bRecomputeTangent = NewSection.bRecomputeTangent;
+				UserSectionData.RecomputeTangentsVertexMaskChannel = NewSection.RecomputeTangentsVertexMaskChannel;
+				UserSectionData.GenerateUpToLodIndex = NewSection.GenerateUpToLodIndex;
+				//The cloth will be rebind later after the reimport is done
+			}
+			//Set the matched section to true to avoid using it again
+			MatchedExistSectionIndex[ExistSectionIndex] = true;
+
+			//find the corresponding current slot name in the skeletal mesh materials list to remap properly the material index, in case the user have change it before re-importing
+			FName ExistMeshSectionSlotName = MeshData->ExistingImportMeshLodSectionMaterialData[LodIndex][ExistSectionIndex].ImportedMaterialSlotName;
+			{
+				for (int32 SkelMeshMaterialIndex = 0; SkelMeshMaterialIndex < SkeletalMeshMaterials.Num(); ++SkelMeshMaterialIndex)
+				{
+					const FSkeletalMaterial &NewSectionMaterial = SkeletalMeshMaterials[SkelMeshMaterialIndex];
+					if (NewSectionMaterial.ImportedMaterialSlotName == ExistMeshSectionSlotName)
+					{
+						if (ExistMeshSectionSlotName != OriginalImportMeshSectionSlotName)
+						{
+							NewSection.MaterialIndex = SkelMeshMaterialIndex;
+						}
+						break;
+					}
+				}
+			}
+			//Break because we found a match and have restore the data for this SectionIndex
+			break;
+		}
+	}
+	//Make sure we reset the User section array to only what we have in the fbx
+	SkeletalMeshLodModel.SyncronizeUserSectionsDataArray(true);
+}
+
+#undef LOCTEXT_NAMESPACE
diff --git a/Engine/Source/Editor/UnrealEd/Private/ImportUtils/StaticMeshImportUtils.cpp b/Engine/Source/Editor/UnrealEd/Private/ImportUtils/StaticMeshImportUtils.cpp
new file mode 100644
--- /dev/null
+++ b/Engine/Source/Editor/UnrealEd/Private/ImportUtils/StaticMeshImportUtils.cpp
@@ -0,0 +1,1289 @@
+// Copyright Epic Games, Inc. All Rights Reserved.
+
+/*=============================================================================
+	StaticMeshImportUtils.cpp: Static mesh edit functions.
+=============================================================================*/
+
+#include "ImportUtils/StaticMeshImportUtils.h"
+
+#include "BSPOps.h"
+#include "EditorFramework/ThumbnailInfo.h"
+#include "Engine/StaticMeshSocket.h"
+#include "Engine/StaticMesh.h"
+#include "IMeshReductionInterfaces.h"
+#include "IMeshReductionManagerModule.h"
+#include "MeshDescription.h"
+#include "Modules/ModuleManager.h"
+#include "UObject/MetaData.h"
+
+DEFINE_LOG_CATEGORY_STATIC(LogStaticMeshImportUtils, Log, All);
+
+#define LOCTEXT_NAMESPACE "StaticMeshImportUtils"
+
+static float MeshToPrimTolerance = 0.001f;
+
+/** Floating point comparitor */
+static FORCEINLINE bool AreEqual(float a, float b)
+{
+	return FMath::Abs(a - b) < MeshToPrimTolerance;
+}
+
+/** Returns 1 if vectors are parallel OR anti-parallel */
+static FORCEINLINE bool AreParallel(const FVector3f& a, const FVector3f& b)
+{
+	float Dot = a | b;
+
+	if( AreEqual(FMath::Abs(Dot), 1.f) )
+	{
+		return true;
+	}
+	else
+	{
+		return false;
+	}
+}
+
+/** Utility struct used in AddBoxGeomFromTris. */
+struct FPlaneInfo
+{
+	FVector3f Normal;
+	int32 DistCount;
+	float PlaneDist[2];
+
+	FPlaneInfo()
+	{
+		Normal = FVector3f::ZeroVector;
+		DistCount = 0;
+		PlaneDist[0] = 0.f;
+		PlaneDist[1] = 0.f;
+	}
+};
+
+struct FMeshConnectivityVertex
+{
+	FVector3f				Position;
+	TArray<int32>		Triangles;
+
+	/** Constructor */
+	FMeshConnectivityVertex( const FVector3f &v )
+		: Position( v )
+	{
+	}
+
+	/** Check if this vertex is in the same place as given point */
+	FORCEINLINE bool IsSame( const FVector3f &v )
+	{
+		const float eps = 0.01f;
+		return v.Equals( Position, eps );
+	}
+
+	/** Add link to triangle */
+	FORCEINLINE void AddTriangleLink( int32 Triangle )
+	{
+		Triangles.Add( Triangle );
+	}
+};
+
+struct FMeshConnectivityTriangle
+{
+	int32				Vertices[3];
+	int32				Group;
+
+	/** Constructor */
+	FMeshConnectivityTriangle( int32 a, int32 b, int32 c )
+		: Group( INDEX_NONE )
+	{
+		Vertices[0] = a;
+		Vertices[1] = b;
+		Vertices[2] = c;
+	}
+};
+
+struct FMeshConnectivityGroup
+{
+	TArray<int32>		Triangles;
+};
+
+class FMeshConnectivityBuilder
+{
+public:
+	TArray<FMeshConnectivityVertex>		Vertices;
+	TArray<FMeshConnectivityTriangle>	Triangles;
+	TArray<FMeshConnectivityGroup>		Groups;
+
+public:
+	/** Add vertex to connectivity information */
+	int32 AddVertex( const FVector3f &v )
+	{
+		// Try to find existing vertex
+		// TODO: should use hash map
+		for ( int32 i=0; i<Vertices.Num(); ++i )
+		{
+			if ( Vertices[i].IsSame( v ) )
+			{
+				return i;
+			}
+		}
+
+		// Add new vertex
+		new ( Vertices ) FMeshConnectivityVertex( v );
+		return Vertices.Num() - 1;
+	}
+
+	/** Add triangle to connectivity information */
+	int32 AddTriangle( const FVector3f &a, const FVector3f &b, const FVector3f &c )
+	{
+		// Map vertices
+		int32 VertexA = AddVertex( a );
+		int32 VertexB = AddVertex( b );
+		int32 VertexC = AddVertex( c );
+
+		// Make sure triangle is not degenerated
+		if ( VertexA!=VertexB && VertexB!=VertexC && VertexC!=VertexA )
+		{
+			// Setup connectivity info
+			int32 TriangleIndex = Triangles.Num();
+			Vertices[ VertexA ].AddTriangleLink( TriangleIndex );
+			Vertices[ VertexB ].AddTriangleLink( TriangleIndex );
+			Vertices[ VertexC ].AddTriangleLink( TriangleIndex );
+
+			// Create triangle
+			new ( Triangles ) FMeshConnectivityTriangle( VertexA, VertexB, VertexC );
+			return TriangleIndex;
+		}
+		else
+		{
+			// Degenerated triangle
+			return INDEX_NONE;
+		}
+	}
+
+	/** Create connectivity groups */
+	void CreateConnectivityGroups()
+	{
+		// Delete group list
+		Groups.Empty();
+
+		// Reset group assignments
+		for ( int32 i=0; i<Triangles.Num(); i++ )
+		{
+			Triangles[i].Group = INDEX_NONE;
+		}
+
+		// Flood fill using connectivity info
+		for ( ;; )
+		{
+			// Find first triangle without group assignment
+			int32 InitialTriangle = INDEX_NONE;
+			for ( int32 i=0; i<Triangles.Num(); i++ )
+			{
+				if ( Triangles[i].Group == INDEX_NONE )
+				{
+					InitialTriangle = i;
+					break;
+				}
+			}
+
+			// No more unassigned triangles, flood fill is done
+			if ( InitialTriangle == INDEX_NONE )
+			{
+				break;
+			}
+
+			// Create group
+			int32 GroupIndex = Groups.AddZeroed( 1 );
+
+			// Start flood fill using connectivity information
+			FloodFillTriangleGroups( InitialTriangle, GroupIndex );
+		}
+	}
+
+private:
+	/** FloodFill core */
+	void FloodFillTriangleGroups( int32 InitialTriangleIndex, int32 GroupIndex )
+	{
+		TArray<int32> TriangleStack;
+
+		// Start with given triangle
+		TriangleStack.Add( InitialTriangleIndex );
+
+		// Set the group for our first triangle
+		Triangles[ InitialTriangleIndex ].Group = GroupIndex;
+
+		// Process until we have triangles in stack
+		while ( TriangleStack.Num() )
+		{
+			// Pop triangle index from stack
+			int32 TriangleIndex = TriangleStack.Pop();
+
+			FMeshConnectivityTriangle &Triangle = Triangles[ TriangleIndex ];
+
+			// All triangles should already have a group before we start processing neighbors
+			checkSlow( Triangle.Group == GroupIndex );
+
+			// Add to list of triangles in group
+			Groups[ GroupIndex ].Triangles.Add( TriangleIndex );
+
+			// Recurse to all other triangles connected with this one
+			for ( int32 i=0; i<3; i++ )
+			{
+				int32 VertexIndex = Triangle.Vertices[ i ];
+				const FMeshConnectivityVertex &Vertex = Vertices[ VertexIndex ];
+
+				for ( int32 j=0; j<Vertex.Triangles.Num(); j++ )
+				{
+					int32 OtherTriangleIndex = Vertex.Triangles[ j ];
+					FMeshConnectivityTriangle &OtherTriangle = Triangles[ OtherTriangleIndex ];
+
+					// Only recurse if triangle was not already assigned to a group
+					if ( OtherTriangle.Group == INDEX_NONE )
+					{
+						// OK, the other triangle now belongs to our group!
+						OtherTriangle.Group = GroupIndex;
+
+						// Add the other triangle to the stack to be processed
+						TriangleStack.Add( OtherTriangleIndex );
+					}
+				}
+			}
+		}
+	}
+};
+
+bool StaticMeshImportUtils::DecomposeUCXMesh( const TArray<FVector3f>& CollisionVertices, const TArray<int32>& CollisionFaceIdx, UBodySetup* BodySetup )
+{
+	// We keep no ref to this Model, so it will be GC'd at some point after the import.
+	auto TempModel = NewObject<UModel>();
+	TempModel->Initialize(nullptr, 1);
+
+	FMeshConnectivityBuilder ConnectivityBuilder;
+
+	// Send triangles to connectivity builder
+	for(int32 x = 0;x < CollisionFaceIdx.Num();x += 3)
+	{
+		const FVector3f &VertexA = CollisionVertices[ CollisionFaceIdx[x + 2] ];
+		const FVector3f &VertexB = CollisionVertices[ CollisionFaceIdx[x + 1] ];
+		const FVector3f &VertexC = CollisionVertices[ CollisionFaceIdx[x + 0] ];
+		ConnectivityBuilder.AddTriangle( VertexA, VertexB, VertexC );
+	}
+
+	ConnectivityBuilder.CreateConnectivityGroups();
+
+	// For each valid group build BSP and extract convex hulls
+	bool bSuccess = true;
+	for ( int32 i=0; i<ConnectivityBuilder.Groups.Num(); i++ )
+	{
+		const FMeshConnectivityGroup &Group = ConnectivityBuilder.Groups[ i ];
+
+		// TODO: add some BSP friendly checks here
+		// e.g. if group triangles form a closed mesh
+
+		// Generate polygons from group triangles
+		TempModel->Polys->Element.Empty();
+
+		for ( int32 j=0; j<Group.Triangles.Num(); j++ )
+		{
+			const FMeshConnectivityTriangle &Triangle = ConnectivityBuilder.Triangles[ Group.Triangles[j] ];
+
+			FPoly*	Poly = new( TempModel->Polys->Element ) FPoly();
+			Poly->Init();
+			Poly->iLink = j / 3;
+
+			// Add vertices
+			new( Poly->Vertices ) FVector3f( ConnectivityBuilder.Vertices[ Triangle.Vertices[0] ].Position );
+			new( Poly->Vertices ) FVector3f( ConnectivityBuilder.Vertices[ Triangle.Vertices[1] ].Position );
+			new( Poly->Vertices ) FVector3f( ConnectivityBuilder.Vertices[ Triangle.Vertices[2] ].Position );
+
+			// Update polygon normal
+			Poly->CalcNormal(1);
+		}
+
+		// Build bounding box.
+		TempModel->BuildBound();
+
+		// Build BSP for the brush.
+		FBSPOps::bspBuild( TempModel,FBSPOps::BSP_Good,15,70,1,0 );
+		FBSPOps::bspRefresh( TempModel, 1 );
+		FBSPOps::bspBuildBounds( TempModel );
+
+		// Convert collision model into a collection of convex hulls.
+		// Generated convex hulls will be added to existing ones
+		bSuccess = BodySetup->CreateFromModel( TempModel, false ) && bSuccess;
+	}
+
+	// Could all meshes be properly decomposed?
+	return bSuccess;
+}
+
+/** 
+ *	Function for adding a box collision primitive to the supplied collision geometry based on the mesh of the box.
+ * 
+ *	We keep a list of triangle normals found so far. For each normal direction,
+ *	we should have 2 distances from the origin (2 parallel box faces). If the 
+ *	mesh is a box, we should have 3 distinct normal directions, and 2 distances
+ *	found for each. The difference between these distances should be the box
+ *	dimensions. The 3 directions give us the key axes, and therefore the
+ *	box transformation matrix. This shouldn't rely on any vertex-ordering on 
+ *	the triangles (normals are compared +ve & -ve). It also shouldn't matter 
+ *	about how many triangles make up each side (but it will take longer). 
+ *	We get the centre of the box from the centre of its AABB.
+ */
+bool StaticMeshImportUtils::AddBoxGeomFromTris( const TArray<FPoly>& Tris, FKAggregateGeom* AggGeom, const TCHAR* ObjName )
+{
+	TArray<FPlaneInfo> Planes;
+
+	for(int32 i=0; i<Tris.Num(); i++)
+	{
+		bool bFoundPlane = false;
+		for(int32 j=0; j<Planes.Num() && !bFoundPlane; j++)
+		{
+			// if this triangle plane is already known...
+			if( AreParallel( Tris[i].Normal, Planes[j].Normal ) )
+			{
+				// Always use the same normal when comparing distances, to ensure consistent sign.
+				float Dist = Tris[i].Vertices[0] | Planes[j].Normal;
+
+				// we only have one distance, and its not that one, add it.
+				if( Planes[j].DistCount == 1 && !AreEqual(Dist, Planes[j].PlaneDist[0]) )
+				{
+					Planes[j].PlaneDist[1] = Dist;
+					Planes[j].DistCount = 2;
+				}
+				// if we have a second distance, and its not that either, something is wrong.
+				else if( Planes[j].DistCount == 2 && !AreEqual(Dist, Planes[j].PlaneDist[1]) )
+				{
+					UE_LOG(LogStaticMeshImportUtils, Log, TEXT("AddBoxGeomFromTris (%s): Found more than 2 planes with different distances."), ObjName);
+					return false;
+				}
+
+				bFoundPlane = true;
+			}
+		}
+
+		// If this triangle does not match an existing plane, add to list.
+		if(!bFoundPlane)
+		{
+			check( Planes.Num() < Tris.Num() );
+
+			FPlaneInfo NewPlane;
+			NewPlane.Normal = Tris[i].Normal;
+			NewPlane.DistCount = 1;
+			NewPlane.PlaneDist[0] = Tris[i].Vertices[0] | NewPlane.Normal;
+			
+			Planes.Add(NewPlane);
+		}
+	}
+
+	// Now we have our candidate planes, see if there are any problems
+
+	// Wrong number of planes.
+	if(Planes.Num() != 3)
+	{
+		UE_LOG(LogStaticMeshImportUtils, Log, TEXT("AddBoxGeomFromTris (%s): Not very box-like (need 3 sets of planes)."), ObjName);
+		return false;
+	}
+
+	// If we don't have 3 pairs, we can't carry on.
+	if((Planes[0].DistCount != 2) || (Planes[1].DistCount != 2) || (Planes[2].DistCount != 2))
+	{
+		UE_LOG(LogStaticMeshImportUtils, Log, TEXT("AddBoxGeomFromTris (%s): Incomplete set of planes (need 2 per axis)."), ObjName);
+		return false;
+	}
+
+	FMatrix BoxTM = FMatrix::Identity;
+
+	BoxTM.SetAxis(0, Planes[0].Normal);
+	BoxTM.SetAxis(1, Planes[1].Normal);
+
+	// ensure valid TM by cross-product
+	FVector3f ZAxis = Planes[0].Normal ^ Planes[1].Normal;
+
+	if( !AreParallel(ZAxis, Planes[2].Normal) )
+	{
+		UE_LOG(LogStaticMeshImportUtils, Log, TEXT("AddBoxGeomFromTris (%s): Box axes are not perpendicular."), ObjName);
+		return false;
+	}
+
+	BoxTM.SetAxis(2, ZAxis);
+
+	// OBB centre == AABB centre.
+	FBox Box(ForceInit);
+	for(int32 i=0; i<Tris.Num(); i++)
+	{
+		Box += Tris[i].Vertices[0];
+		Box += Tris[i].Vertices[1];
+		Box += Tris[i].Vertices[2];
+	}
+
+	BoxTM.SetOrigin( Box.GetCenter() );
+
+	// Allocate box in array
+	FKBoxElem BoxElem;
+	BoxElem.SetTransform( FTransform( BoxTM ) );	
+	// distance between parallel planes is box edge lengths.
+	BoxElem.X = FMath::Abs(Planes[0].PlaneDist[0] - Planes[0].PlaneDist[1]);
+	BoxElem.Y = FMath::Abs(Planes[1].PlaneDist[0] - Planes[1].PlaneDist[1]);
+	BoxElem.Z = FMath::Abs(Planes[2].PlaneDist[0] - Planes[2].PlaneDist[1]);
+	AggGeom->BoxElems.Add(BoxElem);
+
+	return true;
+}
+
+/**
+ *	Function for adding a sphere collision primitive to the supplied collision geometry based on a set of Verts.
+ *	
+ *	Simply put an AABB around mesh and use that to generate centre and radius.
+ *	It checks that the AABB is square, and that all vertices are either at the
+ *	centre, or within 5% of the radius distance away.
+ */
+bool StaticMeshImportUtils::AddSphereGeomFromVerts( const TArray<FVector3f>& Verts, FKAggregateGeom* AggGeom, const TCHAR* ObjName )
+{
+	if(Verts.Num() == 0)
+	{
+		return false;
+	}
+
+	FBox Box(ForceInit);
+
+	for(int32 i=0; i<Verts.Num(); i++)
+	{
+		Box += Verts[i];
+	}
+
+	FVector Center, Extents;
+	Box.GetCenterAndExtents(Center, Extents);
+	float Longest = 2.f * Extents.GetMax();
+	float Shortest = 2.f * Extents.GetMin();
+
+	// check that the AABB is roughly a square (5% tolerance)
+	if((Longest - Shortest)/Longest > 0.05f)
+	{
+		UE_LOG(LogStaticMeshImportUtils, Log, TEXT("AddSphereGeomFromVerts (%s): Sphere bounding box not square."), ObjName);
+		return false;
+	}
+
+	float Radius = 0.5f * Longest;
+
+	// Test that all vertices are a similar radius (5%) from the sphere centre.
+	float MaxR = 0;
+	float MinR = BIG_NUMBER;
+	for(int32 i=0; i<Verts.Num(); i++)
+	{
+		FVector3f CToV = Verts[i] - Center;
+		float RSqr = CToV.SizeSquared();
+
+		MaxR = FMath::Max(RSqr, MaxR);
+
+		// Sometimes vertex at centre, so reject it.
+		if(RSqr > KINDA_SMALL_NUMBER)
+		{
+			MinR = FMath::Min(RSqr, MinR);
+		}
+	}
+
+	MaxR = FMath::Sqrt(MaxR);
+	MinR = FMath::Sqrt(MinR);
+
+	if((MaxR-MinR)/Radius > 0.05f)
+	{
+		UE_LOG(LogStaticMeshImportUtils, Log, TEXT("AddSphereGeomFromVerts (%s): Vertices not at constant radius."), ObjName );
+		return false;
+	}
+
+	// Allocate sphere in array
+	FKSphereElem SphereElem;
+	SphereElem.Center = Center;
+	SphereElem.Radius = Radius;
+	AggGeom->SphereElems.Add(SphereElem);
+	
+	return true;
+}
+
+bool StaticMeshImportUtils::AddCapsuleGeomFromVerts(const TArray<FVector3f>& Verts, FKAggregateGeom* AggGeom, const TCHAR* ObjName)
+{
+	if (Verts.Num() < 3)
+	{
+		return false;
+	}
+
+	FVector3f AxisStart, AxisEnd;
+	float MaxDistSqr = 0.f;
+
+	for (int32 IndexA = 0; IndexA < Verts.Num() - 1; IndexA++)
+	{
+		for (int32 IndexB = IndexA + 1; IndexB < Verts.Num(); IndexB++)
+		{
+			float DistSqr = (Verts[IndexA] - Verts[IndexB]).SizeSquared();
+			if (DistSqr > MaxDistSqr)
+			{
+				AxisStart = Verts[IndexA];
+				AxisEnd = Verts[IndexB];
+				MaxDistSqr = DistSqr;
+			}
+		}
+	}
+
+	// If we got a valid axis, find vertex furthest from it
+	if (MaxDistSqr > SMALL_NUMBER)
+	{
+		float MaxRadius = 0.f;
+
+		const FVector3f LineOrigin = AxisStart;
+		const FVector3f LineDir = (AxisEnd - AxisStart).GetSafeNormal();
+
+		for (int32 IndexA = 0; IndexA < Verts.Num() - 1; IndexA++)
+		{
+			float DistToAxis = FMath::PointDistToLine(Verts[IndexA], LineDir, LineOrigin);
+			if (DistToAxis > MaxRadius)
+			{
+				MaxRadius = DistToAxis;
+			}
+		}
+
+		if (MaxRadius > SMALL_NUMBER)
+		{
+			// Allocate capsule in array
+			FKSphylElem SphylElem;
+			SphylElem.Center = 0.5f * (AxisStart + AxisEnd);
+			SphylElem.Rotation = FQuat::FindBetweenVectors(FVector3f(0,0,1), LineDir).Rotator(); // Get quat that takes you from z axis to desired axis
+			SphylElem.Radius = MaxRadius;
+			SphylElem.Length = FMath::Max(FMath::Sqrt(MaxDistSqr) - (2.f * MaxRadius), 0.f); // subtract two radii from total length to get segment length (ensure > 0)
+			AggGeom->SphylElems.Add(SphylElem);
+			return true;
+		}
+	}
+
+	return false;
+}
+
+
+/** Utility for adding one convex hull from the given verts */
+bool StaticMeshImportUtils::AddConvexGeomFromVertices( const TArray<FVector3f>& Verts, FKAggregateGeom* AggGeom, const TCHAR* ObjName )
+{
+	if(Verts.Num() == 0)
+	{
+		return false;
+	}
+
+	FKConvexElem* ConvexElem = new(AggGeom->ConvexElems) FKConvexElem();
+	ConvexElem->VertexData = LWC::PromoteArrayType<FVector>(Verts);	// LWC_TODO: Perf pessimization
+	ConvexElem->UpdateElemBox();
+
+	return true;
+}
+
+bool IsUsingMaterialSlotNameWorkflow(UAssetImportData* AssetImportData)
+{
+	UFbxStaticMeshImportData* ImportData = Cast<UFbxStaticMeshImportData>(AssetImportData);
+	if (ImportData == nullptr || ImportData->ImportMaterialOriginalNameData.Num() <= 0)
+	{
+		return false;
+	}
+	bool AllNameAreNone = true;
+	for (FName ImportMaterialName : ImportData->ImportMaterialOriginalNameData)
+	{
+		if (ImportMaterialName != NAME_None)
+		{
+			AllNameAreNone = false;
+			break;
+		}
+	}
+	return !AllNameAreNone;
+}
+
+TSharedPtr<FExistingStaticMeshData> StaticMeshImportUtils::SaveExistingStaticMeshData(UStaticMesh* ExistingMesh, UnFbx::FBXImportOptions* ImportOptions, int32 LodIndex)
+{
+	if (!ExistingMesh)
+	{
+		return TSharedPtr<FExistingStaticMeshData>();
+	}
+
+		bool bSaveMaterials = !ImportOptions->bImportMaterials;
+		TSharedPtr<FExistingStaticMeshData> ExistingMeshDataPtr = MakeShared<FExistingStaticMeshData>();
+
+		//Save the package UMetaData
+		TMap<FName, FString>* ExistingUMetaDataTagValues = UMetaData::GetMapForObject(ExistingMesh);
+		if(ExistingUMetaDataTagValues && ExistingUMetaDataTagValues->Num() > 0)
+		{
+			ExistingMeshDataPtr->ExistingUMetaDataTagValues = *ExistingUMetaDataTagValues;
+		}
+
+		ExistingMeshDataPtr->ImportVersion = ExistingMesh->ImportVersion;
+		ExistingMeshDataPtr->UseMaterialNameSlotWorkflow = IsUsingMaterialSlotNameWorkflow(ExistingMesh->AssetImportData);
+
+		FMeshSectionInfoMap OldSectionInfoMap = ExistingMesh->GetSectionInfoMap();
+
+		bool bIsReimportCustomLODOverGeneratedLOD = ExistingMesh->IsSourceModelValid(LodIndex) &&
+			(ExistingMesh->GetSourceModel(LodIndex).IsRawMeshEmpty() || (ExistingMesh->IsReductionActive(LodIndex) && ExistingMesh->GetSourceModel(LodIndex).ReductionSettings.BaseLODModel != LodIndex));
+
+		//We need to reset some data in case we import a custom LOD over a generated LOD
+		if (bIsReimportCustomLODOverGeneratedLOD)
+		{
+			//Reset the section info map for this LOD
+			for (int32 SectionIndex = 0; SectionIndex < ExistingMesh->GetNumSections(LodIndex); ++SectionIndex)
+			{
+				OldSectionInfoMap.Remove(LodIndex, SectionIndex);
+			}
+		}
+
+		ExistingMeshDataPtr->ExistingMaterials.Empty();
+		if (bSaveMaterials)
+		{
+			for (const FStaticMaterial &StaticMaterial : ExistingMesh->GetStaticMaterials())
+			{
+				ExistingMeshDataPtr->ExistingMaterials.Add(StaticMaterial);
+			}
+		}
+
+		ExistingMeshDataPtr->ExistingLODData.AddZeroed(ExistingMesh->GetNumSourceModels());
+
+		// refresh material and section info map here
+		// we have to make sure it only contains valid item
+		// we go through section info and only add it back if used, otherwise we don't want to use
+		if (LodIndex == INDEX_NONE)
+		{
+			ExistingMesh->GetSectionInfoMap().Clear();
+		}
+		else
+		{
+			//Remove only the target section InfoMap, if we destroy more we will not restore the correct material assignment for other Lods contain in the same file.
+			int32 ReimportSectionNumber = ExistingMesh->GetSectionInfoMap().GetSectionNumber(LodIndex);
+			for (int32 SectionIndex = 0; SectionIndex < ReimportSectionNumber; ++SectionIndex)
+			{
+				ExistingMesh->GetSectionInfoMap().Remove(LodIndex, SectionIndex);
+			}
+		}
+
+		int32 TotalMaterialIndex = ExistingMeshDataPtr->ExistingMaterials.Num();
+		for (int32 SourceModelIndex = 0; SourceModelIndex < ExistingMesh->GetNumSourceModels(); SourceModelIndex++)
+		{
+			//If the last import was exceeding the maximum number of LOD the source model will contain more LOD so just break the loop
+			if (SourceModelIndex >= ExistingMesh->GetRenderData()->LODResources.Num())
+				break;
+			FStaticMeshLODResources& LOD = ExistingMesh->GetRenderData()->LODResources[SourceModelIndex];
+			int32 NumSections = LOD.Sections.Num();
+			for(int32 SectionIndex = 0; SectionIndex < NumSections; ++SectionIndex)
+			{
+				FMeshSectionInfo Info = OldSectionInfoMap.Get(SourceModelIndex, SectionIndex);
+				if(bSaveMaterials && ExistingMesh->GetStaticMaterials().IsValidIndex(Info.MaterialIndex))
+				{
+					if (ExistingMeshDataPtr->UseMaterialNameSlotWorkflow)
+					{
+						int32 ExistMaterialIndex = ExistingMeshDataPtr->ExistingMaterials.Find(ExistingMesh->GetStaticMaterials()[Info.MaterialIndex]);
+						if (ExistMaterialIndex == INDEX_NONE)
+						{
+							ExistMaterialIndex = ExistingMeshDataPtr->ExistingMaterials.Add(ExistingMesh->GetStaticMaterials()[Info.MaterialIndex]);
+						}
+						Info.MaterialIndex = ExistMaterialIndex;
+					}
+					else
+					{
+						// we only save per LOD separeate IF the material index isn't added yet. 
+						// if it's already added, we don't have to add another one. 
+						if (Info.MaterialIndex >= TotalMaterialIndex)
+						{
+							ExistingMeshDataPtr->ExistingMaterials.Add(ExistingMesh->GetStaticMaterials()[Info.MaterialIndex]);
+
+							// @Todo @fixme
+							// have to refresh material index since it might be pointing at wrong one
+							// this will break IF the base material number grows or shoterns and index will be off
+							// I think we have to save material index per section, so that we don't have to worry about global index
+							Info.MaterialIndex = TotalMaterialIndex++;
+						}
+					}
+					ExistingMeshDataPtr->ExistingSectionInfoMap.Set(SourceModelIndex, SectionIndex, Info);
+				}
+			}
+
+			const FStaticMeshSourceModel& SourceModel = ExistingMesh->GetSourceModel(SourceModelIndex);
+			FExistingLODMeshData& ExistingLODData = ExistingMeshDataPtr->ExistingLODData[SourceModelIndex];
+			ExistingLODData.ExistingBuildSettings = SourceModel.BuildSettings;
+			ExistingLODData.ExistingReductionSettings = SourceModel.ReductionSettings;
+			if (bIsReimportCustomLODOverGeneratedLOD && (SourceModelIndex == LodIndex))
+			{
+				//Reset the reduction
+				ExistingLODData.ExistingReductionSettings.PercentTriangles = 1.0f;
+				ExistingLODData.ExistingReductionSettings.PercentVertices = 1.0f;
+				ExistingLODData.ExistingReductionSettings.MaxDeviation = 0.0f;
+			}
+			ExistingLODData.ExistingScreenSize = SourceModel.ScreenSize.Default;
+			ExistingLODData.ExistingSourceImportFilename = SourceModel.SourceImportFilename;
+
+			const FMeshDescription* MeshDescription = ExistingMesh->GetMeshDescription(SourceModelIndex);
+			if (MeshDescription)
+			{
+				ExistingLODData.ExistingMeshDescription = MakeUnique<FMeshDescription>(*MeshDescription);
+			}
+		}
+
+		ExistingMeshDataPtr->ExistingSockets = ExistingMesh->Sockets;
+
+		ExistingMeshDataPtr->ExistingCustomizedCollision = ExistingMesh->bCustomizedCollision;
+		ExistingMeshDataPtr->bAutoComputeLODScreenSize = ExistingMesh->bAutoComputeLODScreenSize;
+
+		ExistingMeshDataPtr->ExistingLightMapResolution = ExistingMesh->GetLightMapResolution();
+		ExistingMeshDataPtr->ExistingLightMapCoordinateIndex = ExistingMesh->GetLightMapCoordinateIndex();
+
+		ExistingMeshDataPtr->ExistingImportData = ExistingMesh->AssetImportData;
+		ExistingMeshDataPtr->ExistingThumbnailInfo = ExistingMesh->ThumbnailInfo;
+
+		ExistingMeshDataPtr->ExistingBodySetup = ExistingMesh->GetBodySetup();
+
+		ExistingMeshDataPtr->bHasNavigationData = ExistingMesh->bHasNavigationData;
+		ExistingMeshDataPtr->LODGroup = ExistingMesh->LODGroup;
+		ExistingMeshDataPtr->MinLOD = ExistingMesh->GetMinLOD();
+
+		ExistingMeshDataPtr->ExistingGenerateMeshDistanceField = ExistingMesh->bGenerateMeshDistanceField;
+		ExistingMeshDataPtr->ExistingLODForCollision = ExistingMesh->LODForCollision;
+		ExistingMeshDataPtr->ExistingDistanceFieldSelfShadowBias = ExistingMesh->DistanceFieldSelfShadowBias;
+		ExistingMeshDataPtr->ExistingSupportUniformlyDistributedSampling = ExistingMesh->bSupportUniformlyDistributedSampling;
+		ExistingMeshDataPtr->ExistingAllowCpuAccess = ExistingMesh->bAllowCPUAccess;
+		ExistingMeshDataPtr->ExistingPositiveBoundsExtension = ExistingMesh->GetPositiveBoundsExtension();
+		ExistingMeshDataPtr->ExistingNegativeBoundsExtension = ExistingMesh->GetNegativeBoundsExtension();
+
+		UFbxStaticMeshImportData* ImportData = Cast<UFbxStaticMeshImportData>(ExistingMesh->AssetImportData);
+		if (ImportData && ExistingMeshDataPtr->UseMaterialNameSlotWorkflow)
+		{
+			for (int32 ImportMaterialOriginalNameDataIndex = 0; ImportMaterialOriginalNameDataIndex < ImportData->ImportMaterialOriginalNameData.Num(); ++ImportMaterialOriginalNameDataIndex)
+			{
+				FName MaterialName = ImportData->ImportMaterialOriginalNameData[ImportMaterialOriginalNameDataIndex];
+				ExistingMeshDataPtr->LastImportMaterialOriginalNameData.Add(MaterialName);
+			}
+			for (int32 InternalLodIndex = 0; InternalLodIndex < ImportData->ImportMeshLodData.Num(); ++InternalLodIndex)
+			{
+				ExistingMeshDataPtr->LastImportMeshLodSectionMaterialData.AddZeroed();
+				const FImportMeshLodSectionsData &ImportMeshLodSectionsData = ImportData->ImportMeshLodData[InternalLodIndex];
+				for (int32 SectionIndex = 0; SectionIndex < ImportMeshLodSectionsData.SectionOriginalMaterialName.Num(); ++SectionIndex)
+				{
+					FName MaterialName = ImportMeshLodSectionsData.SectionOriginalMaterialName[SectionIndex];
+					ExistingMeshDataPtr->LastImportMeshLodSectionMaterialData[InternalLodIndex].Add(MaterialName);
+				}
+			}
+		}
+		ExistingMeshDataPtr->ExistingOnMeshChanged = ExistingMesh->OnMeshChanged;
+		ExistingMeshDataPtr->ExistingComplexCollisionMesh = ExistingMesh->ComplexCollisionMesh;
+
+	return ExistingMeshDataPtr;
+}
+
+// Helper to find if some reduction settings are active
+bool IsReductionActive(const FMeshReductionSettings& ReductionSettings)
+{
+	bool bUseQuadricSimplier = true;
+	{
+		// Are we using our tool, or simplygon?  The tool is only changed during editor restarts
+		IMeshReduction* ReductionModule = FModuleManager::Get().LoadModuleChecked<IMeshReductionManagerModule>("MeshReductionInterface").GetStaticMeshReductionInterface();
+		FString VersionString = ReductionModule->GetVersionString();
+		TArray<FString> SplitVersionString;
+		VersionString.ParseIntoArray(SplitVersionString, TEXT("_"), true);
+		bUseQuadricSimplier = SplitVersionString[0].Equals("QuadricMeshReduction");
+	}
+	const bool bVertTermination = (bUseQuadricSimplier) && (ReductionSettings.TerminationCriterion != EStaticMeshReductionTerimationCriterion::Triangles) && (ReductionSettings.PercentVertices < 1.0f);
+	const bool bTriTermination = ReductionSettings.TerminationCriterion != EStaticMeshReductionTerimationCriterion::Vertices && (ReductionSettings.PercentTriangles < 1.0f);
+	return bTriTermination || bVertTermination || (ReductionSettings.MaxDeviation > 0.0f);
+}
+
+/* This function is call before building the mesh when we do a re-import*/
+void StaticMeshImportUtils::RestoreExistingMeshSettings(const FExistingStaticMeshData* ExistingMesh, UStaticMesh* NewMesh, int32 LODIndex)
+{
+	if (!ExistingMesh)
+	{
+		return;
+	}
+	NewMesh->LODGroup = ExistingMesh->LODGroup;
+	NewMesh->SetMinLOD(ExistingMesh->MinLOD);
+	int32 ExistingNumLods = ExistingMesh->ExistingLODData.Num();
+	int32 CurrentNumLods = NewMesh->GetNumSourceModels();
+	if (LODIndex == INDEX_NONE)
+	{
+		if (CurrentNumLods > ExistingNumLods)
+		{
+			NewMesh->SetNumSourceModels(ExistingNumLods);
+		}
+		//Create only the LOD Group we need, extra LOD will be put back when calling RestoreExistingMeshData later in the re import process
+		if (NewMesh->LODGroup != NAME_None)
+		{
+			ITargetPlatform* CurrentPlatform = GetTargetPlatformManagerRef().GetRunningTargetPlatform();
+			check(CurrentPlatform);
+			const FStaticMeshLODGroup& LODGroup = CurrentPlatform->GetStaticMeshLODSettings().GetLODGroup(NewMesh->LODGroup);
+			ExistingNumLods = FMath::Min(ExistingNumLods, LODGroup.GetDefaultNumLODs());
+		}
+
+		for (int32 i = 0; i < ExistingNumLods; i++)
+		{
+			if (NewMesh->GetNumSourceModels() <= i)
+			{
+				NewMesh->AddSourceModel();
+			}
+			FMeshDescription* LODMeshDescription = NewMesh->GetMeshDescription(i);
+			bool bSwapFromGeneratedToImported = !ExistingMesh->ExistingLODData[i].ExistingMeshDescription.IsValid() && (LODMeshDescription && LODMeshDescription->Polygons().Num() > 0);
+			bool bWasReduced = IsReductionActive(ExistingMesh->ExistingLODData[i].ExistingReductionSettings);
+
+			FStaticMeshSourceModel& SourceModel = NewMesh->GetSourceModel(i);
+			if (!bSwapFromGeneratedToImported && bWasReduced)
+			{
+				SourceModel.ReductionSettings = ExistingMesh->ExistingLODData[i].ExistingReductionSettings;
+			}
+			SourceModel.BuildSettings = ExistingMesh->ExistingLODData[i].ExistingBuildSettings;
+			SourceModel.ScreenSize = ExistingMesh->ExistingLODData[i].ExistingScreenSize;
+			SourceModel.SourceImportFilename = ExistingMesh->ExistingLODData[i].ExistingSourceImportFilename;
+		}
+	}
+	else
+	{
+		//Just set the old configuration for the desired LODIndex
+		if(LODIndex >= 0 && LODIndex < CurrentNumLods && LODIndex < ExistingNumLods)
+		{
+			FMeshDescription* LODMeshDescription = NewMesh->GetMeshDescription(LODIndex);
+			bool bSwapFromGeneratedToImported = !ExistingMesh->ExistingLODData[LODIndex].ExistingMeshDescription.IsValid() && (LODMeshDescription && LODMeshDescription->Polygons().Num() > 0);
+			bool bWasReduced = IsReductionActive(ExistingMesh->ExistingLODData[LODIndex].ExistingReductionSettings);
+
+			FStaticMeshSourceModel& SourceModel = NewMesh->GetSourceModel(LODIndex);
+			if (!bSwapFromGeneratedToImported && bWasReduced)
+			{
+				SourceModel.ReductionSettings = ExistingMesh->ExistingLODData[LODIndex].ExistingReductionSettings;
+			}
+			SourceModel.BuildSettings = ExistingMesh->ExistingLODData[LODIndex].ExistingBuildSettings;
+			SourceModel.ScreenSize = ExistingMesh->ExistingLODData[LODIndex].ExistingScreenSize;
+			SourceModel.SourceImportFilename = ExistingMesh->ExistingLODData[LODIndex].ExistingSourceImportFilename;
+		}
+	}
+
+	//We need to fill the import version remap before building the mesh since the
+	//static mesh component will be register at the end of the build.
+	//We do the remap of the material override in the static mesh component in OnRegister()
+	if(ExistingMesh->ImportVersion != EImportStaticMeshVersion::LastVersion)
+	{
+		uint32 MaterialMapKey = 0;
+		TArray<int32> ImportRemapMaterial;
+		MaterialMapKey = ((uint32)((ExistingMesh->ImportVersion & 0xffff) << 16) | (uint32)(EImportStaticMeshVersion::LastVersion & 0xffff));
+		//Avoid matching a material more then once
+		TArray<int32> MatchIndex;
+		ImportRemapMaterial.AddZeroed(ExistingMesh->ExistingMaterials.Num());
+		for (int32 ExistMaterialIndex = 0; ExistMaterialIndex < ExistingMesh->ExistingMaterials.Num(); ++ExistMaterialIndex)
+		{
+			ImportRemapMaterial[ExistMaterialIndex] = ExistMaterialIndex; //Set default value
+			const FStaticMaterial &ExistMaterial = ExistingMesh->ExistingMaterials[ExistMaterialIndex];
+			bool bFoundMatchingMaterial = false;
+			for (int32 MaterialIndex = 0; MaterialIndex < NewMesh->GetStaticMaterials().Num(); ++MaterialIndex)
+			{
+				if (MatchIndex.Contains(MaterialIndex))
+				{
+					continue;
+				}
+				FStaticMaterial &Material = NewMesh->GetStaticMaterials()[MaterialIndex];
+				if (Material.ImportedMaterialSlotName == ExistMaterial.ImportedMaterialSlotName)
+				{
+					MatchIndex.Add(MaterialIndex);
+					ImportRemapMaterial[ExistMaterialIndex] = MaterialIndex;
+					bFoundMatchingMaterial = true;
+					break;
+				}
+			}
+			if (!bFoundMatchingMaterial)
+			{
+				for (int32 MaterialIndex = 0; MaterialIndex < NewMesh->GetStaticMaterials().Num(); ++MaterialIndex)
+				{
+					if (MatchIndex.Contains(MaterialIndex))
+					{
+						continue;
+					}
+
+					FStaticMaterial &Material = NewMesh->GetStaticMaterials()[MaterialIndex];
+					if (ExistMaterial.ImportedMaterialSlotName == NAME_None && Material.MaterialInterface == ExistMaterial.MaterialInterface)
+					{
+						MatchIndex.Add(MaterialIndex);
+						ImportRemapMaterial[ExistMaterialIndex] = MaterialIndex;
+						bFoundMatchingMaterial = true;
+						break;
+					}
+				}
+			}
+			if (!bFoundMatchingMaterial)
+			{
+				ImportRemapMaterial[ExistMaterialIndex] = ExistMaterialIndex;
+			}
+		}
+		NewMesh->MaterialRemapIndexPerImportVersion.Add(FMaterialRemapIndex(MaterialMapKey, ImportRemapMaterial));
+	}
+	// Copy mesh changed delegate data
+	NewMesh->OnMeshChanged = ExistingMesh->ExistingOnMeshChanged;
+}
+
+void StaticMeshImportUtils::UpdateSomeLodsImportMeshData(UStaticMesh* NewMesh, TArray<int32> *ReimportLodList)
+{
+	if (NewMesh == nullptr)
+	{
+		return;
+	}
+	UFbxStaticMeshImportData* ImportData = Cast<UFbxStaticMeshImportData>(NewMesh->AssetImportData);
+	//Update the LOD import data before restoring the data
+	if (ReimportLodList != nullptr && ImportData != nullptr)
+	{
+		for (int32 LodLevelImport : (*ReimportLodList))
+		{
+			if (!ImportData->ImportMeshLodData.IsValidIndex(LodLevelImport))
+			{
+				ImportData->ImportMeshLodData.AddZeroed(LodLevelImport - ImportData->ImportMeshLodData.Num() + 1);
+			}
+			ImportData->ImportMeshLodData[LodLevelImport].SectionOriginalMaterialName.Empty();
+			if (NewMesh->GetRenderData()->LODResources.IsValidIndex(LodLevelImport))
+			{
+				FStaticMeshLODResources& LOD = NewMesh->GetRenderData()->LODResources[LodLevelImport];
+				int32 NumSections = LOD.Sections.Num();
+				for (int32 SectionIndex = 0; SectionIndex < NumSections; ++SectionIndex)
+				{
+					int32 MaterialLodSectionIndex = LOD.Sections[SectionIndex].MaterialIndex;
+					if (NewMesh->GetSectionInfoMap().IsValidSection(LodLevelImport, SectionIndex))
+					{
+						MaterialLodSectionIndex = NewMesh->GetSectionInfoMap().Get(LodLevelImport, SectionIndex).MaterialIndex;
+					}
+
+					if (NewMesh->GetStaticMaterials().IsValidIndex(MaterialLodSectionIndex))
+					{
+						bool bFoundMatch = false;
+						FName OriginalImportName = NewMesh->GetStaticMaterials()[MaterialLodSectionIndex].ImportedMaterialSlotName;
+						//Find the material in the original import data
+						int32 ImportMaterialIndex = 0;
+						for (ImportMaterialIndex = 0; ImportMaterialIndex < ImportData->ImportMaterialOriginalNameData.Num(); ++ImportMaterialIndex)
+						{
+							if (ImportData->ImportMaterialOriginalNameData[ImportMaterialIndex] == OriginalImportName)
+							{
+								bFoundMatch = true;
+								break;
+							}
+						}
+						if (!bFoundMatch)
+						{
+							ImportMaterialIndex = ImportData->ImportMaterialOriginalNameData.Add(OriginalImportName);
+						}
+						ImportData->ImportMeshLodData[LodLevelImport].SectionOriginalMaterialName.Add(ImportData->ImportMaterialOriginalNameData[ImportMaterialIndex]);
+					}
+					else
+					{
+						ImportData->ImportMeshLodData[LodLevelImport].SectionOriginalMaterialName.Add(TEXT("InvalidMaterialIndex"));
+					}
+				}
+			}
+		}
+	}
+}
+
+template<typename T>
+void AddGeneratedGeom(const TArray<T>& SrcData, TArray<T>& DstData) {
+	for (const T& CurrentElement : SrcData)
+	{
+		if (CurrentElement.bIsGenerated)
+		{
+			DstData.Add(CurrentElement);
+		}
+	}
+};
+
+void StaticMeshImportUtils::RestoreExistingMeshData(const TSharedPtr<const FExistingStaticMeshData>& ExistingMeshDataPtr, UStaticMesh* NewMesh, int32 LodLevel, bool bCanShowDialog, bool bForceConflictingMaterialReset)
+{
+	if (!ExistingMeshDataPtr || !NewMesh)
+	{
+		return;
+	}
+
+	//Restore the package metadata
+	if (ExistingMeshDataPtr->ExistingUMetaDataTagValues.Num() > 0)
+	{
+		UMetaData* PackageMetaData = NewMesh->GetOutermost()->GetMetaData();
+		checkSlow(PackageMetaData);
+		
+		TMap<FName, FString> MetaDataToApply(ExistingMeshDataPtr->ExistingUMetaDataTagValues);
+		// Keep all existing metadata, but override/add those who were imported.
+		if (TMap<FName, FString>* ImportedMetaData = UMetaData::GetMapForObject(NewMesh))
+		{
+			for (const TPair<FName, FString>& KeyValue : *ImportedMetaData)
+			{
+				MetaDataToApply.Add(KeyValue);
+			}
+		}
+		PackageMetaData->SetObjectValues(NewMesh, MoveTemp(MetaDataToApply));
+	}
+
+	//Create a remap material Index use to find the matching section later
+	TArray<int32> RemapMaterial;
+	RemapMaterial.AddZeroed(NewMesh->GetStaticMaterials().Num());
+	TArray<FName> RemapMaterialName;
+	RemapMaterialName.AddZeroed(NewMesh->GetStaticMaterials().Num());
+
+	//If user is attended, ask him to verify the match is good
+	UnFbx::EFBXReimportDialogReturnOption ReturnOption;
+	//Ask the user to match the materials conflict
+	UnFbx::FFbxImporter::PrepareAndShowMaterialConflictDialog<FStaticMaterial>(ExistingMeshDataPtr->ExistingMaterials, NewMesh->GetStaticMaterials(), RemapMaterial, RemapMaterialName, bCanShowDialog, false, bForceConflictingMaterialReset, ReturnOption);
+	
+	if (ReturnOption != UnFbx::EFBXReimportDialogReturnOption::FBXRDRO_ResetToFbx)
+	{
+		//Build a ordered material list that try to keep intact the existing material list
+		TArray<FStaticMaterial> MaterialOrdered;
+		TArray<bool> MatchedNewMaterial;
+		MatchedNewMaterial.AddZeroed(NewMesh->GetStaticMaterials().Num());
+		for (int32 ExistMaterialIndex = 0; ExistMaterialIndex < ExistingMeshDataPtr->ExistingMaterials.Num(); ++ExistMaterialIndex)
+		{
+			int32 MaterialIndexOrdered = MaterialOrdered.Add(ExistingMeshDataPtr->ExistingMaterials[ExistMaterialIndex]);
+			FStaticMaterial& OrderedMaterial = MaterialOrdered[MaterialIndexOrdered];
+			int32 NewMaterialIndex = INDEX_NONE;
+			if (RemapMaterial.Find(ExistMaterialIndex, NewMaterialIndex))
+			{
+				MatchedNewMaterial[NewMaterialIndex] = true;
+				RemapMaterial[NewMaterialIndex] = MaterialIndexOrdered;
+				OrderedMaterial.ImportedMaterialSlotName = NewMesh->GetStaticMaterials()[NewMaterialIndex].ImportedMaterialSlotName;
+			}
+			else
+			{
+				//Unmatched material must be conserve
+			}
+		}
+
+		//Add the new material entries (the one that do not match with any existing material)
+		for (int32 NewMaterialIndex = 0; NewMaterialIndex < MatchedNewMaterial.Num(); ++NewMaterialIndex)
+		{
+			if (MatchedNewMaterial[NewMaterialIndex] == false)
+			{
+				int32 NewMeshIndex = MaterialOrdered.Add(NewMesh->GetStaticMaterials()[NewMaterialIndex]);
+				RemapMaterial[NewMaterialIndex] = NewMeshIndex;
+			}
+		}
+
+		//Set the RemapMaterialName array helper
+		for (int32 MaterialIndex = 0; MaterialIndex < RemapMaterial.Num(); ++MaterialIndex)
+		{
+			int32 SourceMaterialMatch = RemapMaterial[MaterialIndex];
+			if (ExistingMeshDataPtr->ExistingMaterials.IsValidIndex(SourceMaterialMatch))
+			{
+				RemapMaterialName[MaterialIndex] = ExistingMeshDataPtr->ExistingMaterials[SourceMaterialMatch].ImportedMaterialSlotName;
+			}
+		}
+
+		//Copy the re ordered materials (this ensure the material array do not change when we re-import)
+		NewMesh->SetStaticMaterials(MaterialOrdered);
+	}
+	int32 NumCommonLODs = FMath::Min<int32>(ExistingMeshDataPtr->ExistingLODData.Num(), NewMesh->GetNumSourceModels());
+	for(int32 i=0; i<NumCommonLODs; i++)
+	{
+		FStaticMeshSourceModel& SourceModel = NewMesh->GetSourceModel(i);
+		SourceModel.BuildSettings = ExistingMeshDataPtr->ExistingLODData[i].ExistingBuildSettings;
+		FMeshDescription* LODMeshDescription = NewMesh->GetMeshDescription(i);
+		//Restore the reduction settings only if the existing data was a using reduction. Because we can set some value if we reimport from existing rawmesh to auto generated.
+		bool bSwapFromGeneratedToImported = !ExistingMeshDataPtr->ExistingLODData[i].ExistingMeshDescription.IsValid() && (LODMeshDescription && LODMeshDescription->Polygons().Num() > 0);
+		bool bWasReduced = IsReductionActive(ExistingMeshDataPtr->ExistingLODData[i].ExistingReductionSettings);
+		if ( !bSwapFromGeneratedToImported && bWasReduced)
+		{
+			SourceModel.ReductionSettings = ExistingMeshDataPtr->ExistingLODData[i].ExistingReductionSettings;
+		}
+		SourceModel.ScreenSize = ExistingMeshDataPtr->ExistingLODData[i].ExistingScreenSize;
+		SourceModel.SourceImportFilename = ExistingMeshDataPtr->ExistingLODData[i].ExistingSourceImportFilename;
+	}
+
+	for(int32 i=NumCommonLODs; i < ExistingMeshDataPtr->ExistingLODData.Num(); ++i)
+	{
+		FStaticMeshSourceModel& SrcModel = NewMesh->AddSourceModel();
+		if (ExistingMeshDataPtr->ExistingLODData[i].ExistingMeshDescription.IsValid())
+		{
+			FMeshDescription* MeshDescription = NewMesh->CreateMeshDescription(i, MoveTemp(*ExistingMeshDataPtr->ExistingLODData[i].ExistingMeshDescription));
+			NewMesh->CommitMeshDescription(i);
+		}
+		SrcModel.BuildSettings = ExistingMeshDataPtr->ExistingLODData[i].ExistingBuildSettings;
+		SrcModel.ReductionSettings = ExistingMeshDataPtr->ExistingLODData[i].ExistingReductionSettings;
+		SrcModel.ScreenSize = ExistingMeshDataPtr->ExistingLODData[i].ExistingScreenSize;
+		SrcModel.SourceImportFilename = ExistingMeshDataPtr->ExistingLODData[i].ExistingSourceImportFilename;
+	}
+
+	// Restore the section info of the just imported LOD so is section info map is remap to fit the mesh material array
+	if (ExistingMeshDataPtr->ExistingSectionInfoMap.Map.Num() > 0)
+	{
+		//Build the mesh we need the render data and the existing section info map build before restoring the data
+		if (NewMesh->GetRenderData()->LODResources.Num() < NewMesh->GetNumSourceModels())
+		{
+			NewMesh->Build();
+		}
+		for (int32 i = 0; i < NewMesh->GetRenderData()->LODResources.Num(); i++)
+		{
+			//If a LOD was specified, only touch the specified LOD
+			if (LodLevel != INDEX_NONE && LodLevel != 0 && LodLevel != i)
+			{
+				continue;
+			}
+
+			
+			//When re-importing the asset, do not touch the LOD that was imported from file, the material array is keep intact so the section should still be valid.
+			bool NoRemapForThisLOD = LodLevel == INDEX_NONE && i != 0 && !NewMesh->GetSourceModel(i).bImportWithBaseMesh && !IsReductionActive(NewMesh->GetSourceModel(i).ReductionSettings);
+
+			FStaticMeshLODResources& LOD = NewMesh->GetRenderData()->LODResources[i];
+			
+			int32 NumSections = LOD.Sections.Num();
+			int32 OldSectionNumber = ExistingMeshDataPtr->ExistingSectionInfoMap.GetSectionNumber(i);
+			for (int32 SectionIndex = 0; SectionIndex < NumSections; ++SectionIndex)
+			{
+				//If the SectionInfoMap is not set yet. Because we re-import LOD 0 but we have other LODs
+				//Just put back the old section Info Map
+				if (NewMesh->GetSectionInfoMap().GetSectionNumber(i) <= SectionIndex)
+				{
+					NewMesh->GetSectionInfoMap().Set(i, SectionIndex, ExistingMeshDataPtr->ExistingSectionInfoMap.Get(i, SectionIndex));
+				}
+				//We recreate the SectionInfoMap from the existing data and we do not remap it if LOD is not auto generate and was not imported
+				if (NoRemapForThisLOD)
+				{
+					continue;
+				}
+
+				FMeshSectionInfo NewSectionInfo =  NewMesh->GetSectionInfoMap().Get(i, SectionIndex);
+				bool bFoundOldMatch = false;
+				bool bKeepOldSectionMaterialIndex = false;
+				int32 OriginalSectionMaterialIndex = INDEX_NONE;
+				if (RemapMaterial.IsValidIndex(NewSectionInfo.MaterialIndex) && NewMesh->GetStaticMaterials().IsValidIndex(NewSectionInfo.MaterialIndex))
+				{
+					//Find the matching old index
+					for (int32 ExistSectionIndex = 0; ExistSectionIndex < OldSectionNumber; ++ExistSectionIndex)
+					{
+						bKeepOldSectionMaterialIndex = false;
+						OriginalSectionMaterialIndex = INDEX_NONE;
+						FMeshSectionInfo OldSectionInfo = ExistingMeshDataPtr->ExistingSectionInfoMap.Get(i, ExistSectionIndex);
+						if (ExistingMeshDataPtr->UseMaterialNameSlotWorkflow)
+						{
+							if (ExistingMeshDataPtr->ExistingImportData.IsValid() &&
+								ExistingMeshDataPtr->LastImportMeshLodSectionMaterialData.IsValidIndex(i) &&
+								ExistingMeshDataPtr->LastImportMeshLodSectionMaterialData[i].IsValidIndex(ExistSectionIndex))
+							{
+								//Keep the old section material index only if the user has change the section mapping
+								bKeepOldSectionMaterialIndex = ExistingMeshDataPtr->LastImportMeshLodSectionMaterialData[i][ExistSectionIndex] != ExistingMeshDataPtr->ExistingMaterials[OldSectionInfo.MaterialIndex].ImportedMaterialSlotName;
+								for (int32 ExistMaterialIndex = 0; ExistMaterialIndex < ExistingMeshDataPtr->ExistingMaterials.Num(); ++ExistMaterialIndex)
+								{
+									if (ExistingMeshDataPtr->LastImportMeshLodSectionMaterialData[i][ExistSectionIndex] == ExistingMeshDataPtr->ExistingMaterials[ExistMaterialIndex].ImportedMaterialSlotName)
+									{
+										OriginalSectionMaterialIndex = ExistMaterialIndex;
+										break;
+									}
+								}
+							}
+						}
+						int32 OldSectionMatchIndex = OriginalSectionMaterialIndex != INDEX_NONE ? OriginalSectionMaterialIndex : OldSectionInfo.MaterialIndex;
+						if (RemapMaterial[NewSectionInfo.MaterialIndex] == OldSectionMatchIndex)
+						{
+							NewMesh->GetSectionInfoMap().Set(i, SectionIndex, OldSectionInfo);
+							bFoundOldMatch = true;
+							break;
+						}
+					}
+				}
+
+				if (!bFoundOldMatch)
+				{
+					if (RemapMaterial.IsValidIndex(NewSectionInfo.MaterialIndex))
+					{
+						//Find the old section that was using the NewSectionInfo.MaterialIndex
+						//This will allow copying the section information: Cast Shadow, enable collision
+						for (int32 ExistSectionIndex = 0; ExistSectionIndex < OldSectionNumber; ++ExistSectionIndex)
+						{
+							FMeshSectionInfo OldSectionInfo = ExistingMeshDataPtr->ExistingSectionInfoMap.Get(i, ExistSectionIndex);
+							if (NewSectionInfo.MaterialIndex == OldSectionInfo.MaterialIndex)
+							{
+								NewSectionInfo.bCastShadow = OldSectionInfo.bCastShadow;
+								NewSectionInfo.bEnableCollision = OldSectionInfo.bEnableCollision;
+								break;
+							}
+						}
+						//If user has change the section info map, we want to keep the change
+						NewSectionInfo.MaterialIndex = RemapMaterial[NewSectionInfo.MaterialIndex];
+						NewMesh->GetSectionInfoMap().Set(i, SectionIndex, NewSectionInfo);
+					}
+				}
+			}
+		}
+		//Store the just imported section info map
+		NewMesh->GetOriginalSectionInfoMap().CopyFrom(NewMesh->GetSectionInfoMap());
+	}
+
+	// Assign sockets from old version of this StaticMesh.
+	for(int32 i=0; i<ExistingMeshDataPtr->ExistingSockets.Num(); i++)
+	{
+		UStaticMeshSocket* ExistingSocket = ExistingMeshDataPtr->ExistingSockets[i];
+		UStaticMeshSocket* Socket = NewMesh->FindSocket(ExistingSocket->SocketName);
+		if (!Socket && !ExistingSocket->bSocketCreatedAtImport)
+		{
+			NewMesh->AddSocket(ExistingSocket);
+		}
+	}
+
+	NewMesh->bCustomizedCollision = ExistingMeshDataPtr->ExistingCustomizedCollision;
+	NewMesh->bAutoComputeLODScreenSize = ExistingMeshDataPtr->bAutoComputeLODScreenSize;
+
+	NewMesh->SetLightMapResolution(ExistingMeshDataPtr->ExistingLightMapResolution);
+	NewMesh->SetLightMapCoordinateIndex(ExistingMeshDataPtr->ExistingLightMapCoordinateIndex);
+
+	if (ExistingMeshDataPtr->ExistingImportData.IsValid())
+	{
+		//RestoredLods
+		UFbxStaticMeshImportData* ImportData = Cast<UFbxStaticMeshImportData>(NewMesh->AssetImportData);
+		TArray<FName> ImportMaterialOriginalNameData;
+		TArray<FImportMeshLodSectionsData> ImportMeshLodData;
+		if (ImportData != nullptr && ImportData->ImportMaterialOriginalNameData.Num() > 0 && ImportData->ImportMeshLodData.Num() > 0)
+		{
+			ImportMaterialOriginalNameData = ImportData->ImportMaterialOriginalNameData;
+			ImportMeshLodData = ImportData->ImportMeshLodData;
+		}
+
+		NewMesh->AssetImportData = ExistingMeshDataPtr->ExistingImportData.Get();
+
+		ImportData = Cast<UFbxStaticMeshImportData>(NewMesh->AssetImportData);
+		if (ImportData != nullptr && ImportMaterialOriginalNameData.Num() > 0 && ImportMeshLodData.Num() > 0)
+		{
+			ImportData->ImportMaterialOriginalNameData = ImportMaterialOriginalNameData;
+			ImportData->ImportMeshLodData = ImportMeshLodData;
+		}
+	}
+		
+
+	NewMesh->ThumbnailInfo = ExistingMeshDataPtr->ExistingThumbnailInfo.Get();
+
+	// If we already had some collision info...
+	if(ExistingMeshDataPtr->ExistingBodySetup)
+	{
+		// If we didn't import anything, keep all collisions otherwise only keep generated collisions.
+		bool bKeepAllCollisions;
+		if(!NewMesh->GetBodySetup() || NewMesh->GetBodySetup()->AggGeom.GetElementCount() == 0)
+		{
+			bKeepAllCollisions = true;
+		}
+		else
+		{
+			bKeepAllCollisions = false;
+		}
+
+		if(bKeepAllCollisions)
+		{
+			NewMesh->SetBodySetup(ExistingMeshDataPtr->ExistingBodySetup);
+		}
+		else
+		{
+			// New collision geometry, but we still want the original settings and the generated collisions
+			NewMesh->GetBodySetup()->CopyBodySetupProperty(ExistingMeshDataPtr->ExistingBodySetup);
+
+			const FKAggregateGeom& ExistingAggGeom = ExistingMeshDataPtr->ExistingBodySetup->AggGeom;
+			FKAggregateGeom& NewAggGeom = NewMesh->GetBodySetup()->AggGeom;
+			AddGeneratedGeom(ExistingAggGeom.BoxElems, NewAggGeom.BoxElems);
+			AddGeneratedGeom(ExistingAggGeom.ConvexElems, NewAggGeom.ConvexElems);
+			AddGeneratedGeom(ExistingAggGeom.SphereElems, NewAggGeom.SphereElems);
+			AddGeneratedGeom(ExistingAggGeom.SphylElems, NewAggGeom.SphylElems);
+		}
+	}
+
+	NewMesh->bHasNavigationData = ExistingMeshDataPtr->bHasNavigationData;
+	NewMesh->LODGroup = ExistingMeshDataPtr->LODGroup;
+
+	NewMesh->bGenerateMeshDistanceField = ExistingMeshDataPtr->ExistingGenerateMeshDistanceField;
+	NewMesh->LODForCollision = ExistingMeshDataPtr->ExistingLODForCollision;
+	NewMesh->DistanceFieldSelfShadowBias = ExistingMeshDataPtr->ExistingDistanceFieldSelfShadowBias;
+	NewMesh->bSupportUniformlyDistributedSampling = ExistingMeshDataPtr->ExistingSupportUniformlyDistributedSampling;
+	NewMesh->bAllowCPUAccess = ExistingMeshDataPtr->ExistingAllowCpuAccess;
+	NewMesh->SetPositiveBoundsExtension(ExistingMeshDataPtr->ExistingPositiveBoundsExtension);
+	NewMesh->SetNegativeBoundsExtension(ExistingMeshDataPtr->ExistingNegativeBoundsExtension);
+
+	NewMesh->ComplexCollisionMesh = ExistingMeshDataPtr->ExistingComplexCollisionMesh;
+}
+
+#undef LOCTEXT_NAMESPACE
diff --git a/Engine/Source/Editor/UnrealEd/Private/SkeletalMeshEdit.cpp b/Engine/Source/Editor/UnrealEd/Private/SkeletalMeshEdit.cpp
--- a/Engine/Source/Editor/UnrealEd/Private/SkeletalMeshEdit.cpp
+++ b/Engine/Source/Editor/UnrealEd/Private/SkeletalMeshEdit.cpp
@@ -30,7 +30,7 @@
 #include "Misc/Paths.h"
 #include "Misc/ScopedSlowTask.h"
 #include "ObjectTools.h"
-#include "SkelImport.h"
+#include "Rendering/SkeletalMeshLODImporterData.h"
 #include "UObject/UObjectIterator.h"
 #include "UObject/Object.h"
 #include "ComponentReregisterContext.h"
diff --git a/Engine/Source/Editor/UnrealEd/Private/StaticMeshEdit.cpp b/Engine/Source/Editor/UnrealEd/Private/StaticMeshEdit.cpp
--- a/Engine/Source/Editor/UnrealEd/Private/StaticMeshEdit.cpp
+++ b/Engine/Source/Editor/UnrealEd/Private/StaticMeshEdit.cpp
@@ -45,559 +45,6 @@ DEFINE_LOG_CATEGORY_STATIC(LogStaticMeshEdit, Log, All);
 
 #define LOCTEXT_NAMESPACE "StaticMeshEdit"
 
-static float MeshToPrimTolerance = 0.001f;
-
-/** Floating point comparitor */
-static FORCEINLINE bool AreEqual(float a, float b)
-{
-	return FMath::Abs(a - b) < MeshToPrimTolerance;
-}
-
-/** Returns 1 if vectors are parallel OR anti-parallel */
-static FORCEINLINE bool AreParallel(const FVector3f& a, const FVector3f& b)
-{
-	float Dot = a | b;
-
-	if( AreEqual(FMath::Abs(Dot), 1.f) )
-	{
-		return true;
-	}
-	else
-	{
-		return false;
-	}
-}
-
-/** Utility struct used in AddBoxGeomFromTris. */
-struct FPlaneInfo
-{
-	FVector3f Normal;
-	int32 DistCount;
-	float PlaneDist[2];
-
-	FPlaneInfo()
-	{
-		Normal = FVector3f::ZeroVector;
-		DistCount = 0;
-		PlaneDist[0] = 0.f;
-		PlaneDist[1] = 0.f;
-	}
-};
-
-struct FMeshConnectivityVertex
-{
-	FVector3f				Position;
-	TArray<int32>		Triangles;
-
-	/** Constructor */
-	FMeshConnectivityVertex( const FVector3f &v )
-		: Position( v )
-	{
-	}
-
-	/** Check if this vertex is in the same place as given point */
-	FORCEINLINE bool IsSame( const FVector3f &v )
-	{
-		const float eps = 0.01f;
-		return v.Equals( Position, eps );
-	}
-
-	/** Add link to triangle */
-	FORCEINLINE void AddTriangleLink( int32 Triangle )
-	{
-		Triangles.Add( Triangle );
-	}
-};
-
-struct FMeshConnectivityTriangle
-{
-	int32				Vertices[3];
-	int32				Group;
-
-	/** Constructor */
-	FMeshConnectivityTriangle( int32 a, int32 b, int32 c )
-		: Group( INDEX_NONE )
-	{
-		Vertices[0] = a;
-		Vertices[1] = b;
-		Vertices[2] = c;
-	}
-};
-
-struct FMeshConnectivityGroup
-{
-	TArray<int32>		Triangles;
-};
-
-class FMeshConnectivityBuilder
-{
-public:
-	TArray<FMeshConnectivityVertex>		Vertices;
-	TArray<FMeshConnectivityTriangle>	Triangles;
-	TArray<FMeshConnectivityGroup>		Groups;
-
-public:
-	/** Add vertex to connectivity information */
-	int32 AddVertex( const FVector3f &v )
-	{
-		// Try to find existing vertex
-		// TODO: should use hash map
-		for ( int32 i=0; i<Vertices.Num(); ++i )
-		{
-			if ( Vertices[i].IsSame( v ) )
-			{
-				return i;
-			}
-		}
-
-		// Add new vertex
-		new ( Vertices ) FMeshConnectivityVertex( v );
-		return Vertices.Num() - 1;
-	}
-
-	/** Add triangle to connectivity information */
-	int32 AddTriangle( const FVector3f &a, const FVector3f &b, const FVector3f &c )
-	{
-		// Map vertices
-		int32 VertexA = AddVertex( a );
-		int32 VertexB = AddVertex( b );
-		int32 VertexC = AddVertex( c );
-
-		// Make sure triangle is not degenerated
-		if ( VertexA!=VertexB && VertexB!=VertexC && VertexC!=VertexA )
-		{
-			// Setup connectivity info
-			int32 TriangleIndex = Triangles.Num();
-			Vertices[ VertexA ].AddTriangleLink( TriangleIndex );
-			Vertices[ VertexB ].AddTriangleLink( TriangleIndex );
-			Vertices[ VertexC ].AddTriangleLink( TriangleIndex );
-
-			// Create triangle
-			new ( Triangles ) FMeshConnectivityTriangle( VertexA, VertexB, VertexC );
-			return TriangleIndex;
-		}
-		else
-		{
-			// Degenerated triangle
-			return INDEX_NONE;
-		}
-	}
-
-	/** Create connectivity groups */
-	void CreateConnectivityGroups()
-	{
-		// Delete group list
-		Groups.Empty();
-
-		// Reset group assignments
-		for ( int32 i=0; i<Triangles.Num(); i++ )
-		{
-			Triangles[i].Group = INDEX_NONE;
-		}
-
-		// Flood fill using connectivity info
-		for ( ;; )
-		{
-			// Find first triangle without group assignment
-			int32 InitialTriangle = INDEX_NONE;
-			for ( int32 i=0; i<Triangles.Num(); i++ )
-			{
-				if ( Triangles[i].Group == INDEX_NONE )
-				{
-					InitialTriangle = i;
-					break;
-				}
-			}
-
-			// No more unassigned triangles, flood fill is done
-			if ( InitialTriangle == INDEX_NONE )
-			{
-				break;
-			}
-
-			// Create group
-			int32 GroupIndex = Groups.AddZeroed( 1 );
-
-			// Start flood fill using connectivity information
-			FloodFillTriangleGroups( InitialTriangle, GroupIndex );
-		}
-	}
-
-private:
-	/** FloodFill core */
-	void FloodFillTriangleGroups( int32 InitialTriangleIndex, int32 GroupIndex )
-	{
-		TArray<int32> TriangleStack;
-
-		// Start with given triangle
-		TriangleStack.Add( InitialTriangleIndex );
-
-		// Set the group for our first triangle
-		Triangles[ InitialTriangleIndex ].Group = GroupIndex;
-
-		// Process until we have triangles in stack
-		while ( TriangleStack.Num() )
-		{
-			// Pop triangle index from stack
-			int32 TriangleIndex = TriangleStack.Pop();
-
-			FMeshConnectivityTriangle &Triangle = Triangles[ TriangleIndex ];
-
-			// All triangles should already have a group before we start processing neighbors
-			checkSlow( Triangle.Group == GroupIndex );
-
-			// Add to list of triangles in group
-			Groups[ GroupIndex ].Triangles.Add( TriangleIndex );
-
-			// Recurse to all other triangles connected with this one
-			for ( int32 i=0; i<3; i++ )
-			{
-				int32 VertexIndex = Triangle.Vertices[ i ];
-				const FMeshConnectivityVertex &Vertex = Vertices[ VertexIndex ];
-
-				for ( int32 j=0; j<Vertex.Triangles.Num(); j++ )
-				{
-					int32 OtherTriangleIndex = Vertex.Triangles[ j ];
-					FMeshConnectivityTriangle &OtherTriangle = Triangles[ OtherTriangleIndex ];
-
-					// Only recurse if triangle was not already assigned to a group
-					if ( OtherTriangle.Group == INDEX_NONE )
-					{
-						// OK, the other triangle now belongs to our group!
-						OtherTriangle.Group = GroupIndex;
-
-						// Add the other triangle to the stack to be processed
-						TriangleStack.Add( OtherTriangleIndex );
-					}
-				}
-			}
-		}
-	}
-};
-
-bool DecomposeUCXMesh( const TArray<FVector3f>& CollisionVertices, const TArray<int32>& CollisionFaceIdx, UBodySetup* BodySetup )
-{
-	// We keep no ref to this Model, so it will be GC'd at some point after the import.
-	auto TempModel = NewObject<UModel>();
-	TempModel->Initialize(nullptr, 1);
-
-	FMeshConnectivityBuilder ConnectivityBuilder;
-
-	// Send triangles to connectivity builder
-	for(int32 x = 0;x < CollisionFaceIdx.Num();x += 3)
-	{
-		const FVector3f &VertexA = CollisionVertices[ CollisionFaceIdx[x + 2] ];
-		const FVector3f &VertexB = CollisionVertices[ CollisionFaceIdx[x + 1] ];
-		const FVector3f &VertexC = CollisionVertices[ CollisionFaceIdx[x + 0] ];
-		ConnectivityBuilder.AddTriangle( VertexA, VertexB, VertexC );
-	}
-
-	ConnectivityBuilder.CreateConnectivityGroups();
-
-	// For each valid group build BSP and extract convex hulls
-	bool bSuccess = true;
-	for ( int32 i=0; i<ConnectivityBuilder.Groups.Num(); i++ )
-	{
-		const FMeshConnectivityGroup &Group = ConnectivityBuilder.Groups[ i ];
-
-		// TODO: add some BSP friendly checks here
-		// e.g. if group triangles form a closed mesh
-
-		// Generate polygons from group triangles
-		TempModel->Polys->Element.Empty();
-
-		for ( int32 j=0; j<Group.Triangles.Num(); j++ )
-		{
-			const FMeshConnectivityTriangle &Triangle = ConnectivityBuilder.Triangles[ Group.Triangles[j] ];
-
-			FPoly*	Poly = new( TempModel->Polys->Element ) FPoly();
-			Poly->Init();
-			Poly->iLink = j / 3;
-
-			// Add vertices
-			new( Poly->Vertices ) FVector3f( ConnectivityBuilder.Vertices[ Triangle.Vertices[0] ].Position );
-			new( Poly->Vertices ) FVector3f( ConnectivityBuilder.Vertices[ Triangle.Vertices[1] ].Position );
-			new( Poly->Vertices ) FVector3f( ConnectivityBuilder.Vertices[ Triangle.Vertices[2] ].Position );
-
-			// Update polygon normal
-			Poly->CalcNormal(1);
-		}
-
-		// Build bounding box.
-		TempModel->BuildBound();
-
-		// Build BSP for the brush.
-		FBSPOps::bspBuild( TempModel,FBSPOps::BSP_Good,15,70,1,0 );
-		FBSPOps::bspRefresh( TempModel, 1 );
-		FBSPOps::bspBuildBounds( TempModel );
-
-		// Convert collision model into a collection of convex hulls.
-		// Generated convex hulls will be added to existing ones
-		bSuccess = BodySetup->CreateFromModel( TempModel, false ) && bSuccess;
-	}
-
-	// Could all meshes be properly decomposed?
-	return bSuccess;
-}
-
-/** 
- *	Function for adding a box collision primitive to the supplied collision geometry based on the mesh of the box.
- * 
- *	We keep a list of triangle normals found so far. For each normal direction,
- *	we should have 2 distances from the origin (2 parallel box faces). If the 
- *	mesh is a box, we should have 3 distinct normal directions, and 2 distances
- *	found for each. The difference between these distances should be the box
- *	dimensions. The 3 directions give us the key axes, and therefore the
- *	box transformation matrix. This shouldn't rely on any vertex-ordering on 
- *	the triangles (normals are compared +ve & -ve). It also shouldn't matter 
- *	about how many triangles make up each side (but it will take longer). 
- *	We get the centre of the box from the centre of its AABB.
- */
-bool AddBoxGeomFromTris( const TArray<FPoly>& Tris, FKAggregateGeom* AggGeom, const TCHAR* ObjName )
-{
-	TArray<FPlaneInfo> Planes;
-
-	for(int32 i=0; i<Tris.Num(); i++)
-	{
-		bool bFoundPlane = false;
-		for(int32 j=0; j<Planes.Num() && !bFoundPlane; j++)
-		{
-			// if this triangle plane is already known...
-			if( AreParallel( Tris[i].Normal, Planes[j].Normal ) )
-			{
-				// Always use the same normal when comparing distances, to ensure consistent sign.
-				float Dist = Tris[i].Vertices[0] | Planes[j].Normal;
-
-				// we only have one distance, and its not that one, add it.
-				if( Planes[j].DistCount == 1 && !AreEqual(Dist, Planes[j].PlaneDist[0]) )
-				{
-					Planes[j].PlaneDist[1] = Dist;
-					Planes[j].DistCount = 2;
-				}
-				// if we have a second distance, and its not that either, something is wrong.
-				else if( Planes[j].DistCount == 2 && !AreEqual(Dist, Planes[j].PlaneDist[1]) )
-				{
-					UE_LOG(LogStaticMeshEdit, Log, TEXT("AddBoxGeomFromTris (%s): Found more than 2 planes with different distances."), ObjName);
-					return false;
-				}
-
-				bFoundPlane = true;
-			}
-		}
-
-		// If this triangle does not match an existing plane, add to list.
-		if(!bFoundPlane)
-		{
-			check( Planes.Num() < Tris.Num() );
-
-			FPlaneInfo NewPlane;
-			NewPlane.Normal = Tris[i].Normal;
-			NewPlane.DistCount = 1;
-			NewPlane.PlaneDist[0] = Tris[i].Vertices[0] | NewPlane.Normal;
-			
-			Planes.Add(NewPlane);
-		}
-	}
-
-	// Now we have our candidate planes, see if there are any problems
-
-	// Wrong number of planes.
-	if(Planes.Num() != 3)
-	{
-		UE_LOG(LogStaticMeshEdit, Log, TEXT("AddBoxGeomFromTris (%s): Not very box-like (need 3 sets of planes)."), ObjName);
-		return false;
-	}
-
-	// If we don't have 3 pairs, we can't carry on.
-	if((Planes[0].DistCount != 2) || (Planes[1].DistCount != 2) || (Planes[2].DistCount != 2))
-	{
-		UE_LOG(LogStaticMeshEdit, Log, TEXT("AddBoxGeomFromTris (%s): Incomplete set of planes (need 2 per axis)."), ObjName);
-		return false;
-	}
-
-	FMatrix BoxTM = FMatrix::Identity;
-
-	BoxTM.SetAxis(0, Planes[0].Normal);
-	BoxTM.SetAxis(1, Planes[1].Normal);
-
-	// ensure valid TM by cross-product
-	FVector3f ZAxis = Planes[0].Normal ^ Planes[1].Normal;
-
-	if( !AreParallel(ZAxis, Planes[2].Normal) )
-	{
-		UE_LOG(LogStaticMeshEdit, Log, TEXT("AddBoxGeomFromTris (%s): Box axes are not perpendicular."), ObjName);
-		return false;
-	}
-
-	BoxTM.SetAxis(2, ZAxis);
-
-	// OBB centre == AABB centre.
-	FBox Box(ForceInit);
-	for(int32 i=0; i<Tris.Num(); i++)
-	{
-		Box += Tris[i].Vertices[0];
-		Box += Tris[i].Vertices[1];
-		Box += Tris[i].Vertices[2];
-	}
-
-	BoxTM.SetOrigin( Box.GetCenter() );
-
-	// Allocate box in array
-	FKBoxElem BoxElem;
-	BoxElem.SetTransform( FTransform( BoxTM ) );	
-	// distance between parallel planes is box edge lengths.
-	BoxElem.X = FMath::Abs(Planes[0].PlaneDist[0] - Planes[0].PlaneDist[1]);
-	BoxElem.Y = FMath::Abs(Planes[1].PlaneDist[0] - Planes[1].PlaneDist[1]);
-	BoxElem.Z = FMath::Abs(Planes[2].PlaneDist[0] - Planes[2].PlaneDist[1]);
-	AggGeom->BoxElems.Add(BoxElem);
-
-	return true;
-}
-
-/**
- *	Function for adding a sphere collision primitive to the supplied collision geometry based on a set of Verts.
- *	
- *	Simply put an AABB around mesh and use that to generate centre and radius.
- *	It checks that the AABB is square, and that all vertices are either at the
- *	centre, or within 5% of the radius distance away.
- */
-bool AddSphereGeomFromVerts( const TArray<FVector3f>& Verts, FKAggregateGeom* AggGeom, const TCHAR* ObjName )
-{
-	if(Verts.Num() == 0)
-	{
-		return false;
-	}
-
-	FBox Box(ForceInit);
-
-	for(int32 i=0; i<Verts.Num(); i++)
-	{
-		Box += Verts[i];
-	}
-
-	FVector Center, Extents;
-	Box.GetCenterAndExtents(Center, Extents);
-	float Longest = 2.f * Extents.GetMax();
-	float Shortest = 2.f * Extents.GetMin();
-
-	// check that the AABB is roughly a square (5% tolerance)
-	if((Longest - Shortest)/Longest > 0.05f)
-	{
-		UE_LOG(LogStaticMeshEdit, Log, TEXT("AddSphereGeomFromVerts (%s): Sphere bounding box not square."), ObjName);
-		return false;
-	}
-
-	float Radius = 0.5f * Longest;
-
-	// Test that all vertices are a similar radius (5%) from the sphere centre.
-	float MaxR = 0;
-	float MinR = BIG_NUMBER;
-	for(int32 i=0; i<Verts.Num(); i++)
-	{
-		FVector3f CToV = Verts[i] - Center;
-		float RSqr = CToV.SizeSquared();
-
-		MaxR = FMath::Max(RSqr, MaxR);
-
-		// Sometimes vertex at centre, so reject it.
-		if(RSqr > KINDA_SMALL_NUMBER)
-		{
-			MinR = FMath::Min(RSqr, MinR);
-		}
-	}
-
-	MaxR = FMath::Sqrt(MaxR);
-	MinR = FMath::Sqrt(MinR);
-
-	if((MaxR-MinR)/Radius > 0.05f)
-	{
-		UE_LOG(LogStaticMeshEdit, Log, TEXT("AddSphereGeomFromVerts (%s): Vertices not at constant radius."), ObjName );
-		return false;
-	}
-
-	// Allocate sphere in array
-	FKSphereElem SphereElem;
-	SphereElem.Center = Center;
-	SphereElem.Radius = Radius;
-	AggGeom->SphereElems.Add(SphereElem);
-	
-	return true;
-}
-
-bool AddCapsuleGeomFromVerts(const TArray<FVector3f>& Verts, FKAggregateGeom* AggGeom, const TCHAR* ObjName)
-{
-	if (Verts.Num() < 3)
-	{
-		return false;
-	}
-
-	FVector3f AxisStart, AxisEnd;
-	float MaxDistSqr = 0.f;
-
-	for (int32 IndexA = 0; IndexA < Verts.Num() - 1; IndexA++)
-	{
-		for (int32 IndexB = IndexA + 1; IndexB < Verts.Num(); IndexB++)
-		{
-			float DistSqr = (Verts[IndexA] - Verts[IndexB]).SizeSquared();
-			if (DistSqr > MaxDistSqr)
-			{
-				AxisStart = Verts[IndexA];
-				AxisEnd = Verts[IndexB];
-				MaxDistSqr = DistSqr;
-			}
-		}
-	}
-
-	// If we got a valid axis, find vertex furthest from it
-	if (MaxDistSqr > SMALL_NUMBER)
-	{
-		float MaxRadius = 0.f;
-
-		const FVector3f LineOrigin = AxisStart;
-		const FVector3f LineDir = (AxisEnd - AxisStart).GetSafeNormal();
-
-		for (int32 IndexA = 0; IndexA < Verts.Num() - 1; IndexA++)
-		{
-			float DistToAxis = FMath::PointDistToLine(Verts[IndexA], LineDir, LineOrigin);
-			if (DistToAxis > MaxRadius)
-			{
-				MaxRadius = DistToAxis;
-			}
-		}
-
-		if (MaxRadius > SMALL_NUMBER)
-		{
-			// Allocate capsule in array
-			FKSphylElem SphylElem;
-			SphylElem.Center = 0.5f * (AxisStart + AxisEnd);
-			SphylElem.Rotation = FQuat::FindBetweenVectors(FVector3f(0,0,1), LineDir).Rotator(); // Get quat that takes you from z axis to desired axis
-			SphylElem.Radius = MaxRadius;
-			SphylElem.Length = FMath::Max(FMath::Sqrt(MaxDistSqr) - (2.f * MaxRadius), 0.f); // subtract two radii from total length to get segment length (ensure > 0)
-			AggGeom->SphylElems.Add(SphylElem);
-			return true;
-		}
-	}
-
-	return false;
-}
-
-
-/** Utility for adding one convex hull from the given verts */
-bool AddConvexGeomFromVertices( const TArray<FVector3f>& Verts, FKAggregateGeom* AggGeom, const TCHAR* ObjName )
-{
-	if(Verts.Num() == 0)
-	{
-		return false;
-	}
-
-	FKConvexElem* ConvexElem = new(AggGeom->ConvexElems) FKConvexElem();
-	ConvexElem->VertexData = LWC::PromoteArrayType<FVector>(Verts);	// LWC_TODO: Perf pessimization
-	ConvexElem->UpdateElemBox();
-
-	return true;
-}
-
-
 /**
  * Creates a static mesh object from raw triangle data.
  */
@@ -968,769 +415,4 @@ static void TransformPolys(UPolys* Polys,const FMatrix& Matrix)
 
 }
 
-// LOD data to copy over
-struct FExistingLODMeshData
-{
-	FMeshBuildSettings				ExistingBuildSettings;
-	FMeshReductionSettings			ExistingReductionSettings;
-	TUniquePtr<FMeshDescription>	ExistingMeshDescription;
-	TArray<FStaticMaterial>			ExistingMaterials;
-	FPerPlatformFloat				ExistingScreenSize;
-	FString							ExistingSourceImportFilename;
-};
-
-struct FExistingStaticMeshData
-{
-	TArray<FStaticMaterial> 	ExistingMaterials;
-
-	FMeshSectionInfoMap			ExistingSectionInfoMap;
-	TArray<FExistingLODMeshData>	ExistingLODData;
-
-	TArray<UStaticMeshSocket*>	ExistingSockets;
-
-	bool						ExistingCustomizedCollision;
-	bool						bAutoComputeLODScreenSize;
-
-	int32						ExistingLightMapResolution;
-	int32						ExistingLightMapCoordinateIndex;
-
-	TWeakObjectPtr<UAssetImportData> ExistingImportData;
-	TWeakObjectPtr<UThumbnailInfo> ExistingThumbnailInfo;
-
-	UModel*						ExistingCollisionModel;
-	UBodySetup*					ExistingBodySetup;
-
-	// A mapping of vertex positions to their color in the existing static mesh
-	TMap<FVector3f, FColor>		ExistingVertexColorData;
-
-	float						LpvBiasMultiplier;
-	bool						bHasNavigationData;
-	FName						LODGroup;
-	FPerPlatformInt				MinLOD;
-
-	int32						ImportVersion;
-
-	bool UseMaterialNameSlotWorkflow;
-	//The last import material data (fbx original data before user changes)
-	TArray<FName> LastImportMaterialOriginalNameData;
-	TArray<TArray<FName>> LastImportMeshLodSectionMaterialData;
-
-	bool						ExistingGenerateMeshDistanceField;
-	int32						ExistingLODForCollision;
-	float						ExistingDistanceFieldSelfShadowBias;
-	bool						ExistingSupportUniformlyDistributedSampling;
-	bool						ExistingAllowCpuAccess;
-	FVector3f					ExistingPositiveBoundsExtension;
-	FVector3f					ExistingNegativeBoundsExtension;
-
-	UStaticMesh::FOnMeshChanged	ExistingOnMeshChanged;
-	UStaticMesh* ExistingComplexCollisionMesh = nullptr;
-
-	TMap<FName, FString> ExistingUMetaDataTagValues;
-};
-
-bool IsUsingMaterialSlotNameWorkflow(UAssetImportData* AssetImportData)
-{
-	UFbxStaticMeshImportData* ImportData = Cast<UFbxStaticMeshImportData>(AssetImportData);
-	if (ImportData == nullptr || ImportData->ImportMaterialOriginalNameData.Num() <= 0)
-	{
-		return false;
-	}
-	bool AllNameAreNone = true;
-	for (FName ImportMaterialName : ImportData->ImportMaterialOriginalNameData)
-	{
-		if (ImportMaterialName != NAME_None)
-		{
-			AllNameAreNone = false;
-			break;
-		}
-	}
-	return !AllNameAreNone;
-}
-
-TSharedPtr<FExistingStaticMeshData> SaveExistingStaticMeshData(UStaticMesh* ExistingMesh, UnFbx::FBXImportOptions* ImportOptions, int32 LodIndex)
-{
-	if (!ExistingMesh)
-	{
-		return TSharedPtr<FExistingStaticMeshData>();
-	}
-
-		bool bSaveMaterials = !ImportOptions->bImportMaterials;
-		TSharedPtr<FExistingStaticMeshData> ExistingMeshDataPtr = MakeShared<FExistingStaticMeshData>();
-
-		//Save the package UMetaData
-		TMap<FName, FString>* ExistingUMetaDataTagValues = UMetaData::GetMapForObject(ExistingMesh);
-		if(ExistingUMetaDataTagValues && ExistingUMetaDataTagValues->Num() > 0)
-		{
-			ExistingMeshDataPtr->ExistingUMetaDataTagValues = *ExistingUMetaDataTagValues;
-		}
-
-		ExistingMeshDataPtr->ImportVersion = ExistingMesh->ImportVersion;
-		ExistingMeshDataPtr->UseMaterialNameSlotWorkflow = IsUsingMaterialSlotNameWorkflow(ExistingMesh->AssetImportData);
-
-		FMeshSectionInfoMap OldSectionInfoMap = ExistingMesh->GetSectionInfoMap();
-
-		bool bIsReimportCustomLODOverGeneratedLOD = ExistingMesh->IsSourceModelValid(LodIndex) &&
-			(ExistingMesh->GetSourceModel(LodIndex).IsRawMeshEmpty() || (ExistingMesh->IsReductionActive(LodIndex) && ExistingMesh->GetSourceModel(LodIndex).ReductionSettings.BaseLODModel != LodIndex));
-
-		//We need to reset some data in case we import a custom LOD over a generated LOD
-		if (bIsReimportCustomLODOverGeneratedLOD)
-		{
-			//Reset the section info map for this LOD
-			for (int32 SectionIndex = 0; SectionIndex < ExistingMesh->GetNumSections(LodIndex); ++SectionIndex)
-			{
-				OldSectionInfoMap.Remove(LodIndex, SectionIndex);
-			}
-		}
-
-		ExistingMeshDataPtr->ExistingMaterials.Empty();
-		if (bSaveMaterials)
-		{
-			for (const FStaticMaterial &StaticMaterial : ExistingMesh->GetStaticMaterials())
-			{
-				ExistingMeshDataPtr->ExistingMaterials.Add(StaticMaterial);
-			}
-		}
-
-		ExistingMeshDataPtr->ExistingLODData.AddZeroed(ExistingMesh->GetNumSourceModels());
-
-		// refresh material and section info map here
-		// we have to make sure it only contains valid item
-		// we go through section info and only add it back if used, otherwise we don't want to use
-		if (LodIndex == INDEX_NONE)
-		{
-			ExistingMesh->GetSectionInfoMap().Clear();
-		}
-		else
-		{
-			//Remove only the target section InfoMap, if we destroy more we will not restore the correct material assignment for other Lods contain in the same file.
-			int32 ReimportSectionNumber = ExistingMesh->GetSectionInfoMap().GetSectionNumber(LodIndex);
-			for (int32 SectionIndex = 0; SectionIndex < ReimportSectionNumber; ++SectionIndex)
-			{
-				ExistingMesh->GetSectionInfoMap().Remove(LodIndex, SectionIndex);
-			}
-		}
-
-		int32 TotalMaterialIndex = ExistingMeshDataPtr->ExistingMaterials.Num();
-		for (int32 SourceModelIndex = 0; SourceModelIndex < ExistingMesh->GetNumSourceModels(); SourceModelIndex++)
-		{
-			//If the last import was exceeding the maximum number of LOD the source model will contain more LOD so just break the loop
-			if (SourceModelIndex >= ExistingMesh->GetRenderData()->LODResources.Num())
-				break;
-			FStaticMeshLODResources& LOD = ExistingMesh->GetRenderData()->LODResources[SourceModelIndex];
-			int32 NumSections = LOD.Sections.Num();
-			for(int32 SectionIndex = 0; SectionIndex < NumSections; ++SectionIndex)
-			{
-				FMeshSectionInfo Info = OldSectionInfoMap.Get(SourceModelIndex, SectionIndex);
-				if(bSaveMaterials && ExistingMesh->GetStaticMaterials().IsValidIndex(Info.MaterialIndex))
-				{
-					if (ExistingMeshDataPtr->UseMaterialNameSlotWorkflow)
-					{
-						int32 ExistMaterialIndex = ExistingMeshDataPtr->ExistingMaterials.Find(ExistingMesh->GetStaticMaterials()[Info.MaterialIndex]);
-						if (ExistMaterialIndex == INDEX_NONE)
-						{
-							ExistMaterialIndex = ExistingMeshDataPtr->ExistingMaterials.Add(ExistingMesh->GetStaticMaterials()[Info.MaterialIndex]);
-						}
-						Info.MaterialIndex = ExistMaterialIndex;
-					}
-					else
-					{
-						// we only save per LOD separeate IF the material index isn't added yet. 
-						// if it's already added, we don't have to add another one. 
-						if (Info.MaterialIndex >= TotalMaterialIndex)
-						{
-							ExistingMeshDataPtr->ExistingMaterials.Add(ExistingMesh->GetStaticMaterials()[Info.MaterialIndex]);
-
-							// @Todo @fixme
-							// have to refresh material index since it might be pointing at wrong one
-							// this will break IF the base material number grows or shoterns and index will be off
-							// I think we have to save material index per section, so that we don't have to worry about global index
-							Info.MaterialIndex = TotalMaterialIndex++;
-						}
-					}
-					ExistingMeshDataPtr->ExistingSectionInfoMap.Set(SourceModelIndex, SectionIndex, Info);
-				}
-			}
-
-			const FStaticMeshSourceModel& SourceModel = ExistingMesh->GetSourceModel(SourceModelIndex);
-			FExistingLODMeshData& ExistingLODData = ExistingMeshDataPtr->ExistingLODData[SourceModelIndex];
-			ExistingLODData.ExistingBuildSettings = SourceModel.BuildSettings;
-			ExistingLODData.ExistingReductionSettings = SourceModel.ReductionSettings;
-			if (bIsReimportCustomLODOverGeneratedLOD && (SourceModelIndex == LodIndex))
-			{
-				//Reset the reduction
-				ExistingLODData.ExistingReductionSettings.PercentTriangles = 1.0f;
-				ExistingLODData.ExistingReductionSettings.PercentVertices = 1.0f;
-				ExistingLODData.ExistingReductionSettings.MaxDeviation = 0.0f;
-			}
-			ExistingLODData.ExistingScreenSize = SourceModel.ScreenSize.Default;
-			ExistingLODData.ExistingSourceImportFilename = SourceModel.SourceImportFilename;
-
-			const FMeshDescription* MeshDescription = ExistingMesh->GetMeshDescription(SourceModelIndex);
-			if (MeshDescription)
-			{
-				ExistingLODData.ExistingMeshDescription = MakeUnique<FMeshDescription>(*MeshDescription);
-			}
-		}
-
-		ExistingMeshDataPtr->ExistingSockets = ExistingMesh->Sockets;
-
-		ExistingMeshDataPtr->ExistingCustomizedCollision = ExistingMesh->bCustomizedCollision;
-		ExistingMeshDataPtr->bAutoComputeLODScreenSize = ExistingMesh->bAutoComputeLODScreenSize;
-
-		ExistingMeshDataPtr->ExistingLightMapResolution = ExistingMesh->GetLightMapResolution();
-		ExistingMeshDataPtr->ExistingLightMapCoordinateIndex = ExistingMesh->GetLightMapCoordinateIndex();
-
-		ExistingMeshDataPtr->ExistingImportData = ExistingMesh->AssetImportData;
-		ExistingMeshDataPtr->ExistingThumbnailInfo = ExistingMesh->ThumbnailInfo;
-
-		ExistingMeshDataPtr->ExistingBodySetup = ExistingMesh->GetBodySetup();
-
-		ExistingMeshDataPtr->bHasNavigationData = ExistingMesh->bHasNavigationData;
-		ExistingMeshDataPtr->LODGroup = ExistingMesh->LODGroup;
-		ExistingMeshDataPtr->MinLOD = ExistingMesh->GetMinLOD();
-
-		ExistingMeshDataPtr->ExistingGenerateMeshDistanceField = ExistingMesh->bGenerateMeshDistanceField;
-		ExistingMeshDataPtr->ExistingLODForCollision = ExistingMesh->LODForCollision;
-		ExistingMeshDataPtr->ExistingDistanceFieldSelfShadowBias = ExistingMesh->DistanceFieldSelfShadowBias;
-		ExistingMeshDataPtr->ExistingSupportUniformlyDistributedSampling = ExistingMesh->bSupportUniformlyDistributedSampling;
-		ExistingMeshDataPtr->ExistingAllowCpuAccess = ExistingMesh->bAllowCPUAccess;
-		ExistingMeshDataPtr->ExistingPositiveBoundsExtension = ExistingMesh->GetPositiveBoundsExtension();
-		ExistingMeshDataPtr->ExistingNegativeBoundsExtension = ExistingMesh->GetNegativeBoundsExtension();
-
-		UFbxStaticMeshImportData* ImportData = Cast<UFbxStaticMeshImportData>(ExistingMesh->AssetImportData);
-		if (ImportData && ExistingMeshDataPtr->UseMaterialNameSlotWorkflow)
-		{
-			for (int32 ImportMaterialOriginalNameDataIndex = 0; ImportMaterialOriginalNameDataIndex < ImportData->ImportMaterialOriginalNameData.Num(); ++ImportMaterialOriginalNameDataIndex)
-			{
-				FName MaterialName = ImportData->ImportMaterialOriginalNameData[ImportMaterialOriginalNameDataIndex];
-				ExistingMeshDataPtr->LastImportMaterialOriginalNameData.Add(MaterialName);
-			}
-			for (int32 InternalLodIndex = 0; InternalLodIndex < ImportData->ImportMeshLodData.Num(); ++InternalLodIndex)
-			{
-				ExistingMeshDataPtr->LastImportMeshLodSectionMaterialData.AddZeroed();
-				const FImportMeshLodSectionsData &ImportMeshLodSectionsData = ImportData->ImportMeshLodData[InternalLodIndex];
-				for (int32 SectionIndex = 0; SectionIndex < ImportMeshLodSectionsData.SectionOriginalMaterialName.Num(); ++SectionIndex)
-				{
-					FName MaterialName = ImportMeshLodSectionsData.SectionOriginalMaterialName[SectionIndex];
-					ExistingMeshDataPtr->LastImportMeshLodSectionMaterialData[InternalLodIndex].Add(MaterialName);
-				}
-			}
-		}
-		ExistingMeshDataPtr->ExistingOnMeshChanged = ExistingMesh->OnMeshChanged;
-		ExistingMeshDataPtr->ExistingComplexCollisionMesh = ExistingMesh->ComplexCollisionMesh;
-
-	return ExistingMeshDataPtr;
-}
-
-// Helper to find if some reduction settings are active
-bool IsReductionActive(const FMeshReductionSettings& ReductionSettings)
-{
-	bool bUseQuadricSimplier = true;
-	{
-		// Are we using our tool, or simplygon?  The tool is only changed during editor restarts
-		IMeshReduction* ReductionModule = FModuleManager::Get().LoadModuleChecked<IMeshReductionManagerModule>("MeshReductionInterface").GetStaticMeshReductionInterface();
-		FString VersionString = ReductionModule->GetVersionString();
-		TArray<FString> SplitVersionString;
-		VersionString.ParseIntoArray(SplitVersionString, TEXT("_"), true);
-		bUseQuadricSimplier = SplitVersionString[0].Equals("QuadricMeshReduction");
-	}
-	const bool bVertTermination = (bUseQuadricSimplier) && (ReductionSettings.TerminationCriterion != EStaticMeshReductionTerimationCriterion::Triangles) && (ReductionSettings.PercentVertices < 1.0f);
-	const bool bTriTermination = ReductionSettings.TerminationCriterion != EStaticMeshReductionTerimationCriterion::Vertices && (ReductionSettings.PercentTriangles < 1.0f);
-	return bTriTermination || bVertTermination || (ReductionSettings.MaxDeviation > 0.0f);
-}
-
-/* This function is call before building the mesh when we do a re-import*/
-void RestoreExistingMeshSettings(const FExistingStaticMeshData* ExistingMesh, UStaticMesh* NewMesh, int32 LODIndex)
-{
-	if (!ExistingMesh)
-	{
-		return;
-	}
-	NewMesh->LODGroup = ExistingMesh->LODGroup;
-	NewMesh->SetMinLOD(ExistingMesh->MinLOD);
-	int32 ExistingNumLods = ExistingMesh->ExistingLODData.Num();
-	int32 CurrentNumLods = NewMesh->GetNumSourceModels();
-	if (LODIndex == INDEX_NONE)
-	{
-		if (CurrentNumLods > ExistingNumLods)
-		{
-			NewMesh->SetNumSourceModels(ExistingNumLods);
-		}
-		//Create only the LOD Group we need, extra LOD will be put back when calling RestoreExistingMeshData later in the re import process
-		if (NewMesh->LODGroup != NAME_None)
-		{
-			ITargetPlatform* CurrentPlatform = GetTargetPlatformManagerRef().GetRunningTargetPlatform();
-			check(CurrentPlatform);
-			const FStaticMeshLODGroup& LODGroup = CurrentPlatform->GetStaticMeshLODSettings().GetLODGroup(NewMesh->LODGroup);
-			ExistingNumLods = FMath::Min(ExistingNumLods, LODGroup.GetDefaultNumLODs());
-		}
-
-		for (int32 i = 0; i < ExistingNumLods; i++)
-		{
-			if (NewMesh->GetNumSourceModels() <= i)
-			{
-				NewMesh->AddSourceModel();
-			}
-			FMeshDescription* LODMeshDescription = NewMesh->GetMeshDescription(i);
-			bool bSwapFromGeneratedToImported = !ExistingMesh->ExistingLODData[i].ExistingMeshDescription.IsValid() && (LODMeshDescription && LODMeshDescription->Polygons().Num() > 0);
-			bool bWasReduced = IsReductionActive(ExistingMesh->ExistingLODData[i].ExistingReductionSettings);
-
-			FStaticMeshSourceModel& SourceModel = NewMesh->GetSourceModel(i);
-			if (!bSwapFromGeneratedToImported && bWasReduced)
-			{
-				SourceModel.ReductionSettings = ExistingMesh->ExistingLODData[i].ExistingReductionSettings;
-			}
-			SourceModel.BuildSettings = ExistingMesh->ExistingLODData[i].ExistingBuildSettings;
-			SourceModel.ScreenSize = ExistingMesh->ExistingLODData[i].ExistingScreenSize;
-			SourceModel.SourceImportFilename = ExistingMesh->ExistingLODData[i].ExistingSourceImportFilename;
-		}
-	}
-	else
-	{
-		//Just set the old configuration for the desired LODIndex
-		if(LODIndex >= 0 && LODIndex < CurrentNumLods && LODIndex < ExistingNumLods)
-		{
-			FMeshDescription* LODMeshDescription = NewMesh->GetMeshDescription(LODIndex);
-			bool bSwapFromGeneratedToImported = !ExistingMesh->ExistingLODData[LODIndex].ExistingMeshDescription.IsValid() && (LODMeshDescription && LODMeshDescription->Polygons().Num() > 0);
-			bool bWasReduced = IsReductionActive(ExistingMesh->ExistingLODData[LODIndex].ExistingReductionSettings);
-
-			FStaticMeshSourceModel& SourceModel = NewMesh->GetSourceModel(LODIndex);
-			if (!bSwapFromGeneratedToImported && bWasReduced)
-			{
-				SourceModel.ReductionSettings = ExistingMesh->ExistingLODData[LODIndex].ExistingReductionSettings;
-			}
-			SourceModel.BuildSettings = ExistingMesh->ExistingLODData[LODIndex].ExistingBuildSettings;
-			SourceModel.ScreenSize = ExistingMesh->ExistingLODData[LODIndex].ExistingScreenSize;
-			SourceModel.SourceImportFilename = ExistingMesh->ExistingLODData[LODIndex].ExistingSourceImportFilename;
-		}
-	}
-
-	//We need to fill the import version remap before building the mesh since the
-	//static mesh component will be register at the end of the build.
-	//We do the remap of the material override in the static mesh component in OnRegister()
-	if(ExistingMesh->ImportVersion != EImportStaticMeshVersion::LastVersion)
-	{
-		uint32 MaterialMapKey = 0;
-		TArray<int32> ImportRemapMaterial;
-		MaterialMapKey = ((uint32)((ExistingMesh->ImportVersion & 0xffff) << 16) | (uint32)(EImportStaticMeshVersion::LastVersion & 0xffff));
-		//Avoid matching a material more then once
-		TArray<int32> MatchIndex;
-		ImportRemapMaterial.AddZeroed(ExistingMesh->ExistingMaterials.Num());
-		for (int32 ExistMaterialIndex = 0; ExistMaterialIndex < ExistingMesh->ExistingMaterials.Num(); ++ExistMaterialIndex)
-		{
-			ImportRemapMaterial[ExistMaterialIndex] = ExistMaterialIndex; //Set default value
-			const FStaticMaterial &ExistMaterial = ExistingMesh->ExistingMaterials[ExistMaterialIndex];
-			bool bFoundMatchingMaterial = false;
-			for (int32 MaterialIndex = 0; MaterialIndex < NewMesh->GetStaticMaterials().Num(); ++MaterialIndex)
-			{
-				if (MatchIndex.Contains(MaterialIndex))
-				{
-					continue;
-				}
-				FStaticMaterial &Material = NewMesh->GetStaticMaterials()[MaterialIndex];
-				if (Material.ImportedMaterialSlotName == ExistMaterial.ImportedMaterialSlotName)
-				{
-					MatchIndex.Add(MaterialIndex);
-					ImportRemapMaterial[ExistMaterialIndex] = MaterialIndex;
-					bFoundMatchingMaterial = true;
-					break;
-				}
-			}
-			if (!bFoundMatchingMaterial)
-			{
-				for (int32 MaterialIndex = 0; MaterialIndex < NewMesh->GetStaticMaterials().Num(); ++MaterialIndex)
-				{
-					if (MatchIndex.Contains(MaterialIndex))
-					{
-						continue;
-					}
-
-					FStaticMaterial &Material = NewMesh->GetStaticMaterials()[MaterialIndex];
-					if (ExistMaterial.ImportedMaterialSlotName == NAME_None && Material.MaterialInterface == ExistMaterial.MaterialInterface)
-					{
-						MatchIndex.Add(MaterialIndex);
-						ImportRemapMaterial[ExistMaterialIndex] = MaterialIndex;
-						bFoundMatchingMaterial = true;
-						break;
-					}
-				}
-			}
-			if (!bFoundMatchingMaterial)
-			{
-				ImportRemapMaterial[ExistMaterialIndex] = ExistMaterialIndex;
-			}
-		}
-		NewMesh->MaterialRemapIndexPerImportVersion.Add(FMaterialRemapIndex(MaterialMapKey, ImportRemapMaterial));
-	}
-	// Copy mesh changed delegate data
-	NewMesh->OnMeshChanged = ExistingMesh->ExistingOnMeshChanged;
-}
-
-void UpdateSomeLodsImportMeshData(UStaticMesh* NewMesh, TArray<int32> *ReimportLodList)
-{
-	if (NewMesh == nullptr)
-	{
-		return;
-	}
-	UFbxStaticMeshImportData* ImportData = Cast<UFbxStaticMeshImportData>(NewMesh->AssetImportData);
-	//Update the LOD import data before restoring the data
-	if (ReimportLodList != nullptr && ImportData != nullptr)
-	{
-		for (int32 LodLevelImport : (*ReimportLodList))
-		{
-			if (!ImportData->ImportMeshLodData.IsValidIndex(LodLevelImport))
-			{
-				ImportData->ImportMeshLodData.AddZeroed(LodLevelImport - ImportData->ImportMeshLodData.Num() + 1);
-			}
-			ImportData->ImportMeshLodData[LodLevelImport].SectionOriginalMaterialName.Empty();
-			if (NewMesh->GetRenderData()->LODResources.IsValidIndex(LodLevelImport))
-			{
-				FStaticMeshLODResources& LOD = NewMesh->GetRenderData()->LODResources[LodLevelImport];
-				int32 NumSections = LOD.Sections.Num();
-				for (int32 SectionIndex = 0; SectionIndex < NumSections; ++SectionIndex)
-				{
-					int32 MaterialLodSectionIndex = LOD.Sections[SectionIndex].MaterialIndex;
-					if (NewMesh->GetSectionInfoMap().IsValidSection(LodLevelImport, SectionIndex))
-					{
-						MaterialLodSectionIndex = NewMesh->GetSectionInfoMap().Get(LodLevelImport, SectionIndex).MaterialIndex;
-					}
-
-					if (NewMesh->GetStaticMaterials().IsValidIndex(MaterialLodSectionIndex))
-					{
-						bool bFoundMatch = false;
-						FName OriginalImportName = NewMesh->GetStaticMaterials()[MaterialLodSectionIndex].ImportedMaterialSlotName;
-						//Find the material in the original import data
-						int32 ImportMaterialIndex = 0;
-						for (ImportMaterialIndex = 0; ImportMaterialIndex < ImportData->ImportMaterialOriginalNameData.Num(); ++ImportMaterialIndex)
-						{
-							if (ImportData->ImportMaterialOriginalNameData[ImportMaterialIndex] == OriginalImportName)
-							{
-								bFoundMatch = true;
-								break;
-							}
-						}
-						if (!bFoundMatch)
-						{
-							ImportMaterialIndex = ImportData->ImportMaterialOriginalNameData.Add(OriginalImportName);
-						}
-						ImportData->ImportMeshLodData[LodLevelImport].SectionOriginalMaterialName.Add(ImportData->ImportMaterialOriginalNameData[ImportMaterialIndex]);
-					}
-					else
-					{
-						ImportData->ImportMeshLodData[LodLevelImport].SectionOriginalMaterialName.Add(TEXT("InvalidMaterialIndex"));
-					}
-				}
-			}
-		}
-	}
-}
-
-template<typename T>
-void AddGeneratedGeom(const TArray<T>& SrcData, TArray<T>& DstData) {
-	for (const T& CurrentElement : SrcData)
-	{
-		if (CurrentElement.bIsGenerated)
-		{
-			DstData.Add(CurrentElement);
-		}
-	}
-};
-
-void RestoreExistingMeshData(TSharedPtr<const FExistingStaticMeshData> ExistingMeshDataPtr, UStaticMesh* NewMesh, int32 LodLevel, bool bCanShowDialog, bool bForceConflictingMaterialReset)
-{
-	if (!ExistingMeshDataPtr || !NewMesh)
-	{
-		return;
-	}
-
-	//Restore the package metadata
-	if (ExistingMeshDataPtr->ExistingUMetaDataTagValues.Num() > 0)
-	{
-		UMetaData* PackageMetaData = NewMesh->GetOutermost()->GetMetaData();
-		checkSlow(PackageMetaData);
-		PackageMetaData->SetObjectValues(NewMesh, ExistingMeshDataPtr->ExistingUMetaDataTagValues);
-	}
-
-	//Create a remap material Index use to find the matching section later
-	TArray<int32> RemapMaterial;
-	RemapMaterial.AddZeroed(NewMesh->GetStaticMaterials().Num());
-	TArray<FName> RemapMaterialName;
-	RemapMaterialName.AddZeroed(NewMesh->GetStaticMaterials().Num());
-
-	//If user is attended, ask him to verify the match is good
-	UnFbx::EFBXReimportDialogReturnOption ReturnOption;
-	//Ask the user to match the materials conflict
-	UnFbx::FFbxImporter::PrepareAndShowMaterialConflictDialog<FStaticMaterial>(ExistingMeshDataPtr->ExistingMaterials, NewMesh->GetStaticMaterials(), RemapMaterial, RemapMaterialName, bCanShowDialog, false, bForceConflictingMaterialReset, ReturnOption);
-	
-	if (ReturnOption != UnFbx::EFBXReimportDialogReturnOption::FBXRDRO_ResetToFbx)
-	{
-		//Build a ordered material list that try to keep intact the existing material list
-		TArray<FStaticMaterial> MaterialOrdered;
-		TArray<bool> MatchedNewMaterial;
-		MatchedNewMaterial.AddZeroed(NewMesh->GetStaticMaterials().Num());
-		for (int32 ExistMaterialIndex = 0; ExistMaterialIndex < ExistingMeshDataPtr->ExistingMaterials.Num(); ++ExistMaterialIndex)
-		{
-			int32 MaterialIndexOrdered = MaterialOrdered.Add(ExistingMeshDataPtr->ExistingMaterials[ExistMaterialIndex]);
-			FStaticMaterial& OrderedMaterial = MaterialOrdered[MaterialIndexOrdered];
-			int32 NewMaterialIndex = INDEX_NONE;
-			if (RemapMaterial.Find(ExistMaterialIndex, NewMaterialIndex))
-			{
-				MatchedNewMaterial[NewMaterialIndex] = true;
-				RemapMaterial[NewMaterialIndex] = MaterialIndexOrdered;
-				OrderedMaterial.ImportedMaterialSlotName = NewMesh->GetStaticMaterials()[NewMaterialIndex].ImportedMaterialSlotName;
-			}
-			else
-			{
-				//Unmatched material must be conserve
-			}
-		}
-
-		//Add the new material entries (the one that do not match with any existing material)
-		for (int32 NewMaterialIndex = 0; NewMaterialIndex < MatchedNewMaterial.Num(); ++NewMaterialIndex)
-		{
-			if (MatchedNewMaterial[NewMaterialIndex] == false)
-			{
-				int32 NewMeshIndex = MaterialOrdered.Add(NewMesh->GetStaticMaterials()[NewMaterialIndex]);
-				RemapMaterial[NewMaterialIndex] = NewMeshIndex;
-			}
-		}
-
-		//Set the RemapMaterialName array helper
-		for (int32 MaterialIndex = 0; MaterialIndex < RemapMaterial.Num(); ++MaterialIndex)
-		{
-			int32 SourceMaterialMatch = RemapMaterial[MaterialIndex];
-			if (ExistingMeshDataPtr->ExistingMaterials.IsValidIndex(SourceMaterialMatch))
-			{
-				RemapMaterialName[MaterialIndex] = ExistingMeshDataPtr->ExistingMaterials[SourceMaterialMatch].ImportedMaterialSlotName;
-			}
-		}
-
-		//Copy the re ordered materials (this ensure the material array do not change when we re-import)
-		NewMesh->SetStaticMaterials(MaterialOrdered);
-	}
-	int32 NumCommonLODs = FMath::Min<int32>(ExistingMeshDataPtr->ExistingLODData.Num(), NewMesh->GetNumSourceModels());
-	for(int32 i=0; i<NumCommonLODs; i++)
-	{
-		FStaticMeshSourceModel& SourceModel = NewMesh->GetSourceModel(i);
-		SourceModel.BuildSettings = ExistingMeshDataPtr->ExistingLODData[i].ExistingBuildSettings;
-		FMeshDescription* LODMeshDescription = NewMesh->GetMeshDescription(i);
-		//Restore the reduction settings only if the existing data was a using reduction. Because we can set some value if we reimport from existing rawmesh to auto generated.
-		bool bSwapFromGeneratedToImported = !ExistingMeshDataPtr->ExistingLODData[i].ExistingMeshDescription.IsValid() && (LODMeshDescription && LODMeshDescription->Polygons().Num() > 0);
-		bool bWasReduced = IsReductionActive(ExistingMeshDataPtr->ExistingLODData[i].ExistingReductionSettings);
-		if ( !bSwapFromGeneratedToImported && bWasReduced)
-		{
-			SourceModel.ReductionSettings = ExistingMeshDataPtr->ExistingLODData[i].ExistingReductionSettings;
-		}
-		SourceModel.ScreenSize = ExistingMeshDataPtr->ExistingLODData[i].ExistingScreenSize;
-		SourceModel.SourceImportFilename = ExistingMeshDataPtr->ExistingLODData[i].ExistingSourceImportFilename;
-	}
-
-	for(int32 i=NumCommonLODs; i < ExistingMeshDataPtr->ExistingLODData.Num(); ++i)
-	{
-		FStaticMeshSourceModel& SrcModel = NewMesh->AddSourceModel();
-		if (ExistingMeshDataPtr->ExistingLODData[i].ExistingMeshDescription.IsValid())
-		{
-			FMeshDescription* MeshDescription = NewMesh->CreateMeshDescription(i, MoveTemp(*ExistingMeshDataPtr->ExistingLODData[i].ExistingMeshDescription));
-			NewMesh->CommitMeshDescription(i);
-		}
-		SrcModel.BuildSettings = ExistingMeshDataPtr->ExistingLODData[i].ExistingBuildSettings;
-		SrcModel.ReductionSettings = ExistingMeshDataPtr->ExistingLODData[i].ExistingReductionSettings;
-		SrcModel.ScreenSize = ExistingMeshDataPtr->ExistingLODData[i].ExistingScreenSize;
-		SrcModel.SourceImportFilename = ExistingMeshDataPtr->ExistingLODData[i].ExistingSourceImportFilename;
-	}
-
-	// Restore the section info of the just imported LOD so is section info map is remap to fit the mesh material array
-	if (ExistingMeshDataPtr->ExistingSectionInfoMap.Map.Num() > 0)
-	{
-		//Build the mesh we need the render data and the existing section info map build before restoring the data
-		if (NewMesh->GetRenderData()->LODResources.Num() < NewMesh->GetNumSourceModels())
-		{
-			NewMesh->Build();
-		}
-		for (int32 i = 0; i < NewMesh->GetRenderData()->LODResources.Num(); i++)
-		{
-			//If a LOD was specified, only touch the specified LOD
-			if (LodLevel != INDEX_NONE && LodLevel != 0 && LodLevel != i)
-			{
-				continue;
-			}
-
-			
-			//When re-importing the asset, do not touch the LOD that was imported from file, the material array is keep intact so the section should still be valid.
-			bool NoRemapForThisLOD = LodLevel == INDEX_NONE && i != 0 && !NewMesh->GetSourceModel(i).bImportWithBaseMesh && !IsReductionActive(NewMesh->GetSourceModel(i).ReductionSettings);
-
-			FStaticMeshLODResources& LOD = NewMesh->GetRenderData()->LODResources[i];
-			
-			int32 NumSections = LOD.Sections.Num();
-			int32 OldSectionNumber = ExistingMeshDataPtr->ExistingSectionInfoMap.GetSectionNumber(i);
-			for (int32 SectionIndex = 0; SectionIndex < NumSections; ++SectionIndex)
-			{
-				//If the SectionInfoMap is not set yet. Because we re-import LOD 0 but we have other LODs
-				//Just put back the old section Info Map
-				if (NewMesh->GetSectionInfoMap().GetSectionNumber(i) <= SectionIndex)
-				{
-					NewMesh->GetSectionInfoMap().Set(i, SectionIndex, ExistingMeshDataPtr->ExistingSectionInfoMap.Get(i, SectionIndex));
-				}
-				//We recreate the SectionInfoMap from the existing data and we do not remap it if LOD is not auto generate and was not imported
-				if (NoRemapForThisLOD)
-				{
-					continue;
-				}
-
-				FMeshSectionInfo NewSectionInfo =  NewMesh->GetSectionInfoMap().Get(i, SectionIndex);
-				bool bFoundOldMatch = false;
-				bool bKeepOldSectionMaterialIndex = false;
-				int32 OriginalSectionMaterialIndex = INDEX_NONE;
-				if (RemapMaterial.IsValidIndex(NewSectionInfo.MaterialIndex) && NewMesh->GetStaticMaterials().IsValidIndex(NewSectionInfo.MaterialIndex))
-				{
-					//Find the matching old index
-					for (int32 ExistSectionIndex = 0; ExistSectionIndex < OldSectionNumber; ++ExistSectionIndex)
-					{
-						bKeepOldSectionMaterialIndex = false;
-						OriginalSectionMaterialIndex = INDEX_NONE;
-						FMeshSectionInfo OldSectionInfo = ExistingMeshDataPtr->ExistingSectionInfoMap.Get(i, ExistSectionIndex);
-						if (ExistingMeshDataPtr->UseMaterialNameSlotWorkflow)
-						{
-							if (ExistingMeshDataPtr->ExistingImportData.IsValid() &&
-								ExistingMeshDataPtr->LastImportMeshLodSectionMaterialData.IsValidIndex(i) &&
-								ExistingMeshDataPtr->LastImportMeshLodSectionMaterialData[i].IsValidIndex(ExistSectionIndex))
-							{
-								//Keep the old section material index only if the user has change the section mapping
-								bKeepOldSectionMaterialIndex = ExistingMeshDataPtr->LastImportMeshLodSectionMaterialData[i][ExistSectionIndex] != ExistingMeshDataPtr->ExistingMaterials[OldSectionInfo.MaterialIndex].ImportedMaterialSlotName;
-								for (int32 ExistMaterialIndex = 0; ExistMaterialIndex < ExistingMeshDataPtr->ExistingMaterials.Num(); ++ExistMaterialIndex)
-								{
-									if (ExistingMeshDataPtr->LastImportMeshLodSectionMaterialData[i][ExistSectionIndex] == ExistingMeshDataPtr->ExistingMaterials[ExistMaterialIndex].ImportedMaterialSlotName)
-									{
-										OriginalSectionMaterialIndex = ExistMaterialIndex;
-										break;
-									}
-								}
-							}
-						}
-						int32 OldSectionMatchIndex = OriginalSectionMaterialIndex != INDEX_NONE ? OriginalSectionMaterialIndex : OldSectionInfo.MaterialIndex;
-						if (RemapMaterial[NewSectionInfo.MaterialIndex] == OldSectionMatchIndex)
-						{
-							NewMesh->GetSectionInfoMap().Set(i, SectionIndex, OldSectionInfo);
-							bFoundOldMatch = true;
-							break;
-						}
-					}
-				}
-
-				if (!bFoundOldMatch)
-				{
-					if (RemapMaterial.IsValidIndex(NewSectionInfo.MaterialIndex))
-					{
-						//Find the old section that was using the NewSectionInfo.MaterialIndex
-						//This will allow copying the section information: Cast Shadow, enable collision
-						for (int32 ExistSectionIndex = 0; ExistSectionIndex < OldSectionNumber; ++ExistSectionIndex)
-						{
-							FMeshSectionInfo OldSectionInfo = ExistingMeshDataPtr->ExistingSectionInfoMap.Get(i, ExistSectionIndex);
-							if (NewSectionInfo.MaterialIndex == OldSectionInfo.MaterialIndex)
-							{
-								NewSectionInfo.bCastShadow = OldSectionInfo.bCastShadow;
-								NewSectionInfo.bEnableCollision = OldSectionInfo.bEnableCollision;
-								break;
-							}
-						}
-						//If user has change the section info map, we want to keep the change
-						NewSectionInfo.MaterialIndex = RemapMaterial[NewSectionInfo.MaterialIndex];
-						NewMesh->GetSectionInfoMap().Set(i, SectionIndex, NewSectionInfo);
-					}
-				}
-			}
-		}
-		//Store the just imported section info map
-		NewMesh->GetOriginalSectionInfoMap().CopyFrom(NewMesh->GetSectionInfoMap());
-	}
-
-	// Assign sockets from old version of this StaticMesh.
-	for(int32 i=0; i<ExistingMeshDataPtr->ExistingSockets.Num(); i++)
-	{
-		UStaticMeshSocket* ExistingSocket = ExistingMeshDataPtr->ExistingSockets[i];
-		UStaticMeshSocket* Socket = NewMesh->FindSocket(ExistingSocket->SocketName);
-		if (!Socket && !ExistingSocket->bSocketCreatedAtImport)
-		{
-			NewMesh->AddSocket(ExistingSocket);
-		}
-	}
-
-	NewMesh->bCustomizedCollision = ExistingMeshDataPtr->ExistingCustomizedCollision;
-	NewMesh->bAutoComputeLODScreenSize = ExistingMeshDataPtr->bAutoComputeLODScreenSize;
-
-	NewMesh->SetLightMapResolution(ExistingMeshDataPtr->ExistingLightMapResolution);
-	NewMesh->SetLightMapCoordinateIndex(ExistingMeshDataPtr->ExistingLightMapCoordinateIndex);
-
-	if (ExistingMeshDataPtr->ExistingImportData.IsValid())
-	{
-		//RestoredLods
-		UFbxStaticMeshImportData* ImportData = Cast<UFbxStaticMeshImportData>(NewMesh->AssetImportData);
-		TArray<FName> ImportMaterialOriginalNameData;
-		TArray<FImportMeshLodSectionsData> ImportMeshLodData;
-		if (ImportData != nullptr && ImportData->ImportMaterialOriginalNameData.Num() > 0 && ImportData->ImportMeshLodData.Num() > 0)
-		{
-			ImportMaterialOriginalNameData = ImportData->ImportMaterialOriginalNameData;
-			ImportMeshLodData = ImportData->ImportMeshLodData;
-		}
-
-		NewMesh->AssetImportData = ExistingMeshDataPtr->ExistingImportData.Get();
-
-		ImportData = Cast<UFbxStaticMeshImportData>(NewMesh->AssetImportData);
-		if (ImportData != nullptr && ImportMaterialOriginalNameData.Num() > 0 && ImportMeshLodData.Num() > 0)
-		{
-			ImportData->ImportMaterialOriginalNameData = ImportMaterialOriginalNameData;
-			ImportData->ImportMeshLodData = ImportMeshLodData;
-		}
-	}
-		
-
-	NewMesh->ThumbnailInfo = ExistingMeshDataPtr->ExistingThumbnailInfo.Get();
-
-	// If we already had some collision info...
-	if(ExistingMeshDataPtr->ExistingBodySetup)
-	{
-		// If we didn't import anything, keep all collisions otherwise only keep generated collisions.
-		bool bKeepAllCollisions;
-		if(!NewMesh->GetBodySetup() || NewMesh->GetBodySetup()->AggGeom.GetElementCount() == 0)
-		{
-			bKeepAllCollisions = true;
-		}
-		else
-		{
-			bKeepAllCollisions = false;
-		}
-
-		if(bKeepAllCollisions)
-		{
-			NewMesh->SetBodySetup(ExistingMeshDataPtr->ExistingBodySetup);
-		}
-		else
-		{
-			// New collision geometry, but we still want the original settings and the generated collisions
-			NewMesh->GetBodySetup()->CopyBodySetupProperty(ExistingMeshDataPtr->ExistingBodySetup);
-
-			const FKAggregateGeom& ExistingAggGeom = ExistingMeshDataPtr->ExistingBodySetup->AggGeom;
-			FKAggregateGeom& NewAggGeom = NewMesh->GetBodySetup()->AggGeom;
-			AddGeneratedGeom(ExistingAggGeom.BoxElems, NewAggGeom.BoxElems);
-			AddGeneratedGeom(ExistingAggGeom.ConvexElems, NewAggGeom.ConvexElems);
-			AddGeneratedGeom(ExistingAggGeom.SphereElems, NewAggGeom.SphereElems);
-			AddGeneratedGeom(ExistingAggGeom.SphylElems, NewAggGeom.SphylElems);
-		}
-	}
-
-	NewMesh->bHasNavigationData = ExistingMeshDataPtr->bHasNavigationData;
-	NewMesh->LODGroup = ExistingMeshDataPtr->LODGroup;
-
-	NewMesh->bGenerateMeshDistanceField = ExistingMeshDataPtr->ExistingGenerateMeshDistanceField;
-	NewMesh->LODForCollision = ExistingMeshDataPtr->ExistingLODForCollision;
-	NewMesh->DistanceFieldSelfShadowBias = ExistingMeshDataPtr->ExistingDistanceFieldSelfShadowBias;
-	NewMesh->bSupportUniformlyDistributedSampling = ExistingMeshDataPtr->ExistingSupportUniformlyDistributedSampling;
-	NewMesh->bAllowCPUAccess = ExistingMeshDataPtr->ExistingAllowCpuAccess;
-	NewMesh->SetPositiveBoundsExtension(ExistingMeshDataPtr->ExistingPositiveBoundsExtension);
-	NewMesh->SetNegativeBoundsExtension(ExistingMeshDataPtr->ExistingNegativeBoundsExtension);
-
-	NewMesh->ComplexCollisionMesh = ExistingMeshDataPtr->ExistingComplexCollisionMesh;
-}
-
 #undef LOCTEXT_NAMESPACE
diff --git a/Engine/Source/Editor/UnrealEd/Public/ImportUtils/SkelImport.h b/Engine/Source/Editor/UnrealEd/Public/ImportUtils/SkelImport.h
new file mode 100644
--- /dev/null
+++ b/Engine/Source/Editor/UnrealEd/Public/ImportUtils/SkelImport.h
@@ -0,0 +1,123 @@
+// Copyright Epic Games, Inc. All Rights Reserved.
+
+/*-----------------------------------------------------------------------------
+	Data structures only used for importing skeletal meshes and animations.
+-----------------------------------------------------------------------------*/
+
+#pragma once
+
+#include "CoreMinimal.h"
+#include "Containers/IndirectArray.h"
+#include "ReferenceSkeleton.h"
+#include "Engine/SkeletalMesh.h"
+#include "Rendering/SkeletalMeshLODImporterData.h"
+
+class UAssetImportData;
+class UMorphTarget;
+class UPhysicsAsset;
+class USkeletalMeshSocket;
+class USkeleton;
+class UThumbnailInfo;
+class FSkeletalMeshLODModel;
+
+struct FExistingMeshLodSectionData
+{
+	FExistingMeshLodSectionData(FName InImportedMaterialSlotName, bool InbCastShadow, bool InbRecomputeTangents, ESkinVertexColorChannel InRecomputeTangentsVertexMaskChannel, int32 InGenerateUpTo, bool InbDisabled)
+	: ImportedMaterialSlotName(InImportedMaterialSlotName)
+	, bCastShadow(InbCastShadow)
+	, bRecomputeTangents(InbRecomputeTangents)
+	, RecomputeTangentsVertexMaskChannel(InRecomputeTangentsVertexMaskChannel)
+	, GenerateUpTo(InGenerateUpTo)
+	, bDisabled(InbDisabled)
+	{}
+	FName ImportedMaterialSlotName;
+	bool bCastShadow;
+	bool bRecomputeTangents;
+	ESkinVertexColorChannel RecomputeTangentsVertexMaskChannel;
+	int32 GenerateUpTo;
+	bool bDisabled;
+};
+
+struct FExistingSkelMeshData
+{
+	TArray<USkeletalMeshSocket*>			ExistingSockets;
+	TArray<TSharedPtr<FReductionBaseSkeletalMeshBulkData>> ExistingOriginalReductionSourceMeshData;
+	TIndirectArray<FSkeletalMeshLODModel>	ExistingLODModels;
+	TArray<FSkeletalMeshLODInfo>			ExistingLODInfo;
+	FReferenceSkeleton						ExistingRefSkeleton;
+	TArray<FSkeletalMaterial>				ExistingMaterials;
+	bool									bSaveRestoreMaterials;
+	TArray<UMorphTarget*>					ExistingMorphTargets;
+	TArray<UPhysicsAsset*>					ExistingPhysicsAssets;
+	UPhysicsAsset*							ExistingShadowPhysicsAsset;
+	USkeleton*								ExistingSkeleton;
+	TArray<FTransform>						ExistingRetargetBasePose;
+	USkeletalMeshLODSettings*				ExistingLODSettings;
+	TSubclassOf<UAnimInstance>				ExistingPostProcessAnimBlueprint;
+	TSoftObjectPtr<UObject>					ExistingDefaultAnimatingRig;
+	//////////////////////////////////////////////////////////////////////////
+
+	bool									bExistingUseFullPrecisionUVs;
+	bool									bExistingUseHighPrecisionTangentBasis;
+
+	UE_DEPRECATED(5.0, "Please use UMirrorDataTable for mirroring ")
+	PRAGMA_DISABLE_DEPRECATION_WARNINGS
+	TArray<FBoneMirrorExport>				ExistingMirrorTable;
+	PRAGMA_ENABLE_DEPRECATION_WARNINGS
+
+	TWeakObjectPtr<UAssetImportData>		ExistingAssetImportData;
+	TWeakObjectPtr<UThumbnailInfo>			ExistingThumbnailInfo;
+
+	TArray<UClothingAssetBase*>				ExistingClothingAssets;
+
+	bool UseMaterialNameSlotWorkflow;
+	//The existing import material data (the state of sections before the reimport)
+	TArray<FName> ExistingImportMaterialOriginalNameData;
+	TArray<TArray<FExistingMeshLodSectionData>> ExistingImportMeshLodSectionMaterialData;
+	//The last import material data (fbx original data before user changes)
+	TArray<FName> LastImportMaterialOriginalNameData;
+	TArray<TArray<FName>> LastImportMeshLodSectionMaterialData;
+
+	FSkeletalMeshSamplingInfo				ExistingSamplingInfo;
+	FPerPlatformInt							MinLOD;
+	FPerPlatformBool						DisableBelowMinLodStripping;
+	bool									bOverrideLODStreamingSettings;
+	FPerPlatformBool						bSupportLODStreaming;
+	FPerPlatformInt							MaxNumStreamedLODs;
+	FPerPlatformInt							MaxNumOptionalLODs;
+
+	TMap<UAssetUserData*, bool>				ExistingAssetUserData;
+
+	USkeletalMesh::FOnMeshChanged			ExistingOnMeshChanged;
+
+	TMap<FName, FString> ExistingUMetaDataTagValues;
+
+	FVector PositiveBoundsExtension;
+	FVector NegativeBoundsExtension;
+};
+
+/** 
+ * Optional data passed in when importing a skeletal mesh LDO
+ */
+class FSkelMeshOptionalImportData
+{
+public:
+	FSkelMeshOptionalImportData() {}
+
+	/** extra data used for importing extra weight/bone influences */
+	FSkeletalMeshImportData RawMeshInfluencesData;
+	int32 MaxBoneCountPerChunk;
+};
+
+/**
+* Data needed for importing an extra set of vertex influences
+*/
+struct FSkelMeshExtraInfluenceImportData
+{
+	FReferenceSkeleton		RefSkeleton;
+	TArray<SkeletalMeshImportData::FVertInfluence> Influences;
+	TArray<SkeletalMeshImportData::FMeshWedge> Wedges;
+	TArray<SkeletalMeshImportData::FMeshFace> Faces;
+	TArray<FVector> Points;
+	int32 MaxBoneCountPerChunk;
+};
\ No newline at end of file
diff --git a/Engine/Source/Editor/UnrealEd/Public/ImportUtils/SkeletalMeshImportUtils.h b/Engine/Source/Editor/UnrealEd/Public/ImportUtils/SkeletalMeshImportUtils.h
new file mode 100644
--- /dev/null
+++ b/Engine/Source/Editor/UnrealEd/Public/ImportUtils/SkeletalMeshImportUtils.h
@@ -0,0 +1,26 @@
+// Copyright Epic Games, Inc. All Rights Reserved.
+
+#pragma once
+
+#include "Containers/Array.h"
+#include "CoreMinimal.h"
+
+struct FExistingSkelMeshData;
+struct FReferenceSkeleton;
+class FSkeletalMeshImportData;
+class FSkeletalMeshLODModel;
+class USkeletalMesh;
+
+namespace SkeletalMeshImportUtils
+{
+	/** Backups the given SkeletalMesh into a FExistingSkelMeshData */
+	UNREALED_API TSharedPtr<FExistingSkelMeshData> SaveExistingSkelMeshData(USkeletalMesh* SourceSkeletalMesh, bool bSaveMaterials, int32 ReimportLODIndex);
+
+	/** Restore a backed up FExistingSkelMeshData into a SkeletalMesh asset */
+	UNREALED_API void RestoreExistingSkelMeshData(const TSharedPtr<const FExistingSkelMeshData>& MeshData, USkeletalMesh* SkeletalMesh, int32 ReimportLODIndex, bool bCanShowDialog, bool bImportSkinningOnly, bool bForceMaterialReset);
+
+	UNREALED_API void ProcessImportMeshInfluences(FSkeletalMeshImportData& ImportData, const FString& SkeletalMeshName);
+	UNREALED_API void ProcessImportMeshMaterials(TArray<FSkeletalMaterial>& Materials, FSkeletalMeshImportData& ImportData);
+	UNREALED_API bool ProcessImportMeshSkeleton(const USkeleton* SkeletonAsset, FReferenceSkeleton& RefSkeleton, int32& SkeletalDepth, FSkeletalMeshImportData& ImportData);
+	UNREALED_API void ApplySkinning(USkeletalMesh* SkeletalMesh, FSkeletalMeshLODModel& SrcLODModel, FSkeletalMeshLODModel& DestLODModel);
+}
\ No newline at end of file
diff --git a/Engine/Source/Editor/UnrealEd/Public/ImportUtils/StaticMeshImportUtils.h b/Engine/Source/Editor/UnrealEd/Public/ImportUtils/StaticMeshImportUtils.h
new file mode 100644
--- /dev/null
+++ b/Engine/Source/Editor/UnrealEd/Public/ImportUtils/StaticMeshImportUtils.h
@@ -0,0 +1,124 @@
+// Copyright Epic Games, Inc. All Rights Reserved.
+
+/*=============================================================================
+	StaticMeshImportUtils.h: Static mesh import functions.
+=============================================================================*/
+
+#pragma once
+
+#include "CoreMinimal.h"
+#include "Engine/EngineTypes.h"
+#include "Engine/MeshMerging.h"
+#include "Engine/StaticMesh.h"
+#include "PerPlatformProperties.h"
+
+struct FMeshDescription;
+class UStaticMeshSocket;
+class UAssetImportData;
+class UThumbnailInfo;
+class UModel;
+class UBodySetup;
+
+// LOD data to copy over
+struct FExistingLODMeshData
+{
+	FMeshBuildSettings				ExistingBuildSettings;
+	FMeshReductionSettings			ExistingReductionSettings;
+	TUniquePtr<FMeshDescription>	ExistingMeshDescription;
+	TArray<FStaticMaterial>			ExistingMaterials;
+	FPerPlatformFloat				ExistingScreenSize;
+	FString							ExistingSourceImportFilename;
+};
+
+struct FExistingStaticMeshData
+{
+	TArray<FStaticMaterial> 	ExistingMaterials;
+
+	FMeshSectionInfoMap			ExistingSectionInfoMap;
+	TArray<FExistingLODMeshData>	ExistingLODData;
+
+	TArray<UStaticMeshSocket*>	ExistingSockets;
+
+	bool						ExistingCustomizedCollision;
+	bool						bAutoComputeLODScreenSize;
+
+	int32						ExistingLightMapResolution;
+	int32						ExistingLightMapCoordinateIndex;
+
+	TWeakObjectPtr<UAssetImportData> ExistingImportData;
+	TWeakObjectPtr<UThumbnailInfo> ExistingThumbnailInfo;
+
+	UModel* ExistingCollisionModel;
+	UBodySetup* ExistingBodySetup;
+
+	// A mapping of vertex positions to their color in the existing static mesh
+	TMap<FVector3f, FColor>		ExistingVertexColorData;
+
+	float						LpvBiasMultiplier;
+	bool						bHasNavigationData;
+	FName						LODGroup;
+	FPerPlatformInt				MinLOD;
+
+	int32						ImportVersion;
+
+	bool UseMaterialNameSlotWorkflow;
+	//The last import material data (fbx original data before user changes)
+	TArray<FName> LastImportMaterialOriginalNameData;
+	TArray<TArray<FName>> LastImportMeshLodSectionMaterialData;
+
+	bool						ExistingGenerateMeshDistanceField;
+	int32						ExistingLODForCollision;
+	float						ExistingDistanceFieldSelfShadowBias;
+	bool						ExistingSupportUniformlyDistributedSampling;
+	bool						ExistingAllowCpuAccess;
+	FVector3f					ExistingPositiveBoundsExtension;
+	FVector3f					ExistingNegativeBoundsExtension;
+
+	UStaticMesh::FOnMeshChanged	ExistingOnMeshChanged;
+	UStaticMesh* ExistingComplexCollisionMesh = nullptr;
+
+	TMap<FName, FString> ExistingUMetaDataTagValues;
+};
+
+namespace StaticMeshImportUtils
+{
+	UNREALED_API bool DecomposeUCXMesh(const TArray<FVector3f>& CollisionVertices, const TArray<int32>& CollisionFaceIdx, UBodySetup* BodySetup);
+
+	/**
+	 *	Function for adding a box collision primitive to the supplied collision geometry based on the mesh of the box.
+	 *
+	 *	We keep a list of triangle normals found so far. For each normal direction,
+	 *	we should have 2 distances from the origin (2 parallel box faces). If the
+	 *	mesh is a box, we should have 3 distinct normal directions, and 2 distances
+	 *	found for each. The difference between these distances should be the box
+	 *	dimensions. The 3 directions give us the key axes, and therefore the
+	 *	box transformation matrix. This shouldn't rely on any vertex-ordering on
+	 *	the triangles (normals are compared +ve & -ve). It also shouldn't matter
+	 *	about how many triangles make up each side (but it will take longer).
+	 *	We get the centre of the box from the centre of its AABB.
+	 */
+	UNREALED_API bool AddBoxGeomFromTris(const TArray<FPoly>& Tris, FKAggregateGeom* AggGeom, const TCHAR* ObjName);
+
+	/**
+	 *	Function for adding a sphere collision primitive to the supplied collision geometry based on a set of Verts.
+	 *
+	 *	Simply put an AABB around mesh and use that to generate center and radius.
+	 *	It checks that the AABB is square, and that all vertices are either at the
+	 *	center, or within 5% of the radius distance away.
+	 */
+	UNREALED_API bool AddSphereGeomFromVerts(const TArray<FVector3f>& Verts, FKAggregateGeom* AggGeom, const TCHAR* ObjName);
+
+	UNREALED_API bool AddCapsuleGeomFromVerts(const TArray<FVector3f>& Verts, FKAggregateGeom* AggGeom, const TCHAR* ObjName);
+
+	/** Utility for adding one convex hull from the given verts */
+	UNREALED_API bool AddConvexGeomFromVertices(const TArray<FVector3f>& Verts, FKAggregateGeom* AggGeom, const TCHAR* ObjName);
+
+	UNREALED_API TSharedPtr<FExistingStaticMeshData> SaveExistingStaticMeshData(UStaticMesh* ExistingMesh, UnFbx::FBXImportOptions* ImportOptions, int32 LodIndex);
+
+	/* This function is call before building the mesh when we do a re-import*/
+	UNREALED_API void RestoreExistingMeshSettings(const FExistingStaticMeshData* ExistingMesh, UStaticMesh* NewMesh, int32 LODIndex);
+
+	UNREALED_API void UpdateSomeLodsImportMeshData(UStaticMesh* NewMesh, TArray<int32>* ReimportLodList);
+
+	UNREALED_API void RestoreExistingMeshData(const TSharedPtr<const FExistingStaticMeshData>& ExistingMeshDataPtr, UStaticMesh* NewMesh, int32 LodLevel, bool bCanShowDialog, bool bForceConflictingMaterialReset);
+}
\ No newline at end of file
diff --git a/Engine/Source/Editor/UnrealEd/Public/SkelImport.h b/Engine/Source/Editor/UnrealEd/Public/SkelImport.h
deleted file mode 100644
diff --git a/Engine/Source/Runtime/Engine/Private/Rendering/SkeletalMeshLODImporterData.cpp b/Engine/Source/Runtime/Engine/Private/Rendering/SkeletalMeshLODImporterData.cpp
--- a/Engine/Source/Runtime/Engine/Private/Rendering/SkeletalMeshLODImporterData.cpp
+++ b/Engine/Source/Runtime/Engine/Private/Rendering/SkeletalMeshLODImporterData.cpp
@@ -3,14 +3,16 @@
 #if WITH_EDITOR
 
 #include "Rendering/SkeletalMeshLODImporterData.h"
-#include "Logging/LogVerbosity.h"
-#include "Logging/LogMacros.h"
-#include "Serialization/BulkDataWriter.h"
-#include "Serialization/BulkDataReader.h"
-#include "Rendering/SkeletalMeshModel.h"
+
 #include "Engine/SkeletalMesh.h"
 #include "Factories/FbxSkeletalMeshImportData.h"
+#include "ImportUtils/SkeletalMeshImportUtils.h"
+#include "Logging/LogMacros.h"
+#include "Logging/LogVerbosity.h"
 #include "Misc/ScopedSlowTask.h"
+#include "Rendering/SkeletalMeshModel.h"
+#include "Serialization/BulkDataReader.h"
+#include "Serialization/BulkDataWriter.h"
 #include "SkeletalMeshAttributes.h"
 #include "SkeletalMeshOperations.h"
 
@@ -653,7 +655,7 @@ FArchive& operator<<(FArchive& Ar, FSkeletalMeshImportData& RawMesh)
 	//We now save it after the processing is done so for old version we do it here when loading
 	if (Ar.IsLoading() && Version < RAW_SKELETAL_MESH_BULKDATA_VER_AlternateInfluence)
 	{
-		SkeletalMeshHelper::ProcessImportMeshInfluences(RawMesh, FString(TEXT("Unknown"))); // Not sure how to get owning mesh name at this point...
+		SkeletalMeshImportUtils::ProcessImportMeshInfluences(RawMesh, FString(TEXT("Unknown"))); // Not sure how to get owning mesh name at this point...
 	}
 
 	
