Duplicating WorldComp map would not properly initialize UPackage::WorldTileInfo

diff --git a/Engine/Source/Editor/UnrealEd/Private/ObjectTools.cpp b/Engine/Source/Editor/UnrealEd/Private/ObjectTools.cpp
--- a/Engine/Source/Editor/UnrealEd/Private/ObjectTools.cpp
+++ b/Engine/Source/Editor/UnrealEd/Private/ObjectTools.cpp
@@ -792,6 +792,12 @@ namespace ObjectTools
 				DupObject->GetOutermost()->SetPackageFlags(PKG_DisallowExport);
 			}
 
+			// When duplicating a World Composition map, make sure to properly initialize WorldTileInfo
+			if (Object->GetOutermost()->WorldTileInfo.IsValid())
+			{
+				DupObject->GetOutermost()->WorldTileInfo = MakeUnique<FWorldTileInfo>(*Object->GetOutermost()->WorldTileInfo);
+			}
+
 			// Notify the asset registry
 			FAssetRegistryModule::AssetCreated(DupObject);
 
@@ -3670,6 +3676,7 @@ namespace ObjectTools
 						NewPackage->SetPackageFlags(PKG_DisallowExport);
 					}
 
+					// When renaming a World Composition map, make sure to properly initialize WorldTileInfo
 					if (Object->GetOutermost()->WorldTileInfo.IsValid())
 					{
 						NewPackage->WorldTileInfo = MakeUnique<FWorldTileInfo>(*Object->GetOutermost()->WorldTileInfo);
