This normally doesn't break things because compiling will usually clear out the default value string as it is stored on the CDO

diff --git a/Engine/Source/Editor/UnrealEd/Private/Kismet2/BlueprintEditorUtils.cpp b/Engine/Source/Editor/UnrealEd/Private/Kismet2/BlueprintEditorUtils.cpp
--- a/Engine/Source/Editor/UnrealEd/Private/Kismet2/BlueprintEditorUtils.cpp
+++ b/Engine/Source/Editor/UnrealEd/Private/Kismet2/BlueprintEditorUtils.cpp
@@ -5304,7 +5304,7 @@ FName FBlueprintEditorUtils::DuplicateMemberVariable(UBlueprint* InFromBlueprint
 				if (OldPropertyAddr)
 				{
 					// if there is a property for variable, it means the original default value was already copied, so it can be safely overridden
-					Variable.DefaultValue.Empty();
+					NewVar.DefaultValue.Empty();
 					TargetProperty->ExportTextItem(NewVar.DefaultValue, OldPropertyAddr, OldPropertyAddr, nullptr, PPF_SerializedAsImportText);
 				}
 			}
