Fix scene import material popup context menu loading all the materials. The global  GIsImportingT3D variable force to load any object we search for, like when we do IsAssetLoaded, this is why the thumbnails logic of the asset picker is loading all materials when GIsImportingT3D is true. The fbx scene importer do not need this flag to be true.

diff --git a/Engine/Source/Editor/UnrealEd/Private/Fbx/FbxSceneImportFactory.cpp b/Engine/Source/Editor/UnrealEd/Private/Fbx/FbxSceneImportFactory.cpp
--- a/Engine/Source/Editor/UnrealEd/Private/Fbx/FbxSceneImportFactory.cpp
+++ b/Engine/Source/Editor/UnrealEd/Private/Fbx/FbxSceneImportFactory.cpp
@@ -1027,11 +1027,6 @@ FFeedbackContext*	Warn
 
 	GEditor->GetEditorSubsystem<UImportSubsystem>()->BroadcastAssetPreImport(this, Class, InParent, Name, Type);
 	
-	//TODO verify if we really need this when instancing actor in a level from an import
-	//In that case we should change the variable name.
-	GEditor->IsImportingT3D = 1;
-	GIsImportingT3D = GEditor->IsImportingT3D;
-
 	// logger for all error/warnings
 	// this one prints all messages that are stored in FFbxImporter
 	UnFbx::FFbxImporter* FbxImporter = UnFbx::FFbxImporter::GetInstance();
@@ -1057,9 +1052,6 @@ FFeedbackContext*	Warn
 		Warn->Log(ELogVerbosity::Error, FbxImporter->GetErrorMessage());
 		FbxImporter->ReleaseScene();
 		FbxImporter = nullptr;
-		// Mark us as no longer importing a T3D.
-		GEditor->IsImportingT3D = 0;
-		GIsImportingT3D = false;
 		Warn->EndSlowTask();
 		GEditor->GetEditorSubsystem<UImportSubsystem>()->BroadcastAssetPostImport(this, World);
 		return nullptr;
@@ -1101,14 +1093,12 @@ FFeedbackContext*	Warn
 		FbxImporter = nullptr;
 		GlobalImportSettings = nullptr;
 		GlobalImportSettingsReference = nullptr;
-		// Mark us as no longer importing a T3D.
-		GEditor->IsImportingT3D = 0;
-		GIsImportingT3D = false;
 		Warn->EndSlowTask();
 		GEditor->GetEditorSubsystem<UImportSubsystem>()->BroadcastAssetPostImport(this, World);
 		return nullptr;
 	}
 
+
 	SFbxSceneOptionWindow::CopyFbxOptionsToFbxOptions(GlobalImportSettingsReference, GlobalImportSettings);
 
 	//Convert the scene to the correct axis system. Option like force front X or ConvertScene affect the scene conversion
@@ -1132,9 +1122,6 @@ FFeedbackContext*	Warn
 			FbxImporter->ReleaseScene();
 			FbxImporter = nullptr;
 			GlobalImportSettings = nullptr;
-			// Mark us as no longer importing a T3D.
-			GEditor->IsImportingT3D = 0;
-			GIsImportingT3D = false;
 			Warn->EndSlowTask();
 			GEditor->GetEditorSubsystem<UImportSubsystem>()->BroadcastAssetPostImport(this, World);
 			return nullptr;
@@ -1215,18 +1202,11 @@ FFeedbackContext*	Warn
 				FString AssetName = TEXT("FbxScene_") + FPaths::GetBaseFilename(UFactory::CurrentFilename);
 				UPackage *Pkg = CreatePackageForNode(FullnameBP, AssetName);
 
-				//IsImportingT3D will force a load of UObject when finding them, Openning the blueprint editor search for all UBlueprint class object and find them.
-				//This force load everything.
-				GEditor->IsImportingT3D = 0;
-				GIsImportingT3D = false;
 				//Create the blueprint from the actor and replace the actor with an instance of the created blueprint
 				FKismetEditorUtilities::FCreateBlueprintFromActorParams Params;
 				Params.bReplaceActor = true;
 				Params.bKeepMobility = true;
 				UBlueprint* SceneBlueprint = FKismetEditorUtilities::CreateBlueprintFromActor(Pkg->GetName(), HierarchyActor, Params);
-				//Put back the T3D state
-				GEditor->IsImportingT3D = 1;
-				GIsImportingT3D = GEditor->IsImportingT3D;
 
 				if (SceneBlueprint != nullptr && ReimportData != nullptr)
 				{
@@ -1253,9 +1233,6 @@ FFeedbackContext*	Warn
 	GlobalImportSettings = nullptr;
 	GlobalImportSettingsReference = nullptr;
 
-	// Mark us as no longer importing a T3D.
-	GEditor->IsImportingT3D = 0;
-	GIsImportingT3D = false;
 	ReimportData = nullptr;
 
 	Warn->EndSlowTask();
