Fix for editor unresponsiveness after making multiple edits to a user-defined enum asset when it's also referenced by one or more user-defined struct assets.

diff --git a/Engine/Source/Editor/UnrealEd/Private/Kismet2/EnumEditorUtils.cpp b/Engine/Source/Editor/UnrealEd/Private/Kismet2/EnumEditorUtils.cpp
--- a/Engine/Source/Editor/UnrealEd/Private/Kismet2/EnumEditorUtils.cpp
+++ b/Engine/Source/Editor/UnrealEd/Private/Kismet2/EnumEditorUtils.cpp
@@ -278,7 +278,7 @@ void FEnumEditorUtils::PostEditUndo(UUserDefinedEnum* Enum)
 
 void FEnumEditorUtils::BroadcastChanges(const UUserDefinedEnum* Enum, const TArray<TPair<FName, int64>>& OldNames, bool bResolveData)
 {
-	check(NULL != Enum);
+	check(nullptr != Enum);
 	if (bResolveData)
 	{
 		FArchiveEnumeratorResolver EnumeratorResolver(Enum, OldNames);
@@ -297,13 +297,13 @@ void FEnumEditorUtils::BroadcastChanges(const UUserDefinedEnum* Enum, const TArr
 					return OwnerClass;	
 				}
 				// Otherwise check for UserDefinedStructs that may have this property
-				else if (UStruct* OwnerStruct = Prop->GetOwnerStruct())
+				else if (UUserDefinedStruct* UserStruct = Cast<UUserDefinedStruct>(Prop->GetOwnerStruct()))
 				{
-					if(UUserDefinedStruct* UserStruct = Cast<UUserDefinedStruct>(OwnerStruct))
+					// Exclude UDS compiler artifacts; they don't require regeneration.
+					if (!UserStruct->PrimaryStruct.IsValid())
 					{
 						EffectedUserStructs.Add(UserStruct);
 					}
-					return OwnerStruct->GetClass();
 				}
 			}
 			return nullptr;
@@ -352,6 +352,8 @@ void FEnumEditorUtils::BroadcastChanges(const UUserDefinedEnum* Enum, const TArr
 		// serialization in order for the structure editor to get updated properly
 		for(UUserDefinedStruct* Struct : EffectedUserStructs)
 		{
+			Struct->Serialize(EnumeratorResolver);
+
 			FStructureEditorUtils::ModifyStructData(Struct);
 			FStructureEditorUtils::OnStructureChanged(Struct);
 		}
