FBX re-import: only non fbx metadata should be restored during re-import

diff --git a/Engine/Source/Editor/UnrealEd/Private/Fbx/FbxMainImport.cpp b/Engine/Source/Editor/UnrealEd/Private/Fbx/FbxMainImport.cpp
--- a/Engine/Source/Editor/UnrealEd/Private/Fbx/FbxMainImport.cpp
+++ b/Engine/Source/Editor/UnrealEd/Private/Fbx/FbxMainImport.cpp
@@ -1271,6 +1271,32 @@ void FFbxImporter::EnsureNodeNameAreValid(const FString& BaseFilename)
 	}
 }
 
+void FFbxImporter::RemoveFBXMetaData(const UObject* Object)
+{
+	TArray<FName> KeysToRemove;
+	if (const TMap<FName, FString>* ExistingUMetaDataTagValues = UMetaData::GetMapForObject(Object))
+	{
+		for (const TPair<FName, FString>& KeyValue : *ExistingUMetaDataTagValues)
+		{
+			if (KeyValue.Key.ToString().StartsWith(FBX_METADATA_PREFIX, ESearchCase::IgnoreCase))
+			{
+				KeysToRemove.Add(KeyValue.Key);
+			}
+		}
+	}
+
+	if (KeysToRemove.Num() > 0)
+	{
+		UMetaData* PackageMetaData = Object->GetOutermost()->GetMetaData();
+		checkSlow(PackageMetaData);
+		for (const FName& KeyToRemove : KeysToRemove)
+		{
+			PackageMetaData->RemoveValue(Object, KeyToRemove);
+		}
+	}
+}
+
+
 FString FFbxImporter::GetFileAxisDirection()
 {
 	FString AxisDirection;
diff --git a/Engine/Source/Editor/UnrealEd/Private/Fbx/FbxSkeletalMeshImport.cpp b/Engine/Source/Editor/UnrealEd/Private/Fbx/FbxSkeletalMeshImport.cpp
--- a/Engine/Source/Editor/UnrealEd/Private/Fbx/FbxSkeletalMeshImport.cpp
+++ b/Engine/Source/Editor/UnrealEd/Private/Fbx/FbxSkeletalMeshImport.cpp
@@ -1776,6 +1776,10 @@ USkeletalMesh* UnFbx::FFbxImporter::ImportSkeletalMesh(FImportSkeletalMeshArgs &
 			//Copy all the build option to reflect any change in the setting using the re-import UI
 			LODInfoPtr->BuildSettings = BuildOptions;
 		}
+
+		// Before calling SaveExistingSkelMeshData() we must remove the existing fbx metadata as we don't want to restore those.
+		RemoveFBXMetaData(ExistingSkelMesh);
+
 		//The backup of the skeletal mesh data empty the LOD array in the ImportedResource of the skeletal mesh
 		//If the import fail after this step the editor can crash when updating the bone later since the LODModel will not exist anymore
 		ExistSkelMeshDataPtr = SkeletalMeshImportUtils::SaveExistingSkelMeshData(ExistingSkelMesh, !ImportOptions->bImportMaterials, ImportSkeletalMeshArgs.LodIndex);
diff --git a/Engine/Source/Editor/UnrealEd/Private/Fbx/FbxStaticMeshImport.cpp b/Engine/Source/Editor/UnrealEd/Private/Fbx/FbxStaticMeshImport.cpp
--- a/Engine/Source/Editor/UnrealEd/Private/Fbx/FbxStaticMeshImport.cpp
+++ b/Engine/Source/Editor/UnrealEd/Private/Fbx/FbxStaticMeshImport.cpp
@@ -1462,6 +1462,8 @@ UStaticMesh* UnFbx::FFbxImporter::ReimportStaticMesh(UStaticMesh* Mesh, UFbxStat
 	ImportOptions->bImportMaterials = false;
 	ImportOptions->bImportTextures = false;
 
+	// Before calling SaveExistingStaticMeshData() we must remove the existing fbx metadata as we don't want to restore those.
+	RemoveFBXMetaData(Mesh);
 	TSharedPtr<FExistingStaticMeshData> ExistMeshDataPtr = StaticMeshImportUtils::SaveExistingStaticMeshData(Mesh, ImportOptions, INDEX_NONE);
 
 	TArray<int32> ReimportLodList;
diff --git a/Engine/Source/Editor/UnrealEd/Private/ImportUtils/InternalImportUtils.cpp b/Engine/Source/Editor/UnrealEd/Private/ImportUtils/InternalImportUtils.cpp
new file mode 100644
--- /dev/null
+++ b/Engine/Source/Editor/UnrealEd/Private/ImportUtils/InternalImportUtils.cpp
@@ -0,0 +1,56 @@
+// Copyright Epic Games, Inc. All Rights Reserved.
+
+/*=============================================================================
+	InternalImportUtils.cpp: Internal import utils functions.
+=============================================================================*/
+
+#include "ImportUtils/InternalImportUtils.h"
+
+#include "Containers/UnrealString.h"
+#include "CoreMinimal.h"
+#include "Engine/EngineTypes.h"
+#include "Factories/FbxMeshImportData.h"
+#include "UObject/MetaData.h"
+#include "UObject/NameTypes.h"
+#include "UObject/Object.h"
+
+
+bool InternalImportUtils::IsUsingMaterialSlotNameWorkflow(UAssetImportData* AssetImportData)
+{
+	UFbxMeshImportData* ImportData = Cast<UFbxMeshImportData>(AssetImportData);
+	if (ImportData == nullptr || ImportData->ImportMaterialOriginalNameData.Num() <= 0)
+	{
+		return false;
+	}
+	bool AllNameAreNone = true;
+	for (FName ImportMaterialName : ImportData->ImportMaterialOriginalNameData)
+	{
+		if (ImportMaterialName != NAME_None)
+		{
+			AllNameAreNone = false;
+			break;
+		}
+	}
+	return !AllNameAreNone;
+}
+
+
+void InternalImportUtils::RestoreMetaData(UObject* TargetObject, const TMap<FName, FString>& ExistingUMetaDataTagValues)
+{
+	if (ExistingUMetaDataTagValues.Num() > 0)
+	{
+		UMetaData* PackageMetaData = TargetObject->GetOutermost()->GetMetaData();
+		checkSlow(PackageMetaData);
+
+		TMap<FName, FString> MetaDataToApply(ExistingUMetaDataTagValues);
+		// Keep all existing metadata, but override/add those who were imported.
+		if (TMap<FName, FString>* ImportedMetaData = UMetaData::GetMapForObject(TargetObject))
+		{
+			for (const TPair<FName, FString>& KeyValue : *ImportedMetaData)
+			{
+				MetaDataToApply.Add(KeyValue);
+			}
+		}
+		PackageMetaData->SetObjectValues(TargetObject, MetaDataToApply);
+	}
+}
diff --git a/Engine/Source/Editor/UnrealEd/Private/ImportUtils/InternalImportUtils.h b/Engine/Source/Editor/UnrealEd/Private/ImportUtils/InternalImportUtils.h
new file mode 100644
--- /dev/null
+++ b/Engine/Source/Editor/UnrealEd/Private/ImportUtils/InternalImportUtils.h
@@ -0,0 +1,26 @@
+// Copyright Epic Games, Inc. All Rights Reserved.
+
+/*=============================================================================
+	InternalImportUtils.h: Internal import utils functions.
+=============================================================================*/
+
+#pragma once
+
+#include "CoreMinimal.h"
+
+class UObject;
+class UAssetImportData;
+
+namespace InternalImportUtils
+{
+	/**
+	 * Returns true if the asset was imported with FBX importer and is using the "Material Slot Name Workflow".
+	 */
+	bool IsUsingMaterialSlotNameWorkflow(UAssetImportData* AssetImportData);
+
+	/**
+	 * Add and override the values of ExistingUMetaDataTagValues metadata map in the TargetObject package. 
+	 * The existing metadata in the package are not removed.
+	 */
+	void RestoreMetaData(UObject* TargetObject, const TMap<FName, FString>& ExistingUMetaDataTagValues);
+}
diff --git a/Engine/Source/Editor/UnrealEd/Private/ImportUtils/SkeletalMeshImportUtils.cpp b/Engine/Source/Editor/UnrealEd/Private/ImportUtils/SkeletalMeshImportUtils.cpp
--- a/Engine/Source/Editor/UnrealEd/Private/ImportUtils/SkeletalMeshImportUtils.cpp
+++ b/Engine/Source/Editor/UnrealEd/Private/ImportUtils/SkeletalMeshImportUtils.cpp
@@ -14,6 +14,7 @@
 #include "Engine/SkeletalMeshSocket.h"
 #include "Factories/FbxSkeletalMeshImportData.h"
 #include "FbxImporter.h"
+#include "ImportUtils/InternalImportUtils.h"
 #include "ImportUtils/SkelImport.h"
 #include "Interfaces/ITargetPlatform.h"
 #include "Interfaces/ITargetPlatformManagerModule.h"
@@ -38,8 +39,6 @@ namespace SkeletalMesUtilsImpl
 	/** Check that root bone is the same, and that any bones that are common have the correct parent. */
 	bool SkeletonsAreCompatible(const FReferenceSkeleton& NewSkel, const FReferenceSkeleton& ExistSkel, bool bFailNoError);
 
-	bool SkeletalMeshIsUsingMaterialSlotNameWorkflow(UAssetImportData* AssetImportData);
-
 	void SaveSkeletalMeshLODModelSections(USkeletalMesh* SourceSkeletalMesh, TSharedPtr<FExistingSkelMeshData>& ExistingMeshDataPtr, int32 LodIndex, bool bSaveNonReducedMeshData);
 
 	void SaveSkeletalMeshMaterialNameWorkflowData(TSharedPtr<FExistingSkelMeshData>& ExistingMeshDataPtr, const USkeletalMesh* SourceSkeletalMesh);
@@ -223,25 +222,6 @@ void SkeletalMeshImportUtils::ProcessImportMeshInfluences(FSkeletalMeshImportDat
 	FLODUtilities::ProcessImportMeshInfluences(ImportData.Wedges.Num(), ImportData.Influences, SkeletalMeshName);
 }
 
-bool SkeletalMesUtilsImpl::SkeletalMeshIsUsingMaterialSlotNameWorkflow(UAssetImportData* AssetImportData)
-{
-	UFbxSkeletalMeshImportData* ImportData = Cast<UFbxSkeletalMeshImportData>(AssetImportData);
-	if (ImportData == nullptr || ImportData->ImportMaterialOriginalNameData.Num() <= 0)
-	{
-		return false;
-	}
-	bool AllNameAreNone = true;
-	for (FName ImportMaterialName : ImportData->ImportMaterialOriginalNameData)
-	{
-		if (ImportMaterialName != NAME_None)
-		{
-			AllNameAreNone = false;
-			break;
-		}
-	}
-	return !AllNameAreNone;
-}
-
 void SkeletalMesUtilsImpl::SaveSkeletalMeshLODModelSections(USkeletalMesh* SourceSkeletalMesh, TSharedPtr<FExistingSkelMeshData>& ExistingMeshDataPtr, int32 LodIndex, bool bSaveNonReducedMeshData)
 {
 	const FSkeletalMeshModel* SourceMeshModel = SourceSkeletalMesh->GetImportedModel();
@@ -394,7 +374,7 @@ TSharedPtr<FExistingSkelMeshData> SkeletalMeshImportUtils::SaveExistingSkelMeshD
 	{
 		ExistingMeshDataPtr->ExistingUMetaDataTagValues = *MetaDataTagValues;
 	}
-	ExistingMeshDataPtr->UseMaterialNameSlotWorkflow = SkeletalMeshIsUsingMaterialSlotNameWorkflow(SourceSkeletalMesh->GetAssetImportData());
+	ExistingMeshDataPtr->UseMaterialNameSlotWorkflow = InternalImportUtils::IsUsingMaterialSlotNameWorkflow(SourceSkeletalMesh->GetAssetImportData());
 	ExistingMeshDataPtr->MinLOD = SourceSkeletalMesh->GetMinLod();
 	ExistingMeshDataPtr->DisableBelowMinLodStripping = SourceSkeletalMesh->GetDisableBelowMinLodStripping();
 	ExistingMeshDataPtr->bOverrideLODStreamingSettings = SourceSkeletalMesh->GetOverrideLODStreamingSettings();
@@ -691,12 +671,7 @@ void SkeletalMeshImportUtils::RestoreExistingSkelMeshData(const TSharedPtr<const
 	}
 
 	//Restore the package metadata
-	if (MeshData->ExistingUMetaDataTagValues.Num() > 0)
-	{
-		UMetaData* PackageMetaData = SkeletalMesh->GetOutermost()->GetMetaData();
-		checkSlow(PackageMetaData);
-		PackageMetaData->SetObjectValues(SkeletalMesh, MeshData->ExistingUMetaDataTagValues);
-	}
+	InternalImportUtils::RestoreMetaData(SkeletalMesh, MeshData->ExistingUMetaDataTagValues);
 
 	int32 SafeReimportLODIndex = ReimportLODIndex < 0 ? 0 : ReimportLODIndex;
 	SkeletalMesh->SetMinLod(MeshData->MinLOD);
diff --git a/Engine/Source/Editor/UnrealEd/Private/ImportUtils/StaticMeshImportUtils.cpp b/Engine/Source/Editor/UnrealEd/Private/ImportUtils/StaticMeshImportUtils.cpp
--- a/Engine/Source/Editor/UnrealEd/Private/ImportUtils/StaticMeshImportUtils.cpp
+++ b/Engine/Source/Editor/UnrealEd/Private/ImportUtils/StaticMeshImportUtils.cpp
@@ -12,6 +12,7 @@
 #include "Engine/StaticMesh.h"
 #include "IMeshReductionInterfaces.h"
 #include "IMeshReductionManagerModule.h"
+#include "ImportUtils/InternalImportUtils.h"
 #include "MeshDescription.h"
 #include "Modules/ModuleManager.h"
 #include "UObject/MetaData.h"
@@ -572,24 +573,6 @@ bool StaticMeshImportUtils::AddConvexGeomFromVertices( const TArray<FVector3f>&
 	return true;
 }
 
-bool IsUsingMaterialSlotNameWorkflow(UAssetImportData* AssetImportData)
-{
-	UFbxStaticMeshImportData* ImportData = Cast<UFbxStaticMeshImportData>(AssetImportData);
-	if (ImportData == nullptr || ImportData->ImportMaterialOriginalNameData.Num() <= 0)
-	{
-		return false;
-	}
-	bool AllNameAreNone = true;
-	for (FName ImportMaterialName : ImportData->ImportMaterialOriginalNameData)
-	{
-		if (ImportMaterialName != NAME_None)
-		{
-			AllNameAreNone = false;
-			break;
-		}
-	}
-	return !AllNameAreNone;
-}
 
 TSharedPtr<FExistingStaticMeshData> StaticMeshImportUtils::SaveExistingStaticMeshData(UStaticMesh* ExistingMesh, UnFbx::FBXImportOptions* ImportOptions, int32 LodIndex)
 {
@@ -598,170 +581,170 @@ TSharedPtr<FExistingStaticMeshData> StaticMeshImportUtils::SaveExistingStaticMes
 		return TSharedPtr<FExistingStaticMeshData>();
 	}
 
-		bool bSaveMaterials = !ImportOptions->bImportMaterials;
-		TSharedPtr<FExistingStaticMeshData> ExistingMeshDataPtr = MakeShared<FExistingStaticMeshData>();
+	bool bSaveMaterials = !ImportOptions->bImportMaterials;
+	TSharedPtr<FExistingStaticMeshData> ExistingMeshDataPtr = MakeShared<FExistingStaticMeshData>();
 
-		//Save the package UMetaData
-		TMap<FName, FString>* ExistingUMetaDataTagValues = UMetaData::GetMapForObject(ExistingMesh);
-		if(ExistingUMetaDataTagValues && ExistingUMetaDataTagValues->Num() > 0)
-		{
-			ExistingMeshDataPtr->ExistingUMetaDataTagValues = *ExistingUMetaDataTagValues;
-		}
+	//Save the package UMetaData
+	TMap<FName, FString>* ExistingUMetaDataTagValues = UMetaData::GetMapForObject(ExistingMesh);
+	if(ExistingUMetaDataTagValues && ExistingUMetaDataTagValues->Num() > 0)
+	{
+		ExistingMeshDataPtr->ExistingUMetaDataTagValues = *ExistingUMetaDataTagValues;
+	}
 
-		ExistingMeshDataPtr->ImportVersion = ExistingMesh->ImportVersion;
-		ExistingMeshDataPtr->UseMaterialNameSlotWorkflow = IsUsingMaterialSlotNameWorkflow(ExistingMesh->AssetImportData);
+	ExistingMeshDataPtr->ImportVersion = ExistingMesh->ImportVersion;
+	ExistingMeshDataPtr->UseMaterialNameSlotWorkflow = InternalImportUtils::IsUsingMaterialSlotNameWorkflow(ExistingMesh->AssetImportData);
 
-		FMeshSectionInfoMap OldSectionInfoMap = ExistingMesh->GetSectionInfoMap();
+	FMeshSectionInfoMap OldSectionInfoMap = ExistingMesh->GetSectionInfoMap();
 
-		bool bIsReimportCustomLODOverGeneratedLOD = ExistingMesh->IsSourceModelValid(LodIndex) &&
-			(ExistingMesh->GetSourceModel(LodIndex).IsRawMeshEmpty() || (ExistingMesh->IsReductionActive(LodIndex) && ExistingMesh->GetSourceModel(LodIndex).ReductionSettings.BaseLODModel != LodIndex));
+	bool bIsReimportCustomLODOverGeneratedLOD = ExistingMesh->IsSourceModelValid(LodIndex) &&
+		(ExistingMesh->GetSourceModel(LodIndex).IsRawMeshEmpty() || (ExistingMesh->IsReductionActive(LodIndex) && ExistingMesh->GetSourceModel(LodIndex).ReductionSettings.BaseLODModel != LodIndex));
 
-		//We need to reset some data in case we import a custom LOD over a generated LOD
-		if (bIsReimportCustomLODOverGeneratedLOD)
+	//We need to reset some data in case we import a custom LOD over a generated LOD
+	if (bIsReimportCustomLODOverGeneratedLOD)
+	{
+		//Reset the section info map for this LOD
+		for (int32 SectionIndex = 0; SectionIndex < ExistingMesh->GetNumSections(LodIndex); ++SectionIndex)
 		{
-			//Reset the section info map for this LOD
-			for (int32 SectionIndex = 0; SectionIndex < ExistingMesh->GetNumSections(LodIndex); ++SectionIndex)
-			{
-				OldSectionInfoMap.Remove(LodIndex, SectionIndex);
-			}
+			OldSectionInfoMap.Remove(LodIndex, SectionIndex);
 		}
+	}
 
-		ExistingMeshDataPtr->ExistingMaterials.Empty();
-		if (bSaveMaterials)
+	ExistingMeshDataPtr->ExistingMaterials.Empty();
+	if (bSaveMaterials)
+	{
+		for (const FStaticMaterial &StaticMaterial : ExistingMesh->GetStaticMaterials())
 		{
-			for (const FStaticMaterial &StaticMaterial : ExistingMesh->GetStaticMaterials())
-			{
-				ExistingMeshDataPtr->ExistingMaterials.Add(StaticMaterial);
-			}
+			ExistingMeshDataPtr->ExistingMaterials.Add(StaticMaterial);
 		}
+	}
 
-		ExistingMeshDataPtr->ExistingLODData.AddZeroed(ExistingMesh->GetNumSourceModels());
+	ExistingMeshDataPtr->ExistingLODData.AddZeroed(ExistingMesh->GetNumSourceModels());
 
-		// refresh material and section info map here
-		// we have to make sure it only contains valid item
-		// we go through section info and only add it back if used, otherwise we don't want to use
-		if (LodIndex == INDEX_NONE)
-		{
-			ExistingMesh->GetSectionInfoMap().Clear();
-		}
-		else
+	// refresh material and section info map here
+	// we have to make sure it only contains valid item
+	// we go through section info and only add it back if used, otherwise we don't want to use
+	if (LodIndex == INDEX_NONE)
+	{
+		ExistingMesh->GetSectionInfoMap().Clear();
+	}
+	else
+	{
+		//Remove only the target section InfoMap, if we destroy more we will not restore the correct material assignment for other Lods contain in the same file.
+		int32 ReimportSectionNumber = ExistingMesh->GetSectionInfoMap().GetSectionNumber(LodIndex);
+		for (int32 SectionIndex = 0; SectionIndex < ReimportSectionNumber; ++SectionIndex)
 		{
-			//Remove only the target section InfoMap, if we destroy more we will not restore the correct material assignment for other Lods contain in the same file.
-			int32 ReimportSectionNumber = ExistingMesh->GetSectionInfoMap().GetSectionNumber(LodIndex);
-			for (int32 SectionIndex = 0; SectionIndex < ReimportSectionNumber; ++SectionIndex)
-			{
-				ExistingMesh->GetSectionInfoMap().Remove(LodIndex, SectionIndex);
-			}
+			ExistingMesh->GetSectionInfoMap().Remove(LodIndex, SectionIndex);
 		}
+	}
 
-		int32 TotalMaterialIndex = ExistingMeshDataPtr->ExistingMaterials.Num();
-		for (int32 SourceModelIndex = 0; SourceModelIndex < ExistingMesh->GetNumSourceModels(); SourceModelIndex++)
+	int32 TotalMaterialIndex = ExistingMeshDataPtr->ExistingMaterials.Num();
+	for (int32 SourceModelIndex = 0; SourceModelIndex < ExistingMesh->GetNumSourceModels(); SourceModelIndex++)
+	{
+		//If the last import was exceeding the maximum number of LOD the source model will contain more LOD so just break the loop
+		if (SourceModelIndex >= ExistingMesh->GetRenderData()->LODResources.Num())
+			break;
+		FStaticMeshLODResources& LOD = ExistingMesh->GetRenderData()->LODResources[SourceModelIndex];
+		int32 NumSections = LOD.Sections.Num();
+		for(int32 SectionIndex = 0; SectionIndex < NumSections; ++SectionIndex)
 		{
-			//If the last import was exceeding the maximum number of LOD the source model will contain more LOD so just break the loop
-			if (SourceModelIndex >= ExistingMesh->GetRenderData()->LODResources.Num())
-				break;
-			FStaticMeshLODResources& LOD = ExistingMesh->GetRenderData()->LODResources[SourceModelIndex];
-			int32 NumSections = LOD.Sections.Num();
-			for(int32 SectionIndex = 0; SectionIndex < NumSections; ++SectionIndex)
+			FMeshSectionInfo Info = OldSectionInfoMap.Get(SourceModelIndex, SectionIndex);
+			if(bSaveMaterials && ExistingMesh->GetStaticMaterials().IsValidIndex(Info.MaterialIndex))
 			{
-				FMeshSectionInfo Info = OldSectionInfoMap.Get(SourceModelIndex, SectionIndex);
-				if(bSaveMaterials && ExistingMesh->GetStaticMaterials().IsValidIndex(Info.MaterialIndex))
+				if (ExistingMeshDataPtr->UseMaterialNameSlotWorkflow)
 				{
-					if (ExistingMeshDataPtr->UseMaterialNameSlotWorkflow)
+					int32 ExistMaterialIndex = ExistingMeshDataPtr->ExistingMaterials.Find(ExistingMesh->GetStaticMaterials()[Info.MaterialIndex]);
+					if (ExistMaterialIndex == INDEX_NONE)
 					{
-						int32 ExistMaterialIndex = ExistingMeshDataPtr->ExistingMaterials.Find(ExistingMesh->GetStaticMaterials()[Info.MaterialIndex]);
-						if (ExistMaterialIndex == INDEX_NONE)
-						{
-							ExistMaterialIndex = ExistingMeshDataPtr->ExistingMaterials.Add(ExistingMesh->GetStaticMaterials()[Info.MaterialIndex]);
-						}
-						Info.MaterialIndex = ExistMaterialIndex;
+						ExistMaterialIndex = ExistingMeshDataPtr->ExistingMaterials.Add(ExistingMesh->GetStaticMaterials()[Info.MaterialIndex]);
 					}
-					else
+					Info.MaterialIndex = ExistMaterialIndex;
+				}
+				else
+				{
+					// we only save per LOD separeate IF the material index isn't added yet. 
+					// if it's already added, we don't have to add another one. 
+					if (Info.MaterialIndex >= TotalMaterialIndex)
 					{
-						// we only save per LOD separeate IF the material index isn't added yet. 
-						// if it's already added, we don't have to add another one. 
-						if (Info.MaterialIndex >= TotalMaterialIndex)
-						{
-							ExistingMeshDataPtr->ExistingMaterials.Add(ExistingMesh->GetStaticMaterials()[Info.MaterialIndex]);
+						ExistingMeshDataPtr->ExistingMaterials.Add(ExistingMesh->GetStaticMaterials()[Info.MaterialIndex]);
 
-							// @Todo @fixme
-							// have to refresh material index since it might be pointing at wrong one
-							// this will break IF the base material number grows or shoterns and index will be off
-							// I think we have to save material index per section, so that we don't have to worry about global index
-							Info.MaterialIndex = TotalMaterialIndex++;
-						}
+						// @Todo @fixme
+						// have to refresh material index since it might be pointing at wrong one
+						// this will break IF the base material number grows or shoterns and index will be off
+						// I think we have to save material index per section, so that we don't have to worry about global index
+						Info.MaterialIndex = TotalMaterialIndex++;
 					}
-					ExistingMeshDataPtr->ExistingSectionInfoMap.Set(SourceModelIndex, SectionIndex, Info);
 				}
+				ExistingMeshDataPtr->ExistingSectionInfoMap.Set(SourceModelIndex, SectionIndex, Info);
 			}
+		}
 
-			const FStaticMeshSourceModel& SourceModel = ExistingMesh->GetSourceModel(SourceModelIndex);
-			FExistingLODMeshData& ExistingLODData = ExistingMeshDataPtr->ExistingLODData[SourceModelIndex];
-			ExistingLODData.ExistingBuildSettings = SourceModel.BuildSettings;
-			ExistingLODData.ExistingReductionSettings = SourceModel.ReductionSettings;
-			if (bIsReimportCustomLODOverGeneratedLOD && (SourceModelIndex == LodIndex))
-			{
-				//Reset the reduction
-				ExistingLODData.ExistingReductionSettings.PercentTriangles = 1.0f;
-				ExistingLODData.ExistingReductionSettings.PercentVertices = 1.0f;
-				ExistingLODData.ExistingReductionSettings.MaxDeviation = 0.0f;
-			}
-			ExistingLODData.ExistingScreenSize = SourceModel.ScreenSize.Default;
-			ExistingLODData.ExistingSourceImportFilename = SourceModel.SourceImportFilename;
+		const FStaticMeshSourceModel& SourceModel = ExistingMesh->GetSourceModel(SourceModelIndex);
+		FExistingLODMeshData& ExistingLODData = ExistingMeshDataPtr->ExistingLODData[SourceModelIndex];
+		ExistingLODData.ExistingBuildSettings = SourceModel.BuildSettings;
+		ExistingLODData.ExistingReductionSettings = SourceModel.ReductionSettings;
+		if (bIsReimportCustomLODOverGeneratedLOD && (SourceModelIndex == LodIndex))
+		{
+			//Reset the reduction
+			ExistingLODData.ExistingReductionSettings.PercentTriangles = 1.0f;
+			ExistingLODData.ExistingReductionSettings.PercentVertices = 1.0f;
+			ExistingLODData.ExistingReductionSettings.MaxDeviation = 0.0f;
+		}
+		ExistingLODData.ExistingScreenSize = SourceModel.ScreenSize.Default;
+		ExistingLODData.ExistingSourceImportFilename = SourceModel.SourceImportFilename;
 
-			const FMeshDescription* MeshDescription = ExistingMesh->GetMeshDescription(SourceModelIndex);
-			if (MeshDescription)
-			{
-				ExistingLODData.ExistingMeshDescription = MakeUnique<FMeshDescription>(*MeshDescription);
-			}
+		const FMeshDescription* MeshDescription = ExistingMesh->GetMeshDescription(SourceModelIndex);
+		if (MeshDescription)
+		{
+			ExistingLODData.ExistingMeshDescription = MakeUnique<FMeshDescription>(*MeshDescription);
 		}
+	}
 
-		ExistingMeshDataPtr->ExistingSockets = ExistingMesh->Sockets;
+	ExistingMeshDataPtr->ExistingSockets = ExistingMesh->Sockets;
 
-		ExistingMeshDataPtr->ExistingCustomizedCollision = ExistingMesh->bCustomizedCollision;
-		ExistingMeshDataPtr->bAutoComputeLODScreenSize = ExistingMesh->bAutoComputeLODScreenSize;
+	ExistingMeshDataPtr->ExistingCustomizedCollision = ExistingMesh->bCustomizedCollision;
+	ExistingMeshDataPtr->bAutoComputeLODScreenSize = ExistingMesh->bAutoComputeLODScreenSize;
 
-		ExistingMeshDataPtr->ExistingLightMapResolution = ExistingMesh->GetLightMapResolution();
-		ExistingMeshDataPtr->ExistingLightMapCoordinateIndex = ExistingMesh->GetLightMapCoordinateIndex();
+	ExistingMeshDataPtr->ExistingLightMapResolution = ExistingMesh->GetLightMapResolution();
+	ExistingMeshDataPtr->ExistingLightMapCoordinateIndex = ExistingMesh->GetLightMapCoordinateIndex();
 
-		ExistingMeshDataPtr->ExistingImportData = ExistingMesh->AssetImportData;
-		ExistingMeshDataPtr->ExistingThumbnailInfo = ExistingMesh->ThumbnailInfo;
+	ExistingMeshDataPtr->ExistingImportData = ExistingMesh->AssetImportData;
+	ExistingMeshDataPtr->ExistingThumbnailInfo = ExistingMesh->ThumbnailInfo;
 
-		ExistingMeshDataPtr->ExistingBodySetup = ExistingMesh->GetBodySetup();
+	ExistingMeshDataPtr->ExistingBodySetup = ExistingMesh->GetBodySetup();
 
-		ExistingMeshDataPtr->bHasNavigationData = ExistingMesh->bHasNavigationData;
-		ExistingMeshDataPtr->LODGroup = ExistingMesh->LODGroup;
-		ExistingMeshDataPtr->MinLOD = ExistingMesh->GetMinLOD();
+	ExistingMeshDataPtr->bHasNavigationData = ExistingMesh->bHasNavigationData;
+	ExistingMeshDataPtr->LODGroup = ExistingMesh->LODGroup;
+	ExistingMeshDataPtr->MinLOD = ExistingMesh->GetMinLOD();
 
-		ExistingMeshDataPtr->ExistingGenerateMeshDistanceField = ExistingMesh->bGenerateMeshDistanceField;
-		ExistingMeshDataPtr->ExistingLODForCollision = ExistingMesh->LODForCollision;
-		ExistingMeshDataPtr->ExistingDistanceFieldSelfShadowBias = ExistingMesh->DistanceFieldSelfShadowBias;
-		ExistingMeshDataPtr->ExistingSupportUniformlyDistributedSampling = ExistingMesh->bSupportUniformlyDistributedSampling;
-		ExistingMeshDataPtr->ExistingAllowCpuAccess = ExistingMesh->bAllowCPUAccess;
-		ExistingMeshDataPtr->ExistingPositiveBoundsExtension = ExistingMesh->GetPositiveBoundsExtension();
-		ExistingMeshDataPtr->ExistingNegativeBoundsExtension = ExistingMesh->GetNegativeBoundsExtension();
+	ExistingMeshDataPtr->ExistingGenerateMeshDistanceField = ExistingMesh->bGenerateMeshDistanceField;
+	ExistingMeshDataPtr->ExistingLODForCollision = ExistingMesh->LODForCollision;
+	ExistingMeshDataPtr->ExistingDistanceFieldSelfShadowBias = ExistingMesh->DistanceFieldSelfShadowBias;
+	ExistingMeshDataPtr->ExistingSupportUniformlyDistributedSampling = ExistingMesh->bSupportUniformlyDistributedSampling;
+	ExistingMeshDataPtr->ExistingAllowCpuAccess = ExistingMesh->bAllowCPUAccess;
+	ExistingMeshDataPtr->ExistingPositiveBoundsExtension = ExistingMesh->GetPositiveBoundsExtension();
+	ExistingMeshDataPtr->ExistingNegativeBoundsExtension = ExistingMesh->GetNegativeBoundsExtension();
 
-		UFbxStaticMeshImportData* ImportData = Cast<UFbxStaticMeshImportData>(ExistingMesh->AssetImportData);
-		if (ImportData && ExistingMeshDataPtr->UseMaterialNameSlotWorkflow)
+	UFbxStaticMeshImportData* ImportData = Cast<UFbxStaticMeshImportData>(ExistingMesh->AssetImportData);
+	if (ImportData && ExistingMeshDataPtr->UseMaterialNameSlotWorkflow)
+	{
+		for (int32 ImportMaterialOriginalNameDataIndex = 0; ImportMaterialOriginalNameDataIndex < ImportData->ImportMaterialOriginalNameData.Num(); ++ImportMaterialOriginalNameDataIndex)
 		{
-			for (int32 ImportMaterialOriginalNameDataIndex = 0; ImportMaterialOriginalNameDataIndex < ImportData->ImportMaterialOriginalNameData.Num(); ++ImportMaterialOriginalNameDataIndex)
-			{
-				FName MaterialName = ImportData->ImportMaterialOriginalNameData[ImportMaterialOriginalNameDataIndex];
-				ExistingMeshDataPtr->LastImportMaterialOriginalNameData.Add(MaterialName);
-			}
-			for (int32 InternalLodIndex = 0; InternalLodIndex < ImportData->ImportMeshLodData.Num(); ++InternalLodIndex)
+			FName MaterialName = ImportData->ImportMaterialOriginalNameData[ImportMaterialOriginalNameDataIndex];
+			ExistingMeshDataPtr->LastImportMaterialOriginalNameData.Add(MaterialName);
+		}
+		for (int32 InternalLodIndex = 0; InternalLodIndex < ImportData->ImportMeshLodData.Num(); ++InternalLodIndex)
+		{
+			ExistingMeshDataPtr->LastImportMeshLodSectionMaterialData.AddZeroed();
+			const FImportMeshLodSectionsData &ImportMeshLodSectionsData = ImportData->ImportMeshLodData[InternalLodIndex];
+			for (int32 SectionIndex = 0; SectionIndex < ImportMeshLodSectionsData.SectionOriginalMaterialName.Num(); ++SectionIndex)
 			{
-				ExistingMeshDataPtr->LastImportMeshLodSectionMaterialData.AddZeroed();
-				const FImportMeshLodSectionsData &ImportMeshLodSectionsData = ImportData->ImportMeshLodData[InternalLodIndex];
-				for (int32 SectionIndex = 0; SectionIndex < ImportMeshLodSectionsData.SectionOriginalMaterialName.Num(); ++SectionIndex)
-				{
-					FName MaterialName = ImportMeshLodSectionsData.SectionOriginalMaterialName[SectionIndex];
-					ExistingMeshDataPtr->LastImportMeshLodSectionMaterialData[InternalLodIndex].Add(MaterialName);
-				}
+				FName MaterialName = ImportMeshLodSectionsData.SectionOriginalMaterialName[SectionIndex];
+				ExistingMeshDataPtr->LastImportMeshLodSectionMaterialData[InternalLodIndex].Add(MaterialName);
 			}
 		}
-		ExistingMeshDataPtr->ExistingOnMeshChanged = ExistingMesh->OnMeshChanged;
-		ExistingMeshDataPtr->ExistingComplexCollisionMesh = ExistingMesh->ComplexCollisionMesh;
+	}
+	ExistingMeshDataPtr->ExistingOnMeshChanged = ExistingMesh->OnMeshChanged;
+	ExistingMeshDataPtr->ExistingComplexCollisionMesh = ExistingMesh->ComplexCollisionMesh;
 
 	return ExistingMeshDataPtr;
 }
@@ -988,22 +971,7 @@ void StaticMeshImportUtils::RestoreExistingMeshData(const TSharedPtr<const FExis
 	}
 
 	//Restore the package metadata
-	if (ExistingMeshDataPtr->ExistingUMetaDataTagValues.Num() > 0)
-	{
-		UMetaData* PackageMetaData = NewMesh->GetOutermost()->GetMetaData();
-		checkSlow(PackageMetaData);
-		
-		TMap<FName, FString> MetaDataToApply(ExistingMeshDataPtr->ExistingUMetaDataTagValues);
-		// Keep all existing metadata, but override/add those who were imported.
-		if (TMap<FName, FString>* ImportedMetaData = UMetaData::GetMapForObject(NewMesh))
-		{
-			for (const TPair<FName, FString>& KeyValue : *ImportedMetaData)
-			{
-				MetaDataToApply.Add(KeyValue);
-			}
-		}
-		PackageMetaData->SetObjectValues(NewMesh, MoveTemp(MetaDataToApply));
-	}
+	InternalImportUtils::RestoreMetaData(NewMesh, ExistingMeshDataPtr->ExistingUMetaDataTagValues);
 
 	//Create a remap material Index use to find the matching section later
 	TArray<int32> RemapMaterial;
diff --git a/Engine/Source/Editor/UnrealEd/Public/FbxImporter.h b/Engine/Source/Editor/UnrealEd/Public/FbxImporter.h
--- a/Engine/Source/Editor/UnrealEd/Public/FbxImporter.h
+++ b/Engine/Source/Editor/UnrealEd/Public/FbxImporter.h
@@ -1231,6 +1231,14 @@ private:
 	*/
 	void EnsureNodeNameAreValid(const FString& BaseFilename);
 
+	/**
+	 * Remove metadata starting with the FBX_METADATA_PREFIX value associated to the specified UObject.
+	 * This is used to ensure we do not restore old fbx meta data during reimports.
+	 * 
+	 * @param Object	The UObject for which we want the FBX metadata removed.
+	 */
+	static void RemoveFBXMetaData(const UObject* Object);
+
 private:
 	/**
 	 * Helper structure to pass around the common animation parameters.
