Missing file for CL16396607

diff --git a/Engine/Source/Editor/UnrealEd/Private/Commandlets/ContentCommandlets.cpp b/Engine/Source/Editor/UnrealEd/Private/Commandlets/ContentCommandlets.cpp
--- a/Engine/Source/Editor/UnrealEd/Private/Commandlets/ContentCommandlets.cpp
+++ b/Engine/Source/Editor/UnrealEd/Private/Commandlets/ContentCommandlets.cpp
@@ -32,6 +32,7 @@
 #include "Engine/MapBuildDataRegistry.h"
 #include "Commandlets/ListMaterialsUsedWithMeshEmittersCommandlet.h"
 #include "Commandlets/ListStaticMeshesImportedFromSpeedTreesCommandlet.h"
+#include "Commandlets/StaticMeshMinLodCommandlet.h"
 #include "Particles/ParticleSystem.h"
 #include "Commandlets/ResavePackagesCommandlet.h"
 #include "Commandlets/WrangleContentCommandlet.h"
@@ -43,6 +44,9 @@
 #include "Engine/Brush.h"
 #include "Editor.h"
 #include "FileHelpers.h"
+#include "PlatformInfo.h"
+#include "CollectionManagerModule.h"
+#include "ICollectionManager.h"
 #include "CommandletSourceControlUtils.h"
 #include "WorldPartition/WorldPartition.h"
 #include "WorldPartition/WorldPartitionHelpers.h"
@@ -3285,5 +3289,270 @@ int32 UListStaticMeshesImportedFromSpeedTreesCommandlet::Main(const FString& Par
 	{
 		UE_LOG(LogContentCommandlet, Display, TEXT("No static meshes were imported from speedtrees in this project."));
 	}
+	return 0;
+}
+
+/* ==========================================================================================================
+	UStaticMeshMinLodCommandlet
+========================================================================================================== */
+
+UStaticMeshMinLodCommandlet::UStaticMeshMinLodCommandlet(const FObjectInitializer& ObjectInitializer)
+	: Super(ObjectInitializer)
+{
+}
+
+int32 UStaticMeshMinLodCommandlet::Main(const FString& Params)
+{
+	TArray<FString> FilesInPath;
+	FEditorFileUtils::FindAllPackageFiles(FilesInPath);
+	if (FilesInPath.Num() == 0)
+	{
+		UE_LOG(LogContentCommandlet, Warning, TEXT("No packages found"));
+		return 1;
+	}
+
+	// parse the mapping from PerPlatform to PerQualityLevel
+	// ex:-mapping=Mobile:Low,Switch:Medium,Desktop:High,PS4:High,XboxOne:High,PS5:Epic,XSX:Epic;console:medium;high;epic,Desktop:high
+	TMultiMap<FName, FName> PerPlatformToQualityLevel;
+	FString MappingStr;
+	if (FParse::Value(*Params, TEXT("-mapping="), MappingStr, false))
+	{
+		TArray<FString> Mappings;
+		MappingStr.ParseIntoArray(Mappings, TEXT(","), true);
+		for (FString& PlatformToQualityLevel : Mappings)
+		{
+			TArray<FString> Entries;
+			PlatformToQualityLevel.ParseIntoArray(Entries, TEXT(":"), false);
+
+			if (Entries.Num() != 2)
+			{
+				UE_LOG(LogContentCommandlet, Error, TEXT("Error bad -mapping argument: %s"), *MappingStr);
+				return 1;
+			}
+
+			TArray<FString> Values;
+			Entries[1].ParseIntoArray(Values, TEXT(";"), false);
+
+			for (const FString& Value : Values)
+			{
+				PerPlatformToQualityLevel.AddUnique(FName(*Entries[0]), FName(*Value));
+			}
+		}
+	}
+
+	bool bConvertToQualityLevel = PerPlatformToQualityLevel.Num() > 0;
+	bool bNoSourceControl = FParse::Param(*Params, TEXT("nosourcecontrol"));
+	bool bGenerateCollections = FParse::Param(*Params, TEXT("collections"));
+	ISourceControlProvider* SourceControlProvider = bNoSourceControl ? nullptr : &ISourceControlModule::Get().GetProvider();
+	ICollectionManager& CollectionManager = FCollectionManagerModule::GetModule().Get();
+
+	struct OverrideGroupInfo
+	{
+		int32 Count;
+		TArray<FString> Names;
+		TArray<FString> Paths;
+	};
+
+	TMap<FName, OverrideGroupInfo> MinLodStats;
+	TArray<FString> StaticMeshList;
+	int32 TotalPackagesChecked = 0;
+	int32 TotalStaticMeshesChecked = 0;
+	int32 GCIndex = 0;
+
+	// Load the asset registry module
+	FAssetRegistryModule& AssetRegistryModule = FModuleManager::LoadModuleChecked<FAssetRegistryModule>(TEXT("AssetRegistry"));
+
+	// Update Registry Module
+	UE_LOG(LogContentCommandlet, Display, TEXT("Searching Asset Registry for static mesh "));
+	AssetRegistryModule.Get().SearchAllAssets(true);
+
+	// Retrieve list of all assets, used to find unreferenced ones.
+	TArray<FAssetData> AssetList;
+	AssetRegistryModule.Get().GetAssetsByClass(UStaticMesh::StaticClass()->GetFName(), AssetList, true);
+	TArray<UPackage*> PackagesToSave;
+
+	// Platform (group) names
+	const TArray<FString> Filters = { TEXT("NoEditor"), TEXT("Client"), TEXT("Server"), TEXT("AllDesktop") };
+	TMultiMap<FName, FName> GroupToPlatform;
+
+	// sanitize all vanilla platform names
+	// generate a list of all supported platform (from the datadrivenplatform files)
+	const TArray<FName>& SanitizedPlatformNameArray = PlatformInfo::GetAllVanillaPlatformNames().FilterByPredicate([&Filters, &GroupToPlatform](const FName& PlatformName)
+		{
+			for (const FString& Filter : Filters)
+			{
+				const int32 Position = PlatformName.ToString().Find(Filter);
+				if (Position != INDEX_NONE)
+				{
+					return false;
+				}
+			}
+
+			if (const PlatformInfo::FPlatformInfo* PlatformInfo = PlatformInfo::FindPlatformInfo(PlatformName))
+			{
+				GroupToPlatform.AddUnique(PlatformInfo->PlatformGroupName, PlatformName);
+			}
+			return true;
+		});
+
+	for (int32 AssetIdx = 0; AssetIdx < AssetList.Num(); ++AssetIdx)
+	{
+		bool SavePackage = false;
+		const FString Filename = AssetList[AssetIdx].ObjectPath.ToString();
+		UE_LOG(LogContentCommandlet, Display, TEXT("Processing static mesh (%i/%i):  %s "), AssetIdx, AssetList.Num(), *Filename);
+
+		UPackage* Package = LoadPackage(NULL, *Filename, LOAD_Quiet);
+		if (Package == NULL)
+		{
+			UE_LOG(LogContentCommandlet, Error, TEXT("Error loading %s!"), *Filename);
+			continue;
+		}
+
+		TotalPackagesChecked++;
+
+		for (TObjectIterator<UStaticMesh> It; It; ++It)
+		{
+			UStaticMesh* StaticMesh = *It;
+			if (StaticMesh->IsIn(Package) && !StaticMesh->IsTemplate() && StaticMesh->GetMinLOD().PerPlatform.Num() > 0)
+			{
+				if (bConvertToQualityLevel)
+				{
+					//Convert default value
+					FPerPlatformInt PerPlatformMinLOD = StaticMesh->GetMinLOD();
+					FPerQualityLevelInt QualityLevelMinLOD;
+					QualityLevelMinLOD.PerQuality.Empty();
+					QualityLevelMinLOD.Default = PerPlatformMinLOD.Default;
+					
+					// convert each platform overrides
+					for (const TPair<FName, int32>& Pair : StaticMesh->GetMinLOD().PerPlatform)
+					{
+						// get all quality levels associated with the PerPlatform override (PS4:high, Console:medium;high;Epic) ...
+						TArray<FName> QLNames;
+						PerPlatformToQualityLevel.MultiFind(Pair.Key, QLNames);
+
+						for (const FName& QLName : QLNames)
+						{
+							int32 QLKey = QualityLevelProperty::FNameToQualityLevel(QLName);
+							if (QLKey != INDEX_NONE)
+							{
+								int32* Value = QualityLevelMinLOD.PerQuality.Find(QLKey);
+
+								// if the quality level already as a value, only change it if the value is lower
+								// this can happen if two mapping key as the same quality level but different min lod idx value
+								// ex: Desktop=2, PS4=1
+								// If Desktop and Console also is mapping to high, we take the smallest min lod idx
+								if (Value != nullptr && Pair.Value < *Value)
+								{
+									// only change the override if its a smaller minLod
+									*Value = Pair.Value;
+								}
+								else
+								{
+									QualityLevelMinLOD.PerQuality.Add(QLKey, Pair.Value);
+								}
+							}
+						}
+					}
+					//StaticMesh->SetQualityLevelMinLOD(MoveTemp(QualityLevelMinLOD));
+				}
+
+				//generate a unique group name to collect stats
+				TArray<FName> PerPlatformNames;
+				StaticMesh->GetMinLOD().PerPlatform.GenerateKeyArray(PerPlatformNames);
+
+				// sort platform names so that (Switch, PS4) and (PS4, Switch) produce the same key
+				PerPlatformNames.Sort(FNameLexicalLess());
+
+				FString FinalCategoryName;
+				for (const FName& SortedPlatformName : PerPlatformNames)
+				{
+					FinalCategoryName += SortedPlatformName.ToString() + TEXT(" ");
+				}
+
+				// save the override info
+				OverrideGroupInfo& Info = MinLodStats.FindOrAdd(FName(*FinalCategoryName));
+				Info.Count++;
+				Info.Names.Add(StaticMesh->GetPathName());
+				Info.Paths.Add(Filename);
+
+				if (!SavePackage)
+				{
+					PackagesToSave.AddUnique(Package);
+					SavePackage = true;
+				}
+				
+				TotalStaticMeshesChecked++;
+			}
+		}
+	}
+
+	// output stats to the log
+	// generate collection uasset
+	int32 TotalOverrideCount = 0;
+	for (const TPair<FName, OverrideGroupInfo>& Pair : MinLodStats)
+	{
+		const OverrideGroupInfo& Info = Pair.Value;
+
+		if (bGenerateCollections)
+		{
+			if (!CollectionManager.CollectionExists(Pair.Key, ECollectionShareType::CST_Local))
+			{
+				CollectionManager.CreateCollection(Pair.Key, ECollectionShareType::CST_Local, ECollectionStorageMode::Static);
+			}
+		}
+
+		UE_LOG(LogContentCommandlet, Display, TEXT("-------------------------------------------------------------------"));
+		UE_LOG(LogContentCommandlet, Display, TEXT("Mask: %s"), *Pair.Key.ToString());
+		UE_LOG(LogContentCommandlet, Display, TEXT("Nb overrides: %d"), Info.Count);
+
+		for (int i = 0; i < Info.Names.Num(); i++)
+		{
+			UE_LOG(LogContentCommandlet, Display, TEXT("	%s"), *Info.Names[i]);
+
+			if (bGenerateCollections)
+			{
+				CollectionManager.AddToCollection(Pair.Key, ECollectionShareType::CST_Local, FName(*Info.Paths[i]));
+			}
+		}
+
+		if (bGenerateCollections)
+		{
+			CollectionManager.SaveCollection(Pair.Key, ECollectionShareType::CST_Local);
+		}
+
+		UE_LOG(LogContentCommandlet, Display, TEXT("-------------------------------------------------------------------"));
+		TotalOverrideCount += Info.Count;
+	}
+
+	UE_LOG(LogContentCommandlet, Display, TEXT("-------------------------------------------------------------------"));
+	UE_LOG(LogContentCommandlet, Display, TEXT("Total overrides: %d"), TotalOverrideCount);
+	UE_LOG(LogContentCommandlet, Display, TEXT("Total ovestatic meshesrrides: %d"), TotalStaticMeshesChecked);
+	UE_LOG(LogContentCommandlet, Display, TEXT("-------------------------------------------------------------------"));
+
+	// save quality level modifications 
+	if (bConvertToQualityLevel)
+	{
+		if (SourceControlProvider)
+		{
+			FEditorFileUtils::CheckoutPackages(PackagesToSave, nullptr, false);
+		}
+		else
+		{
+			for (UPackage* Package : PackagesToSave)
+			{
+				FString PackageFilename = SourceControlHelpers::PackageFilename(Package);
+				if (IPlatformFile::GetPlatformPhysical().FileExists(*PackageFilename))
+				{
+					if (!IPlatformFile::GetPlatformPhysical().SetReadOnly(*PackageFilename, false))
+					{
+						UE_LOG(LogContentCommandlet, Error, TEXT("Error setting %s writable"), *PackageFilename);
+						return 1;
+					}
+				}
+			}
+		}
+		FEditorFileUtils::PromptForCheckoutAndSave(PackagesToSave, false, false, nullptr, true, false);
+	}
+	
 	return 0;
 }
\ No newline at end of file
