Another compile fix, sorry ...

diff --git a/Engine/Source/Editor/PropertyEditor/Private/SPerQualityLevelPropertiesWidget.cpp b/Engine/Source/Editor/PropertyEditor/Private/SPerQualityLevelPropertiesWidget.cpp
--- a/Engine/Source/Editor/PropertyEditor/Private/SPerQualityLevelPropertiesWidget.cpp
+++ b/Engine/Source/Editor/PropertyEditor/Private/SPerQualityLevelPropertiesWidget.cpp
@@ -14,6 +14,7 @@
 #include "DetailLayoutBuilder.h"
 #include "Widgets/Images/SImage.h"
 #include "Scalability.h"
+#include "PerQualityLevelProperties.h"
 
 void SOverridePropertiesWidget::Construct(const typename SOverridePropertiesWidget::FArguments& InArgs)
 {
diff --git a/Engine/Source/Editor/UnrealEd/Private/Commandlets/ContentCommandlets.cpp b/Engine/Source/Editor/UnrealEd/Private/Commandlets/ContentCommandlets.cpp
--- a/Engine/Source/Editor/UnrealEd/Private/Commandlets/ContentCommandlets.cpp
+++ b/Engine/Source/Editor/UnrealEd/Private/Commandlets/ContentCommandlets.cpp
@@ -89,6 +89,7 @@ DEFINE_LOG_CATEGORY(LogContentCommandlet);
 #include "CollectionManagerModule.h"
 #include "UObject/UObjectThreadContext.h"
 #include "Engine/LODActor.h"
+#include "PerQualityLevelProperties.h"
 
 /**-----------------------------------------------------------------------------
  *	UResavePackages commandlet.
