Improve scalability for the way we strip LOD on static mesh.

Change the per platform overrides with  a quality level property  driven by scalability/device profile ini files. User can now enter a Quality level override and specify the qualtiy in the platform engine or scalability ini files. this allow to have different values on one platform, for example XboxOne Base vs XboxOne X. Keep the PerPlatform MinLOD code logic for licensee but allow to choose per instance to use the new system.

diff --git a/Engine/Source/Editor/UnrealEd/Private/EditorEngine.cpp b/Engine/Source/Editor/UnrealEd/Private/EditorEngine.cpp
--- a/Engine/Source/Editor/UnrealEd/Private/EditorEngine.cpp
+++ b/Engine/Source/Editor/UnrealEd/Private/EditorEngine.cpp
@@ -7450,6 +7450,8 @@ void UEditorEngine::SetPreviewPlatform(const FPreviewPlatformInfo& NewPreviewPla
 	{
 		UMaterialShaderQualitySettings* MaterialShaderQualitySettings = UMaterialShaderQualitySettings::Get();
 		MaterialShaderQualitySettings->SetPreviewPlatform(PreviewPlatform.PreviewShaderFormatName);
+
+		UStaticMesh::OnLodStrippingQualityLevelChanged(nullptr);
 	}
 
 	if (bChangedFeatureLevel)
@@ -7503,6 +7505,8 @@ void UEditorEngine::SetPreviewPlatform(const FPreviewPlatformInfo& NewPreviewPla
 
 	UDeviceProfileManager::Get().RestoreDefaultDeviceProfile();
 
+	UStaticMesh::OnLodStrippingQualityLevelChanged(nullptr);
+
 	//Override the current device profile.
 	if (PreviewPlatform.DeviceProfileName != NAME_None)
 	{
@@ -7552,6 +7556,8 @@ void UEditorEngine::ToggleFeatureLevelPreview()
 
 	PreviewPlatformChanged.Broadcast();
 
+	UStaticMesh::OnLodStrippingQualityLevelChanged(nullptr);
+
 	GEditor->RedrawAllViewports();
 	
 	SaveEditorFeatureLevel();
diff --git a/Engine/Source/Runtime/Engine/Private/UnrealEngine.cpp b/Engine/Source/Runtime/Engine/Private/UnrealEngine.cpp
--- a/Engine/Source/Runtime/Engine/Private/UnrealEngine.cpp
+++ b/Engine/Source/Runtime/Engine/Private/UnrealEngine.cpp
@@ -5625,6 +5625,10 @@ bool UEngine::HandleListStaticMeshesCommand(const TCHAR* Cmd, FOutputDevice& Ar)
 		int32		MobileMinLOD = -1;
 
 #if WITH_EDITORONLY_DATA 
+		if (Mesh->GetQualityLevelMinLOD().bIsEnabled && Mesh->GetQualityLevelMinLOD().PerQuality.Find(0/*Low*/) != nullptr)
+		{
+			MobileMinLOD = *Mesh->GetQualityLevelMinLOD().PerQuality.Find(0);
+		}
 		if (Mesh->GetMinLOD().PerPlatform.Find(("Mobile")) != nullptr)
 		{
 			MobileMinLOD = *Mesh->GetMinLOD().PerPlatform.Find(("Mobile"));
