Downgrade the warning during the PSO cache building to a log.

- Now it is possible to build PSO cache without the recorded files, since we add some entries during the cook time.

#rnx

diff --git a/Engine/Source/Editor/UnrealEd/Private/Commandlets/ShaderPipelineCacheToolsCommandlet.cpp b/Engine/Source/Editor/UnrealEd/Private/Commandlets/ShaderPipelineCacheToolsCommandlet.cpp
--- a/Engine/Source/Editor/UnrealEd/Private/Commandlets/ShaderPipelineCacheToolsCommandlet.cpp
+++ b/Engine/Source/Editor/UnrealEd/Private/Commandlets/ShaderPipelineCacheToolsCommandlet.cpp
@@ -57,7 +57,7 @@ void ExpandWildcards(TArray<FString>& Parts)
 				IFileManager::Get().FindFilesRecursive(ExpandedFiles, *CleanPath, *CleanFilename, true, false, false);
 			}
 			
-			UE_CLOG(!ExpandedFiles.Num(), LogShaderPipelineCacheTools, Warning, TEXT("Expanding %s....did not match anything."), *OldPart);
+			UE_CLOG(!ExpandedFiles.Num(), LogShaderPipelineCacheTools, Log, TEXT("Expanding %s....did not match anything."), *OldPart);
 			UE_CLOG(ExpandedFiles.Num(), LogShaderPipelineCacheTools, Log, TEXT("Expanding matched %4d files: %s"), ExpandedFiles.Num(), *OldPart);
 			for (const FString& Item : ExpandedFiles)
 			{
