Delete BP thumbnail scenes when unloading a BP otherwise it will prevent GC of its BPGC (and other BP sub-objects)

diff --git a/Engine/Source/Editor/UnrealEd/Classes/ThumbnailRendering/BlueprintThumbnailRenderer.h b/Engine/Source/Editor/UnrealEd/Classes/ThumbnailRendering/BlueprintThumbnailRenderer.h
--- a/Engine/Source/Editor/UnrealEd/Classes/ThumbnailRendering/BlueprintThumbnailRenderer.h
+++ b/Engine/Source/Editor/UnrealEd/Classes/ThumbnailRendering/BlueprintThumbnailRenderer.h
@@ -13,6 +13,7 @@
 #include "ThumbnailHelpers.h"
 #include "BlueprintThumbnailRenderer.generated.h"
 
+class UBlueprint;
 class FCanvas;
 class FRenderTarget;
 
@@ -21,6 +22,8 @@ class UBlueprintThumbnailRenderer : public UDefaultSizedThumbnailRenderer
 {
 	GENERATED_UCLASS_BODY()
 
+	UNREALED_API virtual ~UBlueprintThumbnailRenderer();
+
 	// Begin UThumbnailRenderer Object
 	UNREALED_API virtual bool CanVisualizeAsset(UObject* Object) override;
 	UNREALED_API virtual void Draw(UObject* Object, int32 X, int32 Y, uint32 Width, uint32 Height, FRenderTarget* RenderTarget, FCanvas* Canvas, bool bAdditionalViewFamily) override;
@@ -31,7 +34,10 @@ class UBlueprintThumbnailRenderer : public UDefaultSizedThumbnailRenderer
 	// End UObject implementation
 
 	/** Notifies the thumbnail scene to refresh components for the specified blueprint */
-	void BlueprintChanged(class UBlueprint* Blueprint);
+	void BlueprintChanged(UBlueprint* Blueprint);
+
+private:
+	void OnBlueprintUnloaded(UBlueprint* Blueprint);
 
 private:
 	TClassInstanceThumbnailScene<FBlueprintThumbnailScene, 100> ThumbnailScenes;
diff --git a/Engine/Source/Editor/UnrealEd/Private/ThumbnailRendering/BlueprintThumbnailRenderer.cpp b/Engine/Source/Editor/UnrealEd/Private/ThumbnailRendering/BlueprintThumbnailRenderer.cpp
--- a/Engine/Source/Editor/UnrealEd/Private/ThumbnailRendering/BlueprintThumbnailRenderer.cpp
+++ b/Engine/Source/Editor/UnrealEd/Private/ThumbnailRendering/BlueprintThumbnailRenderer.cpp
@@ -1,6 +1,7 @@
 // Copyright Epic Games, Inc. All Rights Reserved.
 
 #include "ThumbnailRendering/BlueprintThumbnailRenderer.h"
+#include "Kismet2/KismetEditorUtilities.h"
 #include "ShowFlags.h"
 #include "SceneView.h"
 #include "Misc/App.h"
@@ -11,6 +12,12 @@
 UBlueprintThumbnailRenderer::UBlueprintThumbnailRenderer(const FObjectInitializer& ObjectInitializer)
 	: Super(ObjectInitializer)
 {
+	FKismetEditorUtilities::OnBlueprintUnloaded.AddUObject(this, &UBlueprintThumbnailRenderer::OnBlueprintUnloaded);
+}
+
+UBlueprintThumbnailRenderer::~UBlueprintThumbnailRenderer()
+{
+	FKismetEditorUtilities::OnBlueprintUnloaded.RemoveAll(this);
 }
 
 bool UBlueprintThumbnailRenderer::CanVisualizeAsset(UObject* Object)
@@ -116,7 +123,7 @@ void UBlueprintThumbnailRenderer::BeginDestroy()
 	Super::BeginDestroy();
 }
 
-void UBlueprintThumbnailRenderer::BlueprintChanged(class UBlueprint* Blueprint)
+void UBlueprintThumbnailRenderer::BlueprintChanged(UBlueprint* Blueprint)
 {
 	if (Blueprint && Blueprint->GeneratedClass)
 	{
@@ -127,3 +134,11 @@ void UBlueprintThumbnailRenderer::BlueprintChanged(class UBlueprint* Blueprint)
 		}
 	}
 }
+
+void UBlueprintThumbnailRenderer::OnBlueprintUnloaded(UBlueprint* Blueprint)
+{
+	if (Blueprint && Blueprint->GeneratedClass)
+	{
+		ThumbnailScenes.RemoveThumbnailScene(Blueprint->GeneratedClass);
+	}
+}
diff --git a/Engine/Source/Editor/UnrealEd/Public/ThumbnailHelpers.h b/Engine/Source/Editor/UnrealEd/Public/ThumbnailHelpers.h
--- a/Engine/Source/Editor/UnrealEd/Public/ThumbnailHelpers.h
+++ b/Engine/Source/Editor/UnrealEd/Public/ThumbnailHelpers.h
@@ -342,6 +342,13 @@ public:
 		return ExistingThumbnailScene.ToSharedRef();
 	}
 
+	/** Removes the thumbnail scene instance for the specified class */
+	void RemoveThumbnailScene(const UClass* InClass)
+	{
+		check(InClass);
+		InstancedThumbnailScenes.Remove(InClass->GetFName());
+	}
+
 	/** Clears all thumbnail scenes */
 	void Clear()
 	{
