Advance copy: guard against BP with no GeneratedClass (which can happen if the source BP failed to compile for any reason)

diff --git a/Engine/Source/Editor/UnrealEd/Private/ObjectTools.cpp b/Engine/Source/Editor/UnrealEd/Private/ObjectTools.cpp
--- a/Engine/Source/Editor/UnrealEd/Private/ObjectTools.cpp
+++ b/Engine/Source/Editor/UnrealEd/Private/ObjectTools.cpp
@@ -322,13 +322,16 @@ namespace ObjectTools
 			{
 				if (UBlueprint* BlueprintObject = Cast<UBlueprint>(ObjectToExclude))
 				{
-					TArray<UObject*> ClassSubObjects;
-					GetObjectsWithOuter(BlueprintObject->GeneratedClass, ClassSubObjects, false);
-					for (UObject* ClassSubObject : ClassSubObjects)
+					if (BlueprintObject->GeneratedClass)
 					{
-						if (ClassSubObject->HasAnyFlags(RF_ArchetypeObject))
+						TArray<UObject*> ClassSubObjects;
+						GetObjectsWithOuter(BlueprintObject->GeneratedClass, ClassSubObjects, false);
+						for (UObject* ClassSubObject : ClassSubObjects)
 						{
-							AdditionalObjectsToExclude.Add(ClassSubObject);
+							if (ClassSubObject->HasAnyFlags(RF_ArchetypeObject))
+							{
+								AdditionalObjectsToExclude.Add(ClassSubObject);
+							}
 						}
 					}
 				}
@@ -357,20 +360,23 @@ namespace ObjectTools
 
 				if (UBlueprint* BlueprintObject = Cast<UBlueprint>(InObject))
 				{
-					if (AdditionalObjects.Contains(BlueprintObject->GeneratedClass))
-					{
-						// We don't want to replace within the generated class. 
-						AdditionalObjects.Remove(BlueprintObject->GeneratedClass);
-					}
-					TArray<UObject*> ClassSubObjects;
-					GetObjectsWithOuter(BlueprintObject->GeneratedClass, ClassSubObjects, false);
-					for (UObject* ClassSubObject : ClassSubObjects)
+					if (BlueprintObject->GeneratedClass)
 					{
-						if (ClassSubObject->HasAnyFlags(RF_ArchetypeObject)
-							&& !ObjectsToExclude.Contains(ClassSubObject)
-							&& !InObjects.Contains(ClassSubObject))
+						if (AdditionalObjects.Contains(BlueprintObject->GeneratedClass))
 						{
-							AdditionalObjects.Add(ClassSubObject);
+							// We don't want to replace within the generated class. 
+							AdditionalObjects.Remove(BlueprintObject->GeneratedClass);
+						}
+						TArray<UObject*> ClassSubObjects;
+						GetObjectsWithOuter(BlueprintObject->GeneratedClass, ClassSubObjects, false);
+						for (UObject* ClassSubObject : ClassSubObjects)
+						{
+							if (ClassSubObject->HasAnyFlags(RF_ArchetypeObject)
+								&& !ObjectsToExclude.Contains(ClassSubObject)
+								&& !InObjects.Contains(ClassSubObject))
+							{
+								AdditionalObjects.Add(ClassSubObject);
+							}
 						}
 					}
 				}
@@ -1412,12 +1418,12 @@ namespace ObjectTools
 
 			// See if this is a blueprint consolidate and replace instances of the generated class
 			UBlueprint* BlueprintToConsolidateTo = Cast<UBlueprint>(ObjectToConsolidateTo);
-			if ( BlueprintToConsolidateTo != NULL && ensure(BlueprintToConsolidateTo->GeneratedClass) )
+			if (BlueprintToConsolidateTo != NULL && BlueprintToConsolidateTo->GeneratedClass)
 			{
 				for ( TArray<UObject*>::TConstIterator ConsolIter( ReplaceInfo.ReplaceableObjects ); ConsolIter; ++ConsolIter )
 				{
 					UBlueprint* BlueprintToConsolidate = Cast<UBlueprint>(*ConsolIter);
-					if ( BlueprintToConsolidate != NULL && ensure(BlueprintToConsolidate->GeneratedClass) )
+					if (BlueprintToConsolidate != NULL && BlueprintToConsolidate->GeneratedClass)
 					{
 						// Replace all instances of objects based on the old blueprint's class with objects based on the new class,
 						// then repair the references on the object being consolidated so those objects can be properly disposed of upon deletion.
