CurveTable Factory Options

diff --git a/Engine/Source/Editor/UnrealEd/Classes/Factories/CurveTableFactory.h b/Engine/Source/Editor/UnrealEd/Classes/Factories/CurveTableFactory.h
--- a/Engine/Source/Editor/UnrealEd/Classes/Factories/CurveTableFactory.h
+++ b/Engine/Source/Editor/UnrealEd/Classes/Factories/CurveTableFactory.h
@@ -16,9 +16,14 @@ class UNREALED_API UCurveTableFactory : public UFactory
 {
 	GENERATED_UCLASS_BODY()
 
+  public:
 	// UFactory Interface
+	virtual bool ConfigureProperties() override;
 	virtual UObject* FactoryCreateNew(UClass* Class, UObject* InParent, FName Name, EObjectFlags Flags, UObject* Context, FFeedbackContext* Warn) override;
 
+
 protected:
 	virtual UCurveTable* MakeNewCurveTable(UObject* InParent, FName Name, EObjectFlags Flags);
+
+	ERichCurveInterpMode InterpMode;
 };
diff --git a/Engine/Source/Editor/UnrealEd/Private/Factories/CurveTableFactory.cpp b/Engine/Source/Editor/UnrealEd/Private/Factories/CurveTableFactory.cpp
--- a/Engine/Source/Editor/UnrealEd/Private/Factories/CurveTableFactory.cpp
+++ b/Engine/Source/Editor/UnrealEd/Private/Factories/CurveTableFactory.cpp
@@ -3,6 +3,7 @@
 #include "Factories/CurveTableFactory.h"
 #include "Engine/CurveTable.h"
 #include "Editor.h"
+#include "SCurveTableOptions.h"
 
 #define LOCTEXT_NAMESPACE "CurveTableFactory"
 
@@ -25,9 +26,54 @@ UObject* UCurveTableFactory::FactoryCreateNew(UClass* Class, UObject* InParent,
 	return CurveTable;
 }
 
+
+bool UCurveTableFactory::ConfigureProperties()
+{
+
+	ERichCurveInterpMode SelectedInterpMode = ERichCurveInterpMode::RCIM_Linear;
+	bool bDoCreation = false;
+
+	TSharedPtr<SWindow> Window = SNew(SWindow)
+		.Title( LOCTEXT("NewCurveTableWindowTitle", "Curve Table Options" ))
+		.SizingRule( ESizingRule::Autosized );
+
+	Window->SetContent( 
+		SNew(SCurveTableOptions)
+		.OnCancelClicked_Lambda([Window, &bDoCreation]() 
+		{
+			bDoCreation = false;
+			Window->RequestDestroyWindow();	
+		})
+		.OnCreateClicked_Lambda([Window, &bDoCreation, &SelectedInterpMode] (ERichCurveInterpMode InInterpMode) 
+		{
+			SelectedInterpMode = InInterpMode;
+			bDoCreation = true;
+			Window->RequestDestroyWindow();
+		})
+	);
+
+	GEditor->EditorAddModalWindow(Window.ToSharedRef());
+
+	// Store which interpolation mode was selected
+	InterpMode = SelectedInterpMode;
+
+	return bDoCreation;
+}
+
 UCurveTable* UCurveTableFactory::MakeNewCurveTable(UObject* InParent, FName Name, EObjectFlags Flags)
 {
-	return NewObject<UCurveTable>(InParent, Name, Flags);
+	UCurveTable* NewTable = NewObject<UCurveTable>(InParent, Name, Flags);
+	if (InterpMode != ERichCurveInterpMode::RCIM_Cubic)
+	{
+		FSimpleCurve& NewSimpleCurve = NewTable->AddSimpleCurve(FName("Curve"));
+		NewSimpleCurve.SetKeyInterpMode(InterpMode);
+	}
+	else
+	{
+		FRichCurve& NewRichCurve = NewTable->AddRichCurve(FName("Curve"));
+	}
+
+	return NewTable;
 }
 
 #undef LOCTEXT_NAMESPACE 
diff --git a/Engine/Source/Editor/UnrealEd/Private/SCurveTableOptions.cpp b/Engine/Source/Editor/UnrealEd/Private/SCurveTableOptions.cpp
new file mode 100644
--- /dev/null
+++ b/Engine/Source/Editor/UnrealEd/Private/SCurveTableOptions.cpp
@@ -0,0 +1,131 @@
+// Copyright Epic Games, Inc. All Rights Reserved.
+
+#include "SCurveTableOptions.h"
+#include "Widgets/Input/SButton.h"
+#include "SPrimaryButton.h"
+
+DEFINE_LOG_CATEGORY(LogCurveTableOptions);
+
+#define LOCTEXT_NAMESPACE "CurveTableOptions"
+
+SCurveTableOptions::SCurveTableOptions()
+{
+	// empty	
+}
+
+void SCurveTableOptions::Construct(const FArguments& InArgs)
+{
+	OnCancelClicked = InArgs._OnCancelClicked;
+	OnCreateClicked = InArgs._OnCreateClicked;
+
+	// Populate the valid interpolation modes
+	{
+		CurveInterpModes.Add( MakeShareable( new ERichCurveInterpMode(ERichCurveInterpMode::RCIM_Constant) ) );
+		CurveInterpModes.Add( MakeShareable( new ERichCurveInterpMode(ERichCurveInterpMode::RCIM_Linear) ) );
+		CurveInterpModes.Add( MakeShareable( new ERichCurveInterpMode(ERichCurveInterpMode::RCIM_Cubic) ) );
+	}
+
+	SelectedInterpMode = ERichCurveInterpMode::RCIM_Linear;
+
+
+	ChildSlot
+	[
+		SNew(SBorder)
+		.BorderImage(FAppStyle::Get().GetBrush("Brushes.Panel"))
+		.Padding(16.f)
+		[
+			
+			SNew(SVerticalBox)
+			+SVerticalBox::Slot()
+			.AutoHeight()
+			[
+				SNew(STextBlock)
+				.Text(LOCTEXT("ChooseCurveInterpolationType", "Choose Curve Interpolation Type:"))
+			]
+
+			+SVerticalBox::Slot()
+			.Padding(0.f, 8.f, 0.f, 0.f)
+			.AutoHeight()
+			[
+
+				SNew(SComboBox<CurveInterpModePtr>)
+				.OptionsSource( &CurveInterpModes )
+				.OnGenerateWidget_Lambda( [this] (CurveInterpModePtr Item) { 
+					return SNew(STextBlock).Text( this->GetCurveTypeText(*Item) );
+				})
+				.OnSelectionChanged_Lambda( [this] (CurveInterpModePtr Item, ESelectInfo::Type InSelectInfo)
+				{
+					if (Item.IsValid())
+					{
+						SelectedInterpMode = *Item;
+					}
+				})
+				[
+					SNew(STextBlock)
+					.Text_Lambda([this] () { return GetCurveTypeText(SelectedInterpMode);})
+				]
+			]
+
+			+SVerticalBox::Slot()
+			.HAlign(HAlign_Right)
+			.VAlign(VAlign_Bottom)
+			.Padding(0.f, 8.f, 0.f, 0.f)
+			[
+				SNew(SHorizontalBox)
+
+				+SHorizontalBox::Slot()
+				.AutoWidth()
+				[
+					SNew(SPrimaryButton)
+					.Text(LOCTEXT("CreateCurveTable", "Create Curve Table"))
+					.OnClicked_Lambda( [this] () {
+						if (OnCreateClicked.IsBound())
+						{
+							OnCreateClicked.Execute(SelectedInterpMode);
+						}
+						return FReply::Handled();		
+					})
+				]
+
+				+SHorizontalBox::Slot()
+				.AutoWidth()
+				.Padding(16.f, 0.f, 0.f, 0.f)
+				[
+					SNew(SButton)
+					.Text(LOCTEXT("Cancel", "Cancel"))
+					.HAlign(HAlign_Center)
+					.OnClicked_Lambda( [this] () {
+						if (OnCancelClicked.IsBound())
+						{
+							OnCancelClicked.Execute();
+						}
+						return FReply::Handled();		
+					})
+				]
+			]
+		]
+	];
+}
+
+FText SCurveTableOptions::GetCurveTypeText(const ERichCurveInterpMode& Item) const
+{
+	FString EnumString;
+	switch(Item)
+	{
+		case ERichCurveInterpMode::RCIM_Constant : 
+		EnumString = TEXT("Constant");
+		break;
+
+		case ERichCurveInterpMode::RCIM_Linear : 
+		EnumString = TEXT("Linear");
+		break;
+
+		case ERichCurveInterpMode::RCIM_Cubic : 
+		EnumString = TEXT("Cubic");
+		break;
+	}
+
+	return FText::FromString(EnumString);
+}
+
+#undef LOCTEXT_NAMESPACE 
diff --git a/Engine/Source/Editor/UnrealEd/Private/SCurveTableOptions.h b/Engine/Source/Editor/UnrealEd/Private/SCurveTableOptions.h
new file mode 100644
--- /dev/null
+++ b/Engine/Source/Editor/UnrealEd/Private/SCurveTableOptions.h
@@ -0,0 +1,57 @@
+// Copyright Epic Games, Inc. All Rights Reserved.
+
+#pragma once
+
+#include "CoreMinimal.h"
+#include "UObject/ObjectMacros.h"
+#include "Factories/CurveTableFactory.h"
+
+#include "Widgets/DeclarativeSyntaxSupport.h"
+#include "Delegates/Delegate.h"
+
+
+DECLARE_DELEGATE_OneParam(FOnCreateCurveTable, ERichCurveInterpMode);
+
+DECLARE_DELEGATE(FOnCancelCurveTable);
+
+DECLARE_LOG_CATEGORY_EXTERN(LogCurveTableOptions, Log, All);
+
+typedef TSharedPtr<ERichCurveInterpMode>		CurveInterpModePtr;
+ 
+/**
+ *    UI to allow the user to choose the Interpolation Type when creating a CurveTable.
+ */
+
+class SCurveTableOptions: public SCompoundWidget
+{
+	
+public:
+	SLATE_BEGIN_ARGS(SCurveTableOptions)
+	{}
+		SLATE_EVENT(FOnCreateCurveTable, OnCreateClicked)
+		SLATE_EVENT(FOnCancelCurveTable, OnCancelClicked)
+
+	SLATE_END_ARGS()
+
+	void Construct(const FArguments& InArgs);
+
+	SCurveTableOptions();
+
+  protected:
+
+	FText GetCurveTypeText(const ERichCurveInterpMode& InterMode) const;
+
+
+  private:
+
+	/** All available curve interpolation modes */
+	TArray< CurveInterpModePtr >					CurveInterpModes;
+
+	/** The selected curve interpolation type */
+	ERichCurveInterpMode							SelectedInterpMode;
+
+	FOnCreateCurveTable								OnCreateClicked;
+
+	FOnCancelCurveTable                       		OnCancelClicked;
+
+};
