BugFix: Simulation in physics editor is incorrect with "Simulation updates Component Transform"

diff --git a/Engine/Source/Editor/UnrealEd/Private/Animation/DebugSkelMeshComponent.cpp b/Engine/Source/Editor/UnrealEd/Private/Animation/DebugSkelMeshComponent.cpp
--- a/Engine/Source/Editor/UnrealEd/Private/Animation/DebugSkelMeshComponent.cpp
+++ b/Engine/Source/Editor/UnrealEd/Private/Animation/DebugSkelMeshComponent.cpp
@@ -214,15 +214,7 @@ void UDebugSkelMeshComponent::ConsumeRootMotion(const FVector& FloorMin, const F
 	const float PreviousTime = ConsumeRootMotionPreviousPlaybackTime;
 	ConsumeRootMotionPreviousPlaybackTime = CurrentTime;
 
-	// Ignore Mode: Preview mesh remains in place at the origin regardless of the root motion
-	if (ProcessRootMotionMode == EProcessRootMotionMode::Ignore)
-	{
-		if(!GetRelativeTransform().Equals(FTransform::Identity))
-		{
-			SetRelativeTransform(FTransform::Identity);
-		}
-	}
-	else 
+	if (ProcessRootMotionMode != EProcessRootMotionMode::Ignore && !ShouldBlendPhysicsBones())
 	{
 		if(PreviewInstance->IsPlaying())
 		{
