CurveTable Editor - fixes non-unity build

diff --git a/Engine/Source/Editor/UnrealEd/Classes/Factories/CurveTableFactory.h b/Engine/Source/Editor/UnrealEd/Classes/Factories/CurveTableFactory.h
--- a/Engine/Source/Editor/UnrealEd/Classes/Factories/CurveTableFactory.h
+++ b/Engine/Source/Editor/UnrealEd/Classes/Factories/CurveTableFactory.h
@@ -5,6 +5,7 @@
 #include "CoreMinimal.h"
 #include "UObject/ObjectMacros.h"
 #include "Factories/Factory.h"
+#include "Curves/RealCurve.h"
 #include "CurveTableFactory.generated.h"
 
 class UCurveTable;
diff --git a/Engine/Source/Editor/UnrealEd/Private/SCurveTableOptions.cpp b/Engine/Source/Editor/UnrealEd/Private/SCurveTableOptions.cpp
--- a/Engine/Source/Editor/UnrealEd/Private/SCurveTableOptions.cpp
+++ b/Engine/Source/Editor/UnrealEd/Private/SCurveTableOptions.cpp
@@ -3,6 +3,9 @@
 #include "SCurveTableOptions.h"
 #include "Widgets/Input/SButton.h"
 #include "SPrimaryButton.h"
+#include "Widgets/SBoxPanel.h"
+#include "Widgets/Text/STextBlock.h"
+#include "Widgets/Input/SComboBox.h"
 
 DEFINE_LOG_CATEGORY(LogCurveTableOptions);
 
diff --git a/Engine/Source/Editor/UnrealEd/Private/SCurveTableOptions.h b/Engine/Source/Editor/UnrealEd/Private/SCurveTableOptions.h
--- a/Engine/Source/Editor/UnrealEd/Private/SCurveTableOptions.h
+++ b/Engine/Source/Editor/UnrealEd/Private/SCurveTableOptions.h
@@ -6,6 +6,7 @@
 #include "UObject/ObjectMacros.h"
 #include "Factories/CurveTableFactory.h"
 
+#include "Widgets/SCompoundWidget.h"
 #include "Widgets/DeclarativeSyntaxSupport.h"
 #include "Delegates/Delegate.h"
 
