DDC: Updated IBuildWorkerExecutor to allow filtering of payloads

#rnx

diff --git a/Engine/Source/Developer/DerivedDataCache/Private/DerivedDataBuildJob.cpp b/Engine/Source/Developer/DerivedDataCache/Private/DerivedDataBuildJob.cpp
--- a/Engine/Source/Developer/DerivedDataCache/Private/DerivedDataBuildJob.cpp
+++ b/Engine/Source/Developer/DerivedDataCache/Private/DerivedDataBuildJob.cpp
@@ -791,7 +791,23 @@ void FBuildJob::BeginExecuteRemote()
 	checkf(Worker && WorkerExecutor, TEXT("Job requires a worker in state %s for build of '%s' by %s."),
 		LexToString(State), *Name, *FunctionName);
 	EnumAddFlags(BuildStatus, EBuildStatus::BuildTryRemote);
-	WorkerExecutor->BuildAction(Action.Get(), Inputs, *Worker, BuildSystem, BuildPolicy.GetCombinedPolicy(), Owner,
+
+	FBuildPolicy RemoteBuildPolicy;
+	if (EBuildPolicy Mask = Context->GetBuildPolicyMask(); Mask == ~EBuildPolicy::None)
+	{
+		RemoteBuildPolicy = BuildPolicy;
+	}
+	else
+	{
+		FBuildPolicyBuilder PolicyBuilder(BuildPolicy.GetDefaultPayloadPolicy() & Mask);
+		for (const FBuildPayloadPolicy& Payload : BuildPolicy.GetPayloadPolicies())
+		{
+			PolicyBuilder.AddPayloadPolicy({Payload.Id, Payload.Policy & Mask});
+		}
+		RemoteBuildPolicy = PolicyBuilder.Build();
+	}
+
+	WorkerExecutor->Build(Action.Get(), Inputs, RemoteBuildPolicy, *Worker, BuildSystem, Owner,
 		[this](FBuildWorkerActionCompleteParams&& Params) { EndExecuteRemote(MoveTemp(Params)); });
 }
 
diff --git a/Engine/Source/Developer/DerivedDataCache/Public/DerivedDataBuildWorker.h b/Engine/Source/Developer/DerivedDataCache/Public/DerivedDataBuildWorker.h
--- a/Engine/Source/Developer/DerivedDataCache/Public/DerivedDataBuildWorker.h
+++ b/Engine/Source/Developer/DerivedDataCache/Public/DerivedDataBuildWorker.h
@@ -18,13 +18,13 @@ template <typename FuncType> class TFunctionRef;
 template <typename FuncType> class TUniqueFunction;
 
 namespace UE::DerivedData { class FBuildAction; }
+namespace UE::DerivedData { class FBuildPolicy; }
 namespace UE::DerivedData { class FOptionalBuildInputs; }
 namespace UE::DerivedData { class FOptionalBuildOutput; }
 namespace UE::DerivedData { class IBuild; }
 namespace UE::DerivedData { class IRequestOwner; }
 namespace UE::DerivedData { struct FBuildWorkerActionCompleteParams; }
 namespace UE::DerivedData { struct FBuildWorkerFileDataCompleteParams; }
-namespace UE::DerivedData { enum class EBuildPolicy : uint32; }
 
 namespace UE::DerivedData
 {
@@ -81,12 +81,12 @@ class IBuildWorkerExecutor : public IModularFeature
 public:
 	virtual ~IBuildWorkerExecutor() = default;
 
-	virtual void BuildAction(
+	virtual void Build(
 		const FBuildAction& Action,
 		const FOptionalBuildInputs& Inputs,
+		const FBuildPolicy& Policy,
 		const FBuildWorker& Worker,
 		IBuild& BuildSystem,
-		EBuildPolicy Policy,
 		IRequestOwner& Owner,
 		FOnBuildWorkerActionComplete&& OnComplete) = 0;
 
diff --git a/Engine/Source/Editor/UnrealEd/Private/DerivedDataBuildLocalExecutor.cpp b/Engine/Source/Editor/UnrealEd/Private/DerivedDataBuildLocalExecutor.cpp
--- a/Engine/Source/Editor/UnrealEd/Private/DerivedDataBuildLocalExecutor.cpp
+++ b/Engine/Source/Editor/UnrealEd/Private/DerivedDataBuildLocalExecutor.cpp
@@ -6,6 +6,7 @@
 #include "DerivedDataBuildAction.h"
 #include "DerivedDataBuildInputs.h"
 #include "DerivedDataBuildOutput.h"
+#include "DerivedDataBuildTypes.h"
 #include "DerivedDataBuildWorker.h"
 #include "DerivedDataPayload.h"
 #include "DerivedDataRequest.h"
@@ -32,11 +33,9 @@ DEFINE_LOG_CATEGORY_STATIC(LogDerivedDataBuildLocalExecutor, Log, All);
 /**
  * This implements a simple local DDC2 executor which spawns all jobs
  * as local processes. This is intentionally as simple as possible and
- * everything is synchronous. This is not meant to be used in production, 
- * and is perhaps primarily useful as a debugging aid
- * 
+ * everything is synchronous. This is not meant to be used in production,
+ * and is perhaps primarily useful as a debugging aid.
  */
-
 class FLocalBuildWorkerExecutor final : public IBuildWorkerExecutor
 {
 	FString SandboxRootDir = FPaths::EngineSavedDir() / TEXT("LocalExec");
@@ -48,24 +47,23 @@ public:
 
 		UE_LOG(LogDerivedDataBuildLocalExecutor, Warning, TEXT("Deleting existing local execution state from '%s'"), *SandboxRootDir);
 
-		const bool RequireExists = false;
-		const bool Tree = true;
+		constexpr bool RequireExists = false;
+		constexpr bool Tree = true;
 		IFileManager::Get().DeleteDirectory(*SandboxRootDir, RequireExists, Tree);
 
 		IModularFeatures::Get().RegisterModularFeature(IBuildWorkerExecutor::GetFeatureName(), this);
 	}
 
-	virtual ~FLocalBuildWorkerExecutor()
-	{
-	}
+	~FLocalBuildWorkerExecutor() final = default;
 
-	void BuildAction(const FBuildAction&			Action,
-					 const FOptionalBuildInputs&	Inputs,
-					 const FBuildWorker&			Worker,
-					 IBuild&						BuildSystem,
-					 EBuildPolicy					Policy,
-					 IRequestOwner&					Owner,
-					 FOnBuildWorkerActionComplete&& OnComplete) final
+	void Build(
+		const FBuildAction& Action,
+		const FOptionalBuildInputs& Inputs,
+		const FBuildPolicy& Policy,
+		const FBuildWorker& Worker,
+		IBuild& BuildSystem,
+		IRequestOwner& Owner,
+		FOnBuildWorkerActionComplete&& OnComplete) final
 	{
 		// Review build action inputs to determine if they need to be materialized/propagated 
 		// (right now, they always will be)
@@ -74,22 +72,22 @@ public:
 		TArray<FStringView> MissingInputViews;
 
 		Action.IterateInputs([&](FStringView Key, const FIoHash& RawHash, uint64 RawSize)
+		{
+			if (Inputs.IsNull() || Inputs.Get().FindInput(Key).IsNull())
 			{
-				if (Inputs.IsNull() || Inputs.Get().FindInput(Key).IsNull())
-				{
-					MissingInputs.Emplace(Key);
-					MissingInputViews.Add(MissingInputs.Last());
-				}
-			});
+				MissingInputs.Emplace(Key);
+				MissingInputViews.Add(MissingInputs.Last());
+			}
+		});
 
 		if (!MissingInputViews.IsEmpty())
 		{
 			// Report missing inputs
-			return OnComplete({ Action.GetKey(), {}, MissingInputViews, EStatus::Ok });
+			return OnComplete({Action.GetKey(), {}, MissingInputViews, EStatus::Ok});
 		}
 
 		// This path will execute the build action synchronously in a scratch directory
-		// 
+		//
 		// At this stage, all inputs are available in process
 		//
 		// Currently no cleanup whatsoever is performed, so inputs/outputs can be inspected
@@ -98,8 +96,9 @@ public:
 
 		static std::atomic<int32> SerialNo = 0;
 
-		FString SandboxRoot = SandboxRootDir / TEXT("Scratch");
-		SandboxRoot.AppendInt(++SerialNo);
+		TStringBuilder<256> SandboxRoot;
+		FPathViews::Append(SandboxRoot, SandboxRootDir, TEXT("Scratch"));
+		SandboxRoot.Appendf(TEXT("%06d"), ++SerialNo);
 
 		// Manifest worker in scratch area
 
@@ -108,26 +107,29 @@ public:
 
 			TArray<FIoHash> WorkerFileHashes;
 			TArray<TTuple<FStringView, bool>> WorkerFileMeta;
-			Worker.IterateExecutables([&WorkerFileHashes, &WorkerFileMeta](FStringView Path, const FIoHash& RawHash, uint64 RawSize) {
+			Worker.IterateExecutables([&WorkerFileHashes, &WorkerFileMeta](FStringView Path, const FIoHash& RawHash, uint64 RawSize)
+			{
 				WorkerFileHashes.Emplace(RawHash);
 				WorkerFileMeta.Emplace(Path, true);
 			});
 
-			Worker.IterateFiles([&WorkerFileHashes, &WorkerFileMeta](FStringView Path, const FIoHash& RawHash, uint64 RawSize) {
+			Worker.IterateFiles([&WorkerFileHashes, &WorkerFileMeta](FStringView Path, const FIoHash& RawHash, uint64 RawSize)
+			{
 				WorkerFileHashes.Emplace(RawHash);
 				WorkerFileMeta.Emplace(Path, false);
 			});
 
 			FRequestOwner BlockingOwner(EPriority::Blocking);
-			Worker.FindFileData(WorkerFileHashes, BlockingOwner, [&](FBuildWorkerFileDataCompleteParams&& Params) {
+			Worker.FindFileData(WorkerFileHashes, BlockingOwner, [&SandboxRoot, &WorkerFileMeta](FBuildWorkerFileDataCompleteParams&& Params)
+			{
 				uint32 MetaIndex = 0;
 				for (const FCompressedBuffer& Buffer : Params.Files)
 				{
 					const TTuple<FStringView, bool>& Meta = WorkerFileMeta[MetaIndex];
 
-					FString Path { SandboxRoot / FString(Meta.Key) };
-
-					if (TUniquePtr<FArchive> Ar{ IFileManager::Get().CreateFileWriter(*Path, FILEWRITE_Silent) })
+					TStringBuilder<256> Path;
+					FPathViews::Append(Path, SandboxRoot, Meta.Key);
+					if (TUniquePtr<FArchive> Ar{IFileManager::Get().CreateFileWriter(*Path, FILEWRITE_Silent)})
 					{
 						FCompositeBuffer DecompressedComposite = Buffer.DecompressToComposite();
 
@@ -143,7 +145,8 @@ public:
 			BlockingOwner.Wait();
 
 			// This directory must exist in order for the builder to run correctly
-			FString BinPath{ SandboxRoot / TEXT("Engine/Binaries/Win64") };
+			TStringBuilder<256> BinPath;
+			FPathViews::Append(BinPath, SandboxRoot, TEXT("Engine/Binaries/Win64"));
 			IFileManager::Get().MakeDirectory(*BinPath);
 		}
 
@@ -151,17 +154,15 @@ public:
 
 		if (!Inputs.IsNull())
 		{
-			Inputs.Get().IterateInputs([&](FStringView Key, const FCompressedBuffer& Buffer) {
-				TStringBuilder<128> InputPath;
-				InputPath << TEXT("Inputs/") << FIoHash(Buffer.GetRawHash());
-
-				FString Path{ SandboxRoot / *InputPath };
-
-				if (TUniquePtr<FArchive> Ar{ IFileManager::Get().CreateFileWriter(*Path, FILEWRITE_Silent) })
+			Inputs.Get().IterateInputs([&SandboxRoot](FStringView Key, const FCompressedBuffer& Buffer)
+			{
+				TStringBuilder<256> Path;
+				FPathViews::Append(Path, SandboxRoot, TEXT("Inputs"), FIoHash(Buffer.GetRawHash()));
+				if (TUniquePtr<FArchive> Ar{IFileManager::Get().CreateFileWriter(*Path, FILEWRITE_Silent)})
 				{
-					// Silly workaround for FArchive::operator<< not accepting const objects
-					FCompressedBuffer Comp = FCompressedBuffer::FromCompressed(Buffer.GetCompressed());
-					*Ar << Comp;
+					// Workaround for FArchive::operator<< not accepting const objects
+					FCompressedBuffer Copy = Buffer;
+					*Ar << Copy;
 				}
 			});
 		}
@@ -169,9 +170,9 @@ public:
 		// Serialize action specification
 
 		{
-			FString Path{ SandboxRoot / TEXT("Build.action") };
-
-			if (TUniquePtr<FArchive> Ar{ IFileManager::Get().CreateFileWriter(*Path, FILEWRITE_Silent) })
+			TStringBuilder<256> Path;
+			FPathViews::Append(Path, SandboxRoot, TEXT("Build.action"));
+			if (TUniquePtr<FArchive> Ar{IFileManager::Get().CreateFileWriter(*Path, FILEWRITE_Silent)})
 			{
 				FCbWriter BuildActionWriter;
 				Action.Save(BuildActionWriter);
@@ -185,27 +186,28 @@ public:
 		uint32 ProcessID = 0;
 		const int PriorityModifier = 0;
 		const TCHAR* WorkingDirectory = *SandboxRoot;
-		const FString WorkerPath { SandboxRoot / FString{Worker.GetPath()} };
+		TStringBuilder<256> WorkerPath;
+		FPathViews::Append(WorkerPath, SandboxRoot, Worker.GetPath());
 
 		FProcHandle ProcHandle = FPlatformProcess::CreateProc(
-			*WorkerPath, 
-			TEXT("-Build=build.action"), 
+			*WorkerPath,
+			TEXT("-Build=Build.action"),
 			bLaunchDetached,
-			bLaunchHidden, 
-			bLaunchReallyHidden, 
-			&ProcessID, 
-			PriorityModifier, 
-			WorkingDirectory, 
-			nullptr, 
+			bLaunchHidden,
+			bLaunchReallyHidden,
+			&ProcessID,
+			PriorityModifier,
+			WorkingDirectory,
+			nullptr,
 			nullptr);
 
 		FPlatformProcess::WaitForProc(ProcHandle);
-		
+
 		int32 ExitCode = -1;
 		if (!FPlatformProcess::GetProcReturnCode(ProcHandle, &ExitCode) || ExitCode != 0)
 		{
 			UE_LOG(LogDerivedDataBuildLocalExecutor, Warning, TEXT("Worker process exit code = %d!"), ExitCode);
-			return OnComplete({ Action.GetKey(), {}, {}, EStatus::Error });
+			return OnComplete({Action.GetKey(), {}, {}, EStatus::Error});
 		}
 
 		// Gather results
@@ -213,61 +215,58 @@ public:
 		FOptionalBuildOutput RemoteBuildOutput;
 
 		{
-			FString BuildOutputPath{ SandboxRoot / TEXT("Build.output") };
-
-			if (TUniquePtr<FArchive> Ar{ IFileManager::Get().CreateFileReader(*BuildOutputPath, FILEREAD_Silent) })
+			TStringBuilder<256> OutputPath;
+			FPathViews::Append(OutputPath, SandboxRoot, TEXT("Build.output"));
+			if (TUniquePtr<FArchive> Ar{IFileManager::Get().CreateFileReader(*OutputPath, FILEREAD_Silent)})
 			{
 				FCbObject BuildOutput = LoadCompactBinary(*Ar).AsObject();
 
 				if (Ar->IsError())
 				{
 					UE_LOG(LogDerivedDataBuildLocalExecutor, Warning, TEXT("Worker error: build output structure not valid!"));
-
-					return OnComplete({ Action.GetKey(), {}, {}, EStatus::Error });
+					return OnComplete({Action.GetKey(), {}, {}, EStatus::Error});
 				}
 
-				RemoteBuildOutput = FBuildOutput::Load(	Action.GetName(), 
-														Action.GetFunction(), 
-														BuildOutput);
+				RemoteBuildOutput = FBuildOutput::Load(Action.GetName(), Action.GetFunction(), BuildOutput);
 			}
 		}
 
 		if (RemoteBuildOutput.IsNull())
 		{
 			UE_LOG(LogDerivedDataBuildLocalExecutor, Warning, TEXT("Remote execution system error: build output blob missing!"));
-
-			return OnComplete({ Action.GetKey(), {}, {}, EStatus::Error });
+			return OnComplete({Action.GetKey(), {}, {}, EStatus::Error});
 		}
 
 		FBuildOutputBuilder OutputBuilder = BuildSystem.CreateOutput(Action.GetName(), Action.GetFunction());
-
 		for (const FPayload& Payload : RemoteBuildOutput.Get().GetPayloads())
 		{
-			TStringBuilder<128> OutputPath;
-			OutputPath << TEXT("Outputs/") << Payload.GetRawHash();
-
-			FString Path{ SandboxRoot / *OutputPath };
-
-			FCompressedBuffer BufferForPayload;
-
-			if (TUniquePtr<FArchive> Ar{ IFileManager::Get().CreateFileReader(*Path, FILEREAD_Silent) })
+			if (EnumHasAnyFlags(Policy.GetPayloadPolicy(Payload.GetId()), EBuildPolicy::SkipData))
 			{
-				BufferForPayload = FCompressedBuffer::FromCompressed(*Ar);
+				OutputBuilder.AddPayload(Payload);
 			}
-
-			if (BufferForPayload.IsNull())
+			else
 			{
-				UE_LOG(LogDerivedDataBuildLocalExecutor, Warning, TEXT("Remote execution system error: payload blob missing!"));
+				FCompressedBuffer BufferForPayload;
 
-				return OnComplete({ Action.GetKey(), {}, {}, EStatus::Error });
-			}
+				TStringBuilder<128> Path;
+				FPathViews::Append(Path, SandboxRoot, TEXT("Outputs"), FIoHash(Payload.GetRawHash()));
+				if (TUniquePtr<FArchive> Ar{IFileManager::Get().CreateFileReader(*Path, FILEREAD_Silent)})
+				{
+					BufferForPayload = FCompressedBuffer::FromCompressed(*Ar);
+				}
+
+				if (BufferForPayload.IsNull())
+				{
+					UE_LOG(LogDerivedDataBuildLocalExecutor, Warning, TEXT("Remote execution system error: payload blob missing!"));
+					return OnComplete({Action.GetKey(), {}, {}, EStatus::Error});
+				}
 
-			OutputBuilder.AddPayload(FPayload(Payload.GetId(), BufferForPayload));
+				OutputBuilder.AddPayload(FPayload(Payload.GetId(), BufferForPayload));
+			}
 		}
 
 		FBuildOutput BuildOutput = OutputBuilder.Build();
-
-		return OnComplete({ Action.GetKey(), BuildOutput, {}, EStatus::Ok });
+		return OnComplete({Action.GetKey(), BuildOutput, {}, EStatus::Ok});
 	}
 
 	TConstArrayView<FStringView> GetHostPlatforms() const final
@@ -276,7 +275,9 @@ public:
 		return HostPlatforms;
 	}
 
-	void DumpStats() {}
+	void DumpStats()
+	{
+	}
 };
 
 }  // namespace UE::DerivedData
diff --git a/Engine/Source/Editor/UnrealEd/Private/DerivedDataBuildRemoteExecutor.cpp b/Engine/Source/Editor/UnrealEd/Private/DerivedDataBuildRemoteExecutor.cpp
--- a/Engine/Source/Editor/UnrealEd/Private/DerivedDataBuildRemoteExecutor.cpp
+++ b/Engine/Source/Editor/UnrealEd/Private/DerivedDataBuildRemoteExecutor.cpp
@@ -7,6 +7,7 @@
 #include "DerivedDataBuildAction.h"
 #include "DerivedDataBuildInputs.h"
 #include "DerivedDataBuildOutput.h"
+#include "DerivedDataBuildTypes.h"
 #include "DerivedDataBuildWorker.h"
 #include "DerivedDataPayload.h"
 #include "DerivedDataRequest.h"
@@ -30,8 +31,6 @@
 #include "HttpModule.h"
 #include "RemoteMessages.h"
 
-
-
 namespace UE::DerivedData
 {
 
@@ -46,9 +45,9 @@ public:
 		FRemoteBuildWorkerExecutor& InExecutor,
 		const FBuildAction& Action,
 		const FOptionalBuildInputs& Inputs,
+		const FBuildPolicy& Policy,
 		const FBuildWorker& Worker,
 		IBuild& BuildSystem,
-		EBuildPolicy Policy,
 		IRequestOwner& Owner,
 		FOnBuildWorkerActionComplete&& OnComplete);
 
@@ -137,7 +136,7 @@ private:
 		const FBuildWorker& BuildWorker;
 		IBuild& BuildSystem;
 		IRequestOwner& Owner;
-		EBuildPolicy BuildPolicy;
+		FBuildPolicy BuildPolicy;
 
 		// Unordered arrays that are indexed into
 		TArray<FMerkleTreeDirectoryBuilder> Directories;
@@ -264,12 +263,12 @@ public:
 		}
 	}
 
-	void BuildAction(
+	void Build(
 		const FBuildAction& Action,
 		const FOptionalBuildInputs& Inputs,
+		const FBuildPolicy& Policy,
 		const FBuildWorker& Worker,
 		IBuild& BuildSystem,
-		EBuildPolicy Policy,
 		IRequestOwner& Owner,
 		FOnBuildWorkerActionComplete&& OnComplete) final
 	{
@@ -309,7 +308,7 @@ public:
 			}
 		}
 
-		new FRemoteBuildExecutionRequest(*this, Action, Inputs, Worker, BuildSystem, Policy, Owner, MoveTemp(OnComplete));
+		new FRemoteBuildExecutionRequest(*this, Action, Inputs, Policy, Worker, BuildSystem, Owner, MoveTemp(OnComplete));
 	}
 
 	TConstArrayView<FStringView> GetHostPlatforms() const final
@@ -498,14 +497,14 @@ private:
 };
 
 FRemoteBuildExecutionRequest::FRemoteBuildExecutionRequest(
-		FRemoteBuildWorkerExecutor& InExecutor,
-		const FBuildAction& Action,
-		const FOptionalBuildInputs& Inputs,
-		const FBuildWorker& Worker,
-		IBuild& BuildSystem,
-		EBuildPolicy Policy,
-		IRequestOwner& Owner,
-		FOnBuildWorkerActionComplete&& OnComplete)
+	FRemoteBuildWorkerExecutor& InExecutor,
+	const FBuildAction& Action,
+	const FOptionalBuildInputs& Inputs,
+	const FBuildPolicy& Policy,
+	const FBuildWorker& Worker,
+	IBuild& BuildSystem,
+	IRequestOwner& Owner,
+	FOnBuildWorkerActionComplete&& OnComplete)
 : State{Action, Inputs, Worker, BuildSystem, Owner, Policy}
 , CompletionCallback(MoveTemp(OnComplete))
 , Executor(InExecutor)
@@ -522,7 +521,7 @@ FRemoteBuildExecutionRequest::~FRemoteBuildExecutionRequest()
 	if (bHeuristicBuildStarted)
 	{
 		Executor.LimitingHeuristics.FinishBuild(Executor.Stats);
-	};
+	}
 }
 
 FRemoteBuildExecutionRequest::FMerkleTreeDirectoryBuilder& FRemoteBuildExecutionRequest::GetOrAddMerkleTreeDirectory(FStringView Path, int32& OutDirectoryBuilderIndex)
