DDC: Updated IBuildWorkerExecutor to allow filtering of payloads

#rnx

diff --git a/Engine/Source/Editor/UnrealEd/Private/DerivedDataBuildZenExecutor.cpp b/Engine/Source/Editor/UnrealEd/Private/DerivedDataBuildZenExecutor.cpp
--- a/Engine/Source/Editor/UnrealEd/Private/DerivedDataBuildZenExecutor.cpp
+++ b/Engine/Source/Editor/UnrealEd/Private/DerivedDataBuildZenExecutor.cpp
@@ -6,6 +6,7 @@
 #include "DerivedDataBuildAction.h"
 #include "DerivedDataBuildInputs.h"
 #include "DerivedDataBuildOutput.h"
+#include "DerivedDataBuildTypes.h"
 #include "DerivedDataBuildWorker.h"
 #include "DerivedDataPayload.h"
 #include "DerivedDataRequest.h"
@@ -37,11 +38,10 @@ DEFINE_LOG_CATEGORY_STATIC(LogDerivedDataBuildZenExecutor, Log, All);
 
 /**
  * This implements a simple Zen executor which passes build requests to
- * a local Zen instance for execution. This is intentionally as simple as 
- * possible and everything is synchronous. This is not meant to be used 
+ * a local Zen instance for execution. This is intentionally as simple as
+ * possible and everything is synchronous. This is not meant to be used
  * in production at this point.
  */
-
 class FZenBuildWorkerExecutor final : public IBuildWorkerExecutor
 {
 public:
@@ -52,28 +52,19 @@ public:
 
 		RequestPool = MakeUnique<UE::Zen::FZenHttpRequestPool>(TEXT("http://localhost:1337/"));
 
-		// Clean out any leftovers from a previous run
-
-		UE_LOG(LogDerivedDataBuildZenExecutor, Display, TEXT("Deleting existing local execution state from '%s'"), *SandboxRootDir);
-
-		const bool RequireExists = false;
-		const bool Tree = true;
-		IFileManager::Get().DeleteDirectory(*SandboxRootDir, RequireExists, Tree);
-
 		IModularFeatures::Get().RegisterModularFeature(IBuildWorkerExecutor::GetFeatureName(), this);
 	}
 
-	virtual ~FZenBuildWorkerExecutor()
-	{
-	}
+	~FZenBuildWorkerExecutor() final = default;
 
-	void BuildAction(const FBuildAction&			Action,
-					 const FOptionalBuildInputs&	Inputs,
-					 const FBuildWorker&			Worker,
-					 IBuild&						BuildSystem,
-					 EBuildPolicy					Policy,
-					 IRequestOwner&					Owner,
-					 FOnBuildWorkerActionComplete&& OnComplete) final
+	void Build(
+		const FBuildAction& Action,
+		const FOptionalBuildInputs& Inputs,
+		const FBuildPolicy& Policy,
+		const FBuildWorker& Worker,
+		IBuild& BuildSystem,
+		IRequestOwner& Owner,
+		FOnBuildWorkerActionComplete&& OnComplete) final
 	{
 		// Review build action inputs to determine if they need to be materialized/propagated 
 		// (right now, they always will be)
@@ -81,19 +72,19 @@ public:
 		TArray<FString> MissingInputs;
 		TArray<FStringView> MissingInputViews;
 
-		Action.IterateInputs([&](FStringView Key, const FIoHash& RawHash, uint64 RawSize)
+		Action.IterateInputs([&Inputs, &MissingInputs, &MissingInputViews](FStringView Key, const FIoHash& RawHash, uint64 RawSize)
+		{
+			if (Inputs.IsNull() || Inputs.Get().FindInput(Key).IsNull())
 			{
-				if (Inputs.IsNull() || Inputs.Get().FindInput(Key).IsNull())
-				{
-					MissingInputs.Emplace(Key);
-					MissingInputViews.Add(MissingInputs.Last());
-				}
-			});
+				MissingInputs.Emplace(Key);
+				MissingInputViews.Add(MissingInputs.Last());
+			}
+		});
 
 		if (!MissingInputViews.IsEmpty())
 		{
 			// Report missing inputs
-			return OnComplete({ Action.GetKey(), {}, MissingInputViews, EStatus::Ok });
+			return OnComplete({Action.GetKey(), {}, MissingInputViews, EStatus::Ok});
 		}
 
 		// Describe worker
@@ -107,13 +98,14 @@ public:
 		WorkerDescriptor.AddUuid("buildsystem_version"_ASV, Worker.GetBuildSystemVersion());
 
 		WorkerDescriptor.BeginArray("environment"_ASV);
-		Worker.IterateEnvironment([&](FStringView Name, FStringView Value) {
+		Worker.IterateEnvironment([&WorkerDescriptor](FStringView Name, FStringView Value)
+		{
 			WorkerDescriptor.AddString(WriteToString<256>(Name, "=", Value));
 		});
 		WorkerDescriptor.EndArray();
 
 		WorkerDescriptor.BeginArray("executables"_ASV);
-		Worker.IterateExecutables([&](FStringView Key, const FIoHash& RawHash, uint64 RawSize) 
+		Worker.IterateExecutables([&WorkerDescriptor](FStringView Key, const FIoHash& RawHash, uint64 RawSize) 
 		{
 			WorkerDescriptor.BeginObject();
 			WorkerDescriptor.AddString("name"_ASV, Key);
@@ -124,7 +116,7 @@ public:
 		WorkerDescriptor.EndArray();
 
 		WorkerDescriptor.BeginArray("files");
-		Worker.IterateFiles([&](FStringView Key, const FIoHash& RawHash, uint64 RawSize) 
+		Worker.IterateFiles([&WorkerDescriptor](FStringView Key, const FIoHash& RawHash, uint64 RawSize) 
 		{
 			WorkerDescriptor.BeginObject();
 			WorkerDescriptor.AddString("name"_ASV, Key);
@@ -139,7 +131,7 @@ public:
 		WorkerDescriptor.EndArray();
 
 		WorkerDescriptor.BeginArray("functions"_ASV);
-		Worker.IterateFunctions([&](FStringView Name, const FGuid& Version) 
+		Worker.IterateFunctions([&WorkerDescriptor](FStringView Name, const FGuid& Version) 
 		{
 			WorkerDescriptor.BeginObject();
 			WorkerDescriptor.AddString("name"_ASV, Name);
@@ -178,7 +170,7 @@ public:
 			if (WorkerGetResult == Zen::FZenHttpRequest::Result::Failed)
 			{
 				// TODO: log failure!
-				return OnComplete({ Action.GetKey(), {}, {}, EStatus::Error });
+				return OnComplete({Action.GetKey(), {}, {}, EStatus::Error});
 			}
 			else if (WorkerRequest->GetResponseCode() == 404)
 			{
@@ -191,7 +183,7 @@ public:
 				if (WorkerPrepResult == Zen::FZenHttpRequest::Result::Failed)
 				{
 					// TODO: log failure
-					return OnComplete({ Action.GetKey(), {}, {}, EStatus::Error });
+					return OnComplete({Action.GetKey(), {}, {}, EStatus::Error});
 				}
 				else if (WorkerRequest->GetResponseCode() == 404)
 				{
@@ -209,34 +201,28 @@ public:
 
 					{
 						TArray<FIoHash> WorkerFileHashes;
-						TArray<TTuple<FStringView, bool>> WorkerFileMeta;
-						Worker.IterateExecutables([&](FStringView Path, const FIoHash& RawHash, uint64 RawSize) {
+						Worker.IterateExecutables([&NeedHashes, &WorkerFileHashes](FStringView Path, const FIoHash& RawHash, uint64 RawSize)
+						{
 							if (NeedHashes.Contains(RawHash))
 							{
 								WorkerFileHashes.Emplace(RawHash);
-								WorkerFileMeta.Emplace(Path, true);
 							}
 						});
 
-						Worker.IterateFiles([&](FStringView Path, const FIoHash& RawHash, uint64 RawSize) {
+						Worker.IterateFiles([&NeedHashes, &WorkerFileHashes](FStringView Path, const FIoHash& RawHash, uint64 RawSize)
+						{
 							if (NeedHashes.Contains(RawHash))
 							{
 								WorkerFileHashes.Emplace(RawHash);
-								WorkerFileMeta.Emplace(Path, false);
 							}
 						});
 
 						FRequestOwner BlockingOwner(EPriority::Blocking);
-						Worker.FindFileData(WorkerFileHashes, BlockingOwner, [&](FBuildWorkerFileDataCompleteParams&& Params) {
-							uint32 MetaIndex = 0;
+						Worker.FindFileData(WorkerFileHashes, BlockingOwner, [&Package](FBuildWorkerFileDataCompleteParams&& Params)
+						{
 							for (const FCompressedBuffer& Buffer : Params.Files)
 							{
-								const TTuple<FStringView, bool>& Meta = WorkerFileMeta[MetaIndex];
-
-								FCbAttachment Attachment{ Buffer };
-								Package.AddAttachment(Attachment);
-
-								++MetaIndex;
+								Package.AddAttachment(FCbAttachment{Buffer});
 							}
 						});
 						BlockingOwner.Wait();
@@ -250,7 +236,7 @@ public:
 					if (WorkerTransmitResult == Zen::FZenHttpRequest::Result::Failed)
 					{
 						// TODO: log failure
-						return OnComplete({ Action.GetKey(), {}, {}, EStatus::Error });
+						return OnComplete({Action.GetKey(), {}, {}, EStatus::Error});
 					}
 				}
 			}
@@ -258,7 +244,7 @@ public:
 			if (!UE::Zen::IsSuccessCode(WorkerRequest->GetResponseCode()))
 			{
 				// TODO: log failure
-				return OnComplete({ Action.GetKey(), {}, {}, EStatus::Error });
+				return OnComplete({Action.GetKey(), {}, {}, EStatus::Error});
 			}
 		}
 
@@ -283,7 +269,7 @@ public:
 		if (JobPrepResult == Zen::FZenHttpRequest::Result::Failed)
 		{
 			// TODO: log failure!
-			return OnComplete({ Action.GetKey(), {}, {}, EStatus::Error });
+			return OnComplete({Action.GetKey(), {}, {}, EStatus::Error});
 		}
 		else if (WorkerRequest->GetResponseCode() == 404)
 		{
@@ -301,13 +287,11 @@ public:
 
 			FCbPackage ActionPackage;
 
-			Inputs.Get().IterateInputs([&](FStringView Key, const FCompressedBuffer& Buffer) {
-				FIoHash Hash = Buffer.GetRawHash();
-
-				if (NeedHashes.Contains(Hash))
+			Inputs.Get().IterateInputs([&NeedHashes, &ActionPackage](FStringView Key, const FCompressedBuffer& Buffer)
+			{
+				if (NeedHashes.Contains(Buffer.GetRawHash()))
 				{
-					FCbAttachment Attachment{ Buffer };
-					ActionPackage.AddAttachment(Attachment);
+					ActionPackage.AddAttachment(FCbAttachment{Buffer});
 				}
 			});
 
@@ -319,7 +303,7 @@ public:
 			if (WorkerTransmitResult == Zen::FZenHttpRequest::Result::Failed)
 			{
 				// TODO: log failure!
-				return OnComplete({ Action.GetKey(), {}, {}, EStatus::Error });
+				return OnComplete({Action.GetKey(), {}, {}, EStatus::Error});
 			}
 		}
 
@@ -337,8 +321,7 @@ public:
 		if (RemoteBuildOutput.IsNull())
 		{
 			UE_LOG(LogDerivedDataBuildZenExecutor, Warning, TEXT("Remote execution system error: build output blob missing!"));
-
-			return OnComplete({ Action.GetKey(), {}, {}, EStatus::Error });
+			return OnComplete({Action.GetKey(), {}, {}, EStatus::Error});
 		}
 
 		FBuildOutputBuilder OutputBuilder = BuildSystem.CreateOutput(Action.GetName(), Action.GetFunction());
@@ -346,7 +329,7 @@ public:
 		for (const FPayload& Payload : RemoteBuildOutput.Get().GetPayloads())
 		{
 			FCompressedBuffer BufferForPayload;
-					
+
 			if (const FCbAttachment* Attachment = ResultPackage.FindAttachment(Payload.GetRawHash()))
 			{
 				BufferForPayload = Attachment->AsCompressedBinary();
@@ -355,16 +338,14 @@ public:
 			if (BufferForPayload.IsNull())
 			{
 				UE_LOG(LogDerivedDataBuildZenExecutor, Warning, TEXT("Zen execution system error: payload blob missing!"));
-
-				return OnComplete({ Action.GetKey(), {}, {}, EStatus::Error });
+				return OnComplete({Action.GetKey(), {}, {}, EStatus::Error});
 			}
 
 			OutputBuilder.AddPayload(FPayload(Payload.GetId(), BufferForPayload));
 		}
 
 		FBuildOutput BuildOutput = OutputBuilder.Build();
-
-		return OnComplete({ Action.GetKey(), BuildOutput, {}, EStatus::Ok });
+		return OnComplete({Action.GetKey(), BuildOutput, {}, EStatus::Ok});
 
 #if 0
 		// This path will execute the build action synchronously in a scratch directory
@@ -554,14 +535,15 @@ public:
 		return HostPlatforms;
 	}
 
-	void DumpStats() {}
+	void DumpStats()
+	{
+	}
 
 private:
 	TUniquePtr<UE::Zen::FZenHttpRequestPool> RequestPool;
-	FString SandboxRootDir = FPaths::EngineSavedDir() / TEXT("LocalExec");
 };
 
-}  // namespace UE::DerivedData
+} // UE::DerivedData
 
 TOptional<UE::DerivedData::FZenBuildWorkerExecutor> GZenBuildWorkerExecutor;
 
