Fix wrong icons being used for platform quick launch, resulting in pixelation

diff --git a/Engine/Source/Editor/UnrealEd/Private/Kismet2/DebuggerCommands.cpp b/Engine/Source/Editor/UnrealEd/Private/Kismet2/DebuggerCommands.cpp
--- a/Engine/Source/Editor/UnrealEd/Private/Kismet2/DebuggerCommands.cpp
+++ b/Engine/Source/Editor/UnrealEd/Private/Kismet2/DebuggerCommands.cpp
@@ -1339,7 +1339,7 @@ FSlateIcon FInternalPlayWorldCommandCallbacks::GetRepeatLastPlayIcon()
 
 		// get platform name from DeviceId
 		
-		return FSlateIcon(FEditorStyle::GetStyleSetName(), FDataDrivenPlatformInfoRegistry::GetPlatformInfo(DeviceId.GetPlatformName()).GetIconStyleName(EPlatformIconSize::Large));
+		return FSlateIcon(FEditorStyle::GetStyleSetName(), FDataDrivenPlatformInfoRegistry::GetPlatformInfo(DeviceId.GetPlatformName()).GetIconStyleName(EPlatformIconSize::Normal));
 	}
 
 	return GetLastPlaySessionCommand()->GetIcon();
