Fix unnecessary handling of nullptr found by static analysis; SavePackageContext is now always available.
#rnx

diff --git a/Engine/Source/Editor/UnrealEd/Private/CookOnTheFlyServer.cpp b/Engine/Source/Editor/UnrealEd/Private/CookOnTheFlyServer.cpp
--- a/Engine/Source/Editor/UnrealEd/Private/CookOnTheFlyServer.cpp
+++ b/Engine/Source/Editor/UnrealEd/Private/CookOnTheFlyServer.cpp
@@ -8959,14 +8959,14 @@ uint32 UCookOnTheFlyServer::FullLoadAndSave(uint32& CookedPackageCount)
 						}
 								
 						GIsCookerLoadingPackage = true;
-						check(SavePackageContexts.Num() > 0);
-						FSavePackageContext* const SavePackageContext = (IsCookByTheBookMode() && SavePackageContexts.Num() > 0) ? &SavePackageContexts[PlatformIndex]->SaveContext : nullptr;
+						check(SavePackageContexts.Num() > PlatformIndex);
+						FSavePackageContext& SavePackageContext = SavePackageContexts[PlatformIndex]->SaveContext;
 						IPackageWriter::FBeginPackageInfo BeginInfo;
 						BeginInfo.PackageName = Package->GetFName();
 						BeginInfo.LooseFilePath = PlatFilename;
-						SavePackageContext->PackageWriter->BeginPackage(BeginInfo);
+						SavePackageContext.PackageWriter->BeginPackage(BeginInfo);
 						FSavePackageResultStruct SaveResult = GEditor->Save(Package, World, FlagsToCook, *PlatFilename, GError, NULL, bSwap, false, SaveFlags, Target, FDateTime::MinValue(), 
-																			false, /*DiffMap*/ nullptr,	SavePackageContext);
+																			false, /*DiffMap*/ nullptr,	&SavePackageContext);
 						GIsCookerLoadingPackage = false;
 
 						if (SaveResult == ESavePackageResult::Success && UAssetManager::IsValid())
@@ -8989,7 +8989,7 @@ uint32 UCookOnTheFlyServer::FullLoadAndSave(uint32& CookedPackageCount)
 						CommitInfo.PackageGuid = AssetPackageData->PackageGuid;
 						PRAGMA_ENABLE_DEPRECATION_WARNINGS;
 						CommitInfo.WriteOptions = IPackageWriter::EWriteOptions::Write;
-						SavePackageContext->PackageWriter->CommitPackage(MoveTemp(CommitInfo));
+						SavePackageContext.PackageWriter->CommitPackage(MoveTemp(CommitInfo));
 
 						if (SaveResult.IsSuccessful())
 						{
