Change Save & Browse buttons on asset editor toolbars to Icon only

diff --git a/Engine/Source/Editor/UnrealEd/Private/Toolkits/AssetEditorToolkit.cpp b/Engine/Source/Editor/UnrealEd/Private/Toolkits/AssetEditorToolkit.cpp
--- a/Engine/Source/Editor/UnrealEd/Private/Toolkits/AssetEditorToolkit.cpp
+++ b/Engine/Source/Editor/UnrealEd/Private/Toolkits/AssetEditorToolkit.cpp
@@ -1023,14 +1023,7 @@ void FAssetEditorToolkit::RegisterDefaultToolBar()
 	if (!ToolMenus->IsMenuRegistered(DefaultAssetEditorToolBarName))
 	{
 		UToolMenu* ToolbarBuilder = ToolMenus->RegisterMenu(DefaultAssetEditorToolBarName, NAME_None, EMultiBoxType::SlimHorizontalToolBar);
-#if 0
-		ToolbarBuilder->StyleName = "AssetEditorToolbar";
-#endif
-		{
-			FToolMenuSection& Section = ToolbarBuilder->AddSection("Asset");
-			Section.AddEntry(FToolMenuEntry::InitToolBarButton(FAssetEditorCommonCommands::Get().SaveAsset));
-			Section.AddEntry(FToolMenuEntry::InitToolBarButton(FGlobalEditorCommonCommands::Get().FindInContentBrowser, LOCTEXT("FindInContentBrowserButton", "Browse")));
-		}
+		FToolMenuSection& Section = ToolbarBuilder->AddSection("Asset");
 	}
 }
 
@@ -1039,6 +1032,27 @@ void FAssetEditorToolkit::InitToolMenuContext(FToolMenuContext& MenuContext)
 
 }
 
+UToolMenu* FAssetEditorToolkit::GenerateCommonActionsToolbar(FToolMenuContext& MenuContext)
+{
+	UToolMenus* ToolMenus = UToolMenus::Get();
+	FName ToolBarName = "AssetEditorToolbar.CommonActions";
+
+	UToolMenu* FoundMenu = ToolMenus->FindMenu(ToolBarName);
+
+	if (!FoundMenu || !FoundMenu->IsRegistered())
+	{
+		FoundMenu = ToolMenus->RegisterMenu(ToolBarName, NAME_None, EMultiBoxType::SlimHorizontalToolBar);
+		FoundMenu->StyleName = "AssetEditorToolbar";
+
+		FToolMenuSection& Section = FoundMenu->AddSection("CommonActions");
+		Section.AddEntry(FToolMenuEntry::InitToolBarButton(FAssetEditorCommonCommands::Get().SaveAsset));
+		Section.AddEntry(FToolMenuEntry::InitToolBarButton(FGlobalEditorCommonCommands::Get().FindInContentBrowser, LOCTEXT("FindInContentBrowserButton", "Browse")));
+		Section.AddSeparator(NAME_None);
+	}
+
+	return ToolMenus->GenerateMenu(ToolBarName, MenuContext);;
+}
+
 void FAssetEditorToolkit::GenerateToolbar()
 {
 	TSharedPtr<FExtender> Extender = FExtender::Combine(ToolbarExtenders);
@@ -1067,6 +1081,9 @@ void FAssetEditorToolkit::GenerateToolbar()
 	GeneratedToolbar->bToolBarForceSmallIcons = bIsToolbarUsingSmallIcons;
 	TSharedRef< class SWidget > ToolBarWidget = ToolMenus->GenerateWidget(GeneratedToolbar);
 
+	UToolMenu* CommonActionsToolbar = GenerateCommonActionsToolbar(MenuContext);
+	TSharedRef< class SWidget > CommonActionsToolbarWidget = ToolMenus->GenerateWidget(CommonActionsToolbar);
+
 	TSharedRef<SWidget> MiscWidgets = SNullWidget::NullWidget;
 
 	if(ToolbarWidgets.Num() > 0)
@@ -1089,6 +1106,11 @@ void FAssetEditorToolkit::GenerateToolbar()
 	Toolbar = 
 		SNew(SHorizontalBox)
 		+SHorizontalBox::Slot()
+		.AutoWidth()
+		[
+			CommonActionsToolbarWidget
+		]
+		+SHorizontalBox::Slot()
 		[
 			ToolBarWidget
 		]
diff --git a/Engine/Source/Editor/UnrealEd/Public/Toolkits/AssetEditorToolkit.h b/Engine/Source/Editor/UnrealEd/Public/Toolkits/AssetEditorToolkit.h
--- a/Engine/Source/Editor/UnrealEd/Public/Toolkits/AssetEditorToolkit.h
+++ b/Engine/Source/Editor/UnrealEd/Public/Toolkits/AssetEditorToolkit.h
@@ -288,6 +288,9 @@ protected:
 	/**	Returns an array of all the objects currently being edited. Asserts if editing no objects */
 	const TArray< UObject* >& GetEditingObjects() const;
 
+	/** Generate the toolbar for common asset actions like Save*/
+	UToolMenu* GenerateCommonActionsToolbar(FToolMenuContext& MenuContext);
+
 	/** Get the collection of edited objects that can be saved. */
 	virtual void GetSaveableObjects(TArray<UObject*>& OutObjects) const;
 
