Removes pure virtual from deprecated RunInternal method.
Changes RunInternal FBox version to call FCellInfo version.

#changelist validated

diff --git a/Engine/Source/Editor/UnrealEd/Private/WorldPartition/WorldPartitionBuilder.cpp b/Engine/Source/Editor/UnrealEd/Private/WorldPartition/WorldPartitionBuilder.cpp
--- a/Engine/Source/Editor/UnrealEd/Private/WorldPartition/WorldPartitionBuilder.cpp
+++ b/Engine/Source/Editor/UnrealEd/Private/WorldPartition/WorldPartitionBuilder.cpp
@@ -127,6 +127,15 @@ bool UWorldPartitionBuilder::RunBuilder(UWorldPartitionBuilder* Builder, UWorld*
 	return bResult;
 }
 
+static FIntVector GetCellCoord(const FVector& InPos, const int32 InCellSize)
+{
+	return FIntVector(
+		FMath::FloorToInt(InPos.X / InCellSize),
+		FMath::FloorToInt(InPos.Y / InCellSize),
+		FMath::FloorToInt(InPos.Z / InCellSize)
+	);
+}
+
 bool UWorldPartitionBuilder::Run(UWorld* World, FPackageSourceControlHelper& PackageHelper)
 {
 	// Notify derived classes that partition building process starts
@@ -175,15 +184,6 @@ bool UWorldPartitionBuilder::Run(UWorld* World, FPackageSourceControlHelper& Pac
 	if ((LoadingMode == ELoadingMode::IterativeCells) || (LoadingMode == ELoadingMode::IterativeCells2D))
 	{
 		// do partial loading loop that calls RunInternal
-		auto GetCellCoord = [](const FVector InPos, const int32 InCellSize)
-		{
-			return FIntVector(
-				FMath::FloorToInt(InPos.X / InCellSize),
-				FMath::FloorToInt(InPos.Y / InCellSize),
-				FMath::FloorToInt(InPos.Z / InCellSize)
-			);
-		};
-
 		auto CanIterateZ = [](const bool bInResult, const ELoadingMode InLoadingMode, const int32 InZ, const int32 InMinZ, const int32 InMaxZ) -> bool
 		{
 			if (InLoadingMode == ELoadingMode::IterativeCells2D)
@@ -271,3 +271,16 @@ bool UWorldPartitionBuilder::Run(UWorld* World, FPackageSourceControlHelper& Pac
 
 	return OnPartitionBuildCompleted(World, PackageHelper, bResult);
 }
+
+bool UWorldPartitionBuilder::RunInternal(UWorld* World, const FBox& Bounds, FPackageSourceControlHelper& PackageHelper)
+{
+	UWorldPartition*	WorldPartition = World->GetWorldPartition();
+	FCellInfo			CellInfo;
+
+	CellInfo.EditorBounds = WorldPartition->GetEditorWorldBounds();
+	CellInfo.IterativeCellSize = IterativeCellSize;
+	CellInfo.Bounds = Bounds;
+	CellInfo.Location = GetCellCoord(Bounds.Min, IterativeCellSize);
+
+	return RunInternal(World, CellInfo, PackageHelper);
+}
diff --git a/Engine/Source/Editor/UnrealEd/Public/WorldPartition/WorldPartitionBuilder.h b/Engine/Source/Editor/UnrealEd/Public/WorldPartition/WorldPartitionBuilder.h
--- a/Engine/Source/Editor/UnrealEd/Public/WorldPartition/WorldPartitionBuilder.h
+++ b/Engine/Source/Editor/UnrealEd/Public/WorldPartition/WorldPartitionBuilder.h
@@ -69,7 +69,7 @@ protected:
 	virtual bool RunInternal(UWorld* World, const FCellInfo& InCellInfo, FPackageSourceControlHelper& PackageHelper) PURE_VIRTUAL(UWorldPartition::RunInternal, return false;);
 
 	UE_DEPRECATED(5.0, "Use RunInternal version with FCellInfo instead")
-	virtual bool RunInternal(UWorld* World, const FBox& Bounds, FPackageSourceControlHelper& PackageHelper) PURE_VIRTUAL(UWorldPartition::RunInternal, return false;);
+	virtual bool RunInternal(UWorld* World, const FBox& Bounds, FPackageSourceControlHelper& PackageHelper);
 
 	/**
 	 * Overridable method for derived classed to perform operations when partition building process completes.
