SpawnActorFromClass nodes in Editor Utility Widgets now respect the bSelectActors parameter.

diff --git a/Engine/Source/Editor/UnrealEd/Classes/Editor/EditorEngine.h b/Engine/Source/Editor/UnrealEd/Classes/Editor/EditorEngine.h
--- a/Engine/Source/Editor/UnrealEd/Classes/Editor/EditorEngine.h
+++ b/Engine/Source/Editor/UnrealEd/Classes/Editor/EditorEngine.h
@@ -1098,9 +1098,10 @@ public:
 	 * @param	Transform		The world-space transform to spawn the actor with.
 	 * @param	bSilent			If true, suppress logging (optional, defaults to false).
 	 * @param	ObjectFlags		The object flags to place on the spawned actor.
+	 * @param	bSelectActor	Whether or not to select the spawned actor.
 	 * @return					A pointer to the newly added actor, or NULL if add failed.
 	 */
-	virtual AActor* AddActor(ULevel* InLevel, UClass* Class, const FTransform& Transform, bool bSilent = false, EObjectFlags ObjectFlags = RF_Transactional);
+	virtual AActor* AddActor(ULevel* InLevel, UClass* Class, const FTransform& Transform, bool bSilent = false, EObjectFlags ObjectFlags = RF_Transactional, bool bSelectActor = true);
 
 	/**
 	 * Adds actors to the world at the specified location using export text.
diff --git a/Engine/Source/Editor/UnrealEd/Private/EditorEngine.cpp b/Engine/Source/Editor/UnrealEd/Private/EditorEngine.cpp
--- a/Engine/Source/Editor/UnrealEd/Private/EditorEngine.cpp
+++ b/Engine/Source/Editor/UnrealEd/Private/EditorEngine.cpp
@@ -6104,7 +6104,7 @@ bool UEditorEngine::AreAllWindowsHidden() const
 	return bAllHidden;
 }
 
-AActor* UEditorEngine::AddActor(ULevel* InLevel, UClass* Class, const FTransform& Transform, bool bSilent, EObjectFlags InObjectFlags)
+AActor* UEditorEngine::AddActor(ULevel* InLevel, UClass* Class, const FTransform& Transform, bool bSilent, EObjectFlags InObjectFlags, bool bSelectActor)
 {
 	check( Class );
 
@@ -6152,12 +6152,16 @@ AActor* UEditorEngine::AddActor(ULevel* InLevel, UClass* Class, const FTransform
 	AActor* Actor = NULL;
 	{
 		FScopedTransaction Transaction( NSLOCTEXT("UnrealEd", "AddActor", "Add Actor") );
-		if ( !(InObjectFlags & RF_Transactional) )
+		if (!(InObjectFlags & RF_Transactional))
 		{
 			// Don't attempt a transaction if the actor we are spawning isn't transactional
 			Transaction.Cancel();
 		}
-		SelectNone( false, true );
+
+		if (bSelectActor)
+		{
+			SelectNone(false, true);
+		}
 
 		AActor* Default = Class->GetDefaultObject<AActor>();
 
@@ -6169,9 +6173,13 @@ AActor* UEditorEngine::AddActor(ULevel* InLevel, UClass* Class, const FTransform
 		const auto Rotation = Transform.GetRotation().Rotator();
 		Actor = World->SpawnActor( Class, &Location, &Rotation, SpawnInfo );
 
-		if( Actor )
+		if (Actor)
 		{
-			SelectActor( Actor, 1, 0 );
+			if (bSelectActor)
+			{
+				SelectActor(Actor, 1, 0);
+			}
+
 			Actor->InvalidateLightingCache();
 			Actor->PostEditMove( true );
 		}
@@ -6191,7 +6199,10 @@ AActor* UEditorEngine::AddActor(ULevel* InLevel, UClass* Class, const FTransform
 		ULevel::LevelDirtiedEvent.Broadcast();
 	}
 
-	NoteSelectionChange();
+	if( bSelectActor )
+	{
+		NoteSelectionChange();
+	}
 
 	return Actor;
 }
diff --git a/Engine/Source/Editor/UnrealEd/Private/LevelEditorViewport.cpp b/Engine/Source/Editor/UnrealEd/Private/LevelEditorViewport.cpp
--- a/Engine/Source/Editor/UnrealEd/Private/LevelEditorViewport.cpp
+++ b/Engine/Source/Editor/UnrealEd/Private/LevelEditorViewport.cpp
@@ -225,7 +225,7 @@ TArray<AActor*> FLevelEditorViewportClient::TryPlacingActorFromObject( ULevel* I
 		{
 			// If no actor factory was found or failed, add the actor directly.
 			const FTransform ActorTransform = FActorPositioning::GetCurrentViewportPlacementTransform(*ObjectClass->GetDefaultObject<AActor>(), /*bSnap=*/true, Cursor);
-			PlacedActor = GEditor->AddActor( InLevel, ObjectClass, ActorTransform, /*bSilent=*/false, ObjectFlags );
+			PlacedActor = GEditor->AddActor( InLevel, ObjectClass, ActorTransform, /*bSilent=*/false, ObjectFlags, bSelectActors );
 		}
 
 		if ( PlacedActor != NULL )
