SourceControl: Disable the 'Check Out Selected' button if the SourceControl system does not use checkout (the operation will always fail).

diff --git a/Engine/Source/Editor/UnrealEd/Private/FileHelpers.cpp b/Engine/Source/Editor/UnrealEd/Private/FileHelpers.cpp
--- a/Engine/Source/Editor/UnrealEd/Private/FileHelpers.cpp
+++ b/Engine/Source/Editor/UnrealEd/Private/FileHelpers.cpp
@@ -1505,7 +1505,7 @@ void FEditorFileUtils::Export(bool bExportSelectedActorsOnly)
 
 static bool IsCheckOutSelectedDisabled()
 {
-	return !(ISourceControlModule::Get().IsEnabled() && ISourceControlModule::Get().GetProvider().IsAvailable());
+	return !(ISourceControlModule::Get().IsEnabled() && ISourceControlModule::Get().GetProvider().IsAvailable()) || !ISourceControlModule::Get().GetProvider().UsesCheckout();
 }
 
 bool FEditorFileUtils::AddCheckoutPackageItems(bool bCheckDirty, TArray<UPackage*> PackagesToCheckOut, TArray<UPackage*>* OutPackagesNotNeedingCheckout, bool* bOutHavePackageToCheckOut)
