Cooker: mark the load as failed if the package is IsPendingKill, to prevent an assertion during save where we expect the weakpointer to return non-null.
#rnx

diff --git a/Engine/Source/Editor/UnrealEd/Private/CookOnTheFlyServer.cpp b/Engine/Source/Editor/UnrealEd/Private/CookOnTheFlyServer.cpp
--- a/Engine/Source/Editor/UnrealEd/Private/CookOnTheFlyServer.cpp
+++ b/Engine/Source/Editor/UnrealEd/Private/CookOnTheFlyServer.cpp
@@ -2929,7 +2929,7 @@ bool UCookOnTheFlyServer::LoadPackageForCooking(UE::Cook::FPackageData& PackageD
 
 	bool bSuccess = true;
 	//  if the package is not yet fully loaded then fully load it
-	if (OutPackage == nullptr || !OutPackage->IsFullyLoaded())
+	if (!IsValid(OutPackage) || !OutPackage->IsFullyLoaded())
 	{
 		bool bWasPartiallyLoaded = OutPackage != nullptr;
 		GIsCookerLoadingPackage = true;
@@ -2948,7 +2948,7 @@ bool UCookOnTheFlyServer::LoadPackageForCooking(UE::Cook::FPackageData& PackageD
 		}
 
 		UPackage* LoadedPackage = LoadPackage(LoadIntoPackage, *FileName, LOAD_None);
-		if (LoadedPackage)
+		if (IsValid(LoadedPackage) && LoadedPackage->IsFullyLoaded())
 		{
 			OutPackage = LoadedPackage;
 
