Fixed a compile warning in remote execution

#rnx

diff --git a/Engine/Source/Editor/UnrealEd/Private/DerivedDataBuildRemoteExecutor.cpp b/Engine/Source/Editor/UnrealEd/Private/DerivedDataBuildRemoteExecutor.cpp
--- a/Engine/Source/Editor/UnrealEd/Private/DerivedDataBuildRemoteExecutor.cpp
+++ b/Engine/Source/Editor/UnrealEd/Private/DerivedDataBuildRemoteExecutor.cpp
@@ -550,7 +550,7 @@ FCbObject FRemoteBuildExecutionRequest::BuildWorkerDescriptor(const FBuildWorker
 
 	WorkerDescriptor.EndObject();
 
-	return std::move(WorkerDescriptor.Save().AsObject());
+	return WorkerDescriptor.Save().AsObject();
 }
 
 bool FRemoteBuildExecutionRequest::ProcessCancellation()
