Refresh user-defined struct editor data on save/compile to ensure default object reference values stay in sync with default struct instance data.

diff --git a/Engine/Source/Editor/KismetCompiler/Private/UserDefinedStructureCompilerUtils.cpp b/Engine/Source/Editor/KismetCompiler/Private/UserDefinedStructureCompilerUtils.cpp
--- a/Engine/Source/Editor/KismetCompiler/Private/UserDefinedStructureCompilerUtils.cpp
+++ b/Engine/Source/Editor/KismetCompiler/Private/UserDefinedStructureCompilerUtils.cpp
@@ -150,6 +150,9 @@ struct FUserDefinedStructureCompilerInner
 
 		if (UUserDefinedStructEditorData* EditorData = Cast<UUserDefinedStructEditorData>(StructToClean->EditorData))
 		{
+			// Ensure that editor data is in sync w/ the current default instance (if valid) so that it can be reinitialized later.
+			EditorData->RefreshValuesFromDefaultInstance();
+
 			EditorData->CleanDefaultInstance();
 		}
 
diff --git a/Engine/Source/Editor/UnrealEd/Classes/UserDefinedStructure/UserDefinedStructEditorData.h b/Engine/Source/Editor/UnrealEd/Classes/UserDefinedStructure/UserDefinedStructEditorData.h
--- a/Engine/Source/Editor/UnrealEd/Classes/UserDefinedStructure/UserDefinedStructEditorData.h
+++ b/Engine/Source/Editor/UnrealEd/Classes/UserDefinedStructure/UserDefinedStructEditorData.h
@@ -134,6 +134,7 @@ public:
 	virtual void PostEditUndo() override;
 	virtual void PostEditUndo(TSharedPtr<ITransactionObjectAnnotation> TransactionAnnotation) override;
 	virtual void PostLoadSubobjects(struct FObjectInstancingGraph* OuterInstanceGraph) override;
+	virtual void Serialize(FStructuredArchive::FRecord Record) override;
 	// End of UObject interface.
 
 	// FEditorUndoClient interface
@@ -149,6 +150,7 @@ public:
 	void RecreateDefaultInstance(FString* OutLog = nullptr);
 	void ReinitializeDefaultInstance(FString* OutLog = nullptr);
 	void CleanDefaultInstance();
+	void RefreshValuesFromDefaultInstance();
 
 private:
 
diff --git a/Engine/Source/Editor/UnrealEd/Private/UserDefinedStructEditorData.cpp b/Engine/Source/Editor/UnrealEd/Private/UserDefinedStructEditorData.cpp
--- a/Engine/Source/Editor/UnrealEd/Private/UserDefinedStructEditorData.cpp
+++ b/Engine/Source/Editor/UnrealEd/Private/UserDefinedStructEditorData.cpp
@@ -170,6 +170,18 @@ void UUserDefinedStructEditorData::PostLoadSubobjects(FObjectInstancingGraph* Ou
 	}
 }
 
+void UUserDefinedStructEditorData::Serialize(FStructuredArchive::FRecord Record)
+{
+	// Prior to saving, ensure that editor data matches up with the default instance data.
+	FArchive& UnderlyingArchive = Record.GetUnderlyingArchive();
+	if (UnderlyingArchive.IsSaving())
+	{
+		RefreshValuesFromDefaultInstance();
+	}
+
+	Super::Serialize(Record);
+}
+
 const uint8* UUserDefinedStructEditorData::GetDefaultInstance() const
 {
 	return GetOwnerStruct()->GetDefaultInstance();
@@ -233,4 +245,30 @@ void UUserDefinedStructEditorData::CleanDefaultInstance()
 	ScriptStruct->DefaultStructInstance.Destroy();
 }
 
+void UUserDefinedStructEditorData::RefreshValuesFromDefaultInstance()
+{
+	UUserDefinedStruct* ScriptStruct = GetOwnerStruct();
+	if (!ScriptStruct || !ScriptStruct->DefaultStructInstance.IsValid())
+	{
+		return;
+	}
+
+	ensure(ScriptStruct->DefaultStructInstance.GetStruct() == ScriptStruct);
+	if (uint8* StructData = ScriptStruct->DefaultStructInstance.GetStructMemory())
+	{
+		for (TFieldIterator<FProperty> It(ScriptStruct); It; ++It)
+		{
+			if (const FProperty* Property = *It)
+			{
+				const FGuid VarGuid = FStructureEditorUtils::GetGuidFromPropertyName(Property->GetFName());
+				if (FStructVariableDescription* VarDesc = VariablesDescriptions.FindByPredicate(FStructureEditorUtils::FFindByGuidHelper<FStructVariableDescription>(VarGuid)))
+				{
+					FBlueprintEditorUtils::PropertyValueToString(Property, StructData, VarDesc->CurrentDefaultValue, ScriptStruct);
+					VarDesc->DefaultValue = VarDesc->CurrentDefaultValue;
+				}
+			}
+		}
+	}
+}
+
 #undef LOCTEXT_NAMESPACE
