Added support for Level Instances to have a folder hierarchy in the World Outliner. Actor Folders can now have a root object which serves as a context for folder hierarchy.

diff --git a/Engine/Source/Editor/DataLayerEditor/Private/DataLayer/DataLayerEditorSubsystem.cpp b/Engine/Source/Editor/DataLayerEditor/Private/DataLayer/DataLayerEditorSubsystem.cpp
--- a/Engine/Source/Editor/DataLayerEditor/Private/DataLayer/DataLayerEditorSubsystem.cpp
+++ b/Engine/Source/Editor/DataLayerEditor/Private/DataLayer/DataLayerEditorSubsystem.cpp
@@ -186,7 +186,7 @@ void UDataLayerEditorSubsystem::PostUndoRedo()
 
 bool UDataLayerEditorSubsystem::IsActorValidForDataLayer(AActor* Actor)
 {
-	return Actor && Actor->IsValidForDataLayer() && (Actor->GetLevel() == Actor->GetWorld()->PersistentLevel);
+	return Actor && Actor->SupportsDataLayer() && Actor->IsValidForDataLayer() && (Actor->GetLevel() == Actor->GetWorld()->PersistentLevel);
 }
 
 void UDataLayerEditorSubsystem::InitializeNewActorDataLayers(AActor* Actor)
diff --git a/Engine/Source/Editor/DataLayerEditor/Private/DataLayer/DataLayerMode.cpp b/Engine/Source/Editor/DataLayerEditor/Private/DataLayer/DataLayerMode.cpp
--- a/Engine/Source/Editor/DataLayerEditor/Private/DataLayer/DataLayerMode.cpp
+++ b/Engine/Source/Editor/DataLayerEditor/Private/DataLayer/DataLayerMode.cpp
@@ -375,7 +375,7 @@ FSceneOutlinerDragValidationInfo FDataLayerMode::ValidateDrop(const ISceneOutlin
 	{
 		for (AActor* Actor : PayloadActors)
 		{
-			if (!Actor->SupportsDataLayer())
+			if (!DataLayerEditorSubsystem->IsActorValidForDataLayer(Actor))
 			{
 				return FSceneOutlinerDragValidationInfo(ESceneOutlinerDropCompatibility::IncompatibleGeneric, LOCTEXT("ActorCantBeAssignedToDataLayer", "Can't assign actors to Data Layer"));
 			}
@@ -502,7 +502,7 @@ TArray<AActor*> FDataLayerMode::GetActorsFromOperation(const FDragDropOperation&
 		if (UWorld* World = FolderOp.World.Get())
 		{
 			TArray<TWeakObjectPtr<AActor>> ActorsToDrop;
-			FActorFolders::GetWeakActorsFromFolders(*World, FolderOp.Folders, ActorsToDrop);
+			FActorFolders::GetWeakActorsFromFolders(*World, FolderOp.Folders, ActorsToDrop, FolderOp.RootObject);
 			for (const auto& Actor : ActorsToDrop)
 			{
 				if (AActor* ActorPtr = Actor.Get())
@@ -778,7 +778,7 @@ void FDataLayerMode::RegisterContextMenu()
 			{
 				FToolMenuSection& Section = InMenu->AddSection("DataLayers", LOCTEXT("DataLayers", "Data Layers"));
 				
-				auto CreateNewDataLayer = [this](UDataLayer* ParentDataLayer = nullptr)
+				auto CreateNewDataLayer = [this, SceneOutliner](UDataLayer* ParentDataLayer = nullptr)
 				{
 					const FScopedDataLayerTransaction Transaction(LOCTEXT("CreateNewDataLayer", "Create New Data Layer"), RepresentingWorld.Get());
 					SelectedDataLayersSet.Empty();
@@ -787,10 +787,8 @@ void FDataLayerMode::RegisterContextMenu()
 					{
 						SelectedDataLayersSet.Add(NewDataLayer);
 						DataLayerEditorSubsystem->SetParentDataLayer(NewDataLayer, ParentDataLayer);
-						if (DataLayerBrowser)
-						{
-							DataLayerBrowser->OnSelectionChanged(SelectedDataLayersSet);
-						}
+						// Select it and open a rename when it gets refreshed
+						SceneOutliner->OnItemAdded(NewDataLayer, SceneOutliner::ENewItemAction::Select | SceneOutliner::ENewItemAction::Rename);
 					}
 				};
 
diff --git a/Engine/Source/Editor/Layers/Private/SLayerBrowser.h b/Engine/Source/Editor/Layers/Private/SLayerBrowser.h
--- a/Engine/Source/Editor/Layers/Private/SLayerBrowser.h
+++ b/Engine/Source/Editor/Layers/Private/SLayerBrowser.h
@@ -144,7 +144,7 @@ protected:
 			{
 				if (UWorld* World = FolderDragOp->World.Get())
 				{
-					FActorFolders::GetWeakActorsFromFolders(*World, FolderDragOp->Folders, Actors);
+					FActorFolders::GetWeakActorsFromFolders(*World, FolderDragOp->Folders, Actors, FolderDragOp->RootObject);
 				}
 			}
 		}
@@ -244,7 +244,7 @@ protected:
 		{
 			if (UWorld* World = FolderDragOp->World.Get())
 			{
-				FActorFolders::GetWeakActorsFromFolders(*World, FolderDragOp->Folders, ActorsToDrop);
+				FActorFolders::GetWeakActorsFromFolders(*World, FolderDragOp->Folders, ActorsToDrop, FolderDragOp->RootObject);
 				bHandled = true;
 			}
 		}
diff --git a/Engine/Source/Editor/Layers/Private/SLayersView.h b/Engine/Source/Editor/Layers/Private/SLayersView.h
--- a/Engine/Source/Editor/Layers/Private/SLayersView.h
+++ b/Engine/Source/Editor/Layers/Private/SLayersView.h
@@ -234,7 +234,7 @@ protected:
 		{
 			if (UWorld* World = FolderDragOp->World.Get())
 			{
-				FActorFolders::GetWeakActorsFromFolders(*World, FolderDragOp->Folders, ActorsToDrop);
+				FActorFolders::GetWeakActorsFromFolders(*World, FolderDragOp->Folders, ActorsToDrop, FolderDragOp->RootObject);
 
 				bHandled = true;
 			}
diff --git a/Engine/Source/Editor/Layers/Private/SLayersViewRow.cpp b/Engine/Source/Editor/Layers/Private/SLayersViewRow.cpp
--- a/Engine/Source/Editor/Layers/Private/SLayersViewRow.cpp
+++ b/Engine/Source/Editor/Layers/Private/SLayersViewRow.cpp
@@ -143,7 +143,7 @@ FReply SLayersViewRow::OnDragOver(const FGeometry& MyGeometry, const FDragDropEv
 	{
 		if (UWorld* World = FolderDragOp->World.Get())
 		{
-			FActorFolders::GetWeakActorsFromFolders(*World, FolderDragOp->Folders, Actors);
+			FActorFolders::GetWeakActorsFromFolders(*World, FolderDragOp->Folders, Actors, FolderDragOp->RootObject);
 		}
 	}
 
@@ -202,7 +202,7 @@ FReply SLayersViewRow::OnDrop(const FGeometry& MyGeometry, const FDragDropEvent&
 	{
 		if (UWorld* World = FolderDragOp->World.Get())
 		{
-			FActorFolders::GetWeakActorsFromFolders(*World, FolderDragOp->Folders, ActorsToDrop);
+			FActorFolders::GetWeakActorsFromFolders(*World, FolderDragOp->Folders, ActorsToDrop, FolderDragOp->RootObject);
 			bHandled = true;
 		}
 	}
@@ -217,7 +217,7 @@ FReply SLayersViewRow::OnDrop(const FGeometry& MyGeometry, const FDragDropEvent&
 		{
 			if (UWorld* World = FolderDragOp->World.Get())
 			{
-				FActorFolders::GetWeakActorsFromFolders(*World, FolderSubOp->Folders, ActorsToDrop);
+				FActorFolders::GetWeakActorsFromFolders(*World, FolderSubOp->Folders, ActorsToDrop, FolderDragOp->RootObject);
 				bHandled = true;
 			}
 		}
@@ -274,7 +274,7 @@ FSlateColor SLayersViewRow::GetColorAndOpacity() const
 			auto World = FolderDragOp->World;
 			if (UWorld* WorldPtr = World.Get())
 			{
-				FActorFolders::GetWeakActorsFromFolders(*WorldPtr, FolderDragOp->Folders, DraggedActors);
+				FActorFolders::GetWeakActorsFromFolders(*WorldPtr, FolderDragOp->Folders, DraggedActors, FolderDragOp->RootObject);
 			}
 		}
 
diff --git a/Engine/Source/Editor/LevelEditor/Private/LevelEditorActions.cpp b/Engine/Source/Editor/LevelEditor/Private/LevelEditorActions.cpp
--- a/Engine/Source/Editor/LevelEditor/Private/LevelEditorActions.cpp
+++ b/Engine/Source/Editor/LevelEditor/Private/LevelEditorActions.cpp
@@ -1347,7 +1347,7 @@ bool FLevelEditorActionCallbacks::IsAttachableActor( const AActor* const ParentA
 
 void FLevelEditorActionCallbacks::CreateNewOutlinerFolder_Clicked()
 {
-	const FName NewFolderName = FActorFolders::Get().GetDefaultFolderNameForSelection(*GetWorld());
+	const FFolder NewFolderName = FActorFolders::Get().GetDefaultFolderForSelection(*GetWorld());
 	FActorFolders::Get().CreateFolderContainingSelection(*GetWorld(), NewFolderName);
 }
 
diff --git a/Engine/Source/Editor/SceneOutliner/Private/ActorBrowsingMode.cpp b/Engine/Source/Editor/SceneOutliner/Private/ActorBrowsingMode.cpp
--- a/Engine/Source/Editor/SceneOutliner/Private/ActorBrowsingMode.cpp
+++ b/Engine/Source/Editor/SceneOutliner/Private/ActorBrowsingMode.cpp
@@ -44,6 +44,7 @@ UActorBrowsingModeSettings::UActorBrowsingModeSettings(const FObjectInitializer&
 
 using FActorFilter = TSceneOutlinerPredicateFilter<FActorTreeItem>;
 using FActorDescFilter = TSceneOutlinerPredicateFilter<FActorDescTreeItem>;
+static const FFolder GWorldRoot(NAME_None, nullptr);
 
 FActorBrowsingMode::FActorBrowsingMode(SSceneOutliner* InSceneOutliner, TWeakObjectPtr<UWorld> InSpecifiedWorldToDisplay)
 	: FActorModeInteractive(FActorModeParams(InSceneOutliner, InSpecifiedWorldToDisplay,  /* bHideComponents */ true, /* bHideLevelInstanceHierarchy */ false, /* bHideUnloadedActors */ false))
@@ -545,10 +546,7 @@ void FActorBrowsingMode::OnItemAdded(FSceneOutlinerTreeItemPtr Item)
 	{
 		if (FolderItem->World.IsValid())
 		{
-			if (FActorFolderProps* Props = FActorFolders::Get().GetFolderProperties(*FolderItem->World, FolderItem->Path))
-			{
-				FolderItem->Flags.bIsExpanded = Props->bIsExpanded;
-			}
+			FolderItem->Flags.bIsExpanded = FActorFolders::Get().IsFolderExpanded(*FolderItem->World, FolderItem->GetFolder());
 		}
 	}
 	else if (Item->IsA<FActorDescTreeItem>())
@@ -850,6 +848,36 @@ bool FActorBrowsingMode::CanPasteFoldersOnlyFromClipboard() const
 	return PasteString.StartsWith("BEGIN FOLDERLIST");
 }
 
+bool FActorBrowsingMode::GetFolderNamesFromFolders(const TArray<FFolder>& InFolders, TArray<FName>& OutFolders, FFolder::FRootObject& OutCommonRootObject) const
+{
+	OutFolders.Reset();
+	OutCommonRootObject = FFolder::GetDefaultRootObject();
+
+	TOptional<FFolder::FRootObject> CommonRootObject;
+	for (const FFolder& Folder : InFolders)
+	{
+		OutFolders.Add(Folder.GetPath());
+		if (!CommonRootObject.IsSet())
+		{
+			CommonRootObject = Folder.GetRootObject();
+		}
+		else if (CommonRootObject.GetValue() != Folder.GetRootObject())
+		{
+			OutFolders.Reset();
+			CommonRootObject.Reset();
+			break;
+		}
+	};
+
+	OutCommonRootObject = CommonRootObject.Get(FFolder::GetDefaultRootObject());
+	return CommonRootObject.IsSet();
+}
+
+bool FActorBrowsingMode::GetFolderNamesFromPayload(const FSceneOutlinerDragDropPayload& InPayload, TArray<FName>& OutFolders, FFolder::FRootObject& OutCommonRootObject) const
+{
+	return GetFolderNamesFromFolders(InPayload.GetData<FFolder>(SceneOutliner::FFolderPathSelector()), OutFolders, OutCommonRootObject);
+}
+
 TSharedPtr<FDragDropOperation> FActorBrowsingMode::CreateDragDropOperation(const TArray<FSceneOutlinerTreeItemPtr>& InTreeItems) const
 {
 	FSceneOutlinerDragDropPayload DraggedObjects(InTreeItems);
@@ -871,9 +899,14 @@ TSharedPtr<FDragDropOperation> FActorBrowsingMode::CreateDragDropOperation(const
 
 	if (DraggedObjects.Has<FFolderTreeItem>())
 	{
-		TSharedPtr<FFolderDragDropOp> FolderOperation = MakeShareable(new FFolderDragDropOp);
-		FolderOperation->Init(DraggedObjects.GetData<FName>(SceneOutliner::FFolderPathSelector()), RepresentingWorld.Get());
-		OutlinerOp->AddSubOp(FolderOperation);
+		FFolder::FRootObject CommonRootObject;
+		TArray<FName> DraggedFolders;
+		if (GetFolderNamesFromPayload(DraggedObjects, DraggedFolders, CommonRootObject))
+		{
+			TSharedPtr<FFolderDragDropOp> FolderOperation = MakeShareable(new FFolderDragDropOp);
+			FolderOperation->Init(DraggedFolders, RepresentingWorld.Get(), CommonRootObject);
+			OutlinerOp->AddSubOp(FolderOperation);
+		}
 	}
 	OutlinerOp->Construct();
 	return OutlinerOp;
@@ -888,7 +921,7 @@ bool FActorBrowsingMode::ParseDragDrop(FSceneOutlinerDragDropPayload& OutPayload
 		{
 			for (const auto& Folder : FolderOp->Folders)
 			{
-				OutPayload.DraggedItems.Add(SceneOutliner->GetTreeItem(Folder));
+				OutPayload.DraggedItems.Add(SceneOutliner->GetTreeItem(FFolder(Folder, FolderOp->RootObject)));
 			}
 		}
 		if (const auto& ActorOp = OutlinerOp.GetSubOp<FActorDragDropOp>())
@@ -914,6 +947,22 @@ bool FActorBrowsingMode::ParseDragDrop(FSceneOutlinerDragDropPayload& OutPayload
 
 FSceneOutlinerDragValidationInfo FActorBrowsingMode::ValidateDrop(const ISceneOutlinerTreeItem& DropTarget, const FSceneOutlinerDragDropPayload& Payload) const
 {
+	if (Payload.Has<FFolderTreeItem>())
+	{
+		FFolder::FRootObject TargetRootObject = DropTarget.GetRootObject();
+		FFolder::FRootObject CommonPayloadFoldersRootObject;
+		TArray<FName> PayloadFolders;
+		const bool bHasCommonRootObject = GetFolderNamesFromPayload(Payload, PayloadFolders, CommonPayloadFoldersRootObject);
+		if (!bHasCommonRootObject)
+		{
+			return FSceneOutlinerDragValidationInfo(ESceneOutlinerDropCompatibility::IncompatibleGeneric, LOCTEXT("CantMoveFoldersWithMultipleRoots", "Cannot move folders with multiple roots"));
+		}
+		else if (CommonPayloadFoldersRootObject != TargetRootObject)
+		{
+			return FSceneOutlinerDragValidationInfo(ESceneOutlinerDropCompatibility::IncompatibleGeneric, LOCTEXT("CantChangeFoldersRoot", "Cannot change folders root"));
+		}
+	}
+
 	if (const FActorTreeItem* ActorItem = DropTarget.CastTo<FActorTreeItem>())
 	{
 		if (Payload.Has<FFolderTreeItem>())
@@ -1027,56 +1076,77 @@ FSceneOutlinerDragValidationInfo FActorBrowsingMode::ValidateDrop(const ISceneOu
 	else if (DropTarget.IsA<FFolderTreeItem>() || DropTarget.IsA<FWorldTreeItem>())
 	{
 		const FFolderTreeItem* FolderItem = DropTarget.CastTo<FFolderTreeItem>();
+		const FWorldTreeItem* WorldItem = DropTarget.CastTo<FWorldTreeItem>();
 		// World items are treated as folders with path = none
-		const FName& DestinationPath = FolderItem ? FolderItem->Path : NAME_None;
+		const FFolder DestinationPath = FolderItem ? FolderItem->GetFolder() : GWorldRoot;
+		const FFolder::FRootObject& DestinationRootObject = DestinationPath.GetRootObject();
+		ALevelInstance* LevelInstanceTarget = Cast<ALevelInstance>(DestinationPath.GetRootObjectPtr());
+		if (LevelInstanceTarget && !LevelInstanceTarget->IsEditing())
+		{
+			return FSceneOutlinerDragValidationInfo(ESceneOutlinerDropCompatibility::IncompatibleGeneric, LOCTEXT("Error_DragInNonEditingLevelInstance", "Cannot drag into a LevelInstance which is not being edited"));
+		}
+
 		if (Payload.Has<FFolderTreeItem>())
 		{
-			// Iterate over all the folders that have been dragged
-			for (FName DraggedFolder : Payload.GetData<FName>(SceneOutliner::FFolderPathSelector()))
+			FFolder::FRootObject CommonFolderRootObject;
+			TArray<FName> DraggedFolders;
+			if (GetFolderNamesFromPayload(Payload, DraggedFolders, CommonFolderRootObject))
 			{
-				const FName Leaf = FEditorFolderUtils::GetLeafName(DraggedFolder);
-				const FName Parent = FEditorFolderUtils::GetParentPath(DraggedFolder);
-
-				if (Parent == DestinationPath)
+				// Iterate over all the folders that have been dragged
+				for (const FName& DraggedFolder : DraggedFolders)
 				{
-					FFormatNamedArguments Args;
-					Args.Add(TEXT("SourceName"), FText::FromName(Leaf));
+					const FName Leaf = FEditorFolderUtils::GetLeafName(DraggedFolder);
+					const FName Parent = FEditorFolderUtils::GetParentPath(DraggedFolder);
 
-					FText Text;
-					if (DestinationPath.IsNone())
+					if ((CommonFolderRootObject != DestinationRootObject) && FFolder::HasRootObject(CommonFolderRootObject) && FFolder::HasRootObject(DestinationRootObject))
 					{
-						Text = FText::Format(LOCTEXT("FolderAlreadyAssignedRoot", "{SourceName} is already assigned to root"), Args);
+						FFormatNamedArguments Args;
+						Args.Add(TEXT("SourceName"), FText::FromName(Leaf));
+						FText Text = FText::Format(LOCTEXT("CantChangeFolderRoot", "Cannot change {SourceName} folder root"), Args);
+						return FSceneOutlinerDragValidationInfo(ESceneOutlinerDropCompatibility::IncompatibleGeneric, Text);
 					}
-					else
+
+					if (Parent == DestinationPath.GetPath())
 					{
-						Args.Add(TEXT("DestPath"), FText::FromName(DestinationPath));
-						Text = FText::Format(LOCTEXT("FolderAlreadyAssigned", "{SourceName} is already assigned to {DestPath}"), Args);
-					}
+						FFormatNamedArguments Args;
+						Args.Add(TEXT("SourceName"), FText::FromName(Leaf));
 
-					return FSceneOutlinerDragValidationInfo(ESceneOutlinerDropCompatibility::IncompatibleGeneric, Text);
-				}
+						FText Text;
+						if (DestinationPath.IsNone())
+						{
+							Text = FText::Format(LOCTEXT("FolderAlreadyAssignedRoot", "{SourceName} is already assigned to root"), Args);
+						}
+						else
+						{
+							Args.Add(TEXT("DestPath"), FText::FromName(DestinationPath.GetPath()));
+							Text = FText::Format(LOCTEXT("FolderAlreadyAssigned", "{SourceName} is already assigned to {DestPath}"), Args);
+						}
 
-				const FString DragFolderPath = DraggedFolder.ToString();
-				const FString LeafName = Leaf.ToString();
-				const FString DstFolderPath = DestinationPath.IsNone() ? FString() : DestinationPath.ToString();
-				const FString NewPath = DstFolderPath / LeafName;
+						return FSceneOutlinerDragValidationInfo(ESceneOutlinerDropCompatibility::IncompatibleGeneric, Text);
+					}
 
-				if (FActorFolders::Get().GetFolderProperties(*RepresentingWorld, FName(*NewPath)))
-				{
-					// The folder already exists
-					FFormatNamedArguments Args;
-					Args.Add(TEXT("DragName"), FText::FromString(LeafName));
-					return FSceneOutlinerDragValidationInfo(ESceneOutlinerDropCompatibility::IncompatibleGeneric,
-						FText::Format(LOCTEXT("FolderAlreadyExistsRoot", "A folder called \"{DragName}\" already exists at this level"), Args));
-				}
-				else if (DragFolderPath == DstFolderPath || DstFolderPath.StartsWith(DragFolderPath + "/"))
-				{
-					// Cannot drag as a child of itself
-					FFormatNamedArguments Args;
-					Args.Add(TEXT("FolderPath"), FText::FromName(DraggedFolder));
-					return FSceneOutlinerDragValidationInfo(
-						ESceneOutlinerDropCompatibility::IncompatibleGeneric,
-						FText::Format(LOCTEXT("ChildOfItself", "Cannot move \"{FolderPath}\" to be a child of itself"), Args));
+					const FString DragFolderPath = DraggedFolder.ToString();
+					const FString LeafName = Leaf.ToString();
+					const FString DstFolderPath = DestinationPath.IsNone() ? FString() : DestinationPath.ToString();
+					const FString NewPath = DstFolderPath / LeafName;
+
+					if (FActorFolders::Get().ContainsFolder(*RepresentingWorld, FFolder(FName(*NewPath), DestinationRootObject)))
+					{
+						// The folder already exists
+						FFormatNamedArguments Args;
+						Args.Add(TEXT("DragName"), FText::FromString(LeafName));
+						return FSceneOutlinerDragValidationInfo(ESceneOutlinerDropCompatibility::IncompatibleGeneric,
+							FText::Format(LOCTEXT("FolderAlreadyExistsRoot", "A folder called \"{DragName}\" already exists at this level"), Args));
+					}
+					else if (DragFolderPath == DstFolderPath || DstFolderPath.StartsWith(DragFolderPath + "/"))
+					{
+						// Cannot drag as a child of itself
+						FFormatNamedArguments Args;
+						Args.Add(TEXT("FolderPath"), FText::FromName(DraggedFolder));
+						return FSceneOutlinerDragValidationInfo(
+							ESceneOutlinerDropCompatibility::IncompatibleGeneric,
+							FText::Format(LOCTEXT("ChildOfItself", "Cannot move \"{FolderPath}\" to be a child of itself"), Args));
+					}
 				}
 			}
 		}
@@ -1115,7 +1185,14 @@ FSceneOutlinerDragValidationInfo FActorBrowsingMode::ValidateDrop(const ISceneOu
 				{
 					return FSceneOutlinerDragValidationInfo(ESceneOutlinerDropCompatibility::IncompatibleGeneric, FText::Format(LOCTEXT("Error_AttachChildActor", "Cannot move {0} as it is a child actor."), FText::FromString(Actor->GetActorLabel())));
 				}
-				else if (Actor->GetFolderPath() == DestinationPath && !Actor->GetSceneOutlinerParent() && !bActorContainedInLevelInstance)
+				else if ((Actor->GetFolderRootObject() != DestinationRootObject) && FFolder::HasRootObject(Actor->GetFolderRootObject()) && FFolder::HasRootObject(DestinationRootObject))
+				{
+					FFormatNamedArguments Args;
+					Args.Add(TEXT("SourceName"), FText::FromString(Actor->GetActorLabel()));
+					FText Text = FText::Format(LOCTEXT("CantChangeActorRoot", "Cannot change {SourceName} folder root"), Args);
+					return FSceneOutlinerDragValidationInfo(ESceneOutlinerDropCompatibility::IncompatibleGeneric, Text);
+				}
+				else if (Actor->GetFolderPath() == DestinationPath.GetPath() && !Actor->GetSceneOutlinerParent() && !bActorContainedInLevelInstance)
 				{
 					FFormatNamedArguments Args;
 					Args.Add(TEXT("SourceName"), FText::FromString(Actor->GetActorLabel()));
@@ -1127,7 +1204,7 @@ FSceneOutlinerDragValidationInfo FActorBrowsingMode::ValidateDrop(const ISceneOu
 					}
 					else
 					{
-						Args.Add(TEXT("DestPath"), FText::FromName(DestinationPath));
+						Args.Add(TEXT("DestPath"), FText::FromName(DestinationPath.GetPath()));
 						Text = FText::Format(LOCTEXT("FolderAlreadyAssigned", "{SourceName} is already assigned to {DestPath}"), Args);
 					}
 
@@ -1144,7 +1221,7 @@ FSceneOutlinerDragValidationInfo FActorBrowsingMode::ValidateDrop(const ISceneOu
 		else
 		{
 			FFormatNamedArguments Args;
-			Args.Add(TEXT("DestPath"), FText::FromName(DestinationPath));
+			Args.Add(TEXT("DestPath"), FText::FromName(DestinationPath.GetPath()));
 			return FSceneOutlinerDragValidationInfo(ESceneOutlinerDropCompatibility::CompatibleGeneric, FText::Format(LOCTEXT("MoveInto", "Move into \"{DestPath}\""), Args));
 		}
 	}
@@ -1198,7 +1275,10 @@ void FActorBrowsingMode::OnDrop(ISceneOutlinerTreeItem& DropTarget, const FScene
 
 			if (ALevelInstance* TargetLevelInstance = Cast<ALevelInstance>(DropActor))
 			{
-				// Actors inside LevelInstances cannot have folder paths
+				check(TargetLevelInstance->IsEditing());
+				const FScopedTransaction Transaction(LOCTEXT("UndoAction_MoveActorsToLevelInstance", "Move actors to LevelInstance"));
+
+				// Since target root is directly the Level Instance, clear folder path
 				TArray<AActor*> DraggedActors = Payload.GetData<AActor*>(SceneOutliner::FActorSelector());
 				for (auto& Actor : DraggedActors)
 				{
@@ -1207,10 +1287,6 @@ void FActorBrowsingMode::OnDrop(ISceneOutlinerTreeItem& DropTarget, const FScene
 
 				ULevelInstanceSubsystem* LevelInstanceSubsystem = RepresentingWorld->GetSubsystem<ULevelInstanceSubsystem>();
 				check(LevelInstanceSubsystem);
-
-				check(TargetLevelInstance->IsEditing());
-				const FScopedTransaction Transaction(LOCTEXT("UndoAction_MoveActorsToLevelInstance", "Move actors to LevelInstance"));
-
 				LevelInstanceSubsystem->MoveActorsTo(TargetLevelInstance, DraggedActors);
 			}
 			else
@@ -1268,8 +1344,9 @@ void FActorBrowsingMode::OnDrop(ISceneOutlinerTreeItem& DropTarget, const FScene
 	else if (DropTarget.IsA<FFolderTreeItem>() || DropTarget.IsA<FWorldTreeItem>())
 	{
 		const FFolderTreeItem* FolderItem = DropTarget.CastTo<FFolderTreeItem>();
+		const FWorldTreeItem* WorldItem = DropTarget.CastTo<FWorldTreeItem>();
 		// If the cast fails, the item must be a WorldTreeItem and we set the path to None to reflect this
-		const FName& DestinationPath = FolderItem ? FolderItem->Path : FName(NAME_None);
+		const FFolder DestinationPath = FolderItem ? FolderItem->GetFolder() : GWorldRoot;
 
 		const FScopedTransaction Transaction(LOCTEXT("MoveOutlinerItems", "Move World Outliner Items"));
 
@@ -1285,22 +1362,33 @@ void FActorBrowsingMode::OnDrop(ISceneOutlinerTreeItem& DropTarget, const FScene
 			TSet<const AActor*> ParentActors;
 			TSet<const AActor*> ChildActors;
 
-			Payload.ForEachItem<FActorTreeItem>([&DestinationPath, &ParentActors, &ChildActors](const FActorTreeItem& ActorItem)
+			TArray<AActor*> MovingActorsToValidRootObject;
+			Payload.ForEachItem<FActorTreeItem>([&DestinationPath, &ParentActors, &ChildActors, &MovingActorsToValidRootObject](const FActorTreeItem& ActorItem)
+			{
+				AActor* Actor = ActorItem.Actor.Get();
+				if (Actor)
 				{
-					AActor* Actor = ActorItem.Actor.Get();
-					if (Actor)
+					// First mark this object as a parent, then set its children's path
+					ParentActors.Add(Actor);
+
+					const FFolder ActorFolder = Actor->GetFolder();
+
+					// If the folder root object changes, 1st pass will put actors at root. 2nd pass will set the destination path.
+					FName NewPath = (ActorFolder.GetRootObject() == DestinationPath.GetRootObject()) ? DestinationPath.GetPath() : NAME_None;
+					
+					Actor->SetFolderPath(NewPath);
+					FActorEditorUtils::TraverseActorTree_ParentFirst(Actor, [&](AActor* InActor) {
+						ChildActors.Add(InActor);
+						InActor->SetFolderPath(NewPath);
+						return true;
+						}, false);
+					
+					if ((Actor->GetFolderRootObject() != DestinationPath.GetRootObject()) && !ActorFolder.HasRootObject() && DestinationPath.HasRootObject())
 					{
-						// First mark this object as a parent, then set its children's path
-						ParentActors.Add(Actor);
-						Actor->SetFolderPath(DestinationPath);
-
-						FActorEditorUtils::TraverseActorTree_ParentFirst(Actor, [&](AActor* InActor) {
-							ChildActors.Add(InActor);
-							InActor->SetFolderPath(DestinationPath);
-							return true;
-							}, false);
+						MovingActorsToValidRootObject.Add(Actor);
 					}
-				});
+				}
+			});
 
 			// Detach parent actors
 			for (const AActor* Parent : ParentActors)
@@ -1319,14 +1407,15 @@ void FActorBrowsingMode::OnDrop(ISceneOutlinerTreeItem& DropTarget, const FScene
 				}
 			}
 
-			const ULevelInstanceSubsystem* LevelInstanceSubsystem = RepresentingWorld->GetSubsystem<ULevelInstanceSubsystem>();
-			check(LevelInstanceSubsystem);
-			// Since we are moving to a folder (or root), we must be moving into the persistent level.
-			ULevel* DestinationLevel = RepresentingWorld->PersistentLevel;
-			check(DestinationLevel);
+			if (!DestinationPath.HasRootObject())
+			{
+				const ULevelInstanceSubsystem* LevelInstanceSubsystem = RepresentingWorld->GetSubsystem<ULevelInstanceSubsystem>();
+				check(LevelInstanceSubsystem);
+				ULevel* DestinationLevel = RepresentingWorld->PersistentLevel;
+				check(DestinationLevel);
 
-			TArray<AActor*> ActorsToMove;
-			Payload.ForEachItem<FActorTreeItem>([LevelInstanceSubsystem, &ActorsToMove](const FActorTreeItem& ActorItem)
+				TArray<AActor*> ActorsToMove;
+				Payload.ForEachItem<FActorTreeItem>([LevelInstanceSubsystem, &ActorsToMove](const FActorTreeItem& ActorItem)
 				{
 					AActor* Actor = ActorItem.Actor.Get();
 					if (const ALevelInstance* ParentLevelInstance = LevelInstanceSubsystem->GetParentLevelInstance(Actor))
@@ -1336,48 +1425,68 @@ void FActorBrowsingMode::OnDrop(ISceneOutlinerTreeItem& DropTarget, const FScene
 					}
 				});
 
-			TArray<AActor*> DraggedActors = Payload.GetData<AActor*>(SceneOutliner::FActorSelector());
-			LevelInstanceSubsystem->MoveActorsToLevel(ActorsToMove, DestinationLevel);
-		}
-	}
-}
+				// We are moving actors outside of an editing level instance to a folder (or root) into the persistent level.
+				if (ActorsToMove.Num() > 0)
+				{
+					TArray<AActor*> MovedActors;
+					LevelInstanceSubsystem->MoveActorsToLevel(ActorsToMove, DestinationLevel, &MovedActors);
+					// Once moved, update actors folder path
+					for (AActor* Actor : MovedActors)
+					{
+						Actor->SetFolderPath_Recursively(DestinationPath.GetPath());
+					}
+				}
+			}
+			else if (MovingActorsToValidRootObject.Num())
+			{
+				if (ALevelInstance* TargetLevelInstance = Cast<ALevelInstance>(DestinationPath.GetRootObjectPtr()))
+				{
+					// We are moving actors inside an editing level instance
+					check(TargetLevelInstance->IsEditing());
 
-FName FActorBrowsingMode::CreateNewFolder()
-{
-	const FScopedTransaction Transaction(LOCTEXT("UndoAction_CreateFolder", "Create Folder"));
+					ULevelInstanceSubsystem* LevelInstanceSubsystem = RepresentingWorld->GetSubsystem<ULevelInstanceSubsystem>();
+					check(LevelInstanceSubsystem);
+					TArray<AActor*> MovedActors;
+					LevelInstanceSubsystem->MoveActorsTo(TargetLevelInstance, MovingActorsToValidRootObject, &MovedActors);
 
-	if (ULevelInstanceSubsystem* LevelInstanceSubsystem = RepresentingWorld->GetSubsystem<ULevelInstanceSubsystem>())
-	{
-		for (FSelectionIterator It(*GEditor->GetSelectedActors()); It; ++It)
-		{
-			const AActor* Actor = CastChecked<AActor>(*It);
-			if (LevelInstanceSubsystem->GetParentLevelInstance(Actor) != nullptr)
-			{
-				UE_LOG(LogActorBrowser, Warning, TEXT("Cannot create a folder with actors who are children of a Level Instance"));
-				return NAME_None;
+					// Once moved, update actors folder path
+					for (AActor* Actor : MovedActors)
+					{
+						Actor->SetFolderPath_Recursively(DestinationPath.GetPath());
+					}
+				}
 			}
 		}
 	}
+}
 
-	const FName NewFolderName = FActorFolders::Get().GetDefaultFolderNameForSelection(*RepresentingWorld);
+FFolder FActorBrowsingMode::CreateNewFolder()
+{
+	const FScopedTransaction Transaction(LOCTEXT("UndoAction_CreateFolder", "Create Folder"));
+	TArray<FFolder> SelectedFolders = SceneOutliner->GetSelection().GetData<FFolder>(SceneOutliner::FFolderPathSelector());
+	const FFolder NewFolderName = FActorFolders::Get().GetDefaultFolderForSelection(*RepresentingWorld, &SelectedFolders);
 	FActorFolders::Get().CreateFolderContainingSelection(*RepresentingWorld, NewFolderName);
 
 	return NewFolderName;
 }
 
-FName FActorBrowsingMode::CreateFolder(const FName& ParentPath, const FName& LeafName)
+FFolder FActorBrowsingMode::CreateFolder(const FFolder& ParentPath, const FName& LeafName)
 {
-	const FName NewPath = FActorFolders::Get().GetFolderName(*RepresentingWorld, ParentPath, LeafName);
+	const FFolder NewPath = FActorFolders::Get().GetFolderName(*RepresentingWorld, ParentPath, LeafName);
 	FActorFolders::Get().CreateFolder(*RepresentingWorld, NewPath);
 	return NewPath;
 }
 
-bool FActorBrowsingMode::ReparentItemToFolder(const FName& FolderPath, const FSceneOutlinerTreeItemPtr& Item)
+bool FActorBrowsingMode::ReparentItemToFolder(const FFolder& FolderPath, const FSceneOutlinerTreeItemPtr& Item)
 {
 	if (FActorTreeItem* ActorItem = Item->CastTo<FActorTreeItem>())
 	{
-		ActorItem->Actor->SetFolderPath_Recursively(FolderPath);
-		return true;
+		// Make sure actor has the same root object before updating path
+		if (ActorItem->Actor->GetFolderRootObject() == FolderPath.GetRootObject())
+		{
+			ActorItem->Actor->SetFolderPath_Recursively(FolderPath.GetPath());
+			return true;
+		}
 	}
 	return false;
 }
@@ -1595,24 +1704,31 @@ void FActorBrowsingMode::UnpinSelectedItems()
 	}
 }
 
-FCreateSceneOutlinerMode FActorBrowsingMode::CreateFolderPickerMode() const
+FCreateSceneOutlinerMode FActorBrowsingMode::CreateFolderPickerMode(const FFolder::FRootObject& InRootObject) const
 {
-	auto MoveSelectionTo = [this](const FSceneOutlinerTreeItemRef& NewParent)
+	auto MoveSelectionTo = [this, InRootObject](const FSceneOutlinerTreeItemRef& NewParent)
 	{
-		if (NewParent->IsA<FWorldTreeItem>())
+		if (FWorldTreeItem* WorldItem = NewParent->CastTo<FWorldTreeItem>())
 		{
-			SceneOutliner->MoveSelectionTo(FName());
+			SceneOutliner->MoveSelectionTo(GWorldRoot);
 		}
 		else if (FFolderTreeItem* FolderItem = NewParent->CastTo<FFolderTreeItem>())
 		{
-			SceneOutliner->MoveSelectionTo(FolderItem->Path);
+			SceneOutliner->MoveSelectionTo(FolderItem->GetFolder());
+		}
+		else if (FActorTreeItem* ActorItem = NewParent->CastTo<FActorTreeItem>())
+		{
+			if (FFolder::HasRootObject(InRootObject))
+			{
+				SceneOutliner->MoveSelectionTo(FFolder(NAME_None, InRootObject));
+			}
 		}
 	};
 
-	return FCreateSceneOutlinerMode::CreateLambda([this, MoveSelectionTo](SSceneOutliner* Outliner)
-		{
-			return new FActorFolderPickingMode(Outliner, FOnSceneOutlinerItemPicked::CreateLambda(MoveSelectionTo));
-		});
+	return FCreateSceneOutlinerMode::CreateLambda([this, MoveSelectionTo, InRootObject](SSceneOutliner* Outliner)
+	{
+		return new FActorFolderPickingMode(Outliner, FOnSceneOutlinerItemPicked::CreateLambda(MoveSelectionTo), nullptr, InRootObject);
+	});
 }
 
 void FActorBrowsingMode::OnDuplicateSelected()
@@ -1662,7 +1778,9 @@ void FActorBrowsingMode::OnEditPasteActorsEnd()
 void FActorBrowsingMode::OnDuplicateActorsBegin()
 {
 	// Only a callback in actor browsing mode
-	const TArray<FName> SelectedFolderPaths = SceneOutliner->GetSelection().GetData<FName>(SceneOutliner::FFolderPathSelector());
+	FFolder::FRootObject CommonRootObject;
+	TArray<FName> SelectedFolderPaths;
+	GetFolderNamesFromFolders(SceneOutliner->GetSelection().GetData<FFolder>(SceneOutliner::FFolderPathSelector()), SelectedFolderPaths, CommonRootObject);
 	SceneOutliner->PasteFoldersBegin(SelectedFolderPaths);
 }
 
diff --git a/Engine/Source/Editor/SceneOutliner/Private/ActorFolderHierarchy.cpp b/Engine/Source/Editor/SceneOutliner/Private/ActorFolderHierarchy.cpp
--- a/Engine/Source/Editor/SceneOutliner/Private/ActorFolderHierarchy.cpp
+++ b/Engine/Source/Editor/SceneOutliner/Private/ActorFolderHierarchy.cpp
@@ -3,13 +3,16 @@
 #include "ActorFolderHierarchy.h"
 #include "ISceneOutlinerMode.h"
 #include "WorldTreeItem.h"
+#include "ActorTreeItem.h"
 #include "ActorFolderTreeItem.h"
 #include "EditorActorFolders.h"
 #include "EditorFolderUtils.h"
+#include "LevelInstance/LevelInstanceActor.h"
 
-FActorFolderHierarchy::FActorFolderHierarchy(ISceneOutlinerMode* InMode, const TWeakObjectPtr<UWorld>& World)
+FActorFolderHierarchy::FActorFolderHierarchy(ISceneOutlinerMode* InMode, const TWeakObjectPtr<UWorld>& World, const FFolder::FRootObject& InRootObject)
 	: ISceneOutlinerHierarchy(InMode)
 	, RepresentingWorld(World)
+	, RootObject(InRootObject)
 {
 	// ActorFolderHierarchy should only be used with a mode which is showing folders
 	check(Mode->ShouldShowFolders());
@@ -23,17 +26,24 @@ FSceneOutlinerTreeItemPtr FActorFolderHierarchy::FindParent(const ISceneOutliner
 	}
 	else if (const FActorFolderTreeItem* ActorFolderItem = Item.CastTo<FActorFolderTreeItem>())
 	{
-		const FName ParentPath = FEditorFolderUtils::GetParentPath(ActorFolderItem->Path);
+		const FFolder ParentPath = ActorFolderItem->GetFolder().GetParent();
 
 		const FSceneOutlinerTreeItemPtr* ParentItem = nullptr;
-		// If the folder has no parent path, it must be parented to the root world
+		// If the folder has no parent path
 		if (ParentPath.IsNone())
 		{
-			ParentItem = Items.Find(ActorFolderItem->World.Get());
+			if (UObject* Object = ParentPath.GetRootObjectPtr())
+			{
+				ParentItem = Items.Find(Object);
+			}
+			else
+			{
+				ParentItem = Items.Find(ActorFolderItem->World.Get());
+			}
 		}
 		else
 		{
-			ParentItem = Items.Find(FEditorFolderUtils::GetParentPath(ActorFolderItem->Path));
+			ParentItem = Items.Find(ParentPath);
 		}
 
 		if (ParentItem)
@@ -46,12 +56,24 @@ FSceneOutlinerTreeItemPtr FActorFolderHierarchy::FindParent(const ISceneOutliner
 
 void FActorFolderHierarchy::CreateWorldChildren(UWorld* World, TArray<FSceneOutlinerTreeItemPtr>& OutItems) const
 {
-	for (const auto& Pair : FActorFolders::Get().GetFolderPropertiesForWorld(*World))
+	FActorFolders::Get().ForEachFolderWithRootObject(*World, RootObject, [this, World, &OutItems](const FFolder& Folder)
 	{
-		if (FSceneOutlinerTreeItemPtr FolderItem = Mode->CreateItemFor<FActorFolderTreeItem>(FActorFolderTreeItem(Pair.Key, World)))
+		if (FSceneOutlinerTreeItemPtr FolderItem = Mode->CreateItemFor<FActorFolderTreeItem>(FActorFolderTreeItem(Folder, World)))
 		{
 			OutItems.Add(FolderItem);
 		}
+		return true;
+	});
+
+	if (FFolder::HasRootObject(RootObject))
+	{
+		if (ALevelInstance* RootLevelInstance = Cast<ALevelInstance>(FFolder(NAME_None, RootObject).GetRootObjectPtr()))
+		{
+			if (FSceneOutlinerTreeItemPtr ActorItem = Mode->CreateItemFor<FActorTreeItem>(RootLevelInstance, true))
+			{
+				OutItems.Add(ActorItem);
+			}
+		}
 	}
 }
 
@@ -59,9 +81,12 @@ void FActorFolderHierarchy::CreateItems(TArray<FSceneOutlinerTreeItemPtr>& OutIt
 {
 	check(RepresentingWorld.IsValid());
 
-	if (FSceneOutlinerTreeItemPtr WorldItem = Mode->CreateItemFor<FWorldTreeItem>(RepresentingWorld.Get()))
+	if (!FFolder::HasRootObject(RootObject))
 	{
-		OutItems.Add(WorldItem);
+		if (FSceneOutlinerTreeItemPtr WorldItem = Mode->CreateItemFor<FWorldTreeItem>(RepresentingWorld.Get()))
+		{
+			OutItems.Add(WorldItem);
+		}
 	}
 
 	CreateWorldChildren(RepresentingWorld.Get(), OutItems);
@@ -69,26 +94,44 @@ void FActorFolderHierarchy::CreateItems(TArray<FSceneOutlinerTreeItemPtr>& OutIt
 
 void FActorFolderHierarchy::CreateChildren(const FSceneOutlinerTreeItemPtr& Item, TArray<FSceneOutlinerTreeItemPtr>& OutChildren) const
 {
-	if (const FWorldTreeItem* WorldItem = Item->CastTo<FWorldTreeItem>())
+	auto CreateChildrenFolders = [this](UWorld* InWorld, const FFolder& InParentFolder, const FFolder::FRootObject& InFolderRootObject, TArray<FSceneOutlinerTreeItemPtr>& OutChildren)
 	{
-		check(WorldItem->World == RepresentingWorld);
-		CreateWorldChildren(WorldItem->World.Get(), OutChildren);
-	}
-	else if (FActorFolderTreeItem* FolderItem = Item->CastTo<FActorFolderTreeItem>())
-	{
-		// Since no map of folder->children exist for ActorFolders, must iterate through all
-		// and manually check the path to know if it a child
-		for (const auto& Pair : FActorFolders::Get().GetFolderPropertiesForWorld(*(FolderItem->World)))
+		FActorFolders::Get().ForEachFolderWithRootObject(*InWorld, InFolderRootObject, [this, InWorld, &InParentFolder, &OutChildren](const FFolder& Folder)
 		{
-			if (FEditorFolderUtils::PathIsChildOf(Pair.Key, FolderItem->Path))
+			if (Folder.IsChildOf(InParentFolder))
 			{
-				if (FSceneOutlinerTreeItemPtr NewFolderItem = Mode->CreateItemFor<FActorFolderTreeItem>(FActorFolderTreeItem(Pair.Key, FolderItem->World), true))
+				if (FSceneOutlinerTreeItemPtr NewFolderItem = Mode->CreateItemFor<FActorFolderTreeItem>(FActorFolderTreeItem(Folder, InWorld)))
 				{
 					OutChildren.Add(NewFolderItem);
 				}
 			}
+			return true;
+		});
+	};
+
+	UWorld* World = RepresentingWorld.Get();
+	if (const FWorldTreeItem* WorldItem = Item->CastTo<FWorldTreeItem>())
+	{
+		check(WorldItem->World == RepresentingWorld);
+		CreateWorldChildren(WorldItem->World.Get(), OutChildren);
+	}
+	else if (FActorTreeItem* ParentActorItem = Item->CastTo<FActorTreeItem>())
+	{
+		AActor* ParentActor = ParentActorItem->Actor.Get();
+		if (const ALevelInstance* LevelInstanceParentActor = Cast<ALevelInstance>(ParentActor))
+		{
+			check(ParentActor->GetWorld() == World);
+			FFolder ParentFolder = LevelInstanceParentActor->GetFolder();
+			CreateChildrenFolders(World, ParentFolder, LevelInstanceParentActor, OutChildren);
 		}
 	}
+	else if (FActorFolderTreeItem* FolderItem = Item->CastTo<FActorFolderTreeItem>())
+	{
+		check(FolderItem->World.Get() == World);
+		FFolder ParentFolder = FolderItem->GetFolder();
+		check(!ParentFolder.IsNone());
+		CreateChildrenFolders(World, ParentFolder, ParentFolder.GetRootObject(), OutChildren);
+	}
 }
 
 FSceneOutlinerTreeItemPtr FActorFolderHierarchy::CreateParentItem(const FSceneOutlinerTreeItemPtr& Item) const
@@ -99,16 +142,33 @@ FSceneOutlinerTreeItemPtr FActorFolderHierarchy::CreateParentItem(const FSceneOu
 	}
 	else if (const FActorFolderTreeItem* FolderTreeItem = Item->CastTo<FActorFolderTreeItem>())
 	{
-		const FName ParentPath = FEditorFolderUtils::GetParentPath(FolderTreeItem->Path);
-		if (ParentPath.IsNone())
+		FFolder Folder = FolderTreeItem->GetFolder();
+
+		// Parent Folder
+		const FFolder ParentFolder = Folder.GetParent();
+		if (!ParentFolder.IsNone())
 		{
-			UWorld* OwningWorld = FolderTreeItem->World.Get();
-			check(OwningWorld);
-			return Mode->CreateItemFor<FWorldTreeItem>(OwningWorld, true);
+			return Mode->CreateItemFor<FActorFolderTreeItem>(FActorFolderTreeItem(ParentFolder, FolderTreeItem->World), true);
+		}
+
+		if (FFolder::HasRootObject(RootObject))
+		{
+			// Parent Object
+			if (Folder.GetRootObject() == RootObject)
+			{
+				// If item belongs to a LevelInstance
+				if (ALevelInstance* RootLevelInstance = Cast<ALevelInstance>(Folder.GetRootObjectPtr()))
+				{
+					return Mode->CreateItemFor<FActorTreeItem>(RootLevelInstance, true);
+				}
+			}
 		}
 		else
 		{
-			return Mode->CreateItemFor<FActorFolderTreeItem>(FActorFolderTreeItem(ParentPath, FolderTreeItem->World));
+			// Parent World
+			UWorld* OwningWorld = FolderTreeItem->World.Get();
+			check(OwningWorld);
+			return Mode->CreateItemFor<FWorldTreeItem>(OwningWorld, true);
 		}
 	}
 	return nullptr;
diff --git a/Engine/Source/Editor/SceneOutliner/Private/ActorFolderPickingMode.cpp b/Engine/Source/Editor/SceneOutliner/Private/ActorFolderPickingMode.cpp
--- a/Engine/Source/Editor/SceneOutliner/Private/ActorFolderPickingMode.cpp
+++ b/Engine/Source/Editor/SceneOutliner/Private/ActorFolderPickingMode.cpp
@@ -6,16 +6,17 @@
 
 #define LOCTEXT_NAMESPACE "SceneOutliner_ActoFolderPickingMode"
 
-FActorFolderPickingMode::FActorFolderPickingMode(SSceneOutliner* InSceneOutliner, FOnSceneOutlinerItemPicked InOnItemPicked, TWeakObjectPtr<UWorld> InSpecifiedWorldToDisplay)
+FActorFolderPickingMode::FActorFolderPickingMode(SSceneOutliner* InSceneOutliner, FOnSceneOutlinerItemPicked InOnItemPicked, TWeakObjectPtr<UWorld> InSpecifiedWorldToDisplay, const FFolder::FRootObject& InRootObject)
 	: FActorMode(FActorModeParams(InSceneOutliner, InSpecifiedWorldToDisplay))
 	, OnItemPicked(InOnItemPicked)
+	, RootObject(InRootObject)
 {
 	Rebuild();
 }
 
 TUniquePtr<ISceneOutlinerHierarchy> FActorFolderPickingMode::CreateHierarchy()
 {
-	return MakeUnique<FActorFolderHierarchy>(this, RepresentingWorld);
+	return MakeUnique<FActorFolderHierarchy>(this, RepresentingWorld, RootObject);
 }
 
 void FActorFolderPickingMode::OnItemSelectionChanged(FSceneOutlinerTreeItemPtr Item, ESelectInfo::Type SelectionType, const FSceneOutlinerItemSelection& Selection)
diff --git a/Engine/Source/Editor/SceneOutliner/Private/ActorFolderTreeItem.cpp b/Engine/Source/Editor/SceneOutliner/Private/ActorFolderTreeItem.cpp
--- a/Engine/Source/Editor/SceneOutliner/Private/ActorFolderTreeItem.cpp
+++ b/Engine/Source/Editor/SceneOutliner/Private/ActorFolderTreeItem.cpp
@@ -59,6 +59,7 @@ struct SActorFolderTreeLabel : FSceneOutlinerCommonLabelData, public SCompoundWi
 			[
 				SNew(SImage)
 				.Image(this, &SActorFolderTreeLabel::GetIcon)
+				.ColorAndOpacity(FSlateColor::UseForeground())
 			]
 			]
 
@@ -78,7 +79,7 @@ private:
 	FText GetDisplayText() const
 	{
 		auto Folder = TreeItemPtr.Pin();
-		return Folder.IsValid() ? FText::FromName(Folder->LeafName) : FText();
+		return Folder.IsValid() ? FText::FromString(Folder->GetDisplayString()) : FText();
 	}
 
 	const FSlateBrush* GetIcon() const
@@ -135,7 +136,7 @@ private:
 		}
 
 		const FString LabelString = TrimmedLabel.ToString();
-		if (TreeItem->LeafName.ToString() == LabelString)
+		if (TreeItem->GetLeafName().ToString() == LabelString)
 		{
 			return true;
 		}
@@ -148,7 +149,8 @@ private:
 		}
 
 		// Validate that this folder doesn't exist already
-		FName NewPath = FEditorFolderUtils::GetParentPath(TreeItem->Path);
+		FFolder Folder = TreeItem->GetFolder();
+		FName NewPath = Folder.GetParent().GetPath();
 		if (NewPath.IsNone())
 		{
 			NewPath = FName(*LabelString);
@@ -157,8 +159,9 @@ private:
 		{
 			NewPath = FName(*(NewPath.ToString() / LabelString));
 		}
+		Folder.SetPath(NewPath);
 
-		if (FActorFolders::Get().GetFolderProperties(*GWorld, NewPath))
+		if (FActorFolders::Get().ContainsFolder(*TreeItem->World.Get(), Folder))
 		{
 			OutErrorMessage = LOCTEXT("RenameFailed_AlreadyExists", "A folder with this name already exists at this level");
 			return false;
@@ -170,10 +173,11 @@ private:
 	void OnLabelCommitted(const FText& InLabel, ETextCommit::Type InCommitInfo)
 	{
 		auto TreeItem = TreeItemPtr.Pin();
-		if (TreeItem.IsValid() && !InLabel.ToString().Equals(TreeItem->LeafName.ToString(), ESearchCase::CaseSensitive))
+		if (TreeItem.IsValid() && !InLabel.ToString().Equals(TreeItem->GetLeafName().ToString(), ESearchCase::CaseSensitive))
 		{
 			// Rename the item
-			FName NewPath = FEditorFolderUtils::GetParentPath(TreeItem->Path);
+			FFolder Folder = TreeItem->GetFolder();
+			FName NewPath = Folder.GetParent().GetPath();
 			if (NewPath.IsNone())
 			{
 				NewPath = FName(*InLabel.ToString());
@@ -182,8 +186,9 @@ private:
 			{
 				NewPath = FName(*(NewPath.ToString() / InLabel.ToString()));
 			}
+			FFolder TreeItemNewFolder(NewPath, Folder.GetRootObject());
 
-			FActorFolders::Get().RenameFolderInWorld(*GWorld, TreeItem->Path, NewPath);
+			FActorFolders::Get().RenameFolderInWorld(*TreeItem->World.Get(), Folder, TreeItemNewFolder);
 
 			auto Outliner = WeakSceneOutliner.Pin();
 			if (Outliner.IsValid())
@@ -194,8 +199,8 @@ private:
 	}
 };
 
-FActorFolderTreeItem::FActorFolderTreeItem(FName InPath, TWeakObjectPtr<UWorld> InWorld)
-	: FFolderTreeItem(InPath, Type)
+FActorFolderTreeItem::FActorFolderTreeItem(const FFolder& InFolder, const TWeakObjectPtr<UWorld>& InWorld)
+	: FFolderTreeItem(InFolder, Type)
 	, World(InWorld)
 {
 }
@@ -208,13 +213,10 @@ void FActorFolderTreeItem::OnExpansionChanged()
 	}
 
 	// Update the central store of folder properties with this folder's new expansion state
-	if (FActorFolderProps* Props = FActorFolders::Get().GetFolderProperties(*World, Path))
-	{
-		Props->bIsExpanded = Flags.bIsExpanded;
-	}
+	FActorFolders::Get().SetIsFolderExpanded(*World, GetFolder(), Flags.bIsExpanded);
 }
 
-void FActorFolderTreeItem::Delete(FName InNewParentPath)
+void FActorFolderTreeItem::Delete(const FFolder& InNewParentFolder)
 {
 	if (!World.IsValid())
 	{
@@ -222,6 +224,8 @@ void FActorFolderTreeItem::Delete(FName InNewParentPath)
 	}
 	const FScopedTransaction Transaction(LOCTEXT("DeleteFolderTransaction", "Delete Folder"));
 
+	const FFolder::FRootObject& NewParentRootObject = InNewParentFolder.GetRootObject();
+
 	for (const TWeakPtr<ISceneOutlinerTreeItem>& ChildPtr : GetChildren())
 	{
 		FSceneOutlinerTreeItemPtr Child = ChildPtr.Pin();
@@ -229,31 +233,28 @@ void FActorFolderTreeItem::Delete(FName InNewParentPath)
 		{
 			if (AActor* Actor = ActorItem->Actor.Get())
 			{
-				Actor->SetFolderPath_Recursively(InNewParentPath);
+				check(Actor->GetFolderRootObject() == NewParentRootObject);
+				Actor->SetFolderPath_Recursively(InNewParentFolder.GetPath());
 			}
 		}
 		else if (FFolderTreeItem* FolderItem = Child->CastTo<FFolderTreeItem>())
 		{
-			FolderItem->MoveTo(InNewParentPath);
+			FolderItem->MoveTo(InNewParentFolder);
 		}
 	}
 
-	FActorFolders::Get().DeleteFolder(*World, Path);
+	FActorFolders::Get().DeleteFolder(*World, GetFolder());
 }
 
-FName FActorFolderTreeItem::MoveTo(const FName& NewParent)
+void FActorFolderTreeItem::MoveTo(const FFolder& InNewParentFolder)
 {
 	if (World.IsValid())
 	{
+		check(InNewParentFolder.GetRootObject() == GetRootObject());
 		// Get unique name
-		const FName NewPath = FActorFolders::Get().GetFolderName(*World, NewParent, LeafName);
-
-		if (FActorFolders::Get().RenameFolderInWorld(*World, Path, NewPath))
-		{
-			return NewPath;
-		}
+		const FFolder NewPath = FActorFolders::Get().GetFolderName(*World, InNewParentFolder, GetLeafName());
+		FActorFolders::Get().RenameFolderInWorld(*World, GetFolder(), NewPath);
 	}
-	return FName();
 }
 
 void FActorFolderTreeItem::CreateSubFolder(TWeakPtr<SSceneOutliner> WeakOutliner)
@@ -264,7 +265,7 @@ void FActorFolderTreeItem::CreateSubFolder(TWeakPtr<SSceneOutliner> WeakOutliner
 	{
 		const FScopedTransaction Transaction(LOCTEXT("UndoAction_CreateFolder", "Create Folder"));
 
-		const FName NewFolderName = FActorFolders::Get().GetDefaultFolderName(*World, Path);
+		const FFolder NewFolderName = FActorFolders::Get().GetDefaultFolderName(*World, GetFolder());
 		FActorFolders::Get().CreateFolder(*World, NewFolderName);
 
 		// At this point the new folder will be in our newly added list, so select it and open a rename when it gets refreshed
diff --git a/Engine/Source/Editor/SceneOutliner/Private/ActorHierarchy.cpp b/Engine/Source/Editor/SceneOutliner/Private/ActorHierarchy.cpp
--- a/Engine/Source/Editor/SceneOutliner/Private/ActorHierarchy.cpp
+++ b/Engine/Source/Editor/SceneOutliner/Private/ActorHierarchy.cpp
@@ -56,9 +56,9 @@ TUniquePtr<FActorHierarchy> FActorHierarchy::Create(ISceneOutlinerMode* Mode, co
 	FWorldDelegates::LevelRemovedFromWorld.AddRaw(Hierarchy, &FActorHierarchy::OnLevelRemoved);
 
 	auto& Folders = FActorFolders::Get();
-	Folders.OnFolderCreate.AddRaw(Hierarchy, &FActorHierarchy::OnBroadcastFolderCreate);
-	Folders.OnFolderMove.AddRaw(Hierarchy, &FActorHierarchy::OnBroadcastFolderMove);
-	Folders.OnFolderDelete.AddRaw(Hierarchy, &FActorHierarchy::OnBroadcastFolderDelete);
+	Folders.OnFolderCreated.AddRaw(Hierarchy, &FActorHierarchy::OnBroadcastFolderCreate);
+	Folders.OnFolderMoved.AddRaw(Hierarchy, &FActorHierarchy::OnBroadcastFolderMove);
+	Folders.OnFolderDeleted.AddRaw(Hierarchy, &FActorHierarchy::OnBroadcastFolderDelete);
 
 	return TUniquePtr<FActorHierarchy>(Hierarchy);
 }
@@ -106,9 +106,9 @@ FActorHierarchy::~FActorHierarchy()
 	if (FActorFolders::IsAvailable())
 	{
 		auto& Folders = FActorFolders::Get();
-		Folders.OnFolderCreate.RemoveAll(this);
-		Folders.OnFolderMove.RemoveAll(this);
-		Folders.OnFolderDelete.RemoveAll(this);
+		Folders.OnFolderCreated.RemoveAll(this);
+		Folders.OnFolderMoved.RemoveAll(this);
+		Folders.OnFolderDeleted.RemoveAll(this);
 	}
 }
 
@@ -122,42 +122,25 @@ FSceneOutlinerTreeItemPtr FActorHierarchy::FindParent(const ISceneOutlinerTreeIt
 	{
 		if (AActor* Actor = ActorTreeItem->Actor.Get())
 		{
+			// Parent Actor (Actor attachement / parenting)
 			if (const AActor* ParentActor = Actor->GetSceneOutlinerParent())
 			{
 				if (const FSceneOutlinerTreeItemPtr* ParentItem = Items.Find(ParentActor))
 				{
 					return *ParentItem;
-				} // IF Parent can be listed in SceneOutliner return nullptr so it gets created
+				}
+				// If Parent can be listed in SceneOutliner return nullptr so it gets created
 				else if(ParentActor->IsListedInSceneOutliner())
 				{
 					return nullptr;
 				}
 			}
 
-			if (const ULevelInstanceSubsystem* LevelInstanceSubsystem = RepresentingWorld->GetSubsystem<ULevelInstanceSubsystem>())
+			// Parent Folder
+			FFolder ActorFolder = Actor->GetFolder();
+			if (Mode->ShouldShowFolders() && !ActorFolder.IsNone())
 			{
-				if (const ALevelInstance* OwningLevelInstance = LevelInstanceSubsystem->GetParentLevelInstance(Actor))
-				{
-					const ALevelInstance* LevelInstanceActor = Cast<ALevelInstance>(Actor);
-					const bool bIsAnEditingLevelInstance = LevelInstanceActor ? LevelInstanceActor->IsEditing() : false;
-					// Parent this to a LevelInstance if the parent LevelInstance is being edited or if this is a sub LevelInstance which is being edited
-					if (bShowingLevelInstances || (OwningLevelInstance->IsEditing() || bIsAnEditingLevelInstance))
-					{
-						if (const FSceneOutlinerTreeItemPtr* OwningLevelInstanceItem = Items.Find(OwningLevelInstance))
-						{
-							return *OwningLevelInstanceItem;
-						}
-						else
-						{
-							return nullptr;
-						}
-					}
-				}
-			}
-
-			if (Mode->ShouldShowFolders() && !Actor->GetFolderPath().IsNone())
-			{
-				if (const FSceneOutlinerTreeItemPtr* ParentItem = Items.Find(Actor->GetFolderPath()))
+				if (const FSceneOutlinerTreeItemPtr* ParentItem = Items.Find(ActorFolder))
 				{
 					return *ParentItem;
 				}
@@ -167,7 +150,26 @@ FSceneOutlinerTreeItemPtr FActorHierarchy::FindParent(const ISceneOutlinerTreeIt
 				}
 			}
 
-			// Default to the world
+			// Parent Level Instance
+			if (const ALevelInstance* OwningLevelInstance = Cast<ALevelInstance>(ActorFolder.GetRootObjectPtr()))
+			{
+				const ALevelInstance* LevelInstanceActor = Cast<ALevelInstance>(Actor);
+				const bool bIsAnEditingLevelInstance = LevelInstanceActor ? LevelInstanceActor->IsEditing() : false;
+				// Parent this to a LevelInstance if the parent LevelInstance is being edited or if this is a sub LevelInstance which is being edited
+				if (bShowingLevelInstances || (OwningLevelInstance->IsEditing() || bIsAnEditingLevelInstance))
+				{
+					if (const FSceneOutlinerTreeItemPtr* OwningLevelInstanceItem = Items.Find(OwningLevelInstance))
+					{
+						return *OwningLevelInstanceItem;
+					}
+					else
+					{
+						return nullptr;
+					}
+				}
+			}
+
+			// Parent world
 			if (const FSceneOutlinerTreeItemPtr* ParentItem = Items.Find(RepresentingWorld.Get()))
 			{
 				return *ParentItem;
@@ -179,23 +181,38 @@ FSceneOutlinerTreeItemPtr FActorHierarchy::FindParent(const ISceneOutlinerTreeIt
 		// We should never call FindParents on a folder item if folders are not being shown
 		check(Mode->ShouldShowFolders());
 
-		const FName ParentPath = FEditorFolderUtils::GetParentPath(FolderItem->Path);
+		const FFolder ParentPath = FolderItem->GetFolder().GetParent();
 
-		const FSceneOutlinerTreeItemPtr* ParentItem = nullptr;
-		// If the folder has no parent path, it must be parented to the root world
-		if (ParentPath.IsNone())
+		// Parent Folder
+		if (!ParentPath.IsNone())
 		{
-			ParentItem = Items.Find(RepresentingWorld.Get());
+			if (const FSceneOutlinerTreeItemPtr* ParentItem = Items.Find(ParentPath))
+			{
+				return *ParentItem;
+			}
 		}
-		else
+		// Parent Level Instance
+		else if (ALevelInstance* OwningLevelInstance = Cast<ALevelInstance>(ParentPath.GetRootObjectPtr()))
 		{
-			ParentItem = Items.Find(FEditorFolderUtils::GetParentPath(FolderItem->Path));
+			if (bShowingLevelInstances || OwningLevelInstance->IsEditing())
+			{
+				if (const FSceneOutlinerTreeItemPtr* OwningLevelInstanceItem = Items.Find(OwningLevelInstance))
+				{
+					return *OwningLevelInstanceItem;
+				}
+				else
+				{
+					return nullptr;
+				}
+			}
 		}
-
-		if (ParentItem)
+		// Parent World
+		else if (const FSceneOutlinerTreeItemPtr* WorldItem = Items.Find(RepresentingWorld.Get()))
 		{
-			return *ParentItem;
+			return *WorldItem;
 		}
+
+		return nullptr;
 	}
 	else if (const FComponentTreeItem* ComponentTreeItem = Item.CastTo<FComponentTreeItem>())
 	{
@@ -217,7 +234,7 @@ FSceneOutlinerTreeItemPtr FActorHierarchy::FindParent(const ISceneOutlinerTreeIt
 			const FName FolderPath = ActorDesc->GetFolderPath();
 			if (!FolderPath.IsNone())
 			{
-				if (const FSceneOutlinerTreeItemPtr* UnloadedActorItem = Items.Find(FolderPath))
+				if (const FSceneOutlinerTreeItemPtr* UnloadedActorItem = Items.Find(FFolder(FolderPath)))
 				{
 					return *UnloadedActorItem;
 				}
@@ -258,15 +275,16 @@ void FActorHierarchy::CreateWorldChildren(UWorld* World, TArray<FSceneOutlinerTr
 	if (Mode->ShouldShowFolders())
 	{
 		// Add any folders which might match the current search terms
-		for (const auto& Pair : FActorFolders::Get().GetFolderPropertiesForWorld(*World))
+		FActorFolders::Get().ForEachFolder(*World, [this, &World, &OutItems](const FFolder& Folder)
 		{
-			if (FSceneOutlinerTreeItemPtr FolderItem = Mode->CreateItemFor<FActorFolderTreeItem>(FActorFolderTreeItem(Pair.Key, World)))
+			if (FSceneOutlinerTreeItemPtr FolderItem = Mode->CreateItemFor<FActorFolderTreeItem>(FActorFolderTreeItem(Folder, World)))
 			{
 				OutItems.Add(FolderItem);
 			}
-		}
+			return true;
+		});
 	}
-
+	
 	const ULevelInstanceSubsystem* LevelInstanceSubsystem = World->GetSubsystem<ULevelInstanceSubsystem>();
 	// Create all actor items
 	for (FActorIterator ActorIt(World); ActorIt; ++ActorIt)
@@ -283,7 +301,7 @@ void FActorHierarchy::CreateWorldChildren(UWorld* World, TArray<FSceneOutlinerTr
 				}
 			}
 		}
-
+		
 		if (FSceneOutlinerTreeItemPtr ActorItem = Mode->CreateItemFor<FActorTreeItem>(Actor))
 		{
 			if (bShowingOnlyActorWithValidComponents)
@@ -315,16 +333,16 @@ void FActorHierarchy::CreateWorldChildren(UWorld* World, TArray<FSceneOutlinerTr
 		if (UWorldPartition* WorldPartition = World->GetWorldPartition())
 		{
 			FWorldPartitionHelpers::ForEachActorDesc(WorldPartition, [this, WorldPartition, &OutItems](const FWorldPartitionActorDesc* ActorDesc)
-                {
-                    if (ActorDesc != nullptr && !ActorDesc->IsLoaded(true))
-                    {
-                        if (const FSceneOutlinerTreeItemPtr ActorDescItem = Mode->CreateItemFor<FActorDescTreeItem>(FActorDescTreeItem(ActorDesc->GetGuid(), WorldPartition)))
-                        {
-                            OutItems.Add(ActorDescItem);
-                        }
-                    }
-                    return true;
-                });
+			{
+				if (ActorDesc != nullptr && !ActorDesc->IsLoaded(true))
+				{
+					if (const FSceneOutlinerTreeItemPtr ActorDescItem = Mode->CreateItemFor<FActorDescTreeItem>(FActorDescTreeItem(ActorDesc->GetGuid(), WorldPartition)))
+					{
+						OutItems.Add(ActorDescItem);
+					}
+				}
+				return true;
+			});
 		}
 	}
 }
@@ -346,6 +364,22 @@ void FActorHierarchy::CreateItems(TArray<FSceneOutlinerTreeItemPtr>& OutItems) c
 
 void FActorHierarchy::CreateChildren(const FSceneOutlinerTreeItemPtr& Item, TArray<FSceneOutlinerTreeItemPtr>& OutChildren) const
 {
+	auto CreateChildrenFolders = [this](UWorld* InWorld, const FFolder& InParentFolder, const FFolder::FRootObject& InFolderRootObject, TArray<FSceneOutlinerTreeItemPtr>& OutChildren)
+	{
+		FActorFolders::Get().ForEachFolderWithRootObject(*InWorld, InFolderRootObject, [this, InWorld, &InParentFolder, &OutChildren](const FFolder& Folder)
+		{
+			if (Folder.IsChildOf(InParentFolder))
+			{
+				if (FSceneOutlinerTreeItemPtr NewFolderItem = Mode->CreateItemFor<FActorFolderTreeItem>(FActorFolderTreeItem(Folder, InWorld)))
+				{
+					OutChildren.Add(NewFolderItem);
+				}
+			}
+			return true;
+		});
+	};
+
+	UWorld* World = RepresentingWorld.Get();
 	if (FWorldTreeItem* WorldItem = Item->CastTo<FWorldTreeItem>())
 	{
 		check(WorldItem->World == RepresentingWorld);
@@ -366,15 +400,19 @@ void FActorHierarchy::CreateChildren(const FSceneOutlinerTreeItemPtr& Item, TArr
 			check(LevelInstanceSubsystem);
 
 			LevelInstanceSubsystem->ForEachActorInLevelInstance(LevelInstanceParentActor, [this, LevelInstanceParentActor, LevelInstanceSubsystem, &ChildActors](AActor* SubActor)
+			{
+				const ALevelInstance* LevelInstanceActor = Cast<ALevelInstance>(SubActor);
+				const bool bIsAnEditingLevelInstance = LevelInstanceActor ? LevelInstanceSubsystem->IsEditingLevelInstance(LevelInstanceActor) : false;
+				if (bShowingLevelInstances || (LevelInstanceSubsystem->IsEditingLevelInstance(LevelInstanceParentActor) || bIsAnEditingLevelInstance))
 				{
-					const ALevelInstance* LevelInstanceActor = Cast<ALevelInstance>(SubActor);
-					const bool bIsAnEditingLevelInstance = LevelInstanceActor ? LevelInstanceSubsystem->IsEditingLevelInstance(LevelInstanceActor) : false;
-					if (bShowingLevelInstances || (LevelInstanceSubsystem->IsEditingLevelInstance(LevelInstanceParentActor) || bIsAnEditingLevelInstance))
-					{
-						ChildActors.Add(SubActor);
-					}
-					return true;
-				});
+					ChildActors.Add(SubActor);
+				}
+				return true;
+			});
+
+			check(World == LevelInstanceParentActor->GetWorld());
+			FFolder ParentFolder = LevelInstanceParentActor->GetFolder();
+			CreateChildrenFolders(World, ParentFolder, LevelInstanceParentActor, OutChildren);
 		}
 		else
 		{
@@ -404,17 +442,11 @@ void FActorHierarchy::CreateChildren(const FSceneOutlinerTreeItemPtr& Item, TArr
 	else if (FActorFolderTreeItem* FolderItem = Item->CastTo<FActorFolderTreeItem>())
 	{
 		check(Mode->ShouldShowFolders());
-
-		for (const auto& Pair : FActorFolders::Get().GetFolderPropertiesForWorld(*FolderItem->World))
-		{
-			if (FEditorFolderUtils::PathIsChildOf(Pair.Key, FolderItem->Path))
-			{
-				if (FSceneOutlinerTreeItemPtr NewFolderItem = Mode->CreateItemFor<FActorFolderTreeItem>(FActorFolderTreeItem(Pair.Key, FolderItem->World)))
-				{
-					OutChildren.Add(NewFolderItem);
-				}
-			}
-		}
+		
+		check(World == FolderItem->World.Get());
+		FFolder ParentFolder = FolderItem->GetFolder();
+		check(!ParentFolder.IsNone());
+		CreateChildrenFolders(World, ParentFolder, ParentFolder.GetRootObject(), OutChildren);
 	}
 }
 
@@ -428,6 +460,7 @@ FSceneOutlinerTreeItemPtr FActorHierarchy::CreateParentItem(const FSceneOutliner
 	{
 		if (const AActor* Actor = ActorTreeItem->Actor.Get())
 		{
+			// Parent Actor (Actor attachement / parenting)
 			if (AActor* ParentActor = Actor->GetSceneOutlinerParent())
 			{
 				if (ParentActor->IsListedInSceneOutliner())
@@ -436,10 +469,19 @@ FSceneOutlinerTreeItemPtr FActorHierarchy::CreateParentItem(const FSceneOutliner
 				}
 			}
 
-			// If item belongs to a LevelInstance
-			if (const ULevelInstanceSubsystem* LevelInstanceSubsystem = RepresentingWorld->GetSubsystem<ULevelInstanceSubsystem>())
+			// Parent Folder
+			FFolder ActorFolder = Actor->GetFolder();
+			if (Mode->ShouldShowFolders() && !ActorFolder.IsNone())
 			{
-				if (ALevelInstance* ParentLevelInstance = LevelInstanceSubsystem->GetParentLevelInstance(Actor))
+				return Mode->CreateItemFor<FActorFolderTreeItem>(FActorFolderTreeItem(ActorFolder, ActorTreeItem->Actor->GetWorld()), true);
+			}
+
+			// Parent Object
+			if (ALevelInstance* OwningLevelInstance = Cast<ALevelInstance>(ActorFolder.GetRootObjectPtr()))
+			{
+				const ULevelInstanceSubsystem* LevelInstanceSubsystem = RepresentingWorld->GetSubsystem<ULevelInstanceSubsystem>();
+				ALevelInstance* ParentLevelInstance = LevelInstanceSubsystem ? LevelInstanceSubsystem->GetParentLevelInstance(Actor) : nullptr;
+				check(OwningLevelInstance == ParentLevelInstance);
 				{
 					const ALevelInstance* LevelInstanceActor = Cast<ALevelInstance>(Actor);
 					const bool bIsAnEditingLevelInstance = LevelInstanceActor ? LevelInstanceActor->IsEditing() : false;
@@ -450,13 +492,7 @@ FSceneOutlinerTreeItemPtr FActorHierarchy::CreateParentItem(const FSceneOutliner
 				}
 			}
 
-			// if this item belongs in a folder
-			if (Mode->ShouldShowFolders() && !ActorTreeItem->Actor->GetFolderPath().IsNone())
-			{
-				return Mode->CreateItemFor<FActorFolderTreeItem>(FActorFolderTreeItem(ActorTreeItem->Actor->GetFolderPath(), ActorTreeItem->Actor->GetWorld()), true);
-			}
-
-			// Default to the world
+			// Parent World
 			UWorld* OwningWorld = ActorTreeItem->Actor->GetWorld();
 			check(OwningWorld);
 			return Mode->CreateItemFor<FWorldTreeItem>(OwningWorld, true);
@@ -473,17 +509,28 @@ FSceneOutlinerTreeItemPtr FActorHierarchy::CreateParentItem(const FSceneOutliner
 	{
 		check(Mode->ShouldShowFolders());
 
-		const FName ParentPath = FEditorFolderUtils::GetParentPath(FolderTreeItem->Path);
-		if (ParentPath.IsNone())
+		FFolder Folder = FolderTreeItem->GetFolder();
+
+		// Parent Folder
+		const FFolder ParentFolder = Folder.GetParent();
+		if (!ParentFolder.IsNone())
 		{
-			UWorld* OwningWorld = FolderTreeItem->World.Get();
-			check(OwningWorld);
-			return Mode->CreateItemFor<FWorldTreeItem>(OwningWorld, true);
+			return Mode->CreateItemFor<FActorFolderTreeItem>(FActorFolderTreeItem(ParentFolder, FolderTreeItem->World), true);
 		}
-		else
+
+		// Parent Level Instance
+		if (ALevelInstance* OwningLevelInstance = Cast<ALevelInstance>(Folder.GetRootObjectPtr()))
 		{
-			return Mode->CreateItemFor<FActorFolderTreeItem>(FActorFolderTreeItem(ParentPath, FolderTreeItem->World), true);
+			if (bShowingLevelInstances || OwningLevelInstance->IsEditing())
+			{
+				return Mode->CreateItemFor<FActorTreeItem>(OwningLevelInstance, true);
+			}
 		}
+
+		// Parent World
+		UWorld* OwningWorld = FolderTreeItem->World.Get();
+		check(OwningWorld);
+		return Mode->CreateItemFor<FWorldTreeItem>(OwningWorld, true);
 	}
 	else if (const FActorDescTreeItem* ActorDescItem = Item->CastTo<FActorDescTreeItem>())
 	{
@@ -492,7 +539,7 @@ FSceneOutlinerTreeItemPtr FActorHierarchy::CreateParentItem(const FSceneOutliner
 			const FName ActorDescPath = ActorDesc->GetFolderPath();
 			if (Mode->ShouldShowFolders() && !ActorDescPath.IsNone())
 			{
-				return Mode->CreateItemFor<FActorFolderTreeItem>(FActorFolderTreeItem(ActorDescPath, RepresentingWorld.Get()), true);
+				return Mode->CreateItemFor<FActorFolderTreeItem>(FActorFolderTreeItem(FFolder(ActorDescPath), RepresentingWorld.Get()), true);
 			}
 		}
 	}
@@ -664,39 +711,39 @@ void FActorHierarchy::OnLevelRemoved(ULevel* InLevel, UWorld* InWorld)
 }
 
 /** Called when a folder is to be created */
-void FActorHierarchy::OnBroadcastFolderCreate(UWorld& InWorld, FName NewPath)
+void FActorHierarchy::OnBroadcastFolderCreate(UWorld& InWorld, const FFolder& InNewFolder)
 {
 	if (Mode->ShouldShowFolders() && RepresentingWorld.Get() == &InWorld)
 	{
 		FSceneOutlinerHierarchyChangedData EventData;
 		EventData.Type = FSceneOutlinerHierarchyChangedData::Added;
-		EventData.Items.Add(Mode->CreateItemFor<FActorFolderTreeItem>(FActorFolderTreeItem(NewPath, &InWorld)));
+		EventData.Items.Add(Mode->CreateItemFor<FActorFolderTreeItem>(FActorFolderTreeItem(InNewFolder, &InWorld)));
 		EventData.ItemActions = SceneOutliner::ENewItemAction::Select | SceneOutliner::ENewItemAction::Rename;
 		HierarchyChangedEvent.Broadcast(EventData);
 	}
 }
 
 /** Called when a folder is to be moved */
-void FActorHierarchy::OnBroadcastFolderMove(UWorld& InWorld, FName OldPath, FName NewPath)
+void FActorHierarchy::OnBroadcastFolderMove(UWorld& InWorld, const FFolder& InOldFolder, const FFolder& InNewFolder)
 {
 	if (Mode->ShouldShowFolders() && RepresentingWorld.Get() == &InWorld)
 	{
 		FSceneOutlinerHierarchyChangedData EventData;
 		EventData.Type = FSceneOutlinerHierarchyChangedData::FolderMoved;
-		EventData.ItemIDs.Add(OldPath);
-		EventData.NewPaths.Add(NewPath);
+		EventData.ItemIDs.Add(InOldFolder);
+		EventData.NewPaths.Add(InNewFolder);
 		HierarchyChangedEvent.Broadcast(EventData);
 	}
 }
 
 /** Called when a folder is to be deleted */
-void FActorHierarchy::OnBroadcastFolderDelete(UWorld& InWorld, FName Path)
+void FActorHierarchy::OnBroadcastFolderDelete(UWorld& InWorld, const FFolder& InFolder)
 {
 	if (Mode->ShouldShowFolders() && RepresentingWorld.Get() == &InWorld)
 	{
 		FSceneOutlinerHierarchyChangedData EventData;
 		EventData.Type = FSceneOutlinerHierarchyChangedData::Removed;
-		EventData.ItemIDs.Add(Path);
+		EventData.ItemIDs.Add(InFolder);
 		HierarchyChangedEvent.Broadcast(EventData);
 	}
 }
@@ -710,4 +757,4 @@ void FActorHierarchy::OnLevelActorFolderChanged(const AActor* InActor, FName Old
 		EventData.ItemIDs.Add(FSceneOutlinerTreeItemID(InActor));
 		HierarchyChangedEvent.Broadcast(EventData);
 	}
-}
+}
\ No newline at end of file
diff --git a/Engine/Source/Editor/SceneOutliner/Private/ActorMode.cpp b/Engine/Source/Editor/SceneOutliner/Private/ActorMode.cpp
--- a/Engine/Source/Editor/SceneOutliner/Private/ActorMode.cpp
+++ b/Engine/Source/Editor/SceneOutliner/Private/ActorMode.cpp
@@ -15,11 +15,14 @@
 #include "ComponentTreeItem.h"
 #include "ActorDescTreeItem.h"
 #include "WorldTreeItem.h"
+#include "LevelInstance/LevelInstanceActor.h"
+#include "LevelInstance/LevelInstanceSubsystem.h"
 #include "LevelInstance/LevelInstanceEditorInstanceActor.h"
 
 #define LOCTEXT_NAMESPACE "SceneOutliner_ActorMode"
 
 using FActorFilter = TSceneOutlinerPredicateFilter<FActorTreeItem>;
+using FFolderFilter = TSceneOutlinerPredicateFilter<FFolderTreeItem>;
 
 namespace SceneOutliner
 {
@@ -86,6 +89,23 @@ FActorMode::FActorMode(const FActorModeParams& Params)
 		{
 			return IsActorDisplayable(Actor);
 		}), FSceneOutlinerFilter::EDefaultBehaviour::Pass));
+
+	SceneOutliner->AddFilter(MakeShared<FFolderFilter>(FFolderTreeItem::FFilterPredicate::CreateLambda([this](const FFolder& Folder)
+		{
+			if (!Folder.HasRootObject())
+			{
+				return true;
+			}
+			if (ALevelInstance* LevelInstance = Cast<ALevelInstance>(Folder.GetRootObjectPtr()))
+			{
+				const ULevelInstanceSubsystem* LevelInstanceSubsystem = RepresentingWorld.IsValid() ? RepresentingWorld->GetSubsystem<ULevelInstanceSubsystem>() : nullptr;
+				if (LevelInstanceSubsystem && LevelInstanceSubsystem->IsEditingLevelInstance(LevelInstance))
+				{
+					return true;
+				}
+			}
+			return false;
+		}), FSceneOutlinerFilter::EDefaultBehaviour::Pass));
 }
 
 FActorMode::~FActorMode()
diff --git a/Engine/Source/Editor/SceneOutliner/Private/ActorTreeItem.cpp b/Engine/Source/Editor/SceneOutliner/Private/ActorTreeItem.cpp
--- a/Engine/Source/Editor/SceneOutliner/Private/ActorTreeItem.cpp
+++ b/Engine/Source/Editor/SceneOutliner/Private/ActorTreeItem.cpp
@@ -19,6 +19,7 @@
 #include "SSocketChooser.h"
 #include "LevelInstance/LevelInstanceActor.h"
 #include "WorldPartition/WorldPartition.h"
+#include "ToolMenu.h"
 
 #define LOCTEXT_NAMESPACE "SceneOutliner_ActorTreeItem"
 
@@ -252,20 +253,6 @@ private:
 	{
 		AActor* Actor = ActorPtr.Get();
 
-		// Color LevelInstances differently if they are being edited
-		if (const ALevelInstance* LevelInstanceActor = Cast<ALevelInstance>(Actor))
-		{
-			if (LevelInstanceActor->IsEditing() && !LevelInstanceActor->IsSelected())
-			{
-				return FAppStyle::Get().GetSlateColor("Colors.AccentGreen");
-			}
-			else
-			{
-				return FSlateColor::UseForeground();
-			}
-		}
-
-
 		auto TreeItem = TreeItemPtr.Pin();
 		if (auto BaseColor = FSceneOutlinerCommonLabelData::GetForegroundColor(*TreeItem))
 		{
@@ -300,6 +287,15 @@ private:
 			return FSceneOutlinerCommonLabelData::DarkColor;
 		}
 
+		// Color LevelInstances differently if they are being edited
+		if (const ALevelInstance* LevelInstanceActor = Cast<ALevelInstance>(Actor))
+		{
+			if (LevelInstanceActor->IsEditing() && !LevelInstanceActor->IsSelected())
+			{
+				return FAppStyle::Get().GetSlateColor("Colors.AccentGreen");
+			}
+		}
+
 		return FSlateColor::UseForeground();
 	}
 
@@ -353,6 +349,12 @@ FSceneOutlinerTreeItemID FActorTreeItem::GetID() const
 	return ID;
 }
 
+FFolder::FRootObject FActorTreeItem::GetRootObject() const
+{
+	AActor* ActorPtr = Actor.Get();
+	return ActorPtr ? ActorPtr->GetFolderRootObject() : nullptr;
+}
+
 FString FActorTreeItem::GetDisplayString() const
 {
 	return ActorLabel;
@@ -420,4 +422,17 @@ void FActorTreeItem::OnLabelChanged()
 	}
 }
 
+void FActorTreeItem::GenerateContextMenu(UToolMenu* Menu, SSceneOutliner& Outliner)
+{
+	const AActor* ActorPtr = Actor.Get();
+	const ALevelInstance* LevelInstanceActor = Cast<ALevelInstance>(ActorPtr);
+	if (LevelInstanceActor && LevelInstanceActor->IsEditing())
+	{
+		auto SharedOutliner = StaticCastSharedRef<SSceneOutliner>(Outliner.AsShared());
+		const FSlateIcon NewFolderIcon(FEditorStyle::GetStyleSetName(), "SceneOutliner.NewFolderIcon");
+		FToolMenuSection& Section = Menu->AddSection("Section");
+		Section.AddMenuEntry("CreateFolder", LOCTEXT("CreateFolder", "Create Folder"), FText(), NewFolderIcon, FUIAction(FExecuteAction::CreateSP(&Outliner, &SSceneOutliner::CreateFolder)));
+	}
+}
+
 #undef LOCTEXT_NAMESPACE
diff --git a/Engine/Source/Editor/SceneOutliner/Private/FolderTreeItem.cpp b/Engine/Source/Editor/SceneOutliner/Private/FolderTreeItem.cpp
--- a/Engine/Source/Editor/SceneOutliner/Private/FolderTreeItem.cpp
+++ b/Engine/Source/Editor/SceneOutliner/Private/FolderTreeItem.cpp
@@ -17,13 +17,13 @@
 
 namespace SceneOutliner
 {
-	bool FFolderPathSelector::operator()(TWeakPtr<ISceneOutlinerTreeItem> Item, FName& DataOut) const
+	bool FFolderPathSelector::operator()(TWeakPtr<ISceneOutlinerTreeItem> Item, FFolder& DataOut) const
 	{
 		if (FFolderTreeItem* FolderItem = Item.Pin()->CastTo<FFolderTreeItem>())
 		{
 			if (FolderItem->IsValid())
 			{
-				DataOut = FolderItem->Path;
+				DataOut = FolderItem->GetFolder();
 				return true;
 			}
 		}
@@ -33,23 +33,48 @@ namespace SceneOutliner
 
 const FSceneOutlinerTreeItemType FFolderTreeItem::Type(&ISceneOutlinerTreeItem::Type);
 
+FFolderTreeItem::FFolderTreeItem(const FFolder& InFolder, FSceneOutlinerTreeItemType InType)
+	: ISceneOutlinerTreeItem(InType)
+	, RootObject(InFolder.GetRootObject())
+{
+	SetPath(InFolder.GetPath());
+}
+
+FFolderTreeItem::FFolderTreeItem(const FFolder& InFolder)
+	: ISceneOutlinerTreeItem(Type)
+	, RootObject(InFolder.GetRootObject())
+{
+	SetPath(InFolder.GetPath());
+}
+
 FFolderTreeItem::FFolderTreeItem(FName InPath)
 	: ISceneOutlinerTreeItem(Type)
 	, Path(InPath)
-	, LeafName(FEditorFolderUtils::GetLeafName(InPath))
 {
+	SetPath(InPath);
 }
 
 FFolderTreeItem::FFolderTreeItem(FName InPath, FSceneOutlinerTreeItemType InType)
 	: ISceneOutlinerTreeItem(InType)
-	, Path(InPath)
-	, LeafName(FEditorFolderUtils::GetLeafName(InPath))
 {
+	SetPath(InPath);
+}
+
+void FFolderTreeItem::SetPath(const FName& InNewPath)
+{
+	Path = InNewPath;
+	LeafName = FEditorFolderUtils::GetLeafName(Path);
 }
 
 FSceneOutlinerTreeItemID FFolderTreeItem::GetID() const
 {
-	return FSceneOutlinerTreeItemID(Path);
+	FFolderKey Key(Path, RootObject);
+	return FSceneOutlinerTreeItemID(Key);
+}
+
+FFolder::FRootObject FFolderTreeItem::GetRootObject() const
+{
+	return RootObject;
 }
 
 FString FFolderTreeItem::GetDisplayString() const
diff --git a/Engine/Source/Editor/SceneOutliner/Private/SOutlinerTreeView.cpp b/Engine/Source/Editor/SceneOutliner/Private/SOutlinerTreeView.cpp
--- a/Engine/Source/Editor/SceneOutliner/Private/SOutlinerTreeView.cpp
+++ b/Engine/Source/Editor/SceneOutliner/Private/SOutlinerTreeView.cpp
@@ -123,8 +123,7 @@ void SSceneOutlinerTreeView::FlashHighlightOnItem( FSceneOutlinerTreeItemPtr Fla
 FReply SSceneOutlinerTreeView::OnDragOver(const FGeometry& MyGeometry, const FDragDropEvent& DragDropEvent)
 {
 	FSceneOutlinerDragValidationInfo ValidationInfo = FSceneOutlinerDragValidationInfo::Invalid();
-
-	FFolderTreeItem DropTarget(NAME_None);
+	FFolderTreeItem DropTarget(FFolder(NAME_None, nullptr));
 
 	auto Reply = HandleDropFromWeak(SceneOutlinerWeak, DragDropEvent, DropTarget, ValidationInfo);
 	if (Reply.IsEventHandled())
@@ -146,7 +145,7 @@ void SSceneOutlinerTreeView::OnDragLeave(const FDragDropEvent& DragDropEvent)
 FReply SSceneOutlinerTreeView::OnDrop(const FGeometry& MyGeometry, const FDragDropEvent& DragDropEvent)
 {
 	FSceneOutlinerDragValidationInfo ValidationInfo = FSceneOutlinerDragValidationInfo::Invalid();
-	FFolderTreeItem DropTarget(NAME_None);
+	FFolderTreeItem DropTarget(FFolder(NAME_None, nullptr));
 
 	return HandleDropFromWeak(SceneOutlinerWeak, DragDropEvent, DropTarget, ValidationInfo, true);
 }
diff --git a/Engine/Source/Editor/SceneOutliner/Private/SSceneOutliner.cpp b/Engine/Source/Editor/SceneOutliner/Private/SSceneOutliner.cpp
--- a/Engine/Source/Editor/SceneOutliner/Private/SSceneOutliner.cpp
+++ b/Engine/Source/Editor/SceneOutliner/Private/SSceneOutliner.cpp
@@ -37,6 +37,7 @@
 #include "FolderTreeItem.h"
 #include "EditorFolderUtils.h"
 #include "SceneOutlinerConfig.h"
+#include "Algo/ForEach.h"
 
 #include "Framework/Notifications/NotificationManager.h"
 #include "Widgets/Notifications/SNotificationList.h"
@@ -556,7 +557,7 @@ void SSceneOutliner::Populate()
 	UE_LOG(LogSceneOutliner, Verbose, TEXT("%d Items Processed"), Index);
 	PendingOperations.RemoveAt(0, Index);
 
-	for (FName Folder : PendingFoldersSelect)
+	for (const FFolder& Folder : PendingFoldersSelect)
 	{
 		if (FSceneOutlinerTreeItemPtr* Item = TreeItemMap.Find(Folder))
 		{
@@ -901,7 +902,10 @@ void SSceneOutliner::GetSelectedFolders(TArray<FFolderTreeItem*>& OutFolders) co
 
 TArray<FName> SSceneOutliner::GetSelectedFolderNames() const
 {
-	return GetSelection().GetData<FName>(SceneOutliner::FFolderPathSelector());
+	TArray<FFolder> Folders = GetSelection().GetData<FFolder>(SceneOutliner::FFolderPathSelector());
+	TArray<FName> FolderPaths;
+	Algo::ForEach(Folders, [&FolderPaths](const FFolder& Folder) { FolderPaths.Add(Folder.GetPath()); });
+	return FolderPaths;
 }
 
 TSharedPtr<SWidget> SSceneOutliner::OnOpenContextMenu()
@@ -1055,7 +1059,7 @@ void SSceneOutliner::RemoveFromSelection(const TArray<FSceneOutlinerTreeItemPtr>
 
 void SSceneOutliner::AddFolderToSelection(const FName& FolderName)
 {
-	if (FSceneOutlinerTreeItemPtr* ItemPtr = TreeItemMap.Find(FolderName))
+	if (FSceneOutlinerTreeItemPtr* ItemPtr = TreeItemMap.Find(FFolder(FolderName)))
 	{
 		OutlinerTreeView->SetItemSelection(*ItemPtr, true);
 	}
@@ -1063,7 +1067,7 @@ void SSceneOutliner::AddFolderToSelection(const FName& FolderName)
 
 void SSceneOutliner::RemoveFolderFromSelection(const FName& FolderName)
 {
-	if (FSceneOutlinerTreeItemPtr* ItemPtr = TreeItemMap.Find(FolderName))
+	if (FSceneOutlinerTreeItemPtr* ItemPtr = TreeItemMap.Find(FFolder(FolderName)))
 	{
 		OutlinerTreeView->SetItemSelection(*ItemPtr, false);
 	}
@@ -1079,6 +1083,12 @@ void SSceneOutliner::ClearSelection()
 
 void SSceneOutliner::FillFoldersSubMenu(UToolMenu* Menu) const
 {
+	FFolder::FRootObject TargetRootObject;
+	if (!GetCommonRootObjectFromSelection(TargetRootObject))
+	{
+		return;
+	}
+
 	FToolMenuSection& Section = Menu->AddSection("Section");
 	Section.AddMenuEntry("CreateNew", LOCTEXT( "CreateNew", "Create New Folder" ), LOCTEXT( "CreateNew_ToolTip", "Move to a new folder" ),
 		FSlateIcon(FEditorStyle::GetStyleSetName(), "SceneOutliner.NewFolderIcon"), FExecuteAction::CreateSP(const_cast<SSceneOutliner*>(this), &SSceneOutliner::CreateFolder));
@@ -1086,10 +1096,10 @@ void SSceneOutliner::FillFoldersSubMenu(UToolMenu* Menu) const
 	AddMoveToFolderOutliner(Menu);
 }
 
-TSharedRef<TSet<FName>> SSceneOutliner::GatherInvalidMoveToDestinations() const
+TSharedRef<TSet<FFolder>> SSceneOutliner::GatherInvalidMoveToDestinations() const
 {
 	// We use a pointer here to save copying the whole array for every invocation of the filter delegate
-	TSharedRef<TSet<FName>> ExcludedParents(new TSet<FName>());
+	TSharedRef<TSet<FFolder>> ExcludedParents(new TSet<FFolder>());
 
 	for (const auto& Item : OutlinerTreeView->GetSelectedItems())
 	{
@@ -1117,14 +1127,14 @@ TSharedRef<TSet<FName>> SSceneOutliner::GatherInvalidMoveToDestinations() const
 			
 			if (!bFolderHasSubFolders)
 			{
-				ExcludedParents->Add(ParentFolderItem->Path);
+				ExcludedParents->Add(ParentFolderItem->GetFolder());
 			}
 		}
 			
 		if (FFolderTreeItem* FolderItem = Item->CastTo<FFolderTreeItem>())
 		{
 			// Cannot move into itself, or any child
-			ExcludedParents->Add(FolderItem->Path);
+			ExcludedParents->Add(FolderItem->GetFolder());
 		}
 	}
 
@@ -1139,6 +1149,12 @@ void SSceneOutliner::AddMoveToFolderOutliner(UToolMenu* Menu) const
 		return;
 	}
 
+	FFolder::FRootObject TargetRootObject;
+	if (!GetCommonRootObjectFromSelection(TargetRootObject))
+	{
+		return;
+	}
+
 	// Add a mini scene outliner for choosing an existing folder
 	FSceneOutlinerInitializationOptions MiniSceneOutlinerInitOptions;
 	MiniSceneOutlinerInitOptions.bShowHeaderRow = false;
@@ -1149,10 +1165,11 @@ void SSceneOutliner::AddMoveToFolderOutliner(UToolMenu* Menu) const
 	if (ExcludedParents->Num())
 	{
 		// Add a filter if necessary
-		auto FilterOutChildFolders = [](FName Path, TSharedRef<TSet<FName>> ExcludedParents){
+		auto FilterOutChildFolders = [](const FFolder& Folder, TSharedRef<TSet<FFolder>> ExcludedParents)
+		{
 			for (const auto& Parent : *ExcludedParents)
 			{
-				if (Path == Parent || FEditorFolderUtils::PathIsChildOf(Path, Parent))
+				if (Folder == Parent || Folder.IsChildOf(Parent))
 				{
 					return false;
 				}
@@ -1186,7 +1203,7 @@ void SSceneOutliner::AddMoveToFolderOutliner(UToolMenu* Menu) const
 
 		// Filter in/out root items according to whether it is valid to move to/from the root
 		FSceneOutlinerDragDropPayload DraggedObjects(OutlinerTreeView->GetSelectedItems());
-		const bool bMoveToRootValid = Mode->ValidateDrop(FFolderTreeItem(NAME_None), DraggedObjects).IsValid();
+		const bool bMoveToRootValid = Mode->ValidateDrop(FFolderTreeItem(FFolder(NAME_None, TargetRootObject)), DraggedObjects).IsValid();
 		if (!bMoveToRootValid)
 		{
 			MiniSceneOutlinerInitOptions.Filters->Add(MakeShared<FFilterRoot>(*this));
@@ -1195,7 +1212,7 @@ void SSceneOutliner::AddMoveToFolderOutliner(UToolMenu* Menu) const
 
 	//Let the mode decide how folder selection is handled
 
-	MiniSceneOutlinerInitOptions.ModeFactory = Mode->CreateFolderPickerMode();
+	MiniSceneOutlinerInitOptions.ModeFactory = Mode->CreateFolderPickerMode(TargetRootObject);
 
 	FSceneOutlinerModule& SceneOutlinerModule = FModuleManager::LoadModuleChecked<FSceneOutlinerModule>("SceneOutliner");
 	TSharedRef< SWidget > MiniSceneOutliner =
@@ -1246,7 +1263,7 @@ void SSceneOutliner::SelectFoldersDescendants(bool bSelectImmediateChildrenOnly)
 	Refresh();
 }
 
-void SSceneOutliner::MoveSelectionTo(FName NewParent)
+void SSceneOutliner::MoveSelectionTo(const FFolder& NewParent)
 {
 	FSlateApplication::Get().DismissAllMenus();
 		
@@ -1277,9 +1294,9 @@ FReply SSceneOutliner::OnCreateFolderClicked()
 
 void SSceneOutliner::CreateFolder()
 {
-	const FName& NewFolderName = Mode->CreateNewFolder();
+	const FFolder& NewFolderName = Mode->CreateNewFolder();
 
-	if (NewFolderName != NAME_None)
+	if (!NewFolderName.IsNone())
 	{
 		// Move any selected folders into the new folder
 		auto PreviouslySelectedItems = OutlinerTreeView->GetSelectedItems();
@@ -1287,7 +1304,12 @@ void SSceneOutliner::CreateFolder()
 		{
 			if (FFolderTreeItem* FolderItem = Item->CastTo<FFolderTreeItem>())
 			{
-				FolderItem->MoveTo(NewFolderName);
+				// New folder root object will be identical if whole selection had a common root object.
+				// If not, new folder won't have a root object (world), so this check is needed to skip folders with a different root object.
+				if (FolderItem->GetRootObject() == NewFolderName.GetRootObject())
+				{
+					FolderItem->MoveTo(NewFolderName);
+				}
 			}
 		}
 
@@ -1339,9 +1361,32 @@ void SSceneOutliner::CopyFoldersToClipboard(const TArray<FName>& InFolders, cons
 	}
 }
 
+bool SSceneOutliner::GetCommonRootObjectFromSelection(FFolder::FRootObject& OutCommonRootObject) const
+{
+	TOptional<FFolder::FRootObject> CommonRootObject;
+
+	for (const TWeakPtr<ISceneOutlinerTreeItem>& Item : GetSelection().SelectedItems)
+	{
+		if (auto TreeItem = Item.Pin())
+		{
+			if (!CommonRootObject.IsSet())
+			{
+				CommonRootObject = TreeItem->GetRootObject();
+			}
+			else if (CommonRootObject.GetValue() != TreeItem->GetRootObject())
+			{
+				OutCommonRootObject = FFolder::GetDefaultRootObject();
+				return false;
+			}
+		}
+	}
+	OutCommonRootObject = CommonRootObject.Get(FFolder::GetDefaultRootObject());
+	return true;
+}
+
 void SSceneOutliner::PasteFoldersBegin(TArray<FName> InFolders)
 {
-	auto CacheExistingChildrenAction = [this](const FSceneOutlinerTreeItemPtr& Item)
+	auto CacheExistingChildrenAction = [this](const FSceneOutlinerTreeItemPtr& Item, const FFolder::FRootObject& InTargetRootObject)
 	{
 		if (FFolderTreeItem* FolderItem = Item->CastTo<FFolderTreeItem>())
 		{
@@ -1354,7 +1399,8 @@ void SSceneOutliner::PasteFoldersBegin(TArray<FName> InFolders)
 				}
 			}
 
-			CachePasteFolderExistingChildrenMap.Add(FolderItem->Path, ExistingChildren);
+			check(FolderItem->GetRootObject() == InTargetRootObject);
+			CachePasteFolderExistingChildrenMap.Add(FolderItem->GetFolder(), ExistingChildren);
 		}
 	};
 
@@ -1368,12 +1414,16 @@ void SSceneOutliner::PasteFoldersBegin(TArray<FName> InFolders)
 	// Sort folder names so parents appear before children
 	CacheFoldersEdit.Sort(FNameLexicalLess());
 
+	// Find common root object from selection
+	FFolder::FRootObject TargetRootObject;
+	GetCommonRootObjectFromSelection(TargetRootObject);
+
 	// Cache existing children
 	for (FName Folder : CacheFoldersEdit)
 	{
-		if (FSceneOutlinerTreeItemPtr* TreeItem = TreeItemMap.Find(Folder))
+		if (FSceneOutlinerTreeItemPtr* TreeItem = TreeItemMap.Find(FFolder(Folder, TargetRootObject)))
 		{
-			CacheExistingChildrenAction(*TreeItem);
+			CacheExistingChildrenAction(*TreeItem, TargetRootObject);
 		}
 	}
 }
@@ -1382,6 +1432,10 @@ void SSceneOutliner::PasteFoldersEnd()
 {
 	const FScopedTransaction Transaction(NSLOCTEXT("UnrealEd", "PasteItems", "Paste Items"));
 
+	// Find common root object from selection
+	FFolder::FRootObject TargetRootObject;
+	GetCommonRootObjectFromSelection(TargetRootObject);
+
 	// Create new folder
 	TMap<FName, FName> FolderMap;
 	for (FName Folder : CacheFoldersEdit)
@@ -1395,8 +1449,8 @@ void SSceneOutliner::PasteFoldersEnd()
 				ParentPath = *NewParentPath;
 			}
 
-			FName NewFolderPath = Mode->CreateFolder(ParentPath, LeafName);
-			FolderMap.Add(Folder, NewFolderPath);
+			FFolder NewFolderPath = Mode->CreateFolder(FFolder(ParentPath, TargetRootObject), LeafName);
+			FolderMap.Add(Folder, NewFolderPath.GetPath());
 		}
 	}
 
@@ -1409,20 +1463,23 @@ void SSceneOutliner::PasteFoldersEnd()
 		// Get the new folder that was created from this name
 		if (const FName* NewFolderName = FolderMap.Find(OldFolderName))
 		{
-			if (FSceneOutlinerTreeItemPtr* OldFolderItem = TreeItemMap.Find(OldFolderName))
+			FFolder NewFolder(*NewFolderName, TargetRootObject);
+			FFolder OldFolder(OldFolderName, TargetRootObject);
+
+			if (FSceneOutlinerTreeItemPtr* OldFolderItem = TreeItemMap.Find(OldFolder))
 			{
 				for (const TWeakPtr<ISceneOutlinerTreeItem>& Child : (*OldFolderItem)->GetChildren())
 				{
 					// If this child did not exist in the folder before the paste operation, it should be moved to the new folder
-					TArray<FSceneOutlinerTreeItemID>* ExistingChildren = CachePasteFolderExistingChildrenMap.Find(OldFolderName);
+					TArray<FSceneOutlinerTreeItemID>* ExistingChildren = CachePasteFolderExistingChildrenMap.Find(OldFolder);
 
 					if (ExistingChildren && !ExistingChildren->Contains(Child.Pin()->GetID()))
 					{
-						Mode->ReparentItemToFolder(*NewFolderName, Child.Pin());
+						Mode->ReparentItemToFolder(NewFolder, Child.Pin());
 					}
 				}
 			}
-			PendingFoldersSelect.Add(*NewFolderName);
+			PendingFoldersSelect.Add(NewFolder);
 		}
 	}
 
@@ -1474,16 +1531,16 @@ void SSceneOutliner::DeleteFoldersBegin()
 
 void SSceneOutliner::DeleteFoldersEnd()
 {
-	struct FMatchName
+	struct FMatchFolder
 	{
-		FMatchName(const FName InPathName)
-			: PathName(InPathName) {}
+		FMatchFolder(const FFolder& InFolder)
+			: Folder(InFolder) {}
 
-		const FName PathName;
+		const FFolder Folder;
 
 		bool operator()(const FFolderTreeItem *Entry)
 		{
-			return PathName == Entry->Path;
+			return Folder == Entry->GetFolder();
 		}
 	};
 
@@ -1492,21 +1549,21 @@ void SSceneOutliner::DeleteFoldersEnd()
 		// Sort in descending order so children will be deleted before parents
 		CacheFoldersDelete.Sort([](const FFolderTreeItem& FolderA, const FFolderTreeItem& FolderB)
 		{
-			return FolderB.Path.LexicalLess(FolderA.Path);
+			return FolderB.GetPath().LexicalLess(FolderA.GetPath());
 		});
 
-		for (FFolderTreeItem* Folder : CacheFoldersDelete)
+		for (FFolderTreeItem* FolderItem : CacheFoldersDelete)
 		{
-			if (Folder)
+			if (FolderItem)
 			{
 				// Find lowest parent not being deleted, for reparenting children of current folder
-				FName NewParentPath = FEditorFolderUtils::GetParentPath(Folder->Path);
-				while (!NewParentPath.IsNone() && CacheFoldersDelete.FindByPredicate(FMatchName(NewParentPath)))
+				FFolder NewParentPath = FolderItem->GetFolder().GetParent();
+				while (!NewParentPath.IsNone() && CacheFoldersDelete.FindByPredicate(FMatchFolder(NewParentPath)))
 				{
-					NewParentPath = FEditorFolderUtils::GetParentPath(NewParentPath);
+					NewParentPath = NewParentPath.GetParent();
 				}
 
-				Folder->Delete(NewParentPath);
+				FolderItem->Delete(NewParentPath);
 			}
 		}
 
@@ -1717,6 +1774,10 @@ void SSceneOutliner::OnHierarchyChangedEvent(FSceneOutlinerHierarchyChangedData
 		for (const auto& TreeItemID : Event.ItemIDs)
 		{
 			FSceneOutlinerTreeItemPtr* Item = TreeItemMap.Find(TreeItemID);
+			if (!Item)
+			{
+				Item = PendingTreeItemMap.Find(TreeItemID);
+			}
 
 			if (Item)
 			{
@@ -1746,8 +1807,8 @@ void SSceneOutliner::OnHierarchyChangedEvent(FSceneOutlinerHierarchyChangedData
 
 				// Now change the path and put it back in the map with its new ID
 				auto Folder = StaticCastSharedPtr<FFolderTreeItem>(Item);
-				Folder->Path = Event.NewPaths[i];
-				Folder->LeafName = FEditorFolderUtils::GetLeafName(Event.NewPaths[i]);
+				check(Event.NewPaths[i].GetRootObject() == Folder->GetRootObject());
+				Folder->SetPath(Event.NewPaths[i].GetPath());
 
 				TreeItemMap.Add(Item->GetID(), Item);
 
diff --git a/Engine/Source/Editor/SceneOutliner/Public/ActorBrowsingMode.h b/Engine/Source/Editor/SceneOutliner/Public/ActorBrowsingMode.h
--- a/Engine/Source/Editor/SceneOutliner/Public/ActorBrowsingMode.h
+++ b/Engine/Source/Editor/SceneOutliner/Public/ActorBrowsingMode.h
@@ -14,7 +14,7 @@ public:
 
 	/* Begin ISceneOutlinerMode Interface */
 	virtual void Rebuild() override;
-	virtual FCreateSceneOutlinerMode CreateFolderPickerMode() const override;
+	virtual FCreateSceneOutlinerMode CreateFolderPickerMode(const FFolder::FRootObject& InRootObject = FFolder::GetDefaultRootObject()) const override;
 	virtual void CreateViewContent(FMenuBuilder& MenuBuilder) override;
 	virtual TSharedPtr<FDragDropOperation> CreateDragDropOperation(const TArray<FSceneOutlinerTreeItemPtr>& InTreeItems) const override;
 	virtual bool ParseDragDrop(FSceneOutlinerDragDropPayload& OutPayload, const FDragDropOperation& Operation) const override;
@@ -45,9 +45,9 @@ public:
 	virtual bool CanCopy() const override;
 	virtual bool CanPaste() const override;
 	virtual bool CanSupportDragAndDrop() const { return true; }
-	virtual FName CreateNewFolder() override;
-	virtual FName CreateFolder(const FName& ParentPath, const FName& LeafName) override;
-	virtual bool ReparentItemToFolder(const FName& FolderPath, const FSceneOutlinerTreeItemPtr& Item) override;
+	virtual FFolder CreateNewFolder() override;
+	virtual FFolder CreateFolder(const FFolder& ParentPath, const FName& LeafName) override;
+	virtual bool ReparentItemToFolder(const FFolder& FolderPath, const FSceneOutlinerTreeItemPtr& Item) override;
 	virtual void SelectFoldersDescendants(const TArray<FFolderTreeItem*>& FolderItems, bool bSelectImmediateChildrenOnly) override;
 	virtual void PinItem(const FSceneOutlinerTreeItemPtr& InItem) override;
 	virtual void UnpinItem(const FSceneOutlinerTreeItemPtr& InItem) override;
@@ -101,6 +101,9 @@ private:
 	void RegisterContextMenu();
 	bool CanPasteFoldersOnlyFromClipboard() const;
 
+	bool GetFolderNamesFromFolders(const TArray<FFolder>& InFolders, TArray<FName>& OutFolders, FFolder::FRootObject& OutCommonRootObject) const;
+	bool GetFolderNamesFromPayload(const FSceneOutlinerDragDropPayload& InPayload, TArray<FName>& OutFolders, FFolder::FRootObject& OutCommonRootObject) const;
+
 	/** Filter factories */
 	static TSharedRef<FSceneOutlinerFilter> CreateShowOnlySelectedActorsFilter();
 	static TSharedRef<FSceneOutlinerFilter> CreateHideTemporaryActorsFilter();
diff --git a/Engine/Source/Editor/SceneOutliner/Public/ActorFolderHierarchy.h b/Engine/Source/Editor/SceneOutliner/Public/ActorFolderHierarchy.h
--- a/Engine/Source/Editor/SceneOutliner/Public/ActorFolderHierarchy.h
+++ b/Engine/Source/Editor/SceneOutliner/Public/ActorFolderHierarchy.h
@@ -8,7 +8,7 @@
 class FActorFolderHierarchy : public ISceneOutlinerHierarchy
 {
 public:
-	FActorFolderHierarchy(ISceneOutlinerMode* InMode, const TWeakObjectPtr<UWorld>& World);
+	FActorFolderHierarchy(ISceneOutlinerMode* InMode, const TWeakObjectPtr<UWorld>& World, const FFolder::FRootObject& InRootObject = FFolder::GetDefaultRootObject());
 	virtual ~FActorFolderHierarchy() {}
 
 
@@ -24,4 +24,6 @@ private:
 private:
 	/** The world which this hierarchy is representing */
 	TWeakObjectPtr<UWorld> RepresentingWorld;
+
+	FFolder::FRootObject RootObject;
 };
diff --git a/Engine/Source/Editor/SceneOutliner/Public/ActorFolderPickingMode.h b/Engine/Source/Editor/SceneOutliner/Public/ActorFolderPickingMode.h
--- a/Engine/Source/Editor/SceneOutliner/Public/ActorFolderPickingMode.h
+++ b/Engine/Source/Editor/SceneOutliner/Public/ActorFolderPickingMode.h
@@ -7,7 +7,7 @@
 class SCENEOUTLINER_API FActorFolderPickingMode : public FActorMode
 {
 public:
-	FActorFolderPickingMode(SSceneOutliner* InSceneOutliner, FOnSceneOutlinerItemPicked InOnItemPicked, TWeakObjectPtr<UWorld> InSpecifiedWorldToDisplay = nullptr);
+	FActorFolderPickingMode(SSceneOutliner* InSceneOutliner, FOnSceneOutlinerItemPicked InOnItemPicked, TWeakObjectPtr<UWorld> InSpecifiedWorldToDisplay = nullptr, const FFolder::FRootObject& InRootObject = FFolder::GetDefaultRootObject());
 	virtual ~FActorFolderPickingMode() {}
 
 	/* Begin ISceneOutlinerMode Implementation */
@@ -22,4 +22,6 @@ protected:
 
 	/** Delegate to call when an item is picked */
 	FOnSceneOutlinerItemPicked OnItemPicked;
+
+	FFolder::FRootObject RootObject;
 };
\ No newline at end of file
diff --git a/Engine/Source/Editor/SceneOutliner/Public/ActorFolderTreeItem.h b/Engine/Source/Editor/SceneOutliner/Public/ActorFolderTreeItem.h
--- a/Engine/Source/Editor/SceneOutliner/Public/ActorFolderTreeItem.h
+++ b/Engine/Source/Editor/SceneOutliner/Public/ActorFolderTreeItem.h
@@ -11,7 +11,7 @@ public:
 	/** Static type identifier for this tree item class */
 	static const FSceneOutlinerTreeItemType Type;
 
-	FActorFolderTreeItem(FName InPath, TWeakObjectPtr<UWorld> InWorld);
+	FActorFolderTreeItem(const FFolder& InFolder, const TWeakObjectPtr<UWorld>& InWorld);
 
 	/** The world which this folder belongs to */
 	TWeakObjectPtr<UWorld> World;
@@ -19,12 +19,12 @@ public:
 	/* Begin ISceneOutlinerTreeItem Implementation */
 	virtual bool IsValid() const override { return World.IsValid(); }
 	virtual void OnExpansionChanged() override;
-	virtual void Delete(FName InNewParentPath) override;
+	virtual void Delete(const FFolder& InNewParentFolder) override;
 	virtual TSharedRef<SWidget> GenerateLabelWidget(ISceneOutliner& Outliner, const STableRow<FSceneOutlinerTreeItemPtr>& InRow) override;
 	/* End FFolderTreeItem Implementation */
 		
 	/* Begin FFolderTreeItem Implementation */
-	virtual FName MoveTo(const FName& NewParent) override;
+	virtual void MoveTo(const FFolder& InNewParentFolder) override;
 private:
 	virtual void CreateSubFolder(TWeakPtr<SSceneOutliner> WeakOutliner) override;
 	/* End FFolderTreeItem Implementation */
diff --git a/Engine/Source/Editor/SceneOutliner/Public/ActorHierarchy.h b/Engine/Source/Editor/SceneOutliner/Public/ActorHierarchy.h
--- a/Engine/Source/Editor/SceneOutliner/Public/ActorHierarchy.h
+++ b/Engine/Source/Editor/SceneOutliner/Public/ActorHierarchy.h
@@ -4,6 +4,7 @@
 
 #include "CoreMinimal.h"
 #include "ISceneOutlinerHierarchy.h"
+#include "Folder.h"
 
 class FWorldPartitionActorDesc;
 
@@ -59,13 +60,13 @@ private:
 	void OnLevelRemoved(ULevel* InLevel, UWorld* InWorld);
 
 	/** Called when a folder is to be created */
-	void OnBroadcastFolderCreate(UWorld& InWorld, FName NewPath);
+	void OnBroadcastFolderCreate(UWorld& InWorld, const FFolder& InNewFolder);
 
 	/** Called when a folder is to be moved */
-	void OnBroadcastFolderMove(UWorld& InWorld, FName OldPath, FName NewPath);
+	void OnBroadcastFolderMove(UWorld& InWorld, const FFolder& InOldFolder, const FFolder& InNewFolder);
 
 	/** Called when a folder is to be deleted */
-	void OnBroadcastFolderDelete(UWorld& InWorld, FName Path);
+	void OnBroadcastFolderDelete(UWorld& InWorld, const FFolder& InFolder);
 
 private:
 	/** Send a an event indicating a full refresh of the hierarchy is required */
diff --git a/Engine/Source/Editor/SceneOutliner/Public/ActorTreeItem.h b/Engine/Source/Editor/SceneOutliner/Public/ActorTreeItem.h
--- a/Engine/Source/Editor/SceneOutliner/Public/ActorTreeItem.h
+++ b/Engine/Source/Editor/SceneOutliner/Public/ActorTreeItem.h
@@ -39,6 +39,7 @@ public:
 	/* Begin ISceneOutlinerTreeItem Implementation */
 	virtual bool IsValid() const override { return Actor.IsValid(); }
 	virtual FSceneOutlinerTreeItemID GetID() const override;
+	virtual FFolder::FRootObject GetRootObject() const override;
 	virtual FString GetDisplayString() const override;
 	virtual bool CanInteract() const override;
 	virtual TSharedRef<SWidget> GenerateLabelWidget(ISceneOutliner& Outliner, const STableRow<FSceneOutlinerTreeItemPtr>& InRow) override;
@@ -49,6 +50,7 @@ public:
 	virtual bool HasPinnedStateInfo() const override { return true; };
 	virtual bool GetPinnedState() const override;
 	virtual void OnLabelChanged() override;
+	virtual void GenerateContextMenu(UToolMenu* Menu, SSceneOutliner& Outliner) override;
 	/* End ISceneOutlinerTreeItem Implementation */
 public:
 	/** true if this item exists in both the current world and PIE. */
diff --git a/Engine/Source/Editor/SceneOutliner/Public/FolderTreeItem.h b/Engine/Source/Editor/SceneOutliner/Public/FolderTreeItem.h
--- a/Engine/Source/Editor/SceneOutliner/Public/FolderTreeItem.h
+++ b/Engine/Source/Editor/SceneOutliner/Public/FolderTreeItem.h
@@ -4,6 +4,7 @@
 
 #include "CoreMinimal.h"
 #include "Widgets/SWidget.h"
+#include "Folder.h"
 #include "ISceneOutlinerTreeItem.h"
 
 class UToolMenu;
@@ -12,7 +13,7 @@ namespace SceneOutliner
 {
 	struct FFolderPathSelector
 	{
-		bool operator()(TWeakPtr<ISceneOutlinerTreeItem> Item, FName& DataOut) const;
+		bool operator()(TWeakPtr<ISceneOutlinerTreeItem> Item, FFolder& DataOut) const;
 	};
 }	// namespace SceneOutliner
 
@@ -24,17 +25,17 @@ public:
 	/** Static type identifier for this tree item class */
 	static const FSceneOutlinerTreeItemType Type;
 
-	DECLARE_DELEGATE_RetVal_OneParam(bool, FFilterPredicate, FName);
-	DECLARE_DELEGATE_RetVal_OneParam(bool, FInteractivePredicate, FName);
+	DECLARE_DELEGATE_RetVal_OneParam(bool, FFilterPredicate, const FFolder&);
+	DECLARE_DELEGATE_RetVal_OneParam(bool, FInteractivePredicate, const FFolder&);
 
 	bool Filter(FFilterPredicate Pred) const
 	{
-		return Pred.Execute(Path);
+		return Pred.Execute(GetFolder());
 	}
 
 	bool GetInteractiveState(FInteractivePredicate Pred) const
 	{
-		return Pred.Execute(Path);
+		return Pred.Execute(GetFolder());
 	}
 
 	/** The path of this folder. / separated. */
@@ -43,6 +44,8 @@ public:
 	/** The leaf name of this folder */
 	FName LeafName;
 
+	FFolderTreeItem(const FFolder& InFolder, FSceneOutlinerTreeItemType InType);
+	FFolderTreeItem(const FFolder& InFolder);
 	/** Constructor that takes a path to this folder (including leaf-name) */
 	FFolderTreeItem(FName InPath);
 	/** Constructor that takes a path to this folder and a subclass tree item type (used for subclassing FFolderTreeItem) */
@@ -51,21 +54,30 @@ public:
 	/* Begin ISceneOutlinerTreeItem Implementation */
 	virtual bool IsValid() const override { return true; }
 	virtual FSceneOutlinerTreeItemID GetID() const override;
+	virtual FFolder::FRootObject GetRootObject() const override;
 	virtual FString GetDisplayString() const override;
 	virtual bool CanInteract() const override;
 	virtual void GenerateContextMenu(UToolMenu* Menu, SSceneOutliner& Outliner) override;
 	/** Delete this folder, children will be reparented to provided new parent path */
-	virtual void Delete(FName InNewParentPath) {}
+	virtual void Delete(const FFolder& InNewParentFolder) {}
+	
+	void SetPath(const FName& InNewPath);
+	const FName& GetPath() const { return Path; }
+	const FName& GetLeafName() const { return LeafName; }
+	FFolder GetFolder() const { return FFolder(Path, RootObject); }
 	
 	virtual bool ShouldShowPinnedState() const override { return true; }
 	virtual bool HasPinnedStateInfo() const override { return false; }
 	/* End ISceneOutlinerTreeItem Implementation */
 
 	/** Move this folder to a new parent */
-	virtual FName MoveTo(const FName& NewParent) { return FName(); }
+	virtual void MoveTo(const FFolder& InNewParentFolder) {}
 private:
 	/** Create a new folder as a child of this one */
 	virtual void CreateSubFolder(TWeakPtr<SSceneOutliner> WeakOutliner) {}
 	/** Duplicate folder hierarchy */
 	void DuplicateHierarchy(TWeakPtr<SSceneOutliner> WeakOutliner);
+
+	/** Folder's root object (can be null) */
+	FFolder::FRootObject RootObject;
 };
diff --git a/Engine/Source/Editor/SceneOutliner/Public/ISceneOutlinerMode.h b/Engine/Source/Editor/SceneOutliner/Public/ISceneOutlinerMode.h
--- a/Engine/Source/Editor/SceneOutliner/Public/ISceneOutlinerMode.h
+++ b/Engine/Source/Editor/SceneOutliner/Public/ISceneOutlinerMode.h
@@ -9,7 +9,7 @@
 #include "SceneOutlinerFilters.h"
 #include "SSceneOutliner.h"
 #include "SceneOutlinerDragDrop.h"
-
+#include "Folder.h"
 
 class ISceneOutlinerMode
 {
@@ -38,7 +38,7 @@ public:
 	virtual void CreateViewContent(FMenuBuilder& MenuBuilder) {}
 
 	/** Returns a factory to create a Folder Picker mode which matches this mode */
-	virtual FCreateSceneOutlinerMode CreateFolderPickerMode() const { return FCreateSceneOutlinerMode(); }
+	virtual FCreateSceneOutlinerMode CreateFolderPickerMode(const FFolder::FRootObject& InRootObject = FFolder::GetDefaultRootObject()) const { return FCreateSceneOutlinerMode(); }
 public:
 	/** Construct a new Drag and drop operation */
 	virtual TSharedPtr<FDragDropOperation> CreateDragDropOperation(const TArray<FSceneOutlinerTreeItemPtr>& InTreeItems) const { return nullptr; }
@@ -128,11 +128,11 @@ public:
 	/* Folder management */
 		
 	/** Creates a new folder item at the root with a valid name*/
-	virtual FName CreateNewFolder() { return FName(NAME_None); }
+	virtual FFolder CreateNewFolder() { return FFolder(); }
 	/** Create a folder under a specific parent with a given leaf name */
-	virtual FName CreateFolder(const FName& ParentPath, const FName& LeafName) { return FName(NAME_None); }
+	virtual FFolder CreateFolder(const FFolder& ParentPath, const FName& LeafName) { return FFolder(); }
 	/** Reparent an item to a given folder path. Returns true if the operation is sucessful */
-	virtual bool ReparentItemToFolder(const FName& FolderPath, const FSceneOutlinerTreeItemPtr& Item) { return false; }
+	virtual bool ReparentItemToFolder(const FFolder& FolderPath, const FSceneOutlinerTreeItemPtr& Item) { return false; }
 	/** Select all descendants of a folder. Optionally select only immediate descendants. */
 	virtual void SelectFoldersDescendants(const TArray<FFolderTreeItem*>& FolderItems, bool bSelectImmediateChildrenOnly) {}
 public:
diff --git a/Engine/Source/Editor/SceneOutliner/Public/ISceneOutlinerTreeItem.h b/Engine/Source/Editor/SceneOutliner/Public/ISceneOutlinerTreeItem.h
--- a/Engine/Source/Editor/SceneOutliner/Public/ISceneOutlinerTreeItem.h
+++ b/Engine/Source/Editor/SceneOutliner/Public/ISceneOutlinerTreeItem.h
@@ -120,6 +120,9 @@ public:
 	/** Get the ID that represents this tree item. Used to reference this item in a map */
 	virtual FSceneOutlinerTreeItemID GetID() const = 0;
 
+	/** Get the optional context object root of a tree item. */
+	virtual FFolder::FRootObject GetRootObject() const { return FFolder::GetDefaultRootObject(); }
+
 	/** Get the raw string to display for this tree item - used for sorting */
 	virtual FString GetDisplayString() const = 0;
 
diff --git a/Engine/Source/Editor/SceneOutliner/Public/SSceneOutliner.h b/Engine/Source/Editor/SceneOutliner/Public/SSceneOutliner.h
--- a/Engine/Source/Editor/SceneOutliner/Public/SSceneOutliner.h
+++ b/Engine/Source/Editor/SceneOutliner/Public/SSceneOutliner.h
@@ -18,6 +18,7 @@
 #include "Delegates/DelegateCombinations.h"
 #include "Framework/Views/ITypedTableView.h"
 #include "Types/SlateEnums.h"
+#include "Folder.h"
 
 #include "ISceneOutliner.h"
 #include "SceneOutlinerFwd.h"
@@ -655,9 +656,11 @@ private:
 	FString CacheClipboardContents;
 
 	/** Maps pre-existing children during paste or duplicate */
-	TMap<FName, TArray<FSceneOutlinerTreeItemID>> CachePasteFolderExistingChildrenMap;
+	TMap<FFolder, TArray<FSceneOutlinerTreeItemID>> CachePasteFolderExistingChildrenMap;
 
 private:
+	bool GetCommonRootObjectFromSelection(FFolder::FRootObject& OutCommonRootObject) const;
+
 	/** Miscellaneous bindings required by the UI */
 
 	/** Called by the editable text control when the filter text is changed by the user */
@@ -701,13 +704,13 @@ public:
 	void FillFoldersSubMenu(UToolMenu* Menu) const;
 	void AddMoveToFolderOutliner(UToolMenu* Menu) const;
 	void FillSelectionSubMenu(UToolMenu* Menun) const;
-	TSharedRef<TSet<FName>> GatherInvalidMoveToDestinations() const;
+	TSharedRef<TSet<FFolder>> GatherInvalidMoveToDestinations() const;
 
 	/** Called to select descendants of the currently selected folders */
 	void SelectFoldersDescendants(bool bSelectImmediateChildrenOnly = false);
 
 	/** Moves the current selection to the specified folder path */
-	void MoveSelectionTo(FName NewParent);
+	void MoveSelectionTo(const FFolder& NewParent);
 
 	/** Create a new folder under the specified parent name (NAME_None for root) */
 	void CreateFolder();
@@ -739,7 +742,7 @@ private:
 	FSceneOutlinerTreeItemMap PendingTreeItemMap;
 
 	/** Folders pending selection */
-	TArray<FName> PendingFoldersSelect;
+	TArray<FFolder> PendingFoldersSelect;
 
 	/** Root level tree items */
 	TArray<FSceneOutlinerTreeItemPtr> RootTreeItems;
diff --git a/Engine/Source/Editor/SceneOutliner/Public/SceneOutlinerStandaloneTypes.h b/Engine/Source/Editor/SceneOutliner/Public/SceneOutlinerStandaloneTypes.h
--- a/Engine/Source/Editor/SceneOutliner/Public/SceneOutlinerStandaloneTypes.h
+++ b/Engine/Source/Editor/SceneOutliner/Public/SceneOutlinerStandaloneTypes.h
@@ -6,8 +6,30 @@
 #include "Misc/Paths.h"
 #include "UObject/ObjectKey.h"
 #include "Templates/MaxSizeof.h"
+#include "Folder.h"
 #include "SceneOutlinerFwd.h"
 
+struct FFolderKey
+{
+	FFolderKey(const FName& InPath, const FFolder::FRootObject& InRootObject)
+		: Path(InPath)
+		, RootObjectKey(InRootObject)
+	{}
+
+	FORCEINLINE bool operator == (const FFolderKey& InOther)
+	{
+		return (Path == InOther.Path) && (RootObjectKey == InOther.RootObjectKey);
+	}
+
+	FName Path;
+	FFolder::FRootObject RootObjectKey;
+};
+
+FORCEINLINE uint32 GetTypeHash(FFolderKey Key)
+{
+	return HashCombine(GetTypeHash(Key.Path), GetTypeHash(Key.RootObjectKey));
+}
+
 /** Variant type that defines an identifier for a tree item. Assumes 'trivial relocatability' as with many unreal containers. */
 struct FSceneOutlinerTreeItemID
 {
@@ -33,12 +55,16 @@ public:
 	}
 
 	/** ID representing a folder */
-	FSceneOutlinerTreeItemID(const FName& InFolder) : Type(EType::Folder)
+	FSceneOutlinerTreeItemID(const FFolderKey& InKey) : Type(EType::Folder)
 	{
-		new (Data) FName(InFolder);
+		new (Data) FFolderKey(InKey);
 		CachedHash = CalculateTypeHash();
 	}
 
+	FSceneOutlinerTreeItemID(const FFolder& InFolder)
+	: FSceneOutlinerTreeItemID(FFolderKey(InFolder.GetPath(), InFolder.GetRootObject()))
+	{}
+
 	/** ID representing a generic tree item */
 	FSceneOutlinerTreeItemID(const FUniqueID& CustomID) : Type(EType::UniqueID)
 	{
@@ -63,7 +89,7 @@ public:
 		switch(Type)
 		{
 			case EType::Object:			new (Data) FObjectKey(Other.GetAsObjectKey());		break;
-			case EType::Folder:			new (Data) FName(Other.GetAsFolderRef());			break;
+			case EType::Folder:			new (Data) FFolderKey(Other.GetAsFolderRef());		break;
 			case EType::Guid:			new (Data) FGuid(Other.GetAsGuid());				break;
 			case EType::UniqueID:		new (Data) FUniqueID(Other.GetAsHash());			break;
 			default:																		break;
@@ -89,7 +115,7 @@ public:
 		switch(Type)
 		{
 			case EType::Object:			GetAsObjectKey().~FObjectKey();							break;
-			case EType::Folder:			GetAsFolderRef().~FName();								break;
+			case EType::Folder:			GetAsFolderRef().~FFolderKey();							break;
 			case EType::Guid:			GetAsGuid().~FGuid();									break;
 			case EType::UniqueID:		/* NOP */												break;
 			default:																			break;
@@ -128,7 +154,7 @@ public:
 private:
 
 	FObjectKey& 	GetAsObjectKey() const 		{ return *reinterpret_cast<FObjectKey*>(Data); }
-	FName& 			GetAsFolderRef() const		{ return *reinterpret_cast<FName*>(Data); }
+	FFolderKey& 	GetAsFolderRef() const		{ return *reinterpret_cast<FFolderKey*>(Data); }
 	FGuid&			GetAsGuid() const			{ return *reinterpret_cast<FGuid*>(Data); }
 	FUniqueID&		GetAsHash() const			{ return *reinterpret_cast<FUniqueID*>(Data); }
 
@@ -149,7 +175,7 @@ private:
 	EType Type;
 
 	uint32 CachedHash;
-	static const uint32 MaxSize = TMaxSizeof<FObjectKey, FName, FGuid, FUniqueID>::Value;
+	static const uint32 MaxSize = TMaxSizeof<FObjectKey, FFolderKey, FGuid, FUniqueID>::Value;
 	mutable uint8 Data[MaxSize];
 };
 	
@@ -207,7 +233,7 @@ struct FSceneOutlinerHierarchyChangedData
 
 	TArray<FSceneOutlinerTreeItemID> ItemIDs;
 	// Used for FolderMoved events
-	TArray<FName> NewPaths;
+	TArray<FFolder> NewPaths;
 	/** Actions to apply to items */
 	uint8 ItemActions = 0;
 };
diff --git a/Engine/Source/Editor/Sequencer/Private/SSequencer.cpp b/Engine/Source/Editor/Sequencer/Private/SSequencer.cpp
--- a/Engine/Source/Editor/Sequencer/Private/SSequencer.cpp
+++ b/Engine/Source/Editor/Sequencer/Private/SSequencer.cpp
@@ -3444,21 +3444,25 @@ void SSequencer::OnActorsDropped( FActorDragDropOp& DragDropOp )
 
 void SSequencer::OnFolderDropped( FFolderDragDropOp& DragDropOp )
 {
+	// Sequencer doesn't support dragging folder with a root object
+	if (FFolder::HasRootObject(DragDropOp.RootObject))
+		return;
+
 	const FScopedTransaction Transaction(LOCTEXT("DropActors", "Drop Actors"));
 
 	ESequencerDropResult DropResult = ESequencerDropResult::Unhandled;
 
 	TArray<TWeakObjectPtr<AActor>> DraggedActors;
 
-		// Find any actors in the global editor world that have any of the dragged paths.
-		// WARNING: Actor iteration can be very slow, so this needs to be optimized
-		TSharedPtr<FSequencer> Sequencer = SequencerPtr.Pin();
-		UObject* PlaybackContext = Sequencer->GetPlaybackContext();
-		UWorld* World = PlaybackContext ? PlaybackContext->GetWorld() : nullptr;
-		if (World)
-		{
+	// Find any actors in the global editor world that have any of the dragged paths.
+	// WARNING: Actor iteration can be very slow, so this needs to be optimized
+	TSharedPtr<FSequencer> Sequencer = SequencerPtr.Pin();
+	UObject* PlaybackContext = Sequencer->GetPlaybackContext();
+	UWorld* World = PlaybackContext ? PlaybackContext->GetWorld() : nullptr;
+	if (World)
+	{
 		FActorFolders::GetWeakActorsFromFolders(*World, DragDropOp.Folders, DraggedActors);
-				}
+	}
 	
 	for (FOnFoldersDrop Delegate : OnFoldersDrop)
 	{
diff --git a/Engine/Source/Editor/UnrealEd/Private/EditorActorFolders.cpp b/Engine/Source/Editor/UnrealEd/Private/EditorActorFolders.cpp
--- a/Engine/Source/Editor/UnrealEd/Private/EditorActorFolders.cpp
+++ b/Engine/Source/Editor/UnrealEd/Private/EditorActorFolders.cpp
@@ -15,6 +15,7 @@
 #include "EditorFolderUtils.h"
 #include "ScopedTransaction.h"
 #include "UObject/ObjectSaveContext.h"
+#include "LevelInstance/LevelInstanceActor.h"
 
 #define LOCTEXT_NAMESPACE "FActorFolders"
 
@@ -26,7 +27,7 @@ void UEditorActorFolders::Serialize(FArchive& Ar)
 FString GetWorldStateFilename(UPackage* Package)
 {
 	const FString PathName = Package->GetPathName();
-	const uint32 PathNameCrc = FCrc::MemCrc32(*PathName, sizeof(TCHAR)*PathName.Len());
+	const uint32 PathNameCrc = FCrc::MemCrc32(*PathName, sizeof(TCHAR) * PathName.Len());
 	return FPaths::Combine(*FPaths::ProjectSavedDir(), TEXT("Config"), TEXT("WorldState"), *FString::Printf(TEXT("%u.json"), PathNameCrc));
 }
 
@@ -37,10 +38,19 @@ static FName OldPathToNewPath(const FString& InOldBranch, const FString& InNewBr
 }
 
 // Static member definitions
+
+//~ Begin Deprecated
+PRAGMA_DISABLE_DEPRECATION_WARNINGS
 FOnActorFolderCreate	FActorFolders::OnFolderCreate;
 FOnActorFolderMove		FActorFolders::OnFolderMove;
 FOnActorFolderDelete	FActorFolders::OnFolderDelete;
-FActorFolders*			FActorFolders::Singleton;
+PRAGMA_ENABLE_DEPRECATION_WARNINGS
+//~ End Deprecated
+
+FActorFolders* FActorFolders::Singleton;
+FOnActorFolderCreated	FActorFolders::OnFolderCreated;
+FOnActorFolderMoved		FActorFolders::OnFolderMoved;
+FOnActorFolderDeleted	FActorFolders::OnFolderDeleted;
 
 FActorFolders::FActorFolders()
 {
@@ -96,6 +106,39 @@ void FActorFolders::Housekeeping()
 	}
 }
 
+void FActorFolders::BroadcastOnActorFolderCreated(UWorld& InWorld, const FFolder& InFolder)
+{
+	OnFolderCreated.Broadcast(InWorld, InFolder);
+	PRAGMA_DISABLE_DEPRECATION_WARNINGS
+	if (!InFolder.HasRootObject())
+	{
+		OnFolderCreate.Broadcast(InWorld, InFolder.GetPath());
+	}
+	PRAGMA_ENABLE_DEPRECATION_WARNINGS
+}
+
+void FActorFolders::BroadcastOnActorFolderDeleted(UWorld& InWorld, const FFolder& InFolder)
+{
+	OnFolderDeleted.Broadcast(InWorld, InFolder);
+	PRAGMA_DISABLE_DEPRECATION_WARNINGS
+	if (!InFolder.HasRootObject())
+	{
+		OnFolderDelete.Broadcast(InWorld, InFolder.GetPath());
+	}
+	PRAGMA_ENABLE_DEPRECATION_WARNINGS
+}
+
+void FActorFolders::BroadcastOnActorFolderMoved(UWorld& InWorld, const FFolder& InSrcFolder, const FFolder& InDstFolder)
+{
+	OnFolderMoved.Broadcast(InWorld, InSrcFolder, InDstFolder);
+	PRAGMA_DISABLE_DEPRECATION_WARNINGS
+	if (!InSrcFolder.HasRootObject() && !InDstFolder.HasRootObject())
+	{
+		OnFolderMove.Broadcast(InWorld, InSrcFolder.GetPath(), InDstFolder.GetPath());
+	}
+	PRAGMA_ENABLE_DEPRECATION_WARNINGS
+}
+
 void FActorFolders::OnLevelActorListChanged()
 {
 	Housekeeping();
@@ -148,10 +191,13 @@ void FActorFolders::OnWorldSaved(UWorld* World, FObjectPostSaveContext ObjectSav
 
 			for (const auto& KeyValue : (*ExisingFolders)->Folders)
 			{
-				TSharedRef<FJsonObject> JsonFolder = MakeShareable(new FJsonObject);
-				JsonFolder->SetBoolField(TEXT("bIsExpanded"), KeyValue.Value.bIsExpanded);
-
-				JsonFolders->SetObjectField(KeyValue.Key.ToString(), JsonFolder);
+				// Only write for World root
+				if (!KeyValue.Key.HasRootObject())
+				{
+					TSharedRef<FJsonObject> JsonFolder = MakeShareable(new FJsonObject);
+					JsonFolder->SetBoolField(TEXT("bIsExpanded"), KeyValue.Value.bIsExpanded);
+					JsonFolders->SetObjectField(KeyValue.Key.ToString(), JsonFolder);
+				}
 			}
 
 			RootObject->SetObjectField(TEXT("Folders"), JsonFolders);
@@ -172,11 +218,11 @@ void FActorFolders::OnActorFolderChanged(const AActor* InActor, FName OldPath)
 	FScopedTransaction Transaction(LOCTEXT("UndoAction_FolderChanged", "Actor Folder Changed"));
 
 	UWorld* World = InActor->GetWorld();
-	const auto NewPath = InActor->GetFolderPath();
+	const FFolder NewPath = InActor->GetFolder();
 
 	if (AddFolderToWorld(*World, NewPath))
 	{
-		OnFolderCreate.Broadcast(*World, NewPath);
+		BroadcastOnActorFolderCreated(*World, NewPath);
 	}
 }
 
@@ -184,13 +230,25 @@ void FActorFolders::RebuildFolderListForWorld(UWorld& InWorld)
 {
 	if (TemporaryWorldFolders.Contains(&InWorld))
 	{
-		// We don't empty the existing folders so that we keep empty ones.
-		// Explicitly deleted folders will already be removed from the list
+		// For world folders, we don't empty the existing folders so that we keep empty ones.
+		// Explicitly deleted folders will already be removed from the list.
+		
+		// Clear folders with a Root Object.
+		TArray<FFolder> FoldersToRemove;
+		ForEachFolder(InWorld, [&FoldersToRemove](const FFolder& Folder)
+		{
+			if (Folder.HasRootObject())
+			{
+				FoldersToRemove.Add(Folder);
+			}
+			return true;
+		});
+		RemoveFoldersFromWorld(InWorld, FoldersToRemove, /*bBroadcastDelete*/ false);
 
 		// Iterate over every actor in memory. WARNING: This is potentially very expensive!
 		for (FActorIterator ActorIt(&InWorld); ActorIt; ++ActorIt)
 		{
-			AddFolderToWorld(InWorld, ActorIt->GetFolderPath());
+			AddFolderToWorld(InWorld, ActorIt->GetFolder());
 		}
 	}
 	else
@@ -200,12 +258,12 @@ void FActorFolders::RebuildFolderListForWorld(UWorld& InWorld)
 	}
 }
 
-const TMap<FName, FActorFolderProps>& FActorFolders::GetFolderPropertiesForWorld(UWorld& InWorld)
+const TMap<FFolder, FActorFolderProps>& FActorFolders::GetFolderPropertiesForWorld(UWorld& InWorld)
 {
 	return GetOrCreateFoldersForWorld(InWorld).Folders;
 }
 
-FActorFolderProps* FActorFolders::GetFolderProperties(UWorld& InWorld, FName InPath)
+FActorFolderProps* FActorFolders::GetFolderProperties(UWorld& InWorld, const FFolder& InPath)
 {
 	return GetOrCreateFoldersForWorld(InWorld).Folders.Find(InPath);
 }
@@ -230,6 +288,8 @@ UEditorActorFolders& FActorFolders::InitializeForWorld(UWorld& InWorld)
 	// Clean up any stale worlds
 	Housekeeping();
 
+	InWorld.OnLevelsChanged().AddRaw(this, &FActorFolders::OnLevelActorListChanged);
+
 	// We intentionally don't pass RF_Transactional to ConstructObject so that we don't record the creation of the object into the undo buffer
 	// (to stop it getting deleted on undo as we manage its lifetime), but we still want it to be RF_Transactional so we can record any changes later
 	UEditorActorFolders* Folders = NewObject<UEditorActorFolders>(GetTransientPackage(), NAME_None, RF_NoFlags);
@@ -239,7 +299,7 @@ UEditorActorFolders& FActorFolders::InitializeForWorld(UWorld& InWorld)
 	// Ensure the list is entirely up to date with the world before we write our serialized properties into it.
 	for (FActorIterator ActorIt(&InWorld); ActorIt; ++ActorIt)
 	{
-		AddFolderToWorld(InWorld, ActorIt->GetFolderPath());
+		AddFolderToWorld(InWorld, ActorIt->GetFolder());
 	}
 
 	// Attempt to load the folder properties from this user's saved world state directory
@@ -257,7 +317,7 @@ UEditorActorFolders& FActorFolders::InitializeForWorld(UWorld& InWorld)
 			{
 				// Only pull in the folder's properties if this folder still exists in the world.
 				// This means that old stale folders won't re-appear in the world (they'll won't get serialized when the world is saved anyway)
-				if (FActorFolderProps* FolderInWorld = Folders->Folders.Find(*KeyValue.Key))
+				if (FActorFolderProps* FolderInWorld = Folders->Folders.Find(FFolder(FName(*KeyValue.Key))))
 				{
 					auto FolderProperties = KeyValue.Value->AsObject();
 					FolderInWorld->bIsExpanded = FolderProperties->GetBoolField(TEXT("bIsExpanded"));
@@ -270,31 +330,69 @@ UEditorActorFolders& FActorFolders::InitializeForWorld(UWorld& InWorld)
 	return *Folders;
 }
 
-FName FActorFolders::GetDefaultFolderNameForSelection(UWorld& InWorld)
+FFolder FActorFolders::GetDefaultFolderForSelection(UWorld& InWorld, TArray<FFolder>* InSelectedFolders)
 {
 	// Find a common parent folder, or put it at the root
-	FName CommonParentFolder;
+	TOptional<FFolder> CommonFolder;
+
+	auto MergeFolders = [&CommonFolder](FFolder& Folder)
+	{
+		if (!CommonFolder.IsSet())
+		{
+			CommonFolder = Folder;
+		}
+		else if (CommonFolder.GetValue().GetRootObject() != Folder.GetRootObject())
+		{
+			CommonFolder = FFolder();
+			return false;
+		}
+		else if (CommonFolder.GetValue().GetPath() != Folder.GetPath())
+		{
+			// Empty path and continue iterating as we need to continue validating RootObjects
+			CommonFolder.GetValue().SetPath(NAME_None);
+		}
+		return true;
+	};
+
+	bool bMergeStopped = false;
 	for( FSelectionIterator SelectionIt( *GEditor->GetSelectedActors() ); SelectionIt; ++SelectionIt )
 	{
 		AActor* Actor = CastChecked<AActor>(*SelectionIt);
-		if (CommonParentFolder.IsNone())
+
+		FFolder Folder = Actor->GetFolder();
+		// Special case for Level Instance, make root as level instance if editing
+		if (ALevelInstance* LevelInstance = Cast<ALevelInstance>(Actor))
 		{
-			CommonParentFolder = Actor->GetFolderPath();
+			if (LevelInstance->IsEditing())
+			{
+				Folder = FFolder(Folder.GetPath(), LevelInstance);
+			}
 		}
-		else if (Actor->GetFolderPath() != CommonParentFolder)
+		if (!MergeFolders(Folder))
 		{
-			CommonParentFolder = NAME_None;
+			bMergeStopped = true;
 			break;
 		}
 	}
+	if (!bMergeStopped && InSelectedFolders)
+	{
+		for (FFolder& Folder : *InSelectedFolders)
+		{
+			if (!MergeFolders(Folder))
+			{
+				break;
+			}
+		}
+	}
 
-	return GetDefaultFolderName(InWorld, CommonParentFolder);
+	return GetDefaultFolderName(InWorld, CommonFolder.Get(FFolder()));
 }
 
-FName FActorFolders::GetFolderName(UWorld& InWorld, FName ParentPath, FName InLeafName)
+FFolder FActorFolders::GetFolderName(UWorld& InWorld, const FFolder& InParentFolder, const FName& InLeafName)
 {
 	// This is potentially very slow but necessary to find a unique name
 	const auto& ExistingFolders = GetFolderPropertiesForWorld(InWorld);
+	const FFolder::FRootObject& RootObject = InParentFolder.GetRootObject();
 
 	const FString LeafNameString = InLeafName.ToString();
 
@@ -329,31 +427,33 @@ FName FActorFolders::GetFolderName(UWorld& InWorld, FName ParentPath, FName InLe
 
 	FText LeafName = FText::Format(LOCTEXT("FolderNamePattern", "{0}{1}"), FText::FromString(LeafNameRoot), SuffixLen > 0 ? FText::AsNumber(Suffix++, &NumberFormat) : FText::GetEmpty());
 
-	FString ParentFolderPath = ParentPath.IsNone() ? TEXT("") : ParentPath.ToString();
+	FString ParentFolderPath = InParentFolder.IsNone() ? TEXT("") : InParentFolder.ToString();
 	if (!ParentFolderPath.IsEmpty())
 	{
 		ParentFolderPath += "/";
 	}
 
 	FName FolderName(*(ParentFolderPath + LeafName.ToString()));
-	while (ExistingFolders.Contains(FolderName))
+	while (ExistingFolders.Contains(FFolder(FolderName, RootObject)))
 	{
 		LeafName = FText::Format(LOCTEXT("FolderNamePattern", "{0}{1}"), FText::FromString(LeafNameRoot), FText::AsNumber(Suffix++, &NumberFormat));
 		FolderName = FName(*(ParentFolderPath + LeafName.ToString()));
 		if (Suffix == 0)
 		{
 			// We've wrapped around a 32bit unsigned int - something must be seriously wrong!
-			return FName();
+			FolderName = NAME_None;
+			break;
 		}
 	}
 
-	return FolderName;
+	return FFolder(FolderName, RootObject);
 }
 
-FName FActorFolders::GetDefaultFolderName(UWorld& InWorld, FName ParentPath)
+FFolder FActorFolders::GetDefaultFolderName(UWorld& InWorld, const FFolder& InParentFolder)
 {
 	// This is potentially very slow but necessary to find a unique name
 	const auto& ExistingFolders = GetFolderPropertiesForWorld(InWorld);
+	const FFolder::FRootObject& RootObject = InParentFolder.GetRootObject();
 
 	// Create a valid base name for this folder
 	FNumberFormattingOptions NumberFormat;
@@ -361,38 +461,43 @@ FName FActorFolders::GetDefaultFolderName(UWorld& InWorld, FName ParentPath)
 	uint32 Suffix = 1;
 	FText LeafName = FText::Format(LOCTEXT("DefaultFolderNamePattern", "NewFolder{0}"), FText::AsNumber(Suffix++, &NumberFormat));
 
-	FString ParentFolderPath = ParentPath.IsNone() ? TEXT("") : ParentPath.ToString();
+	FString ParentFolderPath = InParentFolder.IsNone() ? TEXT("") : InParentFolder.ToString();
 	if (!ParentFolderPath.IsEmpty())
 	{
 		ParentFolderPath += "/";
 	}
 
 	FName FolderName(*(ParentFolderPath + LeafName.ToString()));
-	while (ExistingFolders.Contains(FolderName))
+	while (ExistingFolders.Contains(FFolder(FolderName, RootObject)))
 	{
 		LeafName = FText::Format(LOCTEXT("DefaultFolderNamePattern", "NewFolder{0}"), FText::AsNumber(Suffix++, &NumberFormat));
 		FolderName = FName(*(ParentFolderPath + LeafName.ToString()));
 		if (Suffix == 0)
 		{
 			// We've wrapped around a 32bit unsigned int - something must be seriously wrong!
-			return FName();
+			FolderName = NAME_None;
+			break;
 		}
 	}
 
-	return FolderName;
+	return FFolder(FolderName, RootObject);
 }
 
-void FActorFolders::CreateFolderContainingSelection(UWorld& InWorld, FName Path)
+void FActorFolders::CreateFolderContainingSelection(UWorld& InWorld, const FFolder& InFolder)
 {
 	const FScopedTransaction Transaction(LOCTEXT("UndoAction_CreateFolder", "Create Folder"));
-	CreateFolder(InWorld, Path);
-	SetSelectedFolderPath(Path);
+	CreateFolder(InWorld, InFolder);
+	SetSelectedFolderPath(InFolder);
 }
 
-void FActorFolders::SetSelectedFolderPath(FName Path) const
+void FActorFolders::SetSelectedFolderPath(const FFolder& InFolder) const
 {
 	// Move the currently selected actors into the new folder
 	USelection* SelectedActors = GEditor->GetSelectedActors();
+
+	const FFolder::FRootObject& RootObject = InFolder.GetRootObject();
+	FName Path = InFolder.GetPath();
+
 	for (FSelectionIterator SelectionIt(*SelectedActors); SelectionIt; ++SelectionIt)
 	{
 		AActor* Actor = CastChecked<AActor>(*SelectionIt);
@@ -404,35 +509,55 @@ void FActorFolders::SetSelectedFolderPath(FName Path) const
 			continue;
 		}
 
-		Actor->SetFolderPath_Recursively(Path);
+		// Currently not supported to change rootobject through this interface
+		if (RootObject == Actor->GetFolderRootObject())
+		{
+			Actor->SetFolderPath_Recursively(Path);
+		}
 	}
 }
 
-void FActorFolders::CreateFolder(UWorld& InWorld, FName Path)
+void FActorFolders::CreateFolder(UWorld& InWorld, const FFolder& InFolder)
 {
 	FScopedTransaction Transaction(LOCTEXT("UndoAction_CreateFolder", "Create Folder"));
 
-	if (AddFolderToWorld(InWorld, Path))
+	if (AddFolderToWorld(InWorld, InFolder))
 	{
-		OnFolderCreate.Broadcast(InWorld, Path);
+		BroadcastOnActorFolderCreated(InWorld, InFolder);
 	}
 }
 
-void FActorFolders::DeleteFolder(UWorld& InWorld, FName FolderToDelete)
+void FActorFolders::OnFolderRootObjectRemoved(UWorld& InWorld, const FFolder::FRootObject& InFolderRootObject)
+{
+	TArray<FFolder> FoldersToDelete;
+	ForEachFolderWithRootObject(InWorld, InFolderRootObject, [&FoldersToDelete](const FFolder& Folder)
+	{
+		FoldersToDelete.Add(Folder);
+		return true;
+	});
+
+	RemoveFoldersFromWorld(InWorld, FoldersToDelete, /*bBroadcastDelete*/ true);
+}
+
+void FActorFolders::DeleteFolder(UWorld& InWorld, const FFolder& InFolderToDelete)
 {
 	const FScopedTransaction Transaction(LOCTEXT("UndoAction_DeleteFolder", "Delete Folder"));
 
 	UEditorActorFolders& Folders = GetOrCreateFoldersForWorld(InWorld);
-	if (Folders.Folders.Contains(FolderToDelete))
+	if (Folders.Folders.Contains(InFolderToDelete))
 	{
 		Folders.Modify();
-		Folders.Folders.Remove(FolderToDelete);
-		OnFolderDelete.Broadcast(InWorld, FolderToDelete);
+		Folders.Folders.Remove(InFolderToDelete);
+		BroadcastOnActorFolderDeleted(InWorld, InFolderToDelete);
 	}
 }
 
-bool FActorFolders::RenameFolderInWorld(UWorld& World, FName OldPath, FName NewPath)
+bool FActorFolders::RenameFolderInWorld(UWorld& World, const FFolder& OldPath, const FFolder& NewPath)
 {
+	// We currently don't support changing the root object
+	check(OldPath.GetRootObject() == NewPath.GetRootObject());
+	const FFolder::FRootObject& RootObject = OldPath.GetRootObject();
+
 	const FString OldPathString = OldPath.ToString();
 	const FString NewPathString = NewPath.ToString();
 
@@ -443,7 +568,7 @@ bool FActorFolders::RenameFolderInWorld(UWorld& World, FName OldPath, FName NewP
 
 	const FScopedTransaction Transaction(LOCTEXT("UndoAction_RenameFolder", "Rename Folder"));
 
-	TSet<FName> RenamedFolders;
+	TSet<FFolder> RenamedFolders;
 	bool RenamedFolder = false;
 
 	// Move any folders we currently hold - old ones will be deleted later
@@ -453,18 +578,18 @@ bool FActorFolders::RenameFolderInWorld(UWorld& World, FName OldPath, FName NewP
 	auto ExistingFoldersCopy = FoldersInWorld.Folders;
 	for (const auto& Pair : ExistingFoldersCopy)
 	{
-		auto Path = Pair.Key;
+		const FFolder& Path = Pair.Key;
 
 		const FString FolderPath = Path.ToString();
 		if (OldPath == Path || FEditorFolderUtils::PathIsChildOf(FolderPath, OldPathString))
 		{
-			const FName NewFolder = OldPathToNewPath(OldPathString, NewPathString, FolderPath);
+			const FFolder NewFolder = FFolder(OldPathToNewPath(OldPathString, NewPathString, FolderPath), RootObject);
 			
 			// Needs to be done this way otherwise case insensitive comparison is used.
 			bool ContainsFolder = false;
 			for (const auto& FolderPair : FoldersInWorld.Folders)
 			{
-				if (FolderPair.Key.IsEqual(NewFolder, ENameCase::CaseSensitive))
+				if (FolderPair.Key.GetPath().IsEqual(NewFolder.GetPath(), ENameCase::CaseSensitive))
 				{
 					ContainsFolder = true;
 					break;
@@ -483,8 +608,8 @@ bool FActorFolders::RenameFolderInWorld(UWorld& World, FName OldPath, FName NewP
 					// Otherwise use default properties
 					FoldersInWorld.Folders.Add(NewFolder);
 				}
-				OnFolderMove.Broadcast(World, Path, NewFolder);
-				OnFolderCreate.Broadcast(World, NewFolder);
+				BroadcastOnActorFolderMoved(World, Path, NewFolder);
+				BroadcastOnActorFolderCreated(World, NewFolder);
 			}
 
 			// case insensitive compare as we don't want to remove the folder if it has the same name
@@ -501,7 +626,7 @@ bool FActorFolders::RenameFolderInWorld(UWorld& World, FName OldPath, FName NewP
 	for (auto ActorIt = FActorIterator(&World); ActorIt; ++ActorIt)
 	{
 		// copy, otherwise it returns the new value when set later
-		const FName OldActorPath = ActorIt->GetFolderPath();
+		const FFolder OldActorPath = ActorIt->GetFolder();
 		if (OldActorPath.IsNone())
 		{
 			continue;
@@ -510,7 +635,7 @@ bool FActorFolders::RenameFolderInWorld(UWorld& World, FName OldPath, FName NewP
 		if (OldActorPath == OldPath || FEditorFolderUtils::PathIsChildOf(OldActorPath.ToString(), OldPathString))
 		{
 			ActorIt->SetFolderPath_Recursively(OldPathToNewPath(OldPathString, NewPathString, OldActorPath.ToString()));
-			const FName& NewActorPath = ActorIt->GetFolderPath();
+			const FFolder NewActorPath = ActorIt->GetFolder();
 
 			// case insensitive compare as we don't want to remove the folder if it has the same name
 			if (OldActorPath != NewActorPath)
@@ -526,29 +651,29 @@ bool FActorFolders::RenameFolderInWorld(UWorld& World, FName OldPath, FName NewP
 	for (const auto& Path : RenamedFolders)
 	{
 		FoldersInWorld.Folders.Remove(Path);
-		OnFolderDelete.Broadcast(World, Path);
+		BroadcastOnActorFolderDeleted(World, Path);
 	}
 
 	return RenamedFolder;
 }
 
-bool FActorFolders::AddFolderToWorld(UWorld& InWorld, FName Path)
+bool FActorFolders::AddFolderToWorld(UWorld& InWorld, const FFolder& InFolder)
 {
-	if (!Path.IsNone())
+	if (!InFolder.IsNone())
 	{
 		UEditorActorFolders& Folders = GetOrCreateFoldersForWorld(InWorld);
 
-		if (!Folders.Folders.Contains(Path))
+		if (!Folders.Folders.Contains(InFolder))
 		{
 			// Add the parent as well
-			const FName ParentPath(*FPaths::GetPath(Path.ToString()));
-			if (!ParentPath.IsNone())
+			const FFolder ParentFolder = InFolder.GetParent();
+			if (!ParentFolder.IsNone())
 			{
-				AddFolderToWorld(InWorld, ParentPath);
+				AddFolderToWorld(InWorld, ParentFolder);
 			}
 
 			Folders.Modify();
-			Folders.Folders.Add(Path);
+			Folders.Folders.Add(InFolder);
 
 			return true;
 		}
@@ -557,10 +682,78 @@ bool FActorFolders::AddFolderToWorld(UWorld& InWorld, FName Path)
 	return false;
 }
 
-void FActorFolders::ForEachActorInFolders(UWorld& World, const TArray<FName>& Paths, TFunctionRef<bool(AActor*)> Operation)
+void FActorFolders::RemoveFoldersFromWorld(UWorld& InWorld, const TArray<FFolder>& InFolders, bool bBroadcastDelete)
+{
+	if (InFolders.Num() > 0)
+	{
+		UEditorActorFolders& Folders = GetOrCreateFoldersForWorld(InWorld);
+		Folders.Modify();
+		for (const FFolder& Folder : InFolders)
+		{
+			Folders.Folders.Remove(Folder);
+			if (bBroadcastDelete)
+			{
+				BroadcastOnActorFolderDeleted(InWorld, Folder);
+			}
+		}
+	}
+}
+
+bool FActorFolders::ContainsFolder(UWorld& InWorld, const FFolder& InFolder)
+{
+	return GetFolderProperties(InWorld, InFolder) != nullptr;
+}
+
+bool FActorFolders::IsFolderExpanded(UWorld& InWorld, const FFolder& InFolder)
+{
+	FActorFolderProps* FolderProps = GetFolderProperties(InWorld, InFolder);
+	return FolderProps ? FolderProps->bIsExpanded : false;
+}
+
+void FActorFolders::SetIsFolderExpanded(UWorld& InWorld, const FFolder& InFolder, bool bIsExpanded)
+{
+	if (FActorFolderProps* FolderProps = GetFolderProperties(InWorld, InFolder))
+	{
+		FolderProps->bIsExpanded = bIsExpanded;
+	}
+}
+
+void FActorFolders::ForEachFolder(UWorld& InWorld, TFunctionRef<bool(const FFolder&)> Operation)
+{
+	UEditorActorFolders& Folders = GetOrCreateFoldersForWorld(InWorld);
+	for (const auto& Pair : Folders.Folders)
+	{
+		if (!Operation(Pair.Key))
+		{
+			break;
+		}
+	}
+}
+
+void FActorFolders::ForEachFolderWithRootObject(UWorld& InWorld, const FFolder::FRootObject& InFolderRootObject, TFunctionRef<bool(const FFolder&)> Operation)
+{
+	UEditorActorFolders& Folders = GetOrCreateFoldersForWorld(InWorld);
+	for (const auto& Pair : Folders.Folders)
+	{
+		const FFolder& Folder = Pair.Key;
+		if (Folder.GetRootObject() == InFolderRootObject)
+		{
+			if (!Operation(Folder))
+			{
+				break;
+			}
+		}
+	}
+}
+
+void FActorFolders::ForEachActorInFolders(UWorld& World, const TArray<FName>& Paths, TFunctionRef<bool(AActor*)> Operation, const FFolder::FRootObject& InFolderRootObject)
 {
 	for (FActorIterator ActorIt(&World); ActorIt; ++ActorIt)
 	{
+		if (ActorIt->GetFolderRootObject() != InFolderRootObject)
+		{
+			continue;
+		}
 		FName ActorPath = ActorIt->GetFolderPath();
 		if (ActorPath.IsNone() || !Paths.Contains(ActorPath))
 		{
@@ -574,22 +767,76 @@ void FActorFolders::ForEachActorInFolders(UWorld& World, const TArray<FName>& Pa
 	}
 }
 
-void FActorFolders::GetActorsFromFolders(UWorld& World, const TArray<FName>& Paths, TArray<AActor*>& OutActors)
+void FActorFolders::GetActorsFromFolders(UWorld& World, const TArray<FName>& Paths, TArray<AActor*>& OutActors, const FFolder::FRootObject& InFolderRootObject)
 {
 	ForEachActorInFolders(World, Paths, [&OutActors](AActor* InActor)
-		{
-			OutActors.Add(InActor);
-			return true;
-		});
+	{
+		OutActors.Add(InActor);
+		return true;
+	}, InFolderRootObject);
 }
 
-void FActorFolders::GetWeakActorsFromFolders(UWorld& World, const TArray<FName>& Paths, TArray<TWeakObjectPtr<AActor>>& OutActors)
+void FActorFolders::GetWeakActorsFromFolders(UWorld& World, const TArray<FName>& Paths, TArray<TWeakObjectPtr<AActor>>& OutActors, const FFolder::FRootObject& InFolderRootObject)
 {
 	ForEachActorInFolders(World, Paths, [&OutActors](AActor* InActor)
-		{
-			OutActors.Add(InActor);
-			return true;
-		});
+	{
+		OutActors.Add(InActor);
+		return true;
+	}, InFolderRootObject);
+}
+
+////////////////////////////////////////////
+//~ Begin Deprecated
+
+FActorFolderProps* FActorFolders::GetFolderProperties(UWorld& InWorld, FName InPath)
+{
+	return GetFolderProperties(InWorld, FFolder(InPath));
+}
+
+FName FActorFolders::GetDefaultFolderName(UWorld& InWorld, FName ParentPath)
+{
+	FFolder DefaultFolderName = GetDefaultFolderName(InWorld, FFolder(ParentPath));
+	return DefaultFolderName.GetPath();
 }
 
-#undef LOCTEXT_NAMESPACE
+FName FActorFolders::GetDefaultFolderNameForSelection(UWorld& InWorld)
+{
+	FFolder FolderName = GetDefaultFolderForSelection(InWorld);
+	return FolderName.GetPath();
+}
+
+FName FActorFolders::GetFolderName(UWorld& InWorld, FName InParentPath, FName InFolderName)
+{
+	FFolder FolderName = GetFolderName(InWorld, FFolder(InParentPath), InFolderName);
+	return FolderName.GetPath();
+}
+
+void FActorFolders::CreateFolder(UWorld& InWorld, FName Path)
+{
+	CreateFolder(InWorld, FFolder(Path));
+}
+
+void FActorFolders::CreateFolderContainingSelection(UWorld& InWorld, FName Path)
+{
+	CreateFolderContainingSelection(InWorld, FFolder(Path));
+}
+
+void FActorFolders::SetSelectedFolderPath(FName Path) const
+{
+	SetSelectedFolderPath(FFolder(Path));
+}
+
+void FActorFolders::DeleteFolder(UWorld& InWorld, FName FolderToDelete)
+{
+	DeleteFolder(InWorld, FFolder(FolderToDelete));
+}
+
+bool FActorFolders::RenameFolderInWorld(UWorld& InWorld, FName OldPath, FName NewPath)
+{
+	return RenameFolderInWorld(InWorld, FFolder(OldPath), FFolder(NewPath));
+}
+
+//~ End Deprecated
+////////////////////////////////////////////
+
+#undef LOCTEXT_NAMESPACE
\ No newline at end of file
diff --git a/Engine/Source/Editor/UnrealEd/Private/EditorLevelUtils.cpp b/Engine/Source/Editor/UnrealEd/Private/EditorLevelUtils.cpp
--- a/Engine/Source/Editor/UnrealEd/Private/EditorLevelUtils.cpp
+++ b/Engine/Source/Editor/UnrealEd/Private/EditorLevelUtils.cpp
@@ -71,24 +71,24 @@ UEditorLevelUtils::FCanMoveActorToLevelDelegate UEditorLevelUtils::CanMoveActorT
 
 int32 UEditorLevelUtils::MoveActorsToLevel(const TArray<AActor*>& ActorsToMove, ULevelStreaming* DestStreamingLevel, bool bWarnAboutReferences, bool bWarnAboutRenaming)
 {
-	return MoveActorsToLevel(ActorsToMove, DestStreamingLevel ? DestStreamingLevel->GetLoadedLevel() : nullptr, bWarnAboutReferences, bWarnAboutRenaming);
+	return MoveActorsToLevel(ActorsToMove, DestStreamingLevel ? DestStreamingLevel->GetLoadedLevel() : nullptr, bWarnAboutReferences, bWarnAboutRenaming, /*bMoveAllOrFail*/false);
 }
 
-int32 UEditorLevelUtils::MoveActorsToLevel(const TArray<AActor*>& ActorsToMove, ULevel* DestLevel, bool bWarnAboutReferences, bool bWarnAboutRenaming, bool bMoveAllOrFail)
+int32 UEditorLevelUtils::MoveActorsToLevel(const TArray<AActor*>& ActorsToMove, ULevel* DestLevel, bool bWarnAboutReferences, bool bWarnAboutRenaming, bool bMoveAllOrFail, TArray<AActor*>* OutActors /*=nullptr*/)
 {
 	const bool bCutActors = true;
 	const bool bMoveActors = true;
-	return CopyOrMoveActorsToLevel(ActorsToMove, DestLevel, bCutActors, bMoveActors, bWarnAboutReferences, bWarnAboutRenaming, bMoveAllOrFail);
+	return CopyOrMoveActorsToLevel(ActorsToMove, DestLevel, bCutActors, bMoveActors, bWarnAboutReferences, bWarnAboutRenaming, bMoveAllOrFail, OutActors);
 }
 
-int32 UEditorLevelUtils::CopyActorsToLevel(const TArray<AActor*>& ActorsToMove, ULevel* DestLevel, bool bWarnAboutReferences, bool bWarnAboutRenaming, bool bMoveAllOrFail)
+int32 UEditorLevelUtils::CopyActorsToLevel(const TArray<AActor*>& ActorsToMove, ULevel* DestLevel, bool bWarnAboutReferences, bool bWarnAboutRenaming, bool bMoveAllOrFail, TArray<AActor*>* OutActors /*=nullptr*/)
 {
 	const bool bCutActors = false;
 	const bool bMoveActors = false;
 	return CopyOrMoveActorsToLevel(ActorsToMove, DestLevel, bCutActors, bMoveActors, bWarnAboutReferences, bWarnAboutRenaming, bMoveAllOrFail);
 }
 
-int32 UEditorLevelUtils::CopyOrMoveActorsToLevel(const TArray<AActor*>& ActorsToMove, ULevel* DestLevel, bool bCutActors, bool bMoveActors, bool bWarnAboutReferences /*= true*/, bool bWarnAboutRenaming /*= true*/, bool bMoveAllOrFail /*= false*/)
+int32 UEditorLevelUtils::CopyOrMoveActorsToLevel(const TArray<AActor*>& ActorsToMove, ULevel* DestLevel, bool bCutActors, bool bMoveActors, bool bWarnAboutReferences /*= true*/, bool bWarnAboutRenaming /*= true*/, bool bMoveAllOrFail /*= false*/, TArray<AActor*>* OutActors /*=nullptr*/)
 {
 	int32 NumMovedActors = 0;
 
@@ -206,6 +206,10 @@ int32 UEditorLevelUtils::CopyOrMoveActorsToLevel(const TArray<AActor*>& ActorsTo
 						check(Tuple.Value.IsNull());
 
 						Tuple.Value = FSoftObjectPath(Actor);
+						if (OutActors)
+						{
+							OutActors->Add(Actor);
+						}
 					}
 					else
 					{
@@ -863,6 +867,7 @@ void UEditorLevelUtils::PrivateRemoveLevelFromWorld(ULevel* InLevel)
 		InLevel->OwningWorld->RemoveFromWorld(InLevel);
 		check(InLevel->bIsVisible == false);
 	}
+	InLevel->bIsEditorBeingRemoved = true;
 
 	InLevel->ReleaseRenderingResources();
 
diff --git a/Engine/Source/Editor/UnrealEd/Public/DragAndDrop/FolderDragDropOp.h b/Engine/Source/Editor/UnrealEd/Public/DragAndDrop/FolderDragDropOp.h
--- a/Engine/Source/Editor/UnrealEd/Public/DragAndDrop/FolderDragDropOp.h
+++ b/Engine/Source/Editor/UnrealEd/Public/DragAndDrop/FolderDragDropOp.h
@@ -17,12 +17,15 @@ public:
 	TArray<FName> Folders;
 	/** World to which the folders belong */
 	TWeakObjectPtr<UWorld> World;
+	/** Root object of folder (can be invalid) */
+	FFolder::FRootObject RootObject;
 
-	void Init(TArray<FName> InFolders, UWorld* InWorld)
+	void Init(TArray<FName> InFolders, UWorld* InWorld, const FFolder::FRootObject& InRootObject = FFolder::GetDefaultRootObject())
 	{
 		check(InWorld != nullptr);
 		Folders = MoveTemp(InFolders);
 		World = InWorld;
+		RootObject = InRootObject;
 
 		CurrentIconBrush = FEditorStyle::Get().GetBrush(TEXT("SceneOutliner.FolderClosed"));
 		if (Folders.Num() == 1)
diff --git a/Engine/Source/Editor/UnrealEd/Public/EditorActorFolders.h b/Engine/Source/Editor/UnrealEd/Public/EditorActorFolders.h
--- a/Engine/Source/Editor/UnrealEd/Public/EditorActorFolders.h
+++ b/Engine/Source/Editor/UnrealEd/Public/EditorActorFolders.h
@@ -7,15 +7,23 @@
 #include "UObject/Object.h"
 #include "UObject/WeakObjectPtr.h"
 #include "UObject/GCObject.h"
+#include "Folder.h"
 #include "EditorActorFolders.generated.h"
 
-class AActor;
 class FObjectPostSaveContext;
+class AActor;
 
 /** Multicast delegates for broadcasting various folder events */
+
+//~ Begin Deprecated
 DECLARE_MULTICAST_DELEGATE_TwoParams(FOnActorFolderCreate, UWorld&, FName);
 DECLARE_MULTICAST_DELEGATE_TwoParams(FOnActorFolderDelete, UWorld&, FName);
 DECLARE_MULTICAST_DELEGATE_ThreeParams(FOnActorFolderMove, UWorld&, FName /* src */, FName /* dst */);
+//~ End Deprecated
+
+DECLARE_MULTICAST_DELEGATE_TwoParams(FOnActorFolderCreated, UWorld&, const FFolder&);
+DECLARE_MULTICAST_DELEGATE_TwoParams(FOnActorFolderDeleted, UWorld&, const FFolder&);
+DECLARE_MULTICAST_DELEGATE_ThreeParams(FOnActorFolderMoved, UWorld&, const FFolder& /* src */, const FFolder& /* dst */);
 
 USTRUCT()
 struct FActorFolderProps
@@ -42,7 +50,7 @@ public:
 public:
 	virtual void Serialize(FArchive& Ar) override;
 
-	TMap<FName, FActorFolderProps> Folders;
+	TMap<FFolder, FActorFolderProps> Folders;
 };
 
 /** Class responsible for managing an in-memory representation of actor folders in the editor */
@@ -72,50 +80,118 @@ struct UNREALED_API FActorFolders : public FGCObject
 	static void Cleanup();
 
 	/** Folder creation and deletion events. Called whenever a folder is created or deleted in a world. */
+	static FOnActorFolderCreated OnFolderCreated;
+	static FOnActorFolderMoved 	OnFolderMoved;
+	static FOnActorFolderDeleted OnFolderDeleted;
+
+	//~ Begin Deprecated
+
+	UE_DEPRECATED(5.0, "OnFolderCreate has been deprecated. Please use OnFolderCreated.")
 	static FOnActorFolderCreate OnFolderCreate;
+
+	UE_DEPRECATED(5.0, "OnFolderMove has been deprecated. Please use OnFolderMoved.")
 	static FOnActorFolderMove 	OnFolderMove;
-	static FOnActorFolderDelete OnFolderDelete;
 
-	/** Get a map of folder properties for the specified world (map of folder path -> properties) */
-	const TMap<FName, FActorFolderProps>& GetFolderPropertiesForWorld(UWorld& InWorld);
+	UE_DEPRECATED(5.0, "OnFolderDelete has been deprecated. Please use OnFolderDeleted.")
+	static FOnActorFolderDelete OnFolderDelete;
 
-	/** Get the folder properties for the specified path. Returns nullptr if no properties exist */
+	UE_DEPRECATED(5.0, "GetFolderProperties using FName has been deprecated. Please use new interface using FFolder.")
 	FActorFolderProps* GetFolderProperties(UWorld& InWorld, FName InPath);
 
-	/** Get a default folder name under the specified parent path */
+	UE_DEPRECATED(5.0, "GetDefaultFolderName using FName  has been deprecated. Please use new interface using FFolder.")
 	FName GetDefaultFolderName(UWorld& InWorld, FName ParentPath = FName());
 	
-	/** Get a new default folder name that would apply to the current selection */
+	UE_DEPRECATED(5.0, "GetDefaultFolderNameForSelection using FName  has been deprecated. Please use new interface using FFolder.")
 	FName GetDefaultFolderNameForSelection(UWorld& InWorld);
 
-	/** Get folder name that is unique under specified parent path */
+	UE_DEPRECATED(5.0, "GetFolderName using FName  has been deprecated. Please use new interface using FFolder.")
 	FName GetFolderName(UWorld& InWorld, FName ParentPath, FName FolderName);
 
-	/** Create a new folder in the specified world, of the specified path */
+	UE_DEPRECATED(5.0, "CreateFolder using FName  has been deprecated. Please use new interface using FFolder.")
 	void CreateFolder(UWorld& InWorld, FName Path);
 
-	/** Same as CreateFolder, but moves the current actor selection into the new folder as well */
+	UE_DEPRECATED(5.0, "CreateFolderContainingSelection using FName  has been deprecated. Please use new interface using FFolder.")
 	void CreateFolderContainingSelection(UWorld& InWorld, FName Path);
 
-	/** Sets the folder path for all the selected actors */
+	UE_DEPRECATED(5.0, "SetSelectedFolderPath using FName  has been deprecated. Please use new interface using FFolder.")
 	void SetSelectedFolderPath(FName Path) const;
 
-	/** Delete the specified folder in the world */
+	UE_DEPRECATED(5.0, "DeleteFolder using FName  has been deprecated. Please use new interface using FFolder.")
 	void DeleteFolder(UWorld& InWorld, FName FolderToDelete);
 
-	/** Rename the specified path to a new name */
+	UE_DEPRECATED(5.0, "RenameFolderInWorld using FName  has been deprecated. Please use new interface using FFolder.")
 	bool RenameFolderInWorld(UWorld& World, FName OldPath, FName NewPath);
 
+	//~ End Deprecated
+
 	/** Apply an operation to each actor in the given list of folders. Will stop when operation returns false. */
-	static void ForEachActorInFolders(UWorld& World, const TArray<FName>& Paths, TFunctionRef<bool(AActor*)> Operation);
-	
+	static void ForEachActorInFolders(UWorld& InWorld, const TArray<FName>& Paths, TFunctionRef<bool(AActor*)> Operation, const FFolder::FRootObject& InFolderRootObject = FFolder::GetDefaultRootObject());
+
 	/** Get an array of actors from a list of folders */
-	static void GetActorsFromFolders(UWorld& World, const TArray<FName>& Paths, TArray<AActor*>& OutActors);
+	static void GetActorsFromFolders(UWorld& InWorld, const TArray<FName>& Paths, TArray<AActor*>& OutActors, const FFolder::FRootObject& InFolderRootObject = FFolder::GetDefaultRootObject());
 
 	/** Get an array of weak actor pointers from a list of folders */
-	static void GetWeakActorsFromFolders(UWorld& World, const TArray<FName>& Paths, TArray<TWeakObjectPtr<AActor>>& OutActors);
+	static void GetWeakActorsFromFolders(UWorld& InWorld, const TArray<FName>& Paths, TArray<TWeakObjectPtr<AActor>>& OutActors, const FFolder::FRootObject& InFolderRootObject = FFolder::GetDefaultRootObject());
+
+	/** Get a default folder name under the specified parent path */
+	FFolder GetDefaultFolderName(UWorld& InWorld, const FFolder& InParentFolder);
+	
+	/** Get a new default folder name that would apply to the current selection */
+	FFolder GetDefaultFolderForSelection(UWorld& InWorld, TArray<FFolder>* InSelectedFolders = nullptr);
+
+	/** Get folder name that is unique under specified parent path */
+	FFolder GetFolderName(UWorld& InWorld, const FFolder& InParentFolder, const FName& InLeafName);
+
+	/** Create a new folder in the specified world, of the specified path */
+	void CreateFolder(UWorld& InWorld, const FFolder& InFolder);
+
+	/** Same as CreateFolder, but moves the current actor selection into the new folder as well */
+	void CreateFolderContainingSelection(UWorld& InWorld, const FFolder& InFolder);
+
+	/** Sets the folder path for all the selected actors */
+	void SetSelectedFolderPath(const FFolder& InFolder) const;
+
+	/** Delete the specified folder in the world */
+	void DeleteFolder(UWorld& InWorld, const FFolder& InFolderToDelete);
+
+	/** Rename the specified path to a new name */
+	bool RenameFolderInWorld(UWorld& InWorld, const FFolder& OldPath, const FFolder& NewPath);
+
+	/** Notify that a root object has been removed. Cleanup of existing folders with with root object. */
+	void OnFolderRootObjectRemoved(UWorld& InWorld, const FFolder::FRootObject& InFolderRootObject);
+
+	/** Return if folder exists */
+	bool ContainsFolder(UWorld& InWorld, const FFolder& InFolder);
+
+	/** Return the folder expansion state */
+	bool IsFolderExpanded(UWorld& InWorld, const FFolder& InFolder);
+
+	/** Set the folder expansion state */
+	void SetIsFolderExpanded(UWorld& InWorld, const FFolder& InFolder, bool bIsExpanded);
+
+	/** Iterate on all folders of a world and pass it to the provided operation. */
+	void ForEachFolder(UWorld& InWorld, TFunctionRef<bool(const FFolder&)> Operation);
+
+	/** Iterate on all world's folders with the given root object and pass it to the provided operation. */
+	void ForEachFolderWithRootObject(UWorld& InWorld, const FFolder::FRootObject& InFolderRootObject, TFunctionRef<bool(const FFolder&)> Operation);
+
 private:
 
+	/** Broadcast when actor folder is created. */
+	void BroadcastOnActorFolderCreated(UWorld& InWorld, const FFolder& InFolder);
+
+	/** Broadcast when actor folder is deleted. */
+	void BroadcastOnActorFolderDeleted(UWorld& InWorld, const FFolder& InFolder);
+
+	/** Broadcast when actor folder has moved. */
+	void BroadcastOnActorFolderMoved(UWorld& InWorld, const FFolder& InSrcFolder, const FFolder& InDstFolder);
+
+	/** Get the folder properties for the specified path. Returns nullptr if no properties exist */
+	FActorFolderProps* GetFolderProperties(UWorld& InWorld, const FFolder& InFolder);
+
+	/** Get a map of folder properties for the specified world (map of folder path -> properties) */
+	const TMap<FFolder, FActorFolderProps>& GetFolderPropertiesForWorld(UWorld& InWorld);
+
 	/** Returns true if folders have been created for the specified world */
 	bool FoldersExistForWorld(UWorld& InWorld) const;
 
@@ -145,7 +221,10 @@ private:
 	void Housekeeping();
 
 	/** Add a folder to the folder map for the specified world. Does not trigger any events. */
-	bool AddFolderToWorld(UWorld& InWorld, FName Path);
+	bool AddFolderToWorld(UWorld& InWorld, const FFolder& InFolder);
+
+	/** Removed folders from specified world. Can optionally trigger delete events. */
+	void RemoveFoldersFromWorld(UWorld& InWorld, const TArray<FFolder>& InFolders, bool bBroadcastDelete);
 
 	/** Transient map of folders, keyed on world pointer */
 	TMap<TWeakObjectPtr<UWorld>, UEditorActorFolders*> TemporaryWorldFolders;
diff --git a/Engine/Source/Editor/UnrealEd/Public/EditorLevelUtils.h b/Engine/Source/Editor/UnrealEd/Public/EditorLevelUtils.h
--- a/Engine/Source/Editor/UnrealEd/Public/EditorLevelUtils.h
+++ b/Engine/Source/Editor/UnrealEd/Public/EditorLevelUtils.h
@@ -86,9 +86,9 @@ public:
 	static UNREALED_API void MakeLevelCurrent(ULevel* InLevel, bool bEvenIfLocked = false);
 
 	
-	static UNREALED_API int32 MoveActorsToLevel(const TArray<AActor*>& ActorsToMove, ULevel* DestLevel, bool bWarnAboutReferences = true, bool bWarnAboutRenaming = true, bool bMoveAllOrFail = false);
+	static UNREALED_API int32 MoveActorsToLevel(const TArray<AActor*>& ActorsToMove, ULevel* DestLevel, bool bWarnAboutReferences = true, bool bWarnAboutRenaming = true, bool bMoveAllOrFail = false, TArray<AActor*>* OutActors = nullptr);
 
-	static UNREALED_API int32 CopyActorsToLevel(const TArray<AActor*>& ActorsToMove, ULevel* DestLevel, bool bWarnAboutReferences = true, bool bWarnAboutRenaming = true, bool bMoveAllOrFail = false);
+	static UNREALED_API int32 CopyActorsToLevel(const TArray<AActor*>& ActorsToMove, ULevel* DestLevel, bool bWarnAboutReferences = true, bool bWarnAboutRenaming = true, bool bMoveAllOrFail = false, TArray<AActor*>* OutActors = nullptr);
 
 	static UNREALED_API int32 MoveSelectedActorsToLevel(ULevel* DestLevel, bool bWarnAboutReferences = true);
 
@@ -169,7 +169,7 @@ private:
 
 	static UNREALED_API ULevelStreaming* AddLevelToWorld_Internal(UWorld* InWorld, const TCHAR* LevelPackageName, TSubclassOf<ULevelStreaming> LevelStreamingClass, const FTransform& LevelTransform = FTransform::Identity);
 
-	static UNREALED_API int32 CopyOrMoveActorsToLevel(const TArray<AActor*>& ActorsToMove, ULevel* DestLevel, bool bCutActors, bool bMoveActors, bool bWarnAboutReferences = true, bool bWarnAboutRenaming = true, bool bMoveAllOrFail = false);
+	static UNREALED_API int32 CopyOrMoveActorsToLevel(const TArray<AActor*>& ActorsToMove, ULevel* DestLevel, bool bCutActors, bool bMoveActors, bool bWarnAboutReferences = true, bool bWarnAboutRenaming = true, bool bMoveAllOrFail = false, TArray<AActor*>* OutActors = nullptr);
 
 public:
 	/** Sets the LevelStreamingClass for the specified Level 
diff --git a/Engine/Source/Runtime/CoreUObject/Public/UObject/ObjectKey.h b/Engine/Source/Runtime/CoreUObject/Public/UObject/ObjectKey.h
--- a/Engine/Source/Runtime/CoreUObject/Public/UObject/ObjectKey.h
+++ b/Engine/Source/Runtime/CoreUObject/Public/UObject/ObjectKey.h
@@ -75,6 +75,12 @@ public:
 		return ObjectIndex > Other.ObjectIndex || (ObjectIndex == Other.ObjectIndex && ObjectSerialNumber >= Other.ObjectSerialNumber);
 	}
 
+	FORCEINLINE friend FArchive& operator<<(FArchive& Ar, FObjectKey& Key)
+	{
+		check(!Ar.IsPersistent());
+		return Ar << Key.ObjectIndex << Key.ObjectSerialNumber;
+	}
+
 	/**
 	 * Attempt to access the object from which this key was constructed.
 	 * @return The object used to construct this key, or nullptr if it is no longer valid
diff --git a/Engine/Source/Runtime/Engine/Classes/Engine/Level.h b/Engine/Source/Runtime/Engine/Classes/Engine/Level.h
--- a/Engine/Source/Runtime/Engine/Classes/Engine/Level.h
+++ b/Engine/Source/Runtime/Engine/Classes/Engine/Level.h
@@ -617,6 +617,9 @@ public:
 	UPROPERTY()
 	uint8 										bLocked:1;
 
+	/** Whether the level is currently being removed from the world through UEditorLevelUtils::RemoveLevelsFromWorld */
+	uint8										bIsEditorBeingRemoved : 1;
+
 	/** Whether the level has been saved after introducing actor GUIDs */
 	uint8										bContainsStableActorGUIDs:1;
 #endif
@@ -1173,6 +1176,11 @@ public:
 	 */
 	ENGINE_API bool IsInstancedLevel() const;
 
+	/**
+	 * Is level being removed from world
+	 */
+	ENGINE_API bool IsBeingRemoved() const;
+
 	/** 
 	 * Shift level actors by specified offset
 	 * The offset vector will get subtracted from all actors positions and corresponding data structures
diff --git a/Engine/Source/Runtime/Engine/Classes/GameFramework/Actor.h b/Engine/Source/Runtime/Engine/Classes/GameFramework/Actor.h
--- a/Engine/Source/Runtime/Engine/Classes/GameFramework/Actor.h
+++ b/Engine/Source/Runtime/Engine/Classes/GameFramework/Actor.h
@@ -20,6 +20,7 @@
 
 #if WITH_EDITOR
 #include "WorldPartition/DataLayer/ActorDataLayer.h"
+#include "Folder.h"
 #endif
 
 #include "WorldPartition/WorldPartitionActorDescType.h"
@@ -2286,6 +2287,12 @@ public:
 	UFUNCTION(BlueprintCallable, Category = "Editor Scripting | Actor Editing")
 	const FName& GetFolderPath() const;
 
+	/** Returns the actor's folder root object. Null, is interpreted as the actor's world. */
+	FFolder::FRootObject GetFolderRootObject() const;
+
+	/** Returns a FFolder that contains the actor  folder path and its folder root object. */
+	FFolder GetFolder() const;
+
 	/**
 	 * Assigns a new folder to this actor. Actor folder paths are only available in development builds.
 	 * @param	NewFolderPath		The new folder to assign to the actor.
diff --git a/Engine/Source/Runtime/Engine/Private/ActorEditor.cpp b/Engine/Source/Runtime/Engine/Private/ActorEditor.cpp
--- a/Engine/Source/Runtime/Engine/Private/ActorEditor.cpp
+++ b/Engine/Source/Runtime/Engine/Private/ActorEditor.cpp
@@ -27,7 +27,11 @@
 #if WITH_EDITOR
 
 #include "Editor.h"
+#include "EditorActorFolders.h"
+#include "Folder.h"
 #include "WorldPartition/WorldPartitionActorDesc.h"
+#include "LevelInstance/LevelInstanceSubsystem.h"
+#include "LevelInstance/LevelInstanceActor.h"
 
 #define LOCTEXT_NAMESPACE "ErrorChecking"
 
@@ -943,9 +947,36 @@ void AActor::ClearActorLabel()
 	ActorLabel.Reset();
 }
 
+FFolder AActor::GetFolder() const
+{
+	return FFolder(GetFolderPath(), GetFolderRootObject());
+}
+
+static bool IsActorFolderValid(const AActor* InActor)
+{
+	const ULevel* Level = InActor->GetLevel();
+	const bool bIsBeingRemoved = Level && Level->IsBeingRemoved();
+	return !bIsBeingRemoved;
+}
+
+FFolder::FRootObject AActor::GetFolderRootObject() const
+{
+	if (IsActorFolderValid(this))
+	{
+		UWorld* World = GetWorld();
+		if (const ULevelInstanceSubsystem* LevelInstanceSubsystem = World ? World->GetSubsystem<ULevelInstanceSubsystem>() : nullptr)
+		{
+			ALevelInstance* LevelInstance = LevelInstanceSubsystem->GetParentLevelInstance(this);
+			return FFolder::FRootObject(LevelInstance);
+		}
+	}
+	return FFolder::GetDefaultRootObject();
+}
+
 const FName& AActor::GetFolderPath() const
 {
-	return FolderPath;
+	static const FName EmptyPath = NAME_None;
+	return IsActorFolderValid(this) ? FolderPath : EmptyPath;
 }
 
 void AActor::SetFolderPath(const FName& NewFolderPath)
diff --git a/Engine/Source/Runtime/Engine/Private/Level.cpp b/Engine/Source/Runtime/Engine/Private/Level.cpp
--- a/Engine/Source/Runtime/Engine/Private/Level.cpp
+++ b/Engine/Source/Runtime/Engine/Private/Level.cpp
@@ -2875,6 +2875,15 @@ bool ULevel::IsCurrentLevel() const
 	return bIsCurrent;
 }
 
+bool ULevel::IsBeingRemoved() const
+{
+	bool bResult = bIsBeingRemoved;
+#if WITH_EDITOR
+	bResult |= bIsEditorBeingRemoved;
+#endif
+	return bResult;
+}
+
 bool ULevel::IsInstancedLevel() const
 {
 	UPackage* LevelPackage = GetOutermost();
diff --git a/Engine/Source/Runtime/Engine/Private/LevelInstance/LevelInstanceEditorLevelStreaming.cpp b/Engine/Source/Runtime/Engine/Private/LevelInstance/LevelInstanceEditorLevelStreaming.cpp
--- a/Engine/Source/Runtime/Engine/Private/LevelInstance/LevelInstanceEditorLevelStreaming.cpp
+++ b/Engine/Source/Runtime/Engine/Private/LevelInstance/LevelInstanceEditorLevelStreaming.cpp
@@ -66,9 +66,10 @@ ULevelStreamingLevelInstanceEditor* ULevelStreamingLevelInstanceEditor::Load(ALe
 
 void ULevelStreamingLevelInstanceEditor::Unload(ULevelStreamingLevelInstanceEditor* LevelStreaming)
 {
-	// No need to clear the whole editor selection since actor of this level will be removed from the selection by: UEditorEngine::OnLevelRemovedFromWorld
-	const bool bClearSelection = false;
-	EditorLevelUtils::RemoveLevelFromWorld(LevelStreaming->GetLoadedLevel(), bClearSelection);
+	if (ULevelInstanceSubsystem* LevelInstanceSubsystem = LevelStreaming->GetWorld()->GetSubsystem<ULevelInstanceSubsystem>())
+	{
+		LevelInstanceSubsystem->RemoveLevelsFromWorld({ LevelStreaming->GetLoadedLevel() });
+	}
 }
 
 void ULevelStreamingLevelInstanceEditor::OnLevelActorAdded(AActor* InActor)
diff --git a/Engine/Source/Runtime/Engine/Private/LevelInstance/LevelInstanceLevelStreaming.cpp b/Engine/Source/Runtime/Engine/Private/LevelInstance/LevelInstanceLevelStreaming.cpp
--- a/Engine/Source/Runtime/Engine/Private/LevelInstance/LevelInstanceLevelStreaming.cpp
+++ b/Engine/Source/Runtime/Engine/Private/LevelInstance/LevelInstanceLevelStreaming.cpp
@@ -147,7 +147,7 @@ void ULevelStreamingLevelInstance::UnloadInstance(ULevelStreamingLevelInstance*
 			return true;
 		}, true);
 
-		LevelStreaming->GetWorld()->GetSubsystem<ULevelInstanceSubsystem>()->RemoveLevelFromWorld(LevelStreaming->GetLoadedLevel(), bResetTrans);
+		LevelStreaming->GetWorld()->GetSubsystem<ULevelInstanceSubsystem>()->RemoveLevelsFromWorld({ LevelStreaming->GetLoadedLevel() }, bResetTrans);
 	}
 #endif 
 }
diff --git a/Engine/Source/Runtime/Engine/Private/LevelInstance/LevelInstanceSubsystem.cpp b/Engine/Source/Runtime/Engine/Private/LevelInstance/LevelInstanceSubsystem.cpp
--- a/Engine/Source/Runtime/Engine/Private/LevelInstance/LevelInstanceSubsystem.cpp
+++ b/Engine/Source/Runtime/Engine/Private/LevelInstance/LevelInstanceSubsystem.cpp
@@ -40,6 +40,7 @@
 #include "EdGraph/EdGraph.h"
 #include "UObject/ObjectSaveContext.h"
 #include "UObject/UObjectGlobals.h"
+#include "EditorActorFolders.h"
 #endif
 
 #include "HAL/IConsoleManager.h"
@@ -123,12 +124,34 @@ FLevelInstanceID ULevelInstanceSubsystem::RegisterLevelInstance(ALevelInstance*
 	ALevelInstance*& Value = RegisteredLevelInstances.FindOrAdd(LevelInstanceID);
 	check(GIsReinstancing || Value == nullptr || Value == LevelInstanceActor);
 	Value = LevelInstanceActor;
+
+#if WITH_EDITOR
+	FObjectKey Level;
+	if (UnregisteringLevelInstances.RemoveAndCopyValue(LevelInstanceActor, Level))
+	{
+		if (ULevel* LevelPtr = Cast<ULevel>(Level.ResolveObjectPtr()))
+		{
+			LevelPtr->bIsEditorBeingRemoved = false;
+		}
+		ensure(UnregisteringLevelInstanceLevels.Remove(Level));
+	}
+#endif
+
 	return LevelInstanceID;
 }
 
 void ULevelInstanceSubsystem::UnregisterLevelInstance(ALevelInstance* LevelInstanceActor)
 {
 	RegisteredLevelInstances.Remove(LevelInstanceActor->GetLevelInstanceID());
+
+#if WITH_EDITOR
+	if (ULevel* Level = GetLevelInstanceLevel(LevelInstanceActor))
+	{
+		Level->bIsEditorBeingRemoved = true;
+		UnregisteringLevelInstanceLevels.Add(Level, LevelInstanceActor);
+		UnregisteringLevelInstances.Add(LevelInstanceActor, Level);
+	}
+#endif
 }
 
 void ULevelInstanceSubsystem::RequestLoadLevelInstance(ALevelInstance* LevelInstanceActor, bool bForce /* = false */)
@@ -140,6 +163,13 @@ void ULevelInstanceSubsystem::RequestLoadLevelInstance(ALevelInstance* LevelInst
 		if (!IsEditingLevelInstance(LevelInstanceActor))
 #endif
 		{
+#if WITH_EDITOR
+			if (ULevel* Level = GetLevelInstanceLevel(LevelInstanceActor))
+			{
+				Level->bIsEditorBeingRemoved = false;
+			}
+#endif
+			
 			LevelInstancesToUnload.Remove(LevelInstanceActor->GetLevelInstanceID());
 
 			bool* bForcePtr = LevelInstancesToLoadOrUpdate.Find(LevelInstanceActor);
@@ -164,6 +194,13 @@ void ULevelInstanceSubsystem::RequestLoadLevelInstance(ALevelInstance* LevelInst
 
 void ULevelInstanceSubsystem::RequestUnloadLevelInstance(ALevelInstance* LevelInstanceActor)
 {
+#if WITH_EDITOR
+	if (ULevel* Level = GetLevelInstanceLevel(LevelInstanceActor))
+	{
+		Level->bIsEditorBeingRemoved = true;
+	}
+#endif
+
 	const FLevelInstanceID& LevelInstanceID = LevelInstanceActor->GetLevelInstanceID();
 	if (LevelInstances.Contains(LevelInstanceID))
 	{
@@ -196,7 +233,7 @@ void ULevelInstanceSubsystem::UpdateStreamingState()
 	SlowTask.MakeDialogDelayed(1.0f);
 
 	check(!LevelsToRemoveScope);
-	LevelsToRemoveScope.Reset(new FLevelsToRemoveScope());
+	LevelsToRemoveScope.Reset(new FLevelsToRemoveScope(this));
 #endif
 
 	if (LevelInstancesToUnload.Num())
@@ -278,7 +315,7 @@ void ULevelInstanceSubsystem::UnloadLevelInstance(const FLevelInstanceID& LevelI
 	if (!LevelsToRemoveScope)
 	{
 		bReleaseScope = true;
-		LevelsToRemoveScope.Reset(new FLevelsToRemoveScope());
+		LevelsToRemoveScope.Reset(new FLevelsToRemoveScope(this));
 	}
 #endif
 				
@@ -287,6 +324,10 @@ void ULevelInstanceSubsystem::UnloadLevelInstance(const FLevelInstanceID& LevelI
 	{
 		if (ULevel* LoadedLevel = LevelInstance.LevelStreaming->GetLoadedLevel())
 		{
+#if WITH_EDITOR
+			LoadedLevel->bIsEditorBeingRemoved = true;
+#endif
+
 			ForEachActorInLevel(LoadedLevel, [this](AActor* LevelActor)
 			{
 				if (ALevelInstance* LevelInstanceActor = Cast<ALevelInstance>(LevelActor))
@@ -734,14 +775,14 @@ bool ULevelInstanceSubsystem::IsCurrent(const ALevelInstance* LevelInstanceActor
 	return false;
 }
 
-bool ULevelInstanceSubsystem::MoveActorsToLevel(const TArray<AActor*>& ActorsToRemove, ULevel* DestinationLevel) const
+bool ULevelInstanceSubsystem::MoveActorsToLevel(const TArray<AActor*>& ActorsToRemove, ULevel* DestinationLevel, TArray<AActor*>* OutActors /*= nullptr*/) const
 {
 	check(DestinationLevel);
 
 	const bool bWarnAboutReferences = true;
 	const bool bWarnAboutRenaming = true;
 	const bool bMoveAllOrFail = true;
-	if (!EditorLevelUtils::MoveActorsToLevel(ActorsToRemove, DestinationLevel, bWarnAboutReferences, bWarnAboutRenaming, bMoveAllOrFail))
+	if (!EditorLevelUtils::MoveActorsToLevel(ActorsToRemove, DestinationLevel, bWarnAboutReferences, bWarnAboutRenaming, bMoveAllOrFail, OutActors))
 	{
 		UE_LOG(LogLevelInstance, Warning, TEXT("Failed to move actors out of Level Instance because not all actors could be moved"));
 		return false;
@@ -760,13 +801,13 @@ bool ULevelInstanceSubsystem::MoveActorsToLevel(const TArray<AActor*>& ActorsToR
 	return true;
 }
 
-bool ULevelInstanceSubsystem::MoveActorsTo(ALevelInstance* LevelInstanceActor, const TArray<AActor*>& ActorsToMove)
+bool ULevelInstanceSubsystem::MoveActorsTo(ALevelInstance* LevelInstanceActor, const TArray<AActor*>& ActorsToMove, TArray<AActor*>* OutActors /*= nullptr*/)
 {
 	check(IsEditingLevelInstance(LevelInstanceActor));
 	ULevel* LevelInstanceLevel = GetLevelInstanceLevel(LevelInstanceActor);
 	check(LevelInstanceLevel);
 
-	return MoveActorsToLevel(ActorsToMove, LevelInstanceLevel);
+	return MoveActorsToLevel(ActorsToMove, LevelInstanceLevel, OutActors);
 }
 
 ALevelInstance* ULevelInstanceSubsystem::CreateLevelInstanceFrom(const TArray<AActor*>& ActorsToMove, const FNewLevelInstanceParams& CreationParams)
@@ -1146,29 +1187,76 @@ bool ULevelInstanceSubsystem::LevelInstanceHasLevelScriptBlueprint(const ALevelI
 	return false;
 }
 
-void ULevelInstanceSubsystem::RemoveLevelFromWorld(ULevel* Level, bool bResetTrans)
+void ULevelInstanceSubsystem::RemoveLevelsFromWorld(const TArray<ULevel*>& InLevels, bool bResetTrans)
 {
-	if (LevelsToRemoveScope)
+	if (LevelsToRemoveScope && LevelsToRemoveScope->IsValid())
+{
+		for (ULevel* Level : InLevels)
 	{
 		LevelsToRemoveScope->Levels.AddUnique(Level);
+		}
 		LevelsToRemoveScope->bResetTrans |= bResetTrans;
 	}
 	else
 	{
+		TSet<ULevel*> LevelInstanceLevels;
+		for (ULevel* Level : InLevels)
+		{
+			bool bIsAlreadyInSet = false;
+			LevelInstanceLevels.Add(Level, &bIsAlreadyInSet);
+
+			if (!bIsAlreadyInSet)
+			{
+				ForEachActorInLevel(Level, [this, &LevelInstanceLevels](AActor* LevelActor)
+				{
+					if (const ALevelInstance* ChildLevelInstanceActor = Cast<ALevelInstance>(LevelActor))
+					{
+						ForEachLevelInstanceChild(ChildLevelInstanceActor, /*bRecursive*/true, [this, &LevelInstanceLevels](const ALevelInstance* ChildLevelInstanceActor)
+						{
+							if (ULevel* ChildLevel = GetLevelInstanceLevel(ChildLevelInstanceActor))
+							{
+								LevelInstanceLevels.Add(ChildLevel);
+							}
+							return true;
+						});
+					}
+					return true;
+				});
+			}
+		}
+
+		// Flag all levels as bEditorBeingDestroyed. This way, even if child level instances are still pending to be unloaded, TActorIterator won't iterate on them.
+		for (ULevel* Level : LevelInstanceLevels)
+		{
+			Level->bIsEditorBeingRemoved = true;
+			FFolder::FRootObject RootObject;
+			if (UnregisteringLevelInstanceLevels.RemoveAndCopyValue(Level, RootObject))
+			{
+				ensure(UnregisteringLevelInstances.Remove(RootObject));
+				FActorFolders::Get().OnFolderRootObjectRemoved(*GetWorld(), RootObject);
+			}
+		}
+
 		// No need to clear the whole editor selection since actor of this level will be removed from the selection by: UEditorEngine::OnLevelRemovedFromWorld
-		const bool bClearSelection = false;
-		EditorLevelUtils::RemoveLevelFromWorld(Level, bClearSelection, bResetTrans);
+		EditorLevelUtils::RemoveLevelsFromWorld(InLevels, /*bClearSelection*/false, bResetTrans);
+	}
 	}
+
+ULevelInstanceSubsystem::FLevelsToRemoveScope::FLevelsToRemoveScope(ULevelInstanceSubsystem* InOwner)
+	: Owner(InOwner)
+	, bIsBeingDestroyed(false)
+{
 }
 
 ULevelInstanceSubsystem::FLevelsToRemoveScope::~FLevelsToRemoveScope()
 {
 	if (Levels.Num() > 0)
 	{
+		bIsBeingDestroyed = true;
 		double StartTime = FPlatformTime::Seconds();
-		const bool bClearSelection = false;
-		// No need to clear the whole editor selection since actor of this level will be removed from the selection by: UEditorEngine::OnLevelRemovedFromWorld
-		EditorLevelUtils::RemoveLevelsFromWorld(Levels, bClearSelection, bResetTrans);
+		ULevelInstanceSubsystem* LevelInstanceSubsystem = Owner.Get();
+		check(LevelInstanceSubsystem);
+		LevelInstanceSubsystem->RemoveLevelsFromWorld(Levels, bResetTrans);
 		double ElapsedTime = FPlatformTime::Seconds() - StartTime;
 		UE_LOG(LogLevelInstance, Log, TEXT("Unloaded %s levels in %s seconds"), *FText::AsNumber(Levels.Num()).ToString(), *FText::AsNumber(ElapsedTime).ToString());
 	}
@@ -1764,6 +1852,8 @@ ALevelInstance* ULevelInstanceSubsystem::CommitLevelInstance(ALevelInstance* Lev
 		}
 	}
 
+	GEngine->BroadcastLevelActorListChanged();
+
 	if (ALevelInstance* Actor = ActorToSelect.Get())
 	{
 		GEditor->SelectActor(Actor, true, true);
diff --git a/Engine/Source/Runtime/Engine/Private/LevelStreaming.cpp b/Engine/Source/Runtime/Engine/Private/LevelStreaming.cpp
--- a/Engine/Source/Runtime/Engine/Private/LevelStreaming.cpp
+++ b/Engine/Source/Runtime/Engine/Private/LevelStreaming.cpp
@@ -880,6 +880,10 @@ ULevelStreaming* ULevelStreaming::FindStreamingLevel(const ULevel* Level)
 
 void ULevelStreaming::SetLoadedLevel(ULevel* Level)
 { 
+#if WITH_EDITORONLY_DATA
+	check(!Level || !Level->IsBeingRemoved());
+#endif
+
 	// Pending level should be unloaded at this point
 	check(PendingUnloadLevel == nullptr);
 	PendingUnloadLevel = LoadedLevel;
diff --git a/Engine/Source/Runtime/Engine/Public/EngineUtils.h b/Engine/Source/Runtime/Engine/Public/EngineUtils.h
--- a/Engine/Source/Runtime/Engine/Public/EngineUtils.h
+++ b/Engine/Source/Runtime/Engine/Public/EngineUtils.h
@@ -424,7 +424,7 @@ protected:
 	{
 		if (EnumHasAnyFlags(Flags, EActorIteratorFlags::OnlyActiveLevels))
 		{
-			const bool bIsLevelVisibleOrAssociating = (Level->bIsVisible && !Level->bIsBeingRemoved) || Level->bIsAssociatingLevel || Level->bIsDisassociatingLevel;
+			const bool bIsLevelVisibleOrAssociating = (Level->bIsVisible && !Level->IsBeingRemoved()) || Level->bIsAssociatingLevel || Level->bIsDisassociatingLevel;
 
 			// Only allow iteration of Level if it's in the currently active level collection of the world, or is a static level.
 			const FLevelCollection* const ActorLevelCollection = Level->GetCachedLevelCollection();
diff --git a/Engine/Source/Runtime/Engine/Public/Folder.h b/Engine/Source/Runtime/Engine/Public/Folder.h
new file mode 100644
--- /dev/null
+++ b/Engine/Source/Runtime/Engine/Public/Folder.h
@@ -0,0 +1,140 @@
+// Copyright Epic Games, Inc. All Rights Reserved.
+
+#pragma once
+
+#include "CoreMinimal.h"
+#include "UObject/ObjectMacros.h"
+#include "UObject/Object.h"
+#include "UObject/ObjectKey.h"
+#include "UObject/WeakObjectPtr.h"
+#include "Misc/Paths.h"
+
+struct FFolder
+{
+	typedef FObjectKey FRootObject;
+
+	FFolder(const FName& InPath = NAME_None, const FRootObject& InRootObject = GetDefaultRootObject())
+		: Path(InPath)
+		, RootObject(InRootObject)
+	{}
+
+	FORCEINLINE static bool HasRootObject(const FRootObject& Key)
+	{
+		return Key != GetDefaultRootObject();
+	}
+
+	FORCEINLINE static FRootObject GetDefaultRootObject()
+	{
+		return FRootObject();
+	}
+
+	FORCEINLINE bool HasRootObject() const
+	{
+		return FFolder::HasRootObject(RootObject);
+	}
+
+	FORCEINLINE FFolder GetParent() const
+	{
+		const FName ParentPath(*FPaths::GetPath(Path.ToString()));
+		return FFolder(ParentPath, RootObject);
+	}
+
+	FORCEINLINE bool IsChildOf(const FFolder& InParent) const
+	{
+		if (RootObject != InParent.RootObject)
+		{
+			return false;
+		}
+		return PathIsChildOf(GetPath(), InParent.GetPath());
+	}
+
+	FORCEINLINE bool IsNone() const
+	{
+		return Path.IsNone();
+	}
+
+	const FRootObject& GetRootObject() const
+	{
+		return RootObject;
+	}
+
+	FORCEINLINE UObject* GetRootObjectPtr() const
+	{
+		return RootObject.ResolveObjectPtr();
+	}
+
+	FORCEINLINE const FName& GetPath() const
+	{
+		return Path;
+	}
+
+	FORCEINLINE void SetPath(const FName& InPath)
+	{
+		Path = InPath;
+	}
+
+	FORCEINLINE FName GetLeafName() const
+	{
+		FString PathString = Path.ToString();
+		int32 LeafIndex = 0;
+		if (PathString.FindLastChar('/', LeafIndex))
+		{
+			return FName(*PathString.RightChop(LeafIndex + 1));
+		}
+		else
+		{
+			return Path;
+		}
+	}
+
+	FORCEINLINE bool operator == (const FFolder& InOther) const
+	{
+		return (Path == InOther.Path) && (RootObject == InOther.RootObject);
+	}
+
+	FORCEINLINE bool operator != (const FFolder& InOther) const
+	{
+		return !operator==(InOther);
+	}
+
+	FORCEINLINE FString ToString() const
+	{
+		return Path.ToString();
+	}
+
+	FORCEINLINE friend FArchive& operator<<(FArchive& Ar, FFolder& Folder)
+	{
+		check(!Ar.IsPersistent());
+		return Ar << Folder.Path << Folder.RootObject;
+	}
+
+private:
+	bool PathIsChildOf(const FString& InPotentialChild, const FString& InParent) const
+	{
+		const int32 ParentLen = InParent.Len();
+		
+		// If parent is empty and child isn't, consider that path is child of parent
+		if ((InPotentialChild.Len() > 0) && (ParentLen == 0))
+		{
+			return true;
+		}
+
+		return
+			InPotentialChild.Len() > ParentLen &&
+			InPotentialChild[ParentLen] == '/' &&
+			InPotentialChild.Left(ParentLen) == InParent;
+	}
+
+	bool PathIsChildOf(const FName& InPotentialChild, const FName& InParent) const
+	{
+		return PathIsChildOf(InPotentialChild.ToString(), InParent.ToString());
+	}
+
+	FName Path;
+	FRootObject RootObject;
+};
+
+FORCEINLINE uint32 GetTypeHash(const FFolder& InFolder)
+{
+	return HashCombine(GetTypeHash(InFolder.GetPath()), GetTypeHash(InFolder.GetRootObject()));
+}
\ No newline at end of file
diff --git a/Engine/Source/Runtime/Engine/Public/LevelInstance/LevelInstanceSubsystem.h b/Engine/Source/Runtime/Engine/Public/LevelInstance/LevelInstanceSubsystem.h
--- a/Engine/Source/Runtime/Engine/Public/LevelInstance/LevelInstanceSubsystem.h
+++ b/Engine/Source/Runtime/Engine/Public/LevelInstance/LevelInstanceSubsystem.h
@@ -5,6 +5,7 @@
 #include "Containers/Array.h"
 #include "Containers/Map.h"
 #include "UObject/NameTypes.h"
+#include "Folder.h"
 #include "LevelInstance/LevelInstanceTypes.h"
 
 #include "LevelInstanceSubsystem.generated.h"
@@ -78,8 +79,8 @@ public:
 	bool SetCurrent(ALevelInstance* LevelInstanceActor) const;
 	bool IsCurrent(const ALevelInstance* LevelInstanceActor) const;
 	ALevelInstance* CreateLevelInstanceFrom(const TArray<AActor*>& ActorsToMove, const FNewLevelInstanceParams& CreationParams);
-	bool MoveActorsToLevel(const TArray<AActor*>& ActorsToRemove, ULevel* DestinationLevel) const;
-	bool MoveActorsTo(ALevelInstance* LevelInstanceActor, const TArray<AActor*>& ActorsToMove);
+	bool MoveActorsToLevel(const TArray<AActor*>& ActorsToRemove, ULevel* DestinationLevel, TArray<AActor*>* OutActors = nullptr) const;
+	bool MoveActorsTo(ALevelInstance* LevelInstanceActor, const TArray<AActor*>& ActorsToMove, TArray<AActor*>* OutActors = nullptr);
 	bool BreakLevelInstance(ALevelInstance* LevelInstanceActor, uint32 Levels = 1, TArray<AActor*>* OutMovedActors = nullptr);
 
 	bool CanMoveActorToLevel(const AActor* Actor, FText* OutReason = nullptr) const;
@@ -123,14 +124,19 @@ private:
 	
 	struct FLevelsToRemoveScope
 	{
+		FLevelsToRemoveScope(ULevelInstanceSubsystem* InOwner);
 		~FLevelsToRemoveScope();
+		bool IsValid() const { return !bIsBeingDestroyed; }
 
 		TArray<ULevel*> Levels;
+		TWeakObjectPtr<ULevelInstanceSubsystem> Owner;
 		bool bResetTrans = false;
+		bool bIsBeingDestroyed = false;
 	};
 	
 	friend ULevelStreamingLevelInstance;
-	void RemoveLevelFromWorld(ULevel* Level, bool bResetTrans);
+	friend ULevelStreamingLevelInstanceEditor;
+	void RemoveLevelsFromWorld(const TArray<ULevel*>& Levels, bool bResetTrans = true);
 #endif
 
 
@@ -150,5 +156,7 @@ private:
 	TSet<FLevelInstanceID> LevelInstancesToUnload;
 	TMap<FLevelInstanceID, FLevelInstance> LevelInstances;
 	TMap<FLevelInstanceID, ALevelInstance*> RegisteredLevelInstances;
+	TMap<FObjectKey, FFolder::FRootObject> UnregisteringLevelInstanceLevels;
+	TMap<FFolder::FRootObject, FObjectKey > UnregisteringLevelInstances;
 	FLevelInstanceID PendingLevelInstanceToEdit;
 };
\ No newline at end of file
