Animation: add a dedicated setting for timecode custom attributes

This change takes the set of custom attributes that represent timecode and
subframe values and a take name and creates a dedicated setting that maps
each component to an attribute name. This allows the actual attribute name
used in a pipeline to still be configured, but adds the ability to refer directly to
each attribute value through a named struct member. Engine code can then use
that to access the appropriate attribute for each component.

An accessor function was added to UAnimationSettings to get the names of all
bone custom attributes that should be considered during import. This returns
all of the timecode attribute names identified by the new timecode settings as
well as any other attributes in the BoneCustomAttributesNames member, yielding
the same list of attributes as before. Existing call sites in the FBX importer that
used the BoneCustomAttributesNames member directly now use this new accessor
instead.

diff --git a/Engine/Config/BaseEngine.ini b/Engine/Config/BaseEngine.ini
--- a/Engine/Config/BaseEngine.ini
+++ b/Engine/Config/BaseEngine.ini
@@ -2803,12 +2803,7 @@ HMDWornMovementThreshold = 50.0
 
 [/Script/Engine.AnimationSettings]
 bStripAnimationDataOnDedicatedServer=False
-+BoneCustomAttributesNames=(Name="TCHour",Meaning="Timecode Hour")
-+BoneCustomAttributesNames=(Name="TCMinute",Meaning="Timecode Minute")
-+BoneCustomAttributesNames=(Name="TCSecond",Meaning="Timecode Second")
-+BoneCustomAttributesNames=(Name="TCFrame",Meaning="Timecode Frame")
-+BoneCustomAttributesNames=(Name="TCSubframe",Meaning="Timecode Subframe")
-+BoneCustomAttributesNames=(Name="Takename",Meaning="Takename")
+BoneTimecodeCustomAttributeNameSettings=(HourAttributeName="TCHour",MinuteAttributeName="TCMinute",SecondAttributeName="TCSecond",FrameAttributeName="TCFrame",SubframeAttributeName="TCSubframe",TakenameAttributeName="Takename")
 
 [Animation.DefaultObjectSettings]
 BoneCompressionSettings="/Engine/Animation/DefaultAnimBoneCompressionSettings"
diff --git a/Engine/Source/Editor/UnrealEd/Private/FbxAnimUtils.cpp b/Engine/Source/Editor/UnrealEd/Private/FbxAnimUtils.cpp
--- a/Engine/Source/Editor/UnrealEd/Private/FbxAnimUtils.cpp
+++ b/Engine/Source/Editor/UnrealEd/Private/FbxAnimUtils.cpp
@@ -220,6 +220,7 @@ namespace FbxAnimUtils
 	{
 		FbxProperty Property = InNode->GetFirstProperty();
 		UAnimationSettings* AnimationSettings = UAnimationSettings::Get();
+		const TArray<FString> CustomAttributeNamesToImport = AnimationSettings->GetBoneCustomAttributeNamesToImport();
 
 		while (Property.IsValid())
 		{
@@ -229,7 +230,7 @@ namespace FbxAnimUtils
 
 			if (CurveNode && Property.GetFlag(FbxPropertyFlags::eUserDefined) &&
 				CurveNode->IsAnimated() && FbxAnimUtils::IsSupportedCurveDataType(Property.GetPropertyDataType().GetType()) &&
-				!AnimationSettings->BoneCustomAttributesNames.ContainsByPredicate([&PropertyName](const FCustomAttributeSetting& CustomAttributeSetting) { return CustomAttributeSetting.Name.Equals(PropertyName, ESearchCase::IgnoreCase); }))
+				!CustomAttributeNamesToImport.Contains(PropertyName))
 			{
 				FString CurveName = UTF8_TO_TCHAR(CurveNode->GetName());
 				UE_LOG(LogFbx, Log, TEXT("CurveName : %s"), *CurveName);
@@ -261,6 +262,7 @@ namespace FbxAnimUtils
 	{
 		FbxProperty Property = InNode->GetFirstProperty();
 		UAnimationSettings* AnimationSettings = UAnimationSettings::Get();
+		const TArray<FString> CustomAttributeNamesToImport = AnimationSettings->GetBoneCustomAttributeNamesToImport();
 
 		while (Property.IsValid())
 		{
@@ -270,7 +272,7 @@ namespace FbxAnimUtils
 			const bool bIsTypeSupported = FbxAnimUtils::IsSupportedCurveDataType(PropertyType) || eFbxString == PropertyType;
 
 			if (Property.GetFlag(FbxPropertyFlags::eUserDefined) && bIsTypeSupported 
-				&& (bImportAllCustomAttributes || AnimationSettings->BoneCustomAttributesNames.ContainsByPredicate([&PropertyName](const FCustomAttributeSetting& CustomAttributeSetting) { return CustomAttributeSetting.Name.Equals(PropertyName, ESearchCase::IgnoreCase); })))
+				&& (bImportAllCustomAttributes || CustomAttributeNamesToImport.Contains(PropertyName)))
 			{
 				UE_LOG(LogFbx, Log, TEXT("PropertyName : %s"), *PropertyName);
 
diff --git a/Engine/Source/Runtime/Engine/Classes/Animation/AnimationSettings.h b/Engine/Source/Runtime/Engine/Classes/Animation/AnimationSettings.h
--- a/Engine/Source/Runtime/Engine/Classes/Animation/AnimationSettings.h
+++ b/Engine/Source/Runtime/Engine/Classes/Animation/AnimationSettings.h
@@ -61,10 +61,20 @@ class ENGINE_API UAnimationSettings : public UDeveloperSettings
 	UPROPERTY(config, EditAnywhere, Category = Performance)
 	bool bTickAnimationOnSkeletalMeshInit;
 
+	/** Names that identify bone custom attributes representing the individual components of a timecode and a subframe along with a take name.
+	    These will be included in the list of bone custom attribute names to import. */
+	UPROPERTY(config, EditAnywhere, Category = CustomAttributes)
+	FTimecodeCustomAttributeNameSettings BoneTimecodeCustomAttributeNameSettings;
+
 	/** List of custom attribute to import directly on their corresponding bone. The meaning field allows to contextualize the attribute name and customize tooling for it. */
 	UPROPERTY(config, EditAnywhere, Category = CustomAttributes)
 	TArray<FCustomAttributeSetting> BoneCustomAttributesNames;
 
+	/** Gets the complete list of bone custom attribute names to consider for import.
+	    This includes the designated timecode custom attributes as well as other bone custom attributes identified in the settings. */
+	UFUNCTION(BlueprintPure, Category = CustomAttributes)
+	TArray<FString> GetBoneCustomAttributeNamesToImport() const;
+
 	/** List of bone names for which all custom attributes are directly imported on the bone. */
 	UPROPERTY(config, EditAnywhere, Category = CustomAttributes)
 	TArray<FString> BoneNamesWithCustomAttributes;
diff --git a/Engine/Source/Runtime/Engine/Classes/Animation/CustomAttributes.h b/Engine/Source/Runtime/Engine/Classes/Animation/CustomAttributes.h
--- a/Engine/Source/Runtime/Engine/Classes/Animation/CustomAttributes.h
+++ b/Engine/Source/Runtime/Engine/Classes/Animation/CustomAttributes.h
@@ -37,6 +37,40 @@ struct ENGINE_API FCustomAttributeSetting
 	FString Meaning;
 };
 
+/**
+ * Settings that identify the names of custom attributes that represent the individual components of a timecode and a subframe along with a take name.
+ */
+USTRUCT()
+struct ENGINE_API FTimecodeCustomAttributeNameSettings
+{
+	GENERATED_BODY()
+
+	/** Name of the custom attribute representing the hour component of a timecode. */
+	UPROPERTY(EditAnywhere, Category = TimecodeCustomAttributeNameSettings)
+	FName HourAttributeName;
+
+	/** Name of the custom attribute representing the minute component of a timecode. */
+	UPROPERTY(EditAnywhere, Category = TimecodeCustomAttributeNameSettings)
+	FName MinuteAttributeName;
+
+	/** Name of the custom attribute representing the second component of a timecode. */
+	UPROPERTY(EditAnywhere, Category = TimecodeCustomAttributeNameSettings)
+	FName SecondAttributeName;
+
+	/** Name of the custom attribute representing the frame component of a timecode. */
+	UPROPERTY(EditAnywhere, Category = TimecodeCustomAttributeNameSettings)
+	FName FrameAttributeName;
+
+	/** Name of the custom attribute representing a subframe value. Though not strictly a component
+		of a timecode, this attribute can be authored to identify samples in between timecodes. */
+	UPROPERTY(EditAnywhere, Category = TimecodeCustomAttributeNameSettings)
+	FName SubframeAttributeName;
+
+	/** Name of the custom attribute representing the name of a take. */
+	UPROPERTY(EditAnywhere, Category = TimecodeCustomAttributeNameSettings)
+	FName TakenameAttributeName;
+};
+
 struct UE_DEPRECATED(5.0, "FCustomAttribute has been deprecated") FCustomAttribute;
 USTRUCT(Experimental)
 struct ENGINE_API FCustomAttribute
diff --git a/Engine/Source/Runtime/Engine/Private/Animation/AnimationSettings.cpp b/Engine/Source/Runtime/Engine/Private/Animation/AnimationSettings.cpp
--- a/Engine/Source/Runtime/Engine/Private/Animation/AnimationSettings.cpp
+++ b/Engine/Source/Runtime/Engine/Private/Animation/AnimationSettings.cpp
@@ -34,6 +34,25 @@ UAnimationSettings::UAnimationSettings(const FObjectInitializer& ObjectInitializ
 	};
 }
 
+TArray<FString> UAnimationSettings::GetBoneCustomAttributeNamesToImport() const
+{
+	TArray<FString> AttributeNames = {
+		BoneTimecodeCustomAttributeNameSettings.HourAttributeName.ToString(),
+		BoneTimecodeCustomAttributeNameSettings.MinuteAttributeName.ToString(),
+		BoneTimecodeCustomAttributeNameSettings.SecondAttributeName.ToString(),
+		BoneTimecodeCustomAttributeNameSettings.FrameAttributeName.ToString(),
+		BoneTimecodeCustomAttributeNameSettings.SubframeAttributeName.ToString(),
+		BoneTimecodeCustomAttributeNameSettings.TakenameAttributeName.ToString()
+	};
+
+	for (const FCustomAttributeSetting& Setting : BoneCustomAttributesNames)
+	{
+		AttributeNames.AddUnique(Setting.Name);
+	}
+
+	return AttributeNames;
+}
+
 
 #if WITH_EDITOR
 void UAnimationSettings::PostEditChangeProperty(struct FPropertyChangedEvent& PropertyChangedEvent)
