Fixed Static Analysis warning: Expression is always false

diff --git a/Engine/Source/Editor/UnrealEd/Private/Text/HLSLSyntaxHighlighterMarshaller.cpp b/Engine/Source/Editor/UnrealEd/Private/Text/HLSLSyntaxHighlighterMarshaller.cpp
--- a/Engine/Source/Editor/UnrealEd/Private/Text/HLSLSyntaxHighlighterMarshaller.cpp
+++ b/Engine/Source/Editor/UnrealEd/Private/Text/HLSLSyntaxHighlighterMarshaller.cpp
@@ -273,11 +273,6 @@ FTextLayout::FNewLineData FHLSLSyntaxHighlighterMarshaller::ProcessTokenizedLine
 	TSharedRef<FString> ModelString = MakeShareable(new FString());
 	TArray< TSharedRef< IRun > > Runs;
 
-	if(ParseState == EParseState::LookingForSingleLineComment) //-V547
-	{
-		ParseState = EParseState::None;
-	}
-
 	for(const FSyntaxTokenizer::FToken& Token : TokenizedLine.Tokens)
 	{
 		const FString TokenText = SourceString.Mid(Token.Range.BeginIndex, Token.Range.Len());
