Add PackageAccessTrackingOps::NAME_CookerBuildObject to catch references to other packages that are made from an object's BeginCacheForCookedPlatformData call.
#rnx

diff --git a/Engine/Source/Editor/UnrealEd/Private/Commandlets/DerivedDataCacheCommandlet.cpp b/Engine/Source/Editor/UnrealEd/Private/Commandlets/DerivedDataCacheCommandlet.cpp
--- a/Engine/Source/Editor/UnrealEd/Private/Commandlets/DerivedDataCacheCommandlet.cpp
+++ b/Engine/Source/Editor/UnrealEd/Private/Commandlets/DerivedDataCacheCommandlet.cpp
@@ -8,6 +8,7 @@ DerivedDataCacheCommandlet.cpp: Commandlet for DDC maintenence
 #include "UObject/UObjectIterator.h"
 #include "UObject/Package.h"
 #include "Misc/ConfigCacheIni.h"
+#include "Misc/PackageAccessTrackingOps.h"
 #include "Misc/PackageName.h"
 #include "PackageHelperFunctions.h"
 #include "DerivedDataCacheInterface.h"
@@ -226,6 +227,7 @@ void UDerivedDataCacheCommandlet::CacheLoadedPackages(UPackage* CurrentPackage,
 
 				ObjectsWithOuter.Reset();
 				GetObjectsWithOuter(NewPackage, ObjectsWithOuter, true /* bIncludeNestedObjects */, RF_ClassDefaultObject /* ExclusionFlags */);
+				UE_TRACK_REFERENCING_PACKAGE_SCOPED(NewPackage, PackageAccessTrackingOps::NAME_CookerBuildObject);
 				for (UObject* Object : ObjectsWithOuter)
 				{
 					for (auto Platform : Platforms)
@@ -271,11 +273,14 @@ bool UDerivedDataCacheCommandlet::ProcessCachingObjects(const TArray<ITargetPlat
 					bIsFinished = false;
 				}
 
-				for (auto Platform : Platforms)
 				{
-					// IsCachedCookedPlatformDataLoaded can be quite slow for some objects
-					// Do not call it if bIsFinished is already false
-					bIsFinished = bIsFinished && Object->IsCachedCookedPlatformDataLoaded(Platform);
+					UE_TRACK_REFERENCING_PACKAGE_SCOPED(Object->GetPackage(), PackageAccessTrackingOps::NAME_CookerBuildObject);
+					for (auto Platform : Platforms)
+					{
+						// IsCachedCookedPlatformDataLoaded can be quite slow for some objects
+						// Do not call it if bIsFinished is already false
+						bIsFinished = bIsFinished && Object->IsCachedCookedPlatformDataLoaded(Platform);
+					}
 				}
 
 				if (bIsFinished)
diff --git a/Engine/Source/Editor/UnrealEd/Private/CookOnTheFlyServer.cpp b/Engine/Source/Editor/UnrealEd/Private/CookOnTheFlyServer.cpp
--- a/Engine/Source/Editor/UnrealEd/Private/CookOnTheFlyServer.cpp
+++ b/Engine/Source/Editor/UnrealEd/Private/CookOnTheFlyServer.cpp
@@ -82,6 +82,7 @@
 #include "Misc/FileHelper.h"
 #include "Misc/LocalTimestampDirectoryVisitor.h"
 #include "Misc/NetworkVersion.h"
+#include "Misc/PackageAccessTrackingOps.h"
 #include "Misc/PackageName.h"
 #include "Misc/PathViews.h"
 #include "Misc/RedirectCollector.h"
@@ -2661,6 +2662,7 @@ bool UCookOnTheFlyServer::BeginPrepareSave(UE::Cook::FPackageData& PackageData,
 	PackageData.GetRequestedPlatforms(TargetPlatforms);
 	int NumPlatforms = TargetPlatforms.Num();
 	int NumIndexes = CachedObjectsInOuter.Num() * NumPlatforms;
+	UE_TRACK_REFERENCING_PACKAGE_SCOPED(Package, PackageAccessTrackingOps::NAME_CookerBuildObject);
 	while (CookedPlatformDataNextIndex < NumIndexes)
 	{
 		int ObjectIndex = CookedPlatformDataNextIndex / NumPlatforms;
@@ -3404,6 +3406,7 @@ void UCookOnTheFlyServer::TickPrecacheObjectsForPlatforms(const float TimeSlice,
 			continue;
 		}
 
+		UE_TRACK_REFERENCING_PACKAGE_SCOPED(Material->GetPackage(), PackageAccessTrackingOps::NAME_CookerBuildObject);
 		for (const ITargetPlatform* TargetPlatform : TargetPlatforms)
 		{
 			if (!TargetPlatform)
@@ -3444,6 +3447,7 @@ void UCookOnTheFlyServer::TickPrecacheObjectsForPlatforms(const float TimeSlice,
 			continue;
 		}
 
+		UE_TRACK_REFERENCING_PACKAGE_SCOPED(Texture->GetPackage(), PackageAccessTrackingOps::NAME_CookerBuildObject);
 		for (const ITargetPlatform* TargetPlatform : TargetPlatforms)
 		{
 			if (!TargetPlatform)
@@ -8758,6 +8762,7 @@ uint32 UCookOnTheFlyServer::FullLoadAndSave(uint32& CookedPackageCount)
 								if (!bIsTexture || bSaveConcurrent)
 								{
 									UE_SCOPED_HIERARCHICAL_COOKTIMER(FullLoadAndSave_BeginCache);
+									UE_TRACK_REFERENCING_PACKAGE_SCOPED(Package, PackageAccessTrackingOps::NAME_CookerBuildObject);
 									Obj->BeginCacheForCookedPlatformData(TargetPlatform);
 									if (!Obj->IsCachedCookedPlatformDataLoaded(TargetPlatform))
 									{
@@ -8873,6 +8878,7 @@ uint32 UCookOnTheFlyServer::FullLoadAndSave(uint32& CookedPackageCount)
 			{
 				UObject* Obj = ObjectsToWaitForCookedPlatformData[ObjIdx];
 				bool bAllPlatformDataLoaded = true;
+				UE_TRACK_REFERENCING_PACKAGE_SCOPED(Obj->GetPackage(), PackageAccessTrackingOps::NAME_CookerBuildObject);
 				for (const ITargetPlatform* TargetPlatform : TargetPlatforms)
 				{
 					if (!Obj->IsCachedCookedPlatformDataLoaded(TargetPlatform))
@@ -8934,6 +8940,7 @@ uint32 UCookOnTheFlyServer::FullLoadAndSave(uint32& CookedPackageCount)
 						GetObjectsWithOuter(PrecachePackage, ObjsInPackage, false);
 					}
 
+					UE_TRACK_REFERENCING_PACKAGE_SCOPED(PrecachePackage, PackageAccessTrackingOps::NAME_CookerBuildObject);
 					for (UObject* Obj : ObjsInPackage)
 					{
 						if (Obj->HasAnyFlags(RF_Transient) || !Obj->IsA(UTexture::StaticClass()))
diff --git a/Engine/Source/Editor/UnrealEd/Private/Cooker/CookPackageData.cpp b/Engine/Source/Editor/UnrealEd/Private/Cooker/CookPackageData.cpp
--- a/Engine/Source/Editor/UnrealEd/Private/Cooker/CookPackageData.cpp
+++ b/Engine/Source/Editor/UnrealEd/Private/Cooker/CookPackageData.cpp
@@ -16,6 +16,7 @@
 #include "HAL/PlatformTime.h"
 #include "Misc/CommandLine.h"
 #include "Misc/CoreMiscDefines.h"
+#include "Misc/PackageAccessTrackingOps.h"
 #include "Misc/PackageName.h"
 #include "Misc/Parse.h"
 #include "Misc/Paths.h"
@@ -1362,7 +1363,13 @@ namespace UE::Cook
 		}
 
 		UObject* LocalObject = Object.Get();
-		if (!LocalObject || LocalObject->IsCachedCookedPlatformDataLoaded(TargetPlatform))
+		if (!LocalObject)
+		{
+			Release();
+			return true;
+		}
+		UE_TRACK_REFERENCING_PACKAGE_SCOPED(LocalObject->GetPackage(), PackageAccessTrackingOps::NAME_CookerBuildObject);
+		if (LocalObject->IsCachedCookedPlatformDataLoaded(TargetPlatform))
 		{
 			Release();
 			return true;
diff --git a/Engine/Source/Runtime/CoreUObject/Private/Misc/PackageAccessTracking.cpp b/Engine/Source/Runtime/CoreUObject/Private/Misc/PackageAccessTracking.cpp
--- a/Engine/Source/Runtime/CoreUObject/Private/Misc/PackageAccessTracking.cpp
+++ b/Engine/Source/Runtime/CoreUObject/Private/Misc/PackageAccessTracking.cpp
@@ -16,6 +16,7 @@ bool IsBuildOpName(FName OpName)
 	//which occurs from "UObjectLoadAllCompiledInDefaultProperties" and possibly elsewhere.
 	return (OpName == PackageAccessTrackingOps::NAME_Load)
 		|| (OpName == PackageAccessTrackingOps::NAME_Save)
+		|| (OpName == PackageAccessTrackingOps::NAME_CookerBuildObject)
 		// || (OpName == PackageAccessTrackingOps::NAME_CreateDefaultObject)
 		;
 }
diff --git a/Engine/Source/Runtime/CoreUObject/Private/Misc/PackageAccessTrackingOps.cpp b/Engine/Source/Runtime/CoreUObject/Private/Misc/PackageAccessTrackingOps.cpp
--- a/Engine/Source/Runtime/CoreUObject/Private/Misc/PackageAccessTrackingOps.cpp
+++ b/Engine/Source/Runtime/CoreUObject/Private/Misc/PackageAccessTrackingOps.cpp
@@ -8,6 +8,7 @@ FName PackageAccessTrackingOps::NAME_PreLoad(TEXT("PreLoad"));
 FName PackageAccessTrackingOps::NAME_PostLoad(TEXT("PostLoad"));
 FName PackageAccessTrackingOps::NAME_Save(TEXT("Save"));
 FName PackageAccessTrackingOps::NAME_CreateDefaultObject(TEXT("CreateDefaultObject"));
+FName PackageAccessTrackingOps::NAME_CookerBuildObject(TEXT("CookerBuildObject"));
 FName PackageAccessTrackingOps::NAME_ResetContext(TEXT("ResetContext"));
 FName PackageAccessTrackingOps::NAME_NoAccessExpected(TEXT("NoAccessExpected"));
 
diff --git a/Engine/Source/Runtime/CoreUObject/Public/Misc/PackageAccessTrackingOps.h b/Engine/Source/Runtime/CoreUObject/Public/Misc/PackageAccessTrackingOps.h
--- a/Engine/Source/Runtime/CoreUObject/Public/Misc/PackageAccessTrackingOps.h
+++ b/Engine/Source/Runtime/CoreUObject/Public/Misc/PackageAccessTrackingOps.h
@@ -14,6 +14,7 @@ namespace PackageAccessTrackingOps
 	extern COREUOBJECT_API FName NAME_PostLoad;
 	extern COREUOBJECT_API FName NAME_Save;
 	extern COREUOBJECT_API FName NAME_CreateDefaultObject;
+	extern COREUOBJECT_API FName NAME_CookerBuildObject;
 	/**
 	 * A global-scope operation suck as AssetRegistry tick is occurring on the current thread inside of another scope,
 	 * and the references made during this operation should not be associated with the previous scope.
