EditorDomain: Fix spurious parent class log messages due to mishandling of Core Redirects in FPrecacheClassDigest::GetRecursive.
#rnx

diff --git a/Engine/Source/Editor/UnrealEd/Private/EditorDomain/EditorDomainUtils.cpp b/Engine/Source/Editor/UnrealEd/Private/EditorDomain/EditorDomainUtils.cpp
--- a/Engine/Source/Editor/UnrealEd/Private/EditorDomain/EditorDomainUtils.cpp
+++ b/Engine/Source/Editor/UnrealEd/Private/EditorDomain/EditorDomainUtils.cpp
@@ -196,7 +196,7 @@ public:
 		ClassDigestsMap.Lock.WriteUnlock();
 	}
 
-	FClassDigestData* GetRecursive(FName ClassName);
+	FClassDigestData* GetRecursive(FName ClassName, bool bAllowRedirects);
 
 	struct FUnlockScope
 	{
@@ -222,7 +222,7 @@ private:
 	TArray<FName> AncestorShortNames;
 };
 
-FClassDigestData* FPrecacheClassDigest::GetRecursive(FName ClassName)
+FClassDigestData* FPrecacheClassDigest::GetRecursive(FName ClassName, bool bAllowRedirects)
 {
 	// Called within ClassDigestsMap.Lock.WriteLock()
 	FClassDigestData* DigestData = &ClassDigests.FindOrAdd(ClassName);
@@ -234,12 +234,15 @@ FClassDigestData* FPrecacheClassDigest::GetRecursive(FName ClassName)
 
 	FName LookupName = ClassName;
 	ClassName.ToString(NameStringBuffer);
-	FCoreRedirectObjectName ClassNameRedirect(NameStringBuffer);
-	FCoreRedirectObjectName RedirectedClassNameRedirect = FCoreRedirects::GetRedirectedName(ECoreRedirectFlags::Type_Class, ClassNameRedirect);
-	if (ClassNameRedirect != RedirectedClassNameRedirect)
+	if (bAllowRedirects)
 	{
-		NameStringBuffer = RedirectedClassNameRedirect.ToString();
-		LookupName = FName(NameStringBuffer);
+		FCoreRedirectObjectName ClassNameRedirect(NameStringBuffer);
+		FCoreRedirectObjectName RedirectedClassNameRedirect = FCoreRedirects::GetRedirectedName(ECoreRedirectFlags::Type_Class, ClassNameRedirect);
+		if (ClassNameRedirect != RedirectedClassNameRedirect)
+		{
+			NameStringBuffer = RedirectedClassNameRedirect.ToString();
+			LookupName = FName(NameStringBuffer);
+		}
 	}
 
 	UStruct* Struct = nullptr;
@@ -335,19 +338,32 @@ FClassDigestData* FPrecacheClassDigest::GetRecursive(FName ClassName)
 	// Propagate values from the parent
 	if (!ParentName.IsNone())
 	{
-		FClassDigestData* ParentDigest = GetRecursive(ParentName);
+		// CoreRedirects are expected to act only on import classes from packages; they are not expected to act on the parent class pointer
+		// of a native class, which is authoritative, so set bAllowRedirects = false
+		FClassDigestData* ParentDigest = GetRecursive(ParentName, false /* bAllowRedirects */);
 		// The map has possibly been modified so we need to recalculate the address of ClassName's DigestData
 		DigestData = &ClassDigests.FindChecked(ClassName);
 		if (!ParentDigest)
 		{
-			UE_LOG(LogEditorDomain, Display, TEXT("Parent class %s of class %s not found. Allow flags for editordomain and iterative cooking will be invalid."),
+			UE_LOG(LogEditorDomain, Display,
+				TEXT("Parent class %s of class %s not found. Allow flags for editordomain and iterative cooking will be invalid."),
 				*ParentName.ToString(), *ClassName.ToString());
 		}
 		else
 		{
 			if (!ParentDigest->bConstructionComplete)
 			{
-				UE_LOG(LogEditorDomain, Display, TEXT("Cycle detected in parents of class %s. Allow flags for editordomain and iterative cooking will be invalid."), *ClassName.ToString());
+				// Suppress the warning for MulticastDelegateProperty, which has a redirector to its own child class
+				// of MulticastInlineDelegateProperty
+				// We could fix this case by adding bAllowRedirects to the ClassDigestsMap lookup key, but it's not a
+				// problem for MuticastDelegateProperty and we don't have any other cases where it is a problem, so we
+				// avoid the performance cost of doing so.
+				if (ClassName != FName(TEXT("/Script/CoreUObject.MulticastDelegateProperty")))
+				{
+					UE_LOG(LogEditorDomain, Display,
+						TEXT("Cycle detected in parents of class %s. Allow flags for editordomain and iterative cooking will be invalid."),
+						*ClassName.ToString());
+				}
 			}
 			EnumSetFlagsAnd(DigestData->EditorDomainUse, EDomainUse::LoadEnabled | EDomainUse::SaveEnabled,
 				DigestData->EditorDomainUse, ParentDigest->EditorDomainUse);
@@ -366,10 +382,11 @@ FClassDigestData* FPrecacheClassDigest::GetRecursive(FName ClassName)
 		TArray<FGuid> ConstructCustomVersions; // Not a class scratch variable because we use it across recursive calls
 		for (FName ConstructClass : ConstructClasses)
 		{
-			FClassDigestData* ConstructClassDigest = GetRecursive(ConstructClass);
+			FClassDigestData* ConstructClassDigest = GetRecursive(ConstructClass, true /* bAllowRedirects */);
 			if (!ConstructClassDigest)
 			{
-				UE_LOG(LogEditorDomain, Warning, TEXT("Construct class %s specified by Editor.ini:[EditorDomain]:PostLoadConstructClasses for class %s is not found. ")
+				UE_LOG(LogEditorDomain, Warning,
+					TEXT("Construct class %s specified by Editor.ini:[EditorDomain]:PostLoadConstructClasses for class %s is not found. ")
 					TEXT("This is a class that can be constructed by postload upgrades of class %s. ")
 					TEXT("Old packages with class %s will load more slowly."),
 					*ConstructClass.ToString(), *ClassName.ToString(), *ClassName.ToString(), *ClassName.ToString());
@@ -378,7 +395,8 @@ FClassDigestData* FPrecacheClassDigest::GetRecursive(FName ClassName)
 			{
 				if (!ConstructClassDigest->bConstructionComplete)
 				{
-					UE_LOG(LogEditorDomain, Verbose, TEXT("Cycle detected in Editor.ini:[EditorDomain]:PostLoadConstructClasses of class %s. This is unexpected, but not a problem."),
+					UE_LOG(LogEditorDomain, Verbose,
+						TEXT("Cycle detected in Editor.ini:[EditorDomain]:PostLoadConstructClasses of class %s. This is unexpected, but not a problem."),
 						*ClassName.ToString());
 				}
 				ConstructCustomVersions.Append(ConstructClassDigest->CustomVersionGuids);
@@ -401,7 +419,7 @@ void PrecacheClassDigests(TConstArrayView<FName> ClassNames)
 	FPrecacheClassDigest Digester;
 	for (FName ClassName : ClassNames)
 	{
-		Digester.GetRecursive(ClassName);
+		Digester.GetRecursive(ClassName, true /* bAllowRedirects */);
 	}
 }
 
