Fix broken undo/redo of create/delete Layer in ULayersSubsystem.

diff --git a/Engine/Source/Editor/UnrealEd/Private/Layers/LayersSubsystem.cpp b/Engine/Source/Editor/UnrealEd/Private/Layers/LayersSubsystem.cpp
--- a/Engine/Source/Editor/UnrealEd/Private/Layers/LayersSubsystem.cpp
+++ b/Engine/Source/Editor/UnrealEd/Private/Layers/LayersSubsystem.cpp
@@ -1218,6 +1218,7 @@ ULayer* ULayersSubsystem::CreateLayer(const FName& LayerName)
 	ULayer* NewLayer = NewObject<ULayer>(GetWorld(), NAME_None, RF_Transactional | RF_Transient);
 	check(NewLayer != NULL);
 
+	GetWorld()->Modify(false);
 	GetWorld()->Layers.Add(NewLayer);
 
 	NewLayer->SetLayerName(LayerName);
@@ -1256,6 +1257,7 @@ void ULayersSubsystem::DeleteLayers( const TArray< FName >& LayersToDelete )
 	{
 		if( LayersToDelete.Contains( GetWorld()->Layers[ LayerIndex]->GetLayerName() ) )
 		{
+			GetWorld()->Modify(false);
 			GetWorld()->Layers.RemoveAt( LayerIndex );
 			bValidLayerExisted = true;
 		}
@@ -1287,6 +1289,7 @@ void ULayersSubsystem::DeleteLayer( const FName& LayerToDelete )
 	{
 		if( LayerToDelete == GetWorld()->Layers[ LayerIndex]->GetLayerName() )
 		{
+			GetWorld()->Modify(false);
 			GetWorld()->Layers.RemoveAt( LayerIndex );
 			bValidLayerExisted = true;
 		}
