The Open Asset Dialog now opens as a Tab

diff --git a/Engine/Source/Editor/UnrealEd/Private/Toolkits/GlobalEditorCommonCommands.cpp b/Engine/Source/Editor/UnrealEd/Private/Toolkits/GlobalEditorCommonCommands.cpp
--- a/Engine/Source/Editor/UnrealEd/Private/Toolkits/GlobalEditorCommonCommands.cpp
+++ b/Engine/Source/Editor/UnrealEd/Private/Toolkits/GlobalEditorCommonCommands.cpp
@@ -24,6 +24,16 @@
 FGlobalEditorCommonCommands::FGlobalEditorCommonCommands()
 	: TCommands<FGlobalEditorCommonCommands>(TEXT("SystemWideCommands"), NSLOCTEXT("Contexts", "SystemWideCommands", "System-wide"), NAME_None, FEditorStyle::GetStyleSetName())
 {
+	FGlobalTabmanager::Get()->RegisterNomadTabSpawner("GlobalAssetPicker", FOnSpawnTab::CreateStatic(&FGlobalEditorCommonCommands::SpawnAssetPicker))
+		.SetDisplayName(LOCTEXT("AssetPickerTabTitle", "Open Asset"))
+		.SetIcon(FSlateIcon(FAppStyle::GetAppStyleSetName(), "SystemWideCommands.SummonOpenAssetDialog"))
+		.SetMenuType(ETabSpawnerMenuType::Hidden);
+}
+
+
+FGlobalEditorCommonCommands::~FGlobalEditorCommonCommands()
+{
+	FGlobalTabmanager::Get()->UnregisterTabSpawner("GlobalAssetPicker");
 }
 
 void FGlobalEditorCommonCommands::RegisterCommands()
@@ -106,20 +116,39 @@ void FGlobalEditorCommonCommands::OnPressedCtrlTab(TSharedPtr<FUICommandInfo> Tr
 	}
 }
 
-void FGlobalEditorCommonCommands::OnSummonedAssetPicker()
+TSharedRef<SDockTab> FGlobalEditorCommonCommands::SpawnAssetPicker(const FSpawnTabArgs& InArgs)
 {
 	const FVector2D AssetPickerSize(600.0f, 586.0f);
-
 	// Create the contents of the popup
 	TSharedRef<SWidget> ActualWidget = SNew(SGlobalOpenAssetDialog, AssetPickerSize);
 
-	// Wrap the picker widget in a multibox-style menu body
-	FMenuBuilder MenuBuilder(/*BShouldCloseAfterSelection=*/ false, /*CommandList=*/ nullptr);
-	MenuBuilder.BeginSection("AssetPickerOpenAsset", NSLOCTEXT("GlobalAssetPicker", "WindowTitle", "Open Asset"));
-	MenuBuilder.AddWidget(ActualWidget, FText::GetEmpty(), /*bNoIndent=*/ true);
-	MenuBuilder.EndSection();
+	/**
+	 * The Global Asset Picker has been changed to open as a tab. This is because it would close on selecting some menu and sub menu
+	 * options due to it being opened as a menu before, which would lead to the selected menu option doing nothing. Furthermore,
+	 * There was also weird behavior when you had another asset picker in a menu open and tried to open the global asset picker
+	 * where the other picker would close but leave any parent menus hanging forever with no way to close them
+	 */
+
+	const TSharedRef<SDockTab> DockTab = SNew(SDockTab)
+		.TabRole(ETabRole::NomadTab);
+
+	DockTab->SetContent(ActualWidget);
+
+	return DockTab;
 
-	OpenPopupMenu(MenuBuilder.MakeWidget(), AssetPickerSize);
+}
+
+void FGlobalEditorCommonCommands::OnSummonedAssetPicker()
+{
+	if (TSharedPtr<SDockTab> AssetPickerTab = FGlobalTabmanager::Get()->FindExistingLiveTab(FTabId("GlobalAssetPicker")))
+	{
+		AssetPickerTab->RequestCloseTab();
+	}
+	else
+	{
+		FGlobalTabmanager::Get()->TryInvokeTab(FTabId("GlobalAssetPicker"));
+	}
+	
 }
 
 TSharedPtr<IMenu> FGlobalEditorCommonCommands::OpenPopupMenu(TSharedRef<SWidget> WindowContents, const FVector2D& PopupDesiredSize)
diff --git a/Engine/Source/Editor/UnrealEd/Public/Toolkits/GlobalEditorCommonCommands.h b/Engine/Source/Editor/UnrealEd/Public/Toolkits/GlobalEditorCommonCommands.h
--- a/Engine/Source/Editor/UnrealEd/Public/Toolkits/GlobalEditorCommonCommands.h
+++ b/Engine/Source/Editor/UnrealEd/Public/Toolkits/GlobalEditorCommonCommands.h
@@ -6,6 +6,7 @@
 #include "Widgets/SWidget.h"
 #include "Framework/Commands/UICommandList.h"
 #include "Framework/Commands/Commands.h"
+#include "Framework/Docking/TabManager.h"
 
 class IMenu;
 class SWindow;
@@ -16,6 +17,7 @@ class UNREALED_API FGlobalEditorCommonCommands : public TCommands< FGlobalEditor
 {
 public:
 	FGlobalEditorCommonCommands();
+	~FGlobalEditorCommonCommands();
 
 	virtual void RegisterCommands() override;
 
@@ -28,6 +30,8 @@ protected:
 	static void OnOpenContentBrowserDrawer();
 	static void OnOpenOutputLogDrawer();
 
+	static TSharedRef<SDockTab> SpawnAssetPicker(const FSpawnTabArgs& InArgs);
+
 	static TSharedPtr<IMenu> OpenPopupMenu(TSharedRef<SWidget> WindowContents, const FVector2D& PopupDesiredSize);
 public:
 	TSharedPtr<FUICommandInfo> FindInContentBrowser;
