CookCommandlet: Sort includes, remove CoreMinimal.h
#rnx

diff --git a/Engine/Source/Editor/UnrealEd/Classes/Commandlets/CookCommandlet.h b/Engine/Source/Editor/UnrealEd/Classes/Commandlets/CookCommandlet.h
--- a/Engine/Source/Editor/UnrealEd/Classes/Commandlets/CookCommandlet.h
+++ b/Engine/Source/Editor/UnrealEd/Classes/Commandlets/CookCommandlet.h
@@ -6,14 +6,14 @@
 
 #pragma once
 
-#include "CoreMinimal.h"
-#include "UObject/ObjectMacros.h"
+#include "Commandlets/Commandlet.h"
+#include "IPlatformFileSandboxWrapper.h"
 #include "Misc/Guid.h"
-#include "UObject/WeakObjectPtr.h"
 #include "Misc/PackageName.h"
-#include "Commandlets/Commandlet.h"
 #include "Templates/UniquePtr.h"
-#include "IPlatformFileSandboxWrapper.h"
+#include "UObject/ObjectMacros.h"
+#include "UObject/WeakObjectPtr.h"
+
 #include "CookCommandlet.generated.h"
 
 class FSandboxPlatformFile;
diff --git a/Engine/Source/Editor/UnrealEd/Private/Commandlets/CookCommandlet.cpp b/Engine/Source/Editor/UnrealEd/Private/Commandlets/CookCommandlet.cpp
--- a/Engine/Source/Editor/UnrealEd/Private/Commandlets/CookCommandlet.cpp
+++ b/Engine/Source/Editor/UnrealEd/Private/Commandlets/CookCommandlet.cpp
@@ -5,46 +5,47 @@
 =============================================================================*/
 
 #include "Commandlets/CookCommandlet.h"
+
+#include "AssetRegistryModule.h"
+#include "Async/TaskGraphInterfaces.h"
+#include "CookOnTheSide/CookOnTheFlyServer.h"
+#include "Cooker/CookProfiling.h"
+#include "Cooker/PackageBuildDependencyTracker.h"
+#include "CookerSettings.h"
+#include "DerivedDataBuildRemoteExecutor.h"
+#include "Editor.h"
+#include "EngineGlobals.h"
+#include "GameDelegates.h"
+#include "GlobalShader.h"
+#include "HAL/FileManager.h"
 #include "HAL/IConsoleManager.h"
-#include "HAL/PlatformFileManager.h"
+#include "HAL/MemoryMisc.h"
 #include "HAL/PlatformApplicationMisc.h"
-#include "Misc/MessageDialog.h"
-#include "HAL/FileManager.h"
+#include "HAL/PlatformFileManager.h"
+#include "INetworkFileSystemModule.h"
+#include "IPlatformFileSandboxWrapper.h"
+#include "Interfaces/ITargetPlatform.h"
+#include "Interfaces/ITargetPlatformManagerModule.h"
+#include "Misc/App.h"
 #include "Misc/CommandLine.h"
-#include "Misc/FileHelper.h"
-#include "Misc/Paths.h"
-#include "Stats/StatsMisc.h"
 #include "Misc/ConfigCacheIni.h"
+#include "Misc/FileHelper.h"
 #include "Misc/LocalTimestampDirectoryVisitor.h"
-#include "Misc/App.h"
-#include "Modules/ModuleManager.h"
-#include "UObject/Class.h"
-#include "UObject/UObjectIterator.h"
-#include "UObject/Package.h"
-#include "UObject/MetaData.h"
-#include "Async/TaskGraphInterfaces.h"
+#include "Misc/MessageDialog.h"
+#include "Misc/Paths.h"
 #include "Misc/RedirectCollector.h"
-#include "IPlatformFileSandboxWrapper.h"
-#include "CookOnTheSide/CookOnTheFlyServer.h"
-#include "Settings/ProjectPackagingSettings.h"
-#include "EngineGlobals.h"
-#include "Editor.h"
-#include "Serialization/ArrayWriter.h"
+#include "Modules/ModuleManager.h"
 #include "PackageHelperFunctions.h"
-#include "GlobalShader.h"
-#include "ShaderCompiler.h"
-#include "Interfaces/ITargetPlatform.h"
-#include "Interfaces/ITargetPlatformManagerModule.h"
-#include "INetworkFileSystemModule.h"
-#include "GameDelegates.h"
-#include "CookerSettings.h"
-#include "HAL/MemoryMisc.h"
 #include "ProfilingDebugging/CookStats.h"
-#include "AssetRegistryModule.h"
+#include "Serialization/ArrayWriter.h"
+#include "Settings/ProjectPackagingSettings.h"
+#include "ShaderCompiler.h"
+#include "Stats/StatsMisc.h"
 #include "StudioAnalytics.h"
-#include "Cooker/CookProfiling.h"
-#include "Cooker/PackageBuildDependencyTracker.h"
-#include "DerivedDataBuildRemoteExecutor.h"
+#include "UObject/Class.h"
+#include "UObject/MetaData.h"
+#include "UObject/Package.h"
+#include "UObject/UObjectIterator.h"
 
 DEFINE_LOG_CATEGORY_STATIC(LogCookCommandlet, Log, All);
 
