UWorldPartitionResaveActorsBuilder
- Automatically match commandline arguments to properties when they exists (this will be ported to all builders in the future).
- Added `ResaveDirtyActorDescsOnly` to only resave actors that have notable differences in their actor descriptors.
- Added `ReportOnly` option to only report what would be done.

diff --git a/Engine/Source/Editor/UnrealEd/Private/WorldPartition/WorldPartitionResaveActorsBuilder.cpp b/Engine/Source/Editor/UnrealEd/Private/WorldPartition/WorldPartitionResaveActorsBuilder.cpp
--- a/Engine/Source/Editor/UnrealEd/Private/WorldPartition/WorldPartitionResaveActorsBuilder.cpp
+++ b/Engine/Source/Editor/UnrealEd/Private/WorldPartition/WorldPartitionResaveActorsBuilder.cpp
@@ -20,12 +20,58 @@ UWorldPartitionResaveActorsBuilder::UWorldPartitionResaveActorsBuilder(const FOb
 
 bool UWorldPartitionResaveActorsBuilder::PreRun(UWorld* World, FPackageSourceControlHelper& PackageHelper)
 {
-	FParse::Value(FCommandLine::Get(), TEXT("ActorClass="), ActorClassName);
-	
 	TArray<FString> Tokens, Switches;
 	UCommandlet::ParseCommandLine(FCommandLine::Get(), Tokens, Switches);
 
-	bSwitchActorPackagingSchemeToReduced = Switches.Contains(TEXT("SwitchActorPackagingSchemeToReduced"));
+	//@todo_ow: generalize to all builders
+	for (const FString& Switch : Switches)
+	{
+		FString Key, Value;
+		if (!Switch.Split(TEXT("="), &Key, &Value))
+		{
+			Key = Switch;
+		}
+
+		// Lookup property
+		const FProperty* Property = GetClass()->FindPropertyByName(*Key);
+
+		// If we can't find the property, try for properties with the 'b' prefix
+		if (!Property)
+		{
+			Key = TEXT("b") + Key;
+			Property = GetClass()->FindPropertyByName(*Key);
+		}
+
+		if (Property)
+		{
+			// If the property is a bool, treat no values as true
+			if (Property->IsA(FBoolProperty::StaticClass()) && Value.IsEmpty())
+			{
+				Value = TEXT("True");
+			}
+
+			uint8* Container = (uint8*)this;
+			if (!FBlueprintEditorUtils::PropertyValueFromString(Property, Value, Container, nullptr))
+			{
+				UE_LOG(LogWorldPartitionResaveActorsBuilder, Error, TEXT("Cannot set value for '%s': '%s'"), *Key, *Value);
+				return false;
+			}
+		}
+	}
+
+	if (bSwitchActorPackagingSchemeToReduced)
+	{
+		if (!ActorClassName.IsEmpty())
+		{
+			UE_LOG(LogWorldPartitionResaveActorsBuilder, Error, TEXT("SwitchActorPackagingSchemeToReduced is not compatible with ActorClassName"));
+			return false;
+		}
+		else if (bResaveDirtyActorDescsOnly)
+		{
+			UE_LOG(LogWorldPartitionResaveActorsBuilder, Error, TEXT("SwitchActorPackagingSchemeToReduced is not compatible with ResaveDirtyActorDescsOnly"));
+			return false;
+		}
+	}
 
 	return true;
 }
@@ -36,6 +82,7 @@ bool UWorldPartitionResaveActorsBuilder::RunInternal(UWorld* World, const FCellI
 
 	UPackage* WorldPackage = World->GetPackage();
 
+	int32 LoadCount = 0;
 	int32 SaveCount = 0;
 	int32 FailCount = 0;
 
@@ -84,11 +131,11 @@ bool UWorldPartitionResaveActorsBuilder::RunInternal(UWorld* World, const FCellI
 	{
 		World->PersistentLevel->ActorPackagingScheme = EActorPackagingScheme::Reduced;
 
-		FWorldPartitionHelpers::ForEachActorWithLoading(WorldPartition, ActorClass, [this, &SaveCount, &FailCount, &SCCHelper, &PackagesToDelete, WorldPackage](const FWorldPartitionActorDesc* ActorDesc)
+		FWorldPartitionHelpers::ForEachActorWithLoading(WorldPartition, ActorClass, [this, &LoadCount, &SaveCount, &FailCount, &SCCHelper, &PackagesToDelete, WorldPackage](const FWorldPartitionActorDesc* ActorDesc)
 		{
 			ON_SCOPE_EXIT
 			{
-				UE_LOG(LogWorldPartitionResaveActorsBuilder, Display, TEXT("Processed %d packages (%d Saved / %d Failed)"), SaveCount + FailCount, SaveCount, FailCount);
+				UE_LOG(LogWorldPartitionResaveActorsBuilder, Display, TEXT("Processed %d packages (%d Saved / %d Failed)"), LoadCount, SaveCount, FailCount);
 			};
 
 			AActor* Actor = ActorDesc->GetActor();
@@ -101,86 +148,91 @@ bool UWorldPartitionResaveActorsBuilder::RunInternal(UWorld* World, const FCellI
 			}
 			else
 			{
+				++LoadCount;
+
 				UPackage* Package = Actor->GetExternalPackage();
 				check(Package);
 
-				if (SCCHelper.Checkout(Package))
+				if (!bReportOnly)
 				{
-					// Always mark this package for deletion, as it will contain a temporary redirector to fixup references
-					PackagesToDelete.Add(ActorDesc->GetActorPackage().ToString());
+					if (SCCHelper.Checkout(Package))
+					{
+						// Always mark this package for deletion, as it will contain a temporary redirector to fixup references
+						PackagesToDelete.Add(ActorDesc->GetActorPackage().ToString());
 
-					// Move actor back into world's package
-					Actor->SetPackageExternal(false);
+						// Move actor back into world's package
+						Actor->SetPackageExternal(false);
 
-					// Gather dependant objects that also needs to be moved along with the actor
-					TArray<UObject*> DependantObjects;
-					ForEachObjectWithPackage(Package, [&DependantObjects](UObject* Object)
-					{
-						if (!Cast<UMetaData>(Object))
+						// Gather dependant objects that also needs to be moved along with the actor
+						TArray<UObject*> DependantObjects;
+						ForEachObjectWithPackage(Package, [&DependantObjects](UObject* Object)
+						{
+							if (!Cast<UMetaData>(Object))
+							{
+								DependantObjects.Add(Object);
+							}
+							return true;
+						}, false);
+
+						// Move dependant objects into the new world package temporarily, so we can save the original actor package with its redirector
+						for (UObject* DependantObject : DependantObjects)
 						{
-							DependantObjects.Add(Object);
+							DependantObject->Rename(nullptr, WorldPackage, REN_NonTransactional | REN_DontCreateRedirectors | REN_ForceNoResetLoaders);
 						}
-						return true;
-					}, false);
 
-					// Move dependant objects into the new world package temporarily, so we can save the original actor package with its redirector
-					for (UObject* DependantObject : DependantObjects)
-					{
-						DependantObject->Rename(nullptr, WorldPackage, REN_NonTransactional | REN_DontCreateRedirectors | REN_ForceNoResetLoaders);
-					}
+						// Replace old actor package with a redirector to the new actor
+						UObjectRedirector* Redirector = NewObject<UObjectRedirector>(Actor->GetWorld(), Actor->GetFName(), RF_Standalone | RF_Public);
+						Redirector->DestinationObject = Actor;
+						Redirector->SetExternalPackage(Package);
 
-					// Replace old actor package with a redirector to the new actor
-					UObjectRedirector* Redirector = NewObject<UObjectRedirector>(Actor->GetWorld(), Actor->GetFName(), RF_Standalone | RF_Public);
-					Redirector->DestinationObject = Actor;
-					Redirector->SetExternalPackage(Package);
+						// Save original actor package containing a single redirector to the actor
+						FString PackageFileName = SourceControlHelpers::PackageFilename(Package);
+						FSavePackageArgs SaveArgs;
+						SaveArgs.TopLevelFlags = RF_Standalone;
+						if (!UPackage::SavePackage(Package, nullptr, *PackageFileName, SaveArgs))
+						{
+							UE_LOG(LogWorldPartitionResaveActorsBuilder, Error, TEXT("Error saving package redirector %s."), *Package->GetName());
+							++FailCount;
+							return true;
+						}
 
-					// Save original actor package containing a single redirector to the actor
-					FString PackageFileName = SourceControlHelpers::PackageFilename(Package);
-					FSavePackageArgs SaveArgs;
-					SaveArgs.TopLevelFlags = RF_Standalone;
-					if (!UPackage::SavePackage(Package, nullptr, *PackageFileName, SaveArgs))
-					{
-						UE_LOG(LogWorldPartitionResaveActorsBuilder, Error, TEXT("Error saving package redirector %s."), *Package->GetName());
-						++FailCount;
-						return true;
-					}
+						// Move actor in its new package
+						Actor->SetPackageExternal(true);
 
-					// Move actor in its new package
-					Actor->SetPackageExternal(true);
+						// Also move dependant objects into the new package
+						UPackage* NewActorPackage = Actor->GetExternalPackage();
+						for (UObject* DependantObject : DependantObjects)
+						{
+							DependantObject->Rename(nullptr, NewActorPackage, REN_NonTransactional | REN_DontCreateRedirectors | REN_ForceNoResetLoaders);
+						}
 
-					// Also move dependant objects into the new package
-					UPackage* NewActorPackage = Actor->GetExternalPackage();
-					for (UObject* DependantObject : DependantObjects)
-					{
-						DependantObject->Rename(nullptr, NewActorPackage, REN_NonTransactional | REN_DontCreateRedirectors | REN_ForceNoResetLoaders);
-					}
+						// Save the new actor package and add to source control if needed
+						FString NewPackageFileName = SourceControlHelpers::PackageFilename(NewActorPackage);
+						if (!UPackage::SavePackage(NewActorPackage, nullptr, *NewPackageFileName, SaveArgs))
+						{
+							UE_LOG(LogWorldPartitionResaveActorsBuilder, Error, TEXT("Error saving package %s."), *NewActorPackage->GetName());
+							++FailCount;
+							return true;
+						}
 
-					// Save the new actor package and add to source control if needed
-					FString NewPackageFileName = SourceControlHelpers::PackageFilename(NewActorPackage);
-					if (!UPackage::SavePackage(NewActorPackage, nullptr, *NewPackageFileName, SaveArgs))
-					{
-						UE_LOG(LogWorldPartitionResaveActorsBuilder, Error, TEXT("Error saving package %s."), *NewActorPackage->GetName());
-						++FailCount;
+						if (!SCCHelper.AddToSourceControl(NewActorPackage))
+						{
+							UE_LOG(LogWorldPartitionResaveActorsBuilder, Error, TEXT("Error adding package to source control %s."), *NewActorPackage->GetName());
+							++FailCount;
+							return true;
+						}
+
+						UE_LOG(LogWorldPartitionResaveActorsBuilder, Display, TEXT("Saved package %s."), *NewActorPackage->GetName());
+						++SaveCount;
 						return true;
 					}
-
-					if (!SCCHelper.AddToSourceControl(NewActorPackage))
+					else
 					{
-						UE_LOG(LogWorldPartitionResaveActorsBuilder, Error, TEXT("Error adding package to source control %s."), *NewActorPackage->GetName());
+						// It is possible the resave can't checkout everything. Continue processing.
+						UE_LOG(LogWorldPartitionResaveActorsBuilder, Warning, TEXT("Error checking out package %s."), *Package->GetName());
 						++FailCount;
 						return true;
 					}
-
-					UE_LOG(LogWorldPartitionResaveActorsBuilder, Display, TEXT("Saved package %s."), *NewActorPackage->GetName());
-					++SaveCount;
-					return true;
-				}
-				else
-				{
-					// It is possible the resave can't checkout everything. Continue processing.
-					UE_LOG(LogWorldPartitionResaveActorsBuilder, Warning, TEXT("Error checking out package %s."), *Package->GetName());
-					++FailCount;
-					return true;
 				}
 			}
 
@@ -189,11 +241,11 @@ bool UWorldPartitionResaveActorsBuilder::RunInternal(UWorld* World, const FCellI
 	}
 
 	// Normal save pass
-	FWorldPartitionHelpers::ForEachActorWithLoading(WorldPartition, ActorClass, [this, &SaveCount, &FailCount, &SCCHelper, &PackagesToDelete](const FWorldPartitionActorDesc* ActorDesc)
+	FWorldPartitionHelpers::ForEachActorWithLoading(WorldPartition, ActorClass, [this, &LoadCount, &SaveCount, &FailCount, &SCCHelper, &PackagesToDelete](const FWorldPartitionActorDesc* ActorDesc)
 	{
 		ON_SCOPE_EXIT
 		{
-			UE_LOG(LogWorldPartitionResaveActorsBuilder, Display, TEXT("Processed %d packages (%d Saved / %d Failed)"), SaveCount + FailCount, SaveCount, FailCount);
+			UE_LOG(LogWorldPartitionResaveActorsBuilder, Display, TEXT("Processed %d packages (%d Saved / %d Failed)"), LoadCount, SaveCount, FailCount);
 		};
 
 		AActor* Actor = ActorDesc->GetActor();
@@ -206,41 +258,61 @@ bool UWorldPartitionResaveActorsBuilder::RunInternal(UWorld* World, const FCellI
 		}
 		else
 		{
+			++LoadCount;
+
 			UPackage* Package = Actor->GetExternalPackage();
 			check(Package);
 
-			if (SCCHelper.Checkout(Package))
+			if (bResaveDirtyActorDescsOnly)
 			{
-				// Save package
-				FString PackageFileName = SourceControlHelpers::PackageFilename(Package);
-				FSavePackageArgs SaveArgs;
-				SaveArgs.TopLevelFlags = RF_Standalone;
-				if (!UPackage::SavePackage(Package, nullptr, *PackageFileName, SaveArgs))
+				TUniquePtr<FWorldPartitionActorDesc> NewActorDesc = Actor->CreateActorDesc();
+
+				if (ActorDesc->Equals(NewActorDesc.Get()))
 				{
-					UE_LOG(LogWorldPartitionResaveActorsBuilder, Error, TEXT("Error saving package %s."), *Package->GetName());
-					++FailCount;
 					return true;
 				}
 
-				UE_LOG(LogWorldPartitionResaveActorsBuilder, Display, TEXT("Saved package %s."), *Package->GetName());
-				++SaveCount;
-				return true;
+				UE_LOG(LogWorldPartitionResaveActorsBuilder, Log, TEXT("Package %s needs to be resaved."), *Package->GetName());
 			}
-			else
+
+			if (!bReportOnly)
 			{
-				// It is possible the resave can't checkout everything. Continue processing.
-				UE_LOG(LogWorldPartitionResaveActorsBuilder, Warning, TEXT("Error checking out package %s."), *Package->GetName());
-				++FailCount;
-				return true;
+				if (SCCHelper.Checkout(Package))
+				{
+					// Save package
+					FString PackageFileName = SourceControlHelpers::PackageFilename(Package);
+					FSavePackageArgs SaveArgs;
+					SaveArgs.TopLevelFlags = RF_Standalone;
+					if (!UPackage::SavePackage(Package, nullptr, *PackageFileName, SaveArgs))
+					{
+						UE_LOG(LogWorldPartitionResaveActorsBuilder, Error, TEXT("Error saving package %s."), *Package->GetName());
+						++FailCount;
+						return true;
+					}
+
+					UE_LOG(LogWorldPartitionResaveActorsBuilder, Display, TEXT("Saved package %s."), *Package->GetName());
+					++SaveCount;
+					return true;
+				}
+				else
+				{
+					// It is possible the resave can't checkout everything. Continue processing.
+					UE_LOG(LogWorldPartitionResaveActorsBuilder, Warning, TEXT("Error checking out package %s."), *Package->GetName());
+					++FailCount;
+					return true;
+				}
 			}
 		}
 
 		return true;
 	});
 
-	SCCHelper.Delete(PackagesToDelete);
+	if (!bReportOnly)
+	{
+		SCCHelper.Delete(PackagesToDelete);
+	}
 
-	if (bSwitchActorPackagingSchemeToReduced)
+	if (bSwitchActorPackagingSchemeToReduced && !bReportOnly)
 	{
 		// checkout world package
 		if (!SCCHelper.Checkout(WorldPackage))
@@ -261,4 +333,4 @@ bool UWorldPartitionResaveActorsBuilder::RunInternal(UWorld* World, const FCellI
 	}
 
 	return true;
-}
+}
\ No newline at end of file
diff --git a/Engine/Source/Editor/UnrealEd/Public/WorldPartition/WorldPartitionResaveActorsBuilder.h b/Engine/Source/Editor/UnrealEd/Public/WorldPartition/WorldPartitionResaveActorsBuilder.h
--- a/Engine/Source/Editor/UnrealEd/Public/WorldPartition/WorldPartitionResaveActorsBuilder.h
+++ b/Engine/Source/Editor/UnrealEd/Public/WorldPartition/WorldPartitionResaveActorsBuilder.h
@@ -21,6 +21,15 @@ public:
 	// UWorldPartitionBuilder interface end
 
 private:
+	UPROPERTY()
 	FString ActorClassName;
-	bool bSwitchActorPackagingSchemeToReduced;
+
+	UPROPERTY()
+	bool bReportOnly;
+
+	UPROPERTY()
+	bool bResaveDirtyActorDescsOnly;
+
+	UPROPERTY()
+	bool bSwitchActorPackagingSchemeToReduced;	
 };
\ No newline at end of file
diff --git a/Engine/Source/Runtime/Engine/Classes/GameFramework/Actor.h b/Engine/Source/Runtime/Engine/Classes/GameFramework/Actor.h
--- a/Engine/Source/Runtime/Engine/Classes/GameFramework/Actor.h
+++ b/Engine/Source/Runtime/Engine/Classes/GameFramework/Actor.h
@@ -885,10 +885,6 @@ public:
 	virtual void OnPlayFromHere();
 
 	bool CanPlayFromHere() const { return bCanPlayFromHere; }
-private:
-	friend class FActorDescList;
-	friend class UActorDescContainer;
-	friend class FWorldPartitionActorViewProxy;
 
 	/**
 	 * Creates an uninitialized actor descriptor. Can be overriden.
diff --git a/Engine/Source/Runtime/Engine/Private/WorldPartition/ActorPartition/PartitionActorDesc.cpp b/Engine/Source/Runtime/Engine/Private/WorldPartition/ActorPartition/PartitionActorDesc.cpp
--- a/Engine/Source/Runtime/Engine/Private/WorldPartition/ActorPartition/PartitionActorDesc.cpp
+++ b/Engine/Source/Runtime/Engine/Private/WorldPartition/ActorPartition/PartitionActorDesc.cpp
@@ -35,4 +35,21 @@ void FPartitionActorDesc::Serialize(FArchive& Ar)
 		Ar << GridGuid;
 	}
 }
+
+bool FPartitionActorDesc::Equals(const FWorldPartitionActorDesc* Other) const
+{
+	if (FWorldPartitionActorDesc::Equals(Other))
+	{
+		const FPartitionActorDesc* PartitionActorDesc = (FPartitionActorDesc*)Other;
+
+		return
+			GridSize == PartitionActorDesc->GridSize &&
+			GridIndexX == PartitionActorDesc->GridIndexX &&
+			GridIndexY == PartitionActorDesc->GridIndexY &&
+			GridIndexZ == PartitionActorDesc->GridIndexZ &&
+			GridGuid == PartitionActorDesc->GridGuid;		
+	}
+
+	return false;
+}
 #endif
diff --git a/Engine/Source/Runtime/Engine/Private/WorldPartition/HLOD/HLODActorDesc.cpp b/Engine/Source/Runtime/Engine/Private/WorldPartition/HLOD/HLODActorDesc.cpp
--- a/Engine/Source/Runtime/Engine/Private/WorldPartition/HLOD/HLODActorDesc.cpp
+++ b/Engine/Source/Runtime/Engine/Private/WorldPartition/HLOD/HLODActorDesc.cpp
@@ -52,6 +52,25 @@ void FHLODActorDesc::Serialize(FArchive& Ar)
 	}
 }
 
+bool FHLODActorDesc::Equals(const FWorldPartitionActorDesc* Other) const
+{
+	if (FWorldPartitionActorDesc::Equals(Other))
+	{
+		const FHLODActorDesc* HLODActorDesc = (FHLODActorDesc*)Other;
+
+		if (CellHash == HLODActorDesc->CellHash && SubActors.Num() == HLODActorDesc->SubActors.Num())
+		{
+			TArray<FGuid> SortedSubActors(SubActors);
+			TArray<FGuid> SortedSubActorsOther(HLODActorDesc->SubActors);
+			SortedSubActors.Sort();
+			SortedSubActorsOther.Sort();
+			return SortedSubActors == SortedSubActorsOther;
+		}
+	}
+
+	return false;
+}
+
 uint64 FHLODActorDesc::ComputeCellHash(const FString HLODLayerName, uint64 GridIndexX, uint64 GridIndexY, uint64 GridIndexZ, FDataLayersID DataLayersID)
 {
 	uint64 CellHash = FCrc::StrCrc32(*HLODLayerName);
diff --git a/Engine/Source/Runtime/Engine/Private/WorldPartition/Landscape/LandscapeSplineActorDesc.cpp b/Engine/Source/Runtime/Engine/Private/WorldPartition/Landscape/LandscapeSplineActorDesc.cpp
--- a/Engine/Source/Runtime/Engine/Private/WorldPartition/Landscape/LandscapeSplineActorDesc.cpp
+++ b/Engine/Source/Runtime/Engine/Private/WorldPartition/Landscape/LandscapeSplineActorDesc.cpp
@@ -15,10 +15,20 @@ void FLandscapeSplineActorDesc::Init(const AActor* InActor)
 	FWorldPartitionActorDesc::Init(InActor);
 
 	const ALandscapeSplineActor* LandscapeSplineActor = CastChecked<ALandscapeSplineActor>(InActor);
-	check(LandscapeSplineActor);
 	LandscapeGuid = LandscapeSplineActor->GetLandscapeGuid();
 }
 
+bool FLandscapeSplineActorDesc::Equals(const FWorldPartitionActorDesc* Other) const
+{
+	if (FWorldPartitionActorDesc::Equals(Other))
+	{
+		const FLandscapeSplineActorDesc* LandscapeSplineActorDesc = (FLandscapeSplineActorDesc*)Other;
+		return LandscapeGuid == LandscapeSplineActorDesc->LandscapeGuid;
+	}
+
+	return false;
+}
+
 void FLandscapeSplineActorDesc::Serialize(FArchive& Ar)
 {
 	Ar.UsingCustomVersion(FUE5MainStreamObjectVersion::GUID);
diff --git a/Engine/Source/Runtime/Engine/Private/WorldPartition/LevelInstance/LevelInstanceActorDesc.cpp b/Engine/Source/Runtime/Engine/Private/WorldPartition/LevelInstance/LevelInstanceActorDesc.cpp
--- a/Engine/Source/Runtime/Engine/Private/WorldPartition/LevelInstance/LevelInstanceActorDesc.cpp
+++ b/Engine/Source/Runtime/Engine/Private/WorldPartition/LevelInstance/LevelInstanceActorDesc.cpp
@@ -59,6 +59,21 @@ void FLevelInstanceActorDesc::Init(UActorDescContainer* InContainer, const FWorl
 	FWorldPartitionActorDesc::Init(InContainer, DescData);
 }
 
+bool FLevelInstanceActorDesc::Equals(const FWorldPartitionActorDesc* Other) const
+{
+	if (FWorldPartitionActorDesc::Equals(Other))
+	{
+		const FLevelInstanceActorDesc* LevelInstanceActorDesc = (FLevelInstanceActorDesc*)Other;
+
+		return
+			LevelPackage == LevelInstanceActorDesc->LevelPackage &&
+			LevelInstanceTransform.Equals(LevelInstanceActorDesc->LevelInstanceTransform, 0.1f) &&
+			DesiredRuntimeBehavior == LevelInstanceActorDesc->DesiredRuntimeBehavior;
+	}
+
+	return false;
+}
+
 void FLevelInstanceActorDesc::OnRegister(UWorld* InWorld)
 {
 	FWorldPartitionActorDesc::OnRegister(InWorld);
diff --git a/Engine/Source/Runtime/Engine/Private/WorldPartition/WorldPartitionActorDesc.cpp b/Engine/Source/Runtime/Engine/Private/WorldPartition/WorldPartitionActorDesc.cpp
--- a/Engine/Source/Runtime/Engine/Private/WorldPartition/WorldPartitionActorDesc.cpp
+++ b/Engine/Source/Runtime/Engine/Private/WorldPartition/WorldPartitionActorDesc.cpp
@@ -119,6 +119,40 @@ void FWorldPartitionActorDesc::Init(UActorDescContainer* InContainer, const FWor
 	}
 }
 
+bool FWorldPartitionActorDesc::Equals(const FWorldPartitionActorDesc* Other) const
+{
+	if (Guid == Other->Guid && 
+		Class == Other->Class && 
+		ActorPackage == Other->ActorPackage && 
+		ActorPath == Other->ActorPath && 
+		ActorLabel == Other->ActorLabel && 
+		BoundsLocation.Equals(Other->BoundsLocation, 0.1f) &&
+		BoundsExtent.Equals(Other->BoundsExtent, 0.1f) && 
+		RuntimeGrid == Other->RuntimeGrid && 
+		bActorIsEditorOnly == Other->bActorIsEditorOnly && 
+		bLevelBoundsRelevant == Other->bLevelBoundsRelevant && 
+		bActorIsHLODRelevant == Other->bActorIsHLODRelevant && 
+		HLODLayer == Other->HLODLayer && 
+		FolderPath == Other->FolderPath &&
+		DataLayers.Num() == Other->DataLayers.Num() &&
+		References.Num() == Other->References.Num())
+	{
+		TArray<FName> SortedDataLayers(DataLayers);
+		TArray<FName> SortedDataLayersOther(Other->DataLayers);
+		SortedDataLayers.Sort();
+		SortedDataLayersOther.Sort();
+
+		TArray<FGuid> SortedReferences(References);
+		TArray<FGuid> SortedReferencesOther(Other->References);
+		SortedReferences.Sort();
+		SortedReferencesOther.Sort();
+
+		return SortedDataLayers == SortedDataLayersOther && SortedReferences == SortedReferencesOther;
+	}
+
+	return false;
+}
+
 void FWorldPartitionActorDesc::SerializeTo(TArray<uint8>& OutData)
 {
 	// Serialize to archive and gather custom versions
diff --git a/Engine/Source/Runtime/Engine/Public/WorldPartition/ActorPartition/PartitionActorDesc.h b/Engine/Source/Runtime/Engine/Public/WorldPartition/ActorPartition/PartitionActorDesc.h
--- a/Engine/Source/Runtime/Engine/Public/WorldPartition/ActorPartition/PartitionActorDesc.h
+++ b/Engine/Source/Runtime/Engine/Public/WorldPartition/ActorPartition/PartitionActorDesc.h
@@ -20,6 +20,7 @@ public:
 	FGuid GridGuid;
 protected:
 	virtual void Init(const AActor* InActor) override;
+	virtual bool Equals(const FWorldPartitionActorDesc* Other) const override;
 	virtual void Serialize(FArchive& Ar) override;
 #endif
 };
diff --git a/Engine/Source/Runtime/Engine/Public/WorldPartition/HLOD/HLODActorDesc.h b/Engine/Source/Runtime/Engine/Public/WorldPartition/HLOD/HLODActorDesc.h
--- a/Engine/Source/Runtime/Engine/Public/WorldPartition/HLOD/HLODActorDesc.h
+++ b/Engine/Source/Runtime/Engine/Public/WorldPartition/HLOD/HLODActorDesc.h
@@ -27,6 +27,7 @@ public:
 
 protected:
 	virtual void Init(const AActor* InActor) override;
+	virtual bool Equals(const FWorldPartitionActorDesc* Other) const override;
 	virtual void Serialize(FArchive& Ar) override;
 
 	TArray<FGuid> SubActors;
diff --git a/Engine/Source/Runtime/Engine/Public/WorldPartition/Landscape/LandscapeSplineActorDesc.h b/Engine/Source/Runtime/Engine/Public/WorldPartition/Landscape/LandscapeSplineActorDesc.h
--- a/Engine/Source/Runtime/Engine/Public/WorldPartition/Landscape/LandscapeSplineActorDesc.h
+++ b/Engine/Source/Runtime/Engine/Public/WorldPartition/Landscape/LandscapeSplineActorDesc.h
@@ -17,6 +17,7 @@ public:
 
 protected:
 	virtual void Init(const AActor* InActor) override;
+	virtual bool Equals(const FWorldPartitionActorDesc* Other) const override;
 	virtual void Serialize(FArchive& Ar) override;
 	virtual void OnRegister(UWorld* InWorld) override;
 	virtual void OnUnregister() override;
diff --git a/Engine/Source/Runtime/Engine/Public/WorldPartition/LevelInstance/LevelInstanceActorDesc.h b/Engine/Source/Runtime/Engine/Public/WorldPartition/LevelInstance/LevelInstanceActorDesc.h
--- a/Engine/Source/Runtime/Engine/Public/WorldPartition/LevelInstance/LevelInstanceActorDesc.h
+++ b/Engine/Source/Runtime/Engine/Public/WorldPartition/LevelInstance/LevelInstanceActorDesc.h
@@ -32,6 +32,7 @@ protected:
 	FLevelInstanceActorDesc();
 	virtual void Init(const AActor* InActor) override;
 	virtual void Init(UActorDescContainer* InContainer, const FWorldPartitionActorDescInitData& DescData) override;
+	virtual bool Equals(const FWorldPartitionActorDesc* Other) const override;
 	virtual void TransferFrom(const FWorldPartitionActorDesc* From) override;
 	virtual void Serialize(FArchive& Ar) override;
 	virtual void AddReferencedObjects(FReferenceCollector& Collector);
diff --git a/Engine/Source/Runtime/Engine/Public/WorldPartition/WorldPartitionActorDesc.h b/Engine/Source/Runtime/Engine/Public/WorldPartition/WorldPartitionActorDesc.h
--- a/Engine/Source/Runtime/Engine/Public/WorldPartition/WorldPartitionActorDesc.h
+++ b/Engine/Source/Runtime/Engine/Public/WorldPartition/WorldPartitionActorDesc.h
@@ -141,6 +141,8 @@ public:
 	virtual void Init(const AActor* InActor);
 	virtual void Init(UActorDescContainer* InContainer, const FWorldPartitionActorDescInitData& DescData);
 
+	virtual bool Equals(const FWorldPartitionActorDesc* Other) const;
+
 	void SerializeTo(TArray<uint8>& OutData);
 
 protected:
