EditorDomain: Fix "uses an unexpected custom version" performance messages from older packages in EngineTest.
#rnx

diff --git a/Engine/Config/BaseEditor.ini b/Engine/Config/BaseEditor.ini
--- a/Engine/Config/BaseEditor.ini
+++ b/Engine/Config/BaseEditor.ini
@@ -399,6 +399,8 @@ SaveUnversioned=false
 ; +PostLoadCanConstructClasses=<LongPackageNameOfExistingClass>,<LongPackageNameOfClassItCanCreateInPostLoad>
 +PostLoadCanConstructClasses=/Script/Engine.StaticMesh,/Script/Engine.StaticMeshDescriptionBulkData
 +PostLoadCanConstructClasses=/Script/Engine.MapBuildDataRegistry,/Script/Engine.TextureCube
++PostLoadCanConstructClasses=/Script/ClothingSystemRuntimeCommon.ClothingAssetCommon,/Script/ChaosCloth.ChaosClothSharedSimConfig
++PostLoadCanConstructClasses=/Script/ClothingSystemRuntimeCommon.ClothingAssetCommon,/Script/ChaosCloth.ChaosClothConfig
 
 [TargetDomain]
 ; IterativeClassAllowListEnabled/IterativeClassAllowList/IterativeClassBlockList:
diff --git a/Engine/Source/Editor/UnrealEd/Private/EditorDomain/EditorDomainUtils.cpp b/Engine/Source/Editor/UnrealEd/Private/EditorDomain/EditorDomainUtils.cpp
--- a/Engine/Source/Editor/UnrealEd/Private/EditorDomain/EditorDomainUtils.cpp
+++ b/Engine/Source/Editor/UnrealEd/Private/EditorDomain/EditorDomainUtils.cpp
@@ -377,6 +377,10 @@ FClassDigestData* FPrecacheClassDigest::GetRecursive(FName ClassName, bool bAllo
 	// Propagate values from the ConstructClasses
 	TArray<FName> ConstructClasses; // Not a class scratch variable because we use it across recursive calls
 	GConstructClasses.MultiFind(ClassName, ConstructClasses);
+	if (LookupName != ClassName)
+	{
+		GConstructClasses.MultiFind(LookupName, ConstructClasses);
+	}
 	if (!ConstructClasses.IsEmpty())
 	{
 		TArray<FGuid> ConstructCustomVersions; // Not a class scratch variable because we use it across recursive calls
@@ -470,8 +474,12 @@ class FCustomVersionCollectorArchive : public FArchiveUObject
 public:
 	FCustomVersionCollectorArchive()
 	{
+		// Use the same Archive properties that are used by FPackageHarvester, since that
+		// is the authoritative way of collecting CustomVersions used in the save
 		SetIsSaving(true);
 		SetIsPersistent(true);
+		ArIsObjectReferenceCollector = true;
+		ArShouldSkipBulkData = true;
 	}
 	// The base class functionality does most of what we need:
 	// ignore Serialize(void*,int), ignore Serialize(UObject*), collect customversions
@@ -1109,41 +1117,65 @@ bool TrySavePackage(UPackage* Package)
 	if (!UnknownGuids.IsEmpty())
 	{
 		TMap<FGuid, UObject*> Culprits = FindCustomVersionCulprits(UnknownGuids, Package);
+
+		// First check whether the culprit for (one of) the missing CustomVersion is an instance
+		// that was added during PostLoad. If so, advise adding an entry to PostLoadCanConstructClasses.
 		UObject* ConstructedCulprit = nullptr;
 		TOptional<FAssetPackageData> PackageData = IAssetRegistry::Get()->GetAssetPackageDataCopy(Package->GetFName());
-		TStringBuilder<256> Message;
-		TStringBuilder<128> FixupSuggestion;
-		Message << TEXT("Could not save ") << Package->GetFName() << TEXT(" to EditorDomain: It uses an unexpected custom version. ")
-			<< TEXT("Optimized loading and iterative cooking will be disabled for this package. ");
-		FixupSuggestion << TEXT("Modify the classes or structs used in the package to call Ar.UsingCustomVersion(Guid) at the beginning of their serialize function.");
 		for (const FGuid& CustomVersionGuid : UnknownGuids)
 		{
-			TOptional<FCustomVersion> CustomVersion = FCurrentCustomVersions::Get(CustomVersionGuid);
 			UObject* Culprit = Culprits.FindOrAdd(CustomVersionGuid);
-			FixupSuggestion << TEXT("\nCustomVersion(Guid=") << CustomVersionGuid << TEXT(", Name=")
-				<< (CustomVersion ? *CustomVersion->GetFriendlyName().ToString() : TEXT("<Unknown>"))
-				<< TEXT("): Used by ")
-				<< (Culprit ? *Culprit->GetClass()->GetPathName() : TEXT("<CulpritUnknown>"));
-			if (Culprit)
+			FName CulpritClassName = Culprit ? FName(*Culprit->GetClass()->GetPathName()) : NAME_None;
+			if (CulpritClassName.IsNone() || PackageData->ImportedClasses.Contains(CulpritClassName))
 			{
-				FName CulpritClassName = FName(*Culprit->GetClass()->GetPathName());
-				if (!PackageData->ImportedClasses.Contains(CulpritClassName))
-				{
-					ConstructedCulprit = Culprit;
-					break;
-				}
+				continue;
+			}
+			// If the culprit class does not declare the version either, then we still need to give the message
+			// advising adding an entry in DeclareCustomVersions
+			bool bConstructedClassDeclaresTheVersion = true;
+			{
+				PrecacheClassDigests({ CulpritClassName });
+				FClassDigestMap& ClassDigests = GetClassDigests();
+				FReadScopeLock ClassDigestsScopeLock(ClassDigests.Lock);
+				FClassDigestData* ExistingData = ClassDigests.Map.Find(CulpritClassName);
+				bConstructedClassDeclaresTheVersion = ExistingData && ExistingData->CustomVersionGuids.Contains(CustomVersionGuid);
+			}
+			if (bConstructedClassDeclaresTheVersion)
+			{
+				ConstructedCulprit = Culprit;
+				break;
 			}
 		}
+		TStringBuilder<128> FixupSuggestion;
 		if (ConstructedCulprit)
 		{
-			FixupSuggestion.Reset();
+			// Suggested debugging technique for this message: Add a conditional breakpoint on the packagename
+			// at the start of LoadPackageInternal. After it gets hit, add a breakpoint in the constructor
+			// of the ConstructedCulprit class.
 			FixupSuggestion << TEXT("The custom version is used by a class which was created after load of the package. ")
 				<< TEXT("Find the class that added ") << ConstructedCulprit->GetFullName() << TEXT(" and add ")
 				<< TEXT("Editor.ini:[EditorDomain]:+PostLoadCanConstructClasses=<ConstructingClass>,")
 				<< ConstructedCulprit->GetClass()->GetPathName();
 		}
-		Message << FixupSuggestion;
-		UE_LOG(LogEditorDomain, Display, TEXT("%s"), Message.ToString());
+		else
+		{
+			// Suggested debugging technique for this message: SetNextStatement back to beginning of the function,
+			// add a conditional breakpoint in FArchive::UsingCustomVersion with Key.A == 0x<FirstHexWordFromGuid>
+			FixupSuggestion << TEXT("Modify the classes or structs used in the package to call Ar.UsingCustomVersion(Guid) in Serialize or DeclareCustomVersions.");
+			for (const FGuid& CustomVersionGuid : UnknownGuids)
+			{
+				TOptional<FCustomVersion> CustomVersion = FCurrentCustomVersions::Get(CustomVersionGuid);
+				UObject* Culprit = Culprits.FindOrAdd(CustomVersionGuid);
+				FixupSuggestion << TEXT("\n\tCustomVersion(Guid=") << CustomVersionGuid << TEXT(", Name=")
+					<< (CustomVersion ? *CustomVersion->GetFriendlyName().ToString() : TEXT("<Unknown>"))
+					<< TEXT("): Used by ")
+					<< (Culprit ? *Culprit->GetClass()->GetPathName() : TEXT("<CulpritUnknown>"));
+			}
+
+		}
+		UE_LOG(LogEditorDomain, Display, TEXT("Could not save %s to EditorDomain: It uses an unexpected custom version. ")
+			TEXT("Optimized loading and iterative cooking will be disabled for this package.\n\t%s"),
+			*Package->GetName(), FixupSuggestion.ToString());
 		return false;
 	}
 
