Running WP navigation data generation from the editor
Change UWorldPartitionNavigationDataBuilder loading mode to IterativeCells2D
Fix issue on loading cell borders
Make that navigation data chunk are added on PostLoad()
Addition of a CleanPackages switch in UWorldPartitionBuilderCommandlet to allow builders to erase their packages

diff --git a/Engine/Source/Editor/UnrealEd/Private/EditorBuildUtils.cpp b/Engine/Source/Editor/UnrealEd/Private/EditorBuildUtils.cpp
--- a/Engine/Source/Editor/UnrealEd/Private/EditorBuildUtils.cpp
+++ b/Engine/Source/Editor/UnrealEd/Private/EditorBuildUtils.cpp
@@ -39,6 +39,8 @@
 #include "MaterialUtilities.h"
 #include "UnrealEngine.h"
 #include "DebugViewModeHelpers.h"
+#include "IDirectoryWatcher.h"
+#include "DirectoryWatcherModule.h"
 #include "MaterialStatsCommon.h"
 #include "Materials/MaterialInstance.h"
 #include "VirtualTexturingEditorModule.h"
@@ -838,9 +840,9 @@ void FEditorBuildUtils::SubmitPackagesForAutomatedBuild( const TSet<UPackage*>&
 	SourceControlProvider.Execute( CheckInOperation, LevelsToSubmit, EConcurrency::Synchronous );
 }
 
-void FEditorBuildUtils::TriggerNavigationBuilder(UWorld* InWorld, FName Id)
+void FEditorBuildUtils::TriggerNavigationBuilder(UWorld*& InOutWorld, FName Id)
 {
-	if (InWorld)
+	if (InOutWorld)
 	{
 		if (Id == FBuildOptions::BuildAIPaths ||
 			Id == FBuildOptions::BuildSelectedAIPaths ||
@@ -855,9 +857,105 @@ void FEditorBuildUtils::TriggerNavigationBuilder(UWorld* InWorld, FName Id)
 			bBuildingNavigationFromUserRequest = false;
 		}
 
-		// Invoke navmesh generator
-		FNavigationSystem::Build(*InWorld);
+		const bool bAllowPartitionedBuildingFromEditor = false; // Still experimental, not enabled by default yet.
+		const FString& LongPackageName = InOutWorld->GetPackage()->GetName();
+		if (bAllowPartitionedBuildingFromEditor && ULevel::GetIsLevelPartitionedFromPackage(*LongPackageName))
+		{
+			WorldPartitionBuildNavigation(LongPackageName);
+			InOutWorld = GEditor->GetEditorWorldContext().World();
+		}
+		else
+		{
+			// Invoke navmesh generator
+			FNavigationSystem::Build(*InOutWorld);
+		}
+	}
+}
+
+bool FEditorBuildUtils::WorldPartitionBuildNavigation(const FString& InLongPackageName)
+{
+	// Ask user to save dirty packages
+	if (!FEditorFileUtils::SaveDirtyPackages(/*bPromptUserToSave=*/true, /*bSaveMapPackages=*/true, /*bSaveContentPackages=*/false))
+	{
+		return false;
+	}
+
+	// Unload any loaded map
+	if (!UEditorLoadingAndSavingUtils::NewBlankMap(/*bSaveExistingMap*/false))
+	{
+		return false;
+	}
+	
+	FProcHandle ProcessHandle;
+	bool bCancelled = false;
+
+	// Task scope
+	{
+		FScopedSlowTask SlowTask(0, LOCTEXT("WorldPartitionBuildNavigationProgress", "Building navigation..."));
+		SlowTask.MakeDialog(true);
+
+		const FString CurrentExecutableName = FPlatformProcess::ExecutablePath();
+
+		// Try to provide complete Path, if we can't try with project name
+		const FString ProjectPath = FPaths::IsProjectFilePathSet() ? FPaths::GetProjectFilePath() : FApp::GetProjectName();
+
+		uint32 ProcessID;
+
+		const FString Arguments = FString::Printf(TEXT("\"%s\" -run=WorldPartitionBuilderCommandlet %s  %s"), *ProjectPath, *InLongPackageName, TEXT(" -AllowCommandletRendering -Builder=WorldPartitionNavigationDataBuilder -SCCProvider=None -log=WPNavigationBuilderLog.txt"));
+		ProcessHandle = FPlatformProcess::CreateProc(*CurrentExecutableName, *Arguments, true, false, false, &ProcessID, 0, nullptr, nullptr);
+		
+		while (FPlatformProcess::IsProcRunning(ProcessHandle))
+		{
+			if (SlowTask.ShouldCancel())
+			{
+				bCancelled = true;
+				FPlatformProcess::TerminateProc(ProcessHandle);
+				break;
+			}
+
+			SlowTask.EnterProgressFrame(0);
+			FPlatformProcess::Sleep(0.1);
+		}
+	}
+
+	int32 Result = 0;
+	if (!bCancelled && FPlatformProcess::GetProcReturnCode(ProcessHandle, &Result))
+	{	
+		if (Result == 0)
+		{
+			// Unload any loaded map
+			if (!UEditorLoadingAndSavingUtils::NewBlankMap(/*bSaveExistingMap*/false))
+			{
+				return false;
+			}
+			
+			// Force registry update before loading converted map
+			const FAssetRegistryModule& AssetRegistryModule = FModuleManager::LoadModuleChecked<FAssetRegistryModule>("AssetRegistry");
+			IAssetRegistry& AssetRegistry = AssetRegistryModule.Get();
+											
+			FString MapToLoad = InLongPackageName;
+			
+			AssetRegistry.ScanModifiedAssetFiles({ MapToLoad });
+			AssetRegistry.ScanPathsSynchronous({ ULevel::GetExternalActorsPath(MapToLoad) }, true);
+
+			// Force a directory watcher tick for the asset registry to get notified of the changes
+			FDirectoryWatcherModule& DirectoryWatcherModule = FModuleManager::Get().LoadModuleChecked<FDirectoryWatcherModule>(TEXT("DirectoryWatcher"));
+			DirectoryWatcherModule.Get()->Tick(-1.0f);
+			
+			FEditorFileUtils::LoadMap(MapToLoad);
+		}
 	}
+	else if (bCancelled)
+	{
+		FMessageDialog::Open(EAppMsgType::Ok, LOCTEXT("WorldPartitionBuildNavigationCancelled", "Building navigation cancelled!"));
+	}
+	
+	if (Result != 0)
+	{
+		FMessageDialog::Open(EAppMsgType::Ok, LOCTEXT("WorldPartitionBuildNavigationFailed", "Building navigation failed!"));
+	}
+
+	return false;
 }
 
 /**
diff --git a/Engine/Source/Editor/UnrealEd/Private/WorldPartition/WorldPartitionNavigationDataBuilder.cpp b/Engine/Source/Editor/UnrealEd/Private/WorldPartition/WorldPartitionNavigationDataBuilder.cpp
--- a/Engine/Source/Editor/UnrealEd/Private/WorldPartition/WorldPartitionNavigationDataBuilder.cpp
+++ b/Engine/Source/Editor/UnrealEd/Private/WorldPartition/WorldPartitionNavigationDataBuilder.cpp
@@ -8,6 +8,7 @@
 #include "StaticMeshCompiler.h"
 #include "Logging/LogMacros.h"
 #include "UObject/SavePackage.h"
+#include "Commandlets/Commandlet.h"
 
 #include "WorldPartition/WorldPartition.h"
 #include "WorldPartition/WorldPartitionSubsystem.h"
@@ -24,11 +25,26 @@ UWorldPartitionNavigationDataBuilder::UWorldPartitionNavigationDataBuilder(const
 	
 	// Extra padding around loaded cell.
 	// @todo: set value programatically.
-	IterativeCellOverlapSize = 2000;
+	IterativeCellOverlapSize = 2000 + 51200;	// tile size + data chunk actor half size (because chunks are currently centered)
+}
+
+bool UWorldPartitionNavigationDataBuilder::PreRun(UWorld* World, FPackageSourceControlHelper& PackageHelper)
+{
+	TArray<FString> Tokens, Switches;
+	UCommandlet::ParseCommandLine(FCommandLine::Get(), Tokens, Switches);
+
+	bCleanBuilderPackages = Switches.Contains(TEXT("CleanPackages"));
+
+	return true;
 }
 
 bool UWorldPartitionNavigationDataBuilder::RunInternal(UWorld* World, const FCellInfo& InCellInfo, FPackageSourceControlHelper& PackageHelper)
 {
+	UE_LOG(LogWorldPartitionNavigationDataBuilder, Verbose, TEXT(" "));
+	UE_LOG(LogWorldPartitionNavigationDataBuilder, Verbose, TEXT("============================================================================================================"));
+	UE_LOG(LogWorldPartitionNavigationDataBuilder, Verbose, TEXT("RunInternal"));
+	UE_LOG(LogWorldPartitionNavigationDataBuilder, Verbose, TEXT("   Bounds %s ."), *InCellInfo.Bounds.ToString());
+	
 	UWorldPartitionSubsystem* WorldPartitionSubsystem = World->GetSubsystem<UWorldPartitionSubsystem>();
 	check(WorldPartitionSubsystem);
 
@@ -36,6 +52,7 @@ bool UWorldPartitionNavigationDataBuilder::RunInternal(UWorld* World, const FCel
 	check(WorldPartition);
 	
 	TSet<UPackage*> NavigationDataChunkActorPackages;
+	TSet<UPackage*> PackagesToClean;
 
 	// Gather all packages before any navigation data chunk actors are deleted
 	for (TActorIterator<ANavigationDataChunkActor> ItActor(World); ItActor; ++ItActor)
@@ -44,15 +61,57 @@ bool UWorldPartitionNavigationDataBuilder::RunInternal(UWorld* World, const FCel
 	}
 
 	// Destroy any existing navigation data chunk actors within bounds we are generating, we will make new ones.
-	const FBox GeneratingBounds = InCellInfo.Bounds.ExpandBy(-IterativeCellOverlapSize);
+	int32 Count = 0;
+	FBox GeneratingBounds = InCellInfo.Bounds.ExpandBy(-IterativeCellOverlapSize);
+
+	UE_LOG(LogWorldPartitionNavigationDataBuilder, Verbose, TEXT("   GeneratingBounds %s"), *GeneratingBounds.ToString());
+	
 	for (TActorIterator<ANavigationDataChunkActor> It(World); It; ++It)
 	{
+		Count++;
+		
 		ANavigationDataChunkActor* Actor = *It;
 		const FVector Location = Actor->GetActorLocation();
-		if (GeneratingBounds.IsInside(Location))
+		
+		auto IsInside2D = [](const FBox Bounds, const FVector& In) -> bool
 		{
-			World->DestroyActor(Actor);	
+			return ((In.X >= Bounds.Min.X) && (In.X < Bounds.Max.X) && (In.Y >= Bounds.Min.Y) && (In.Y < Bounds.Max.Y));
+		};
+
+		UE_LOG(LogWorldPartitionNavigationDataBuilder, Verbose, TEXT("   Location %s %s (%s %s)"),
+			*Location.ToCompactString(), IsInside2D(GeneratingBounds, Location) ? TEXT("inside") : TEXT("outside"), *Actor->GetName(), *Actor->GetPackage()->GetName());
+		
+		if (IsInside2D(GeneratingBounds, Location))
+		{
+			if (bCleanBuilderPackages)
+			{
+				check(!PackagesToClean.Find(Actor->GetPackage()));
+				PackagesToClean.Add(Actor->GetPackage());
+			}
+			
+			UE_LOG(LogWorldPartitionNavigationDataBuilder, Verbose, TEXT("   Destroy actor %s in package %s."), *Actor->GetName(), *Actor->GetPackage()->GetName());
+			World->DestroyActor(Actor);
+		}
+	}
+	UE_LOG(LogWorldPartitionNavigationDataBuilder, Verbose, TEXT("   Number of ANavigationDataChunkActor: %i"), Count);
+
+	// Check if we are just in cleaning mode.
+	if (bCleanBuilderPackages)
+	{
+		UE_LOG(LogWorldPartitionNavigationDataBuilder, Verbose, TEXT("   Number of packages to clear: %i"), PackagesToClean.Num());
+		
+		// Just delete all ANavigationDataChunkActor packages
+		if (!PackageHelper.Delete(PackagesToClean.Array()))
+		{
+			UE_LOG(LogWorldPartitionNavigationDataBuilder, Error, TEXT("Error deleting packages."));
 		}
+
+		if (!SavePackages(PackagesToClean.Array()))
+		{
+			return true;
+		}
+
+		return true;
 	}
 
 	// Make sure static meshes have compiled before generating navigation data
@@ -65,6 +124,16 @@ bool UWorldPartitionNavigationDataBuilder::RunInternal(UWorld* World, const FCel
 	for (TActorIterator<ANavigationDataChunkActor> ItActor(World); ItActor; ++ItActor)
 	{
 		NavigationDataChunkActorPackages.Add(ItActor->GetPackage());
+
+		// Log
+		FString String;
+		if (ItActor->GetPackage())
+		{
+			String += ItActor->GetPackage()->GetName();
+			String += UPackage::IsEmptyPackage(ItActor->GetPackage()) ? " empty" : " ";
+			String += ItActor->GetPackage()->IsDirty() ? " dirty" : " ";
+		}
+		UE_LOG(LogWorldPartitionNavigationDataBuilder, Verbose, TEXT("   Adding package %s (from actor %s)."), *String, *ItActor->GetName());
 	}
 
 	TArray<UPackage*> PackagesToSave;
@@ -79,10 +148,9 @@ bool UWorldPartitionNavigationDataBuilder::RunInternal(UWorld* World, const FCel
 			{
 				PackagesToDelete.Add(ActorPackage);
 			}
-			else
-			{
-				PackagesToSave.Add(ActorPackage);
-			}
+
+			// Save all packages (we need to also save the ones we are deleting).
+			PackagesToSave.Add(ActorPackage);
 		}
 	}
 
@@ -131,25 +199,10 @@ bool UWorldPartitionNavigationDataBuilder::RunInternal(UWorld* World, const FCel
 				}
 			}
 		}
-		
-		{
-			// Save packages
-			TRACE_CPUPROFILER_EVENT_SCOPE(SavingPackages);
-			UE_LOG(LogWorldPartitionNavigationDataBuilder, Log, TEXT("Saving %d packages."), PackagesToSave.Num());
 
-			for (UPackage* Package : PackagesToSave)
-			{
-				UE_LOG(LogWorldPartitionNavigationDataBuilder, Verbose, TEXT("   Saving package  %s."), *Package->GetName());
-				FString PackageFileName = SourceControlHelpers::PackageFilename(Package);
-				FSavePackageArgs SaveArgs;
-				SaveArgs.TopLevelFlags = RF_Standalone;
-				SaveArgs.SaveFlags = SAVE_Async;
-				if (!UPackage::SavePackage(Package, nullptr, *PackageFileName, SaveArgs))
-				{
-					UE_LOG(LogWorldPartitionNavigationDataBuilder, Error, TEXT("Error saving package %s."), *Package->GetName());
-					return true;
-				}
-			}
+		if (!SavePackages(PackagesToSave))
+		{
+			return true;
 		}
 		
 		{
@@ -172,3 +225,27 @@ bool UWorldPartitionNavigationDataBuilder::RunInternal(UWorld* World, const FCel
 
 	return true;
 }
+
+bool UWorldPartitionNavigationDataBuilder::SavePackages(const TArray<UPackage*>& PackagesToSave)
+{
+	// Save packages
+	TRACE_CPUPROFILER_EVENT_SCOPE(SavingPackages);
+	UE_LOG(LogWorldPartitionNavigationDataBuilder, Log, TEXT("Saving %d packages."), PackagesToSave.Num());
+
+	for (UPackage* Package : PackagesToSave)
+	{
+		UE_LOG(LogWorldPartitionNavigationDataBuilder, Verbose, TEXT("   Saving package  %s."), *Package->GetName());
+		FString PackageFileName = SourceControlHelpers::PackageFilename(Package);
+		FSavePackageArgs SaveArgs;
+		SaveArgs.TopLevelFlags = RF_Standalone;
+		SaveArgs.SaveFlags = SAVE_Async;
+		if (!UPackage::SavePackage(Package, nullptr, *PackageFileName, SaveArgs))
+		{
+			UE_LOG(LogWorldPartitionNavigationDataBuilder, Error, TEXT("   Error saving package %s."), *Package->GetName());
+			return false;
+		}
+	}
+
+	return true;
+}
+
diff --git a/Engine/Source/Editor/UnrealEd/Public/EditorBuildUtils.h b/Engine/Source/Editor/UnrealEd/Public/EditorBuildUtils.h
--- a/Engine/Source/Editor/UnrealEd/Public/EditorBuildUtils.h
+++ b/Engine/Source/Editor/UnrealEd/Public/EditorBuildUtils.h
@@ -254,10 +254,12 @@ private:
 	/** 
 	 * Trigger navigation builder to (re)generate NavMesh 
 	 *
-	 * @param	InWorld			WorldContext
+	 * @param	InOutWorld		WorldContext
 	 * @param	BuildSettings	Build settings that will be used for the editor build
 	 */
-	static void TriggerNavigationBuilder(UWorld* InWorld, FName Id);
+	static void TriggerNavigationBuilder(UWorld*& InOutWorld, FName Id);
+
+	static bool WorldPartitionBuildNavigation(const FString& InLongPackageName);
 
 	/** 
 	 * Trigger LOD builder to (re)generate LODActors
diff --git a/Engine/Source/Editor/UnrealEd/Public/WorldPartition/WorldPartitionNavigationDataBuilder.h b/Engine/Source/Editor/UnrealEd/Public/WorldPartition/WorldPartitionNavigationDataBuilder.h
--- a/Engine/Source/Editor/UnrealEd/Public/WorldPartition/WorldPartitionNavigationDataBuilder.h
+++ b/Engine/Source/Editor/UnrealEd/Public/WorldPartition/WorldPartitionNavigationDataBuilder.h
@@ -13,9 +13,14 @@ class UWorldPartitionNavigationDataBuilder : public UWorldPartitionBuilder
 public:
 	// UWorldPartitionBuilder interface begin
 	virtual bool RequiresCommandletRendering() const override { return false; }
-	virtual ELoadingMode GetLoadingMode() const override { return ELoadingMode::IterativeCells; }
+	virtual ELoadingMode GetLoadingMode() const override { return ELoadingMode::IterativeCells2D; }
 
 protected:
+	virtual bool PreRun(UWorld* World, FPackageSourceControlHelper& PackageHelper) override;
 	virtual bool RunInternal(UWorld* World, const FCellInfo& InCellInfo, FPackageSourceControlHelper& PackageHelper) override;
 	// UWorldPartitionBuilder interface end
+
+	bool SavePackages(const TArray<UPackage*>& PackagesToSave);
+
+	bool bCleanBuilderPackages = false;
 };
diff --git a/Engine/Source/Runtime/Engine/Private/WorldPartition/NavigationData/NavigationDataChunkActor.cpp b/Engine/Source/Runtime/Engine/Private/WorldPartition/NavigationData/NavigationDataChunkActor.cpp
--- a/Engine/Source/Runtime/Engine/Private/WorldPartition/NavigationData/NavigationDataChunkActor.cpp
+++ b/Engine/Source/Runtime/Engine/Private/WorldPartition/NavigationData/NavigationDataChunkActor.cpp
@@ -26,23 +26,29 @@ void ANavigationDataChunkActor::PostLoad()
 {
 	Super::PostLoad();
 
-	if (GEditor)
+	if (UWorld* World = GetWorld())
 	{
-		const bool bIsInPIE = (GEditor->PlayWorld != NULL) && (!GEditor->bIsSimulatingInEditor);
-		if (!bIsInPIE)
+		if (UNavigationSystemBase* NavSys = World->GetNavigationSystem())
 		{
-			Log(ANSI_TO_TCHAR(__FUNCTION__));
-			UE_LOG(LogNavigation, Verbose, TEXT("   pos: %s ext: %s"), *DataChunkActorBounds.GetCenter().ToCompactString(), *DataChunkActorBounds.GetExtent().ToCompactString());
-			AddNavigationDataChunkToWorld();
+			if (NavSys->IsWorldInitDone())
+			{
+				AddNavigationDataChunkInEditor(*NavSys);
+			}
+			else
+			{
+				UNavigationSystemBase::OnNavigationInitDoneStaticDelegate().AddUObject(this, &ANavigationDataChunkActor::AddNavigationDataChunkInEditor);
+			}
 		}
 	}
 }
 
 void ANavigationDataChunkActor::BeginDestroy()
 {
+	UNavigationSystemBase::OnNavigationInitDoneStaticDelegate().RemoveAll(this);
+	
 	if (GEditor)
 	{
-		const bool bIsInPIE = (GEditor->PlayWorld != NULL) && (!GEditor->bIsSimulatingInEditor);
+		const bool bIsInPIE = (GEditor->PlayWorld != nullptr) && (!GEditor->bIsSimulatingInEditor);
 		if (!bIsInPIE)
 		{
 			Log(ANSI_TO_TCHAR(__FUNCTION__));
@@ -52,14 +58,31 @@ void ANavigationDataChunkActor::BeginDestroy()
 
 	Super::BeginDestroy();
 }
+
+void ANavigationDataChunkActor::AddNavigationDataChunkInEditor(const UNavigationSystemBase& NavSys)
+{
+	if (GEditor)
+	{
+		const bool bIsInPIE = (GEditor->PlayWorld != nullptr) && (!GEditor->bIsSimulatingInEditor);
+		if (!bIsInPIE)
+		{
+			if (NavSys.GetWorld() == GetWorld())
+			{
+				Log(ANSI_TO_TCHAR(__FUNCTION__));
+				UE_LOG(LogNavigation, Verbose, TEXT("   pos: %s ext: %s"), *DataChunkActorBounds.GetCenter().ToCompactString(), *DataChunkActorBounds.GetExtent().ToCompactString());
+				AddNavigationDataChunkToWorld();
+			}
+		}
+	}
+}
 #endif // WITH_EDITOR
 
 void ANavigationDataChunkActor::CollectNavData(const FBox& QueryBounds, FBox& OutTilesBounds)
 {
+	TRACE_CPUPROFILER_EVENT_SCOPE(ANavigationDataChunkActor::CollectNavData);
 	Log(ANSI_TO_TCHAR(__FUNCTION__));
 
-	UWorld* World = GetWorld();
-	if (World)
+	if (UWorld* World = GetWorld())
 	{
 		if (UNavigationSystemBase* NavSys = World->GetNavigationSystem())
 		{
@@ -91,8 +114,7 @@ void ANavigationDataChunkActor::EndPlay(const EEndPlayReason::Type EndPlayReason
 
 void ANavigationDataChunkActor::AddNavigationDataChunkToWorld()
 {
-	UWorld* World = GetWorld();
-	if (World)
+	if (UWorld* World = GetWorld())
 	{
 		if (UNavigationSystemBase* NavSys = World->GetNavigationSystem())
 		{
@@ -103,8 +125,7 @@ void ANavigationDataChunkActor::AddNavigationDataChunkToWorld()
 
 void ANavigationDataChunkActor::RemoveNavigationDataChunkFromWorld()
 {
-	UWorld* World = GetWorld();
-	if (World)
+	if (UWorld* World = GetWorld())
 	{
 		if (UNavigationSystemBase* NavSys = World->GetNavigationSystem())
 		{
diff --git a/Engine/Source/Runtime/Engine/Private/WorldPartition/RuntimeSpatialHash/RuntimeSpatialHashNavigation.cpp b/Engine/Source/Runtime/Engine/Private/WorldPartition/RuntimeSpatialHash/RuntimeSpatialHashNavigation.cpp
--- a/Engine/Source/Runtime/Engine/Private/WorldPartition/RuntimeSpatialHash/RuntimeSpatialHashNavigation.cpp
+++ b/Engine/Source/Runtime/Engine/Private/WorldPartition/RuntimeSpatialHash/RuntimeSpatialHashNavigation.cpp
@@ -73,7 +73,12 @@ bool UWorldPartitionRuntimeSpatialHash::GenerateNavigationData(const FBox& Loade
 
 	const FSquare2DGridHelper::FGridLevel& GridLevelHelper = GridHelper.Levels[GridLevel];
 
-	GridLevelHelper.ForEachIntersectingCells(LoadedBounds, [GridLevel, &GridHelper, &GridLevelHelper, RuntimeGrid, &ActorCount, World, &ValidNavigationDataChunkActors, &NavDataBounds, this](const FIntVector2& CellCoord)
+	// Loaded bounds include an half cell border buffer to allow to build all cell withing the iterative loaded bounds
+	//  so here we substract the half cell size to make the building bounds.
+	// A DataChunkActor will be generated for each tile touching the building bounds.
+	const FBox BuildingBounds = LoadedBounds.ExpandBy(-(0.5*GridLevelHelper.CellSize));
+	
+	GridLevelHelper.ForEachIntersectingCells(BuildingBounds, [GridLevel, &GridHelper, &GridLevelHelper, RuntimeGrid, &ActorCount, World, &ValidNavigationDataChunkActors, &NavDataBounds, this](const FIntVector2& CellCoord)
 		{
 			TRACE_CPUPROFILER_EVENT_SCOPE(MakeNavigationDataChunkActorForGridCell);
 		
diff --git a/Engine/Source/Runtime/Engine/Private/WorldPartition/WorldPartition.cpp b/Engine/Source/Runtime/Engine/Private/WorldPartition/WorldPartition.cpp
--- a/Engine/Source/Runtime/Engine/Private/WorldPartition/WorldPartition.cpp
+++ b/Engine/Source/Runtime/Engine/Private/WorldPartition/WorldPartition.cpp
@@ -964,8 +964,6 @@ void UWorldPartition::UpdateLoadingEditorCell(UWorldPartitionEditorCell* Cell, b
 		{
 			if (ActorHandle.IsValid())
 			{
-				AActor* Actor = ActorHandle->GetActor();
-
 				// Filter actor against DataLayers
 				if (ShouldActorBeLoadedByEditorCells(*ActorHandle))
 				{
diff --git a/Engine/Source/Runtime/Engine/Public/WorldPartition/NavigationData/NavigationDataChunkActor.h b/Engine/Source/Runtime/Engine/Public/WorldPartition/NavigationData/NavigationDataChunkActor.h
--- a/Engine/Source/Runtime/Engine/Public/WorldPartition/NavigationData/NavigationDataChunkActor.h
+++ b/Engine/Source/Runtime/Engine/Public/WorldPartition/NavigationData/NavigationDataChunkActor.h
@@ -7,6 +7,7 @@
 #include "NavigationDataChunkActor.generated.h"
 
 class UNavigationDataChunk;
+class UNavigationSystemBase;
 
 UCLASS(NotPlaceable)
 class ENGINE_API ANavigationDataChunkActor : public AActor
@@ -14,13 +15,15 @@ class ENGINE_API ANavigationDataChunkActor : public AActor
 	GENERATED_UCLASS_BODY()
 
 public:
-	//~ Begin UObject Interface
 #if WITH_EDITOR
+	//~ Begin UObject Interface
 	virtual void PostLoad() override;
 	virtual void BeginDestroy() override;
-#endif //WITH_EDITOR
 	//~ End UObject Interface
 
+	void AddNavigationDataChunkInEditor(const UNavigationSystemBase& NavSys);
+#endif // WITH_EDITOR
+
 	const TArray<UNavigationDataChunk*>& GetNavDataChunk() const { return NavDataChunks; }
 	TArray<UNavigationDataChunk*>& GetMutableNavDataChunk() { return NavDataChunks; }
 
