Level & WorldPartition:
- Renamed bLevelOkayForPlacementWhileCheckedIn -> bPromptWhenAddingToLevelBeforeCheckout (transient, matches editor setting name)
- Add bPromptWhenAddingToLevelOutsideBounds (transient, matches editor setting name)
- WorldPartition level initializes both to false
- Remove WorldPartition specific check in ActorPlacementUtils::IsLevelValidForActorPlacement
- Set bLevelOkayForPlacementWhileCheckedIn = true once user discards the prompt

diff --git a/Engine/Source/Editor/UnrealEd/Private/AssetSelection.cpp b/Engine/Source/Editor/UnrealEd/Private/AssetSelection.cpp
--- a/Engine/Source/Editor/UnrealEd/Private/AssetSelection.cpp
+++ b/Engine/Source/Editor/UnrealEd/Private/AssetSelection.cpp
@@ -431,11 +431,6 @@ namespace ActorPlacementUtils
 {
 	bool IsLevelValidForActorPlacement(ULevel* InLevel, TArray<FTransform>& InActorTransforms)
 	{
-		// Checks below don't apply to partitioned world
-		if (InLevel && InLevel->GetWorldPartition())
-		{
-			return true;
-		}
 		if (FLevelUtils::IsLevelLocked(InLevel))
 		{
 			FNotificationInfo Info(NSLOCTEXT("UnrealEd", "Error_OperationDisallowedOnLockedLevel", "The requested operation could not be completed because the level is locked."));
@@ -453,12 +448,12 @@ namespace ActorPlacementUtils
 			// Don't prompt user for checks in unattended mode
 			return true;
 		}
-		if (InLevel && GetDefault<ULevelEditorMiscSettings>()->bPromptWhenAddingToLevelBeforeCheckout && SourceControlHelpers::IsAvailable())
+		if (InLevel && InLevel->GetPromptWhenAddingToLevelBeforeCheckout() && SourceControlHelpers::IsAvailable())
 		{
 			FString FileName = SourceControlHelpers::PackageFilename(InLevel->GetPathName());
 			// Query file state also checks the source control status
 			FSourceControlStatePtr SCState = ISourceControlModule::Get().GetProvider().GetState(FileName, EStateCacheUsage::Use);
-			if (!InLevel->bLevelOkayForPlacementWhileCheckedIn && !(SCState->IsCheckedOut() || SCState->IsAdded() || SCState->CanAdd() || SCState->IsUnknown()))
+			if (!(SCState->IsCheckedOut() || SCState->IsAdded() || SCState->CanAdd() || SCState->IsUnknown()))
 			{
 				FText Title = NSLOCTEXT("UnrealEd", "LevelCheckout_Title", "Level Checkout Warning");
 				if (EAppReturnType::Ok != FMessageDialog::Open(EAppMsgType::OkCancel, NSLOCTEXT("UnrealEd","LevelNotCheckedOutMsg", "This actor will be placed in a level that is in source control but not currently checked out. Continue?"), &Title))
@@ -467,7 +462,7 @@ namespace ActorPlacementUtils
 				}
 				else
 				{
-					InLevel->bLevelOkayForPlacementWhileCheckedIn = true;
+					InLevel->bPromptWhenAddingToLevelBeforeCheckout = false;
 				}
 			}
 		}
@@ -486,7 +481,7 @@ namespace ActorPlacementUtils
 				return true;
 			}
 		}
-		if (InLevel && GetDefault<ULevelEditorMiscSettings>()->bPromptWhenAddingToLevelOutsideBounds)
+		if (InLevel && InLevel->GetPromptWhenAddingToLevelOutsideBounds())
 		{
 			FBox CurrentLevelBounds(ForceInit);
 			if (InLevel->LevelBoundsActor.IsValid())
@@ -519,6 +514,11 @@ namespace ActorPlacementUtils
 					{
 						return false;
 					}
+					else
+					{
+						InLevel->bPromptWhenAddingToLevelOutsideBounds = false;
+						break;
+					}
 				}
 			}
 		}
diff --git a/Engine/Source/Runtime/Engine/Classes/Engine/Level.h b/Engine/Source/Runtime/Engine/Classes/Engine/Level.h
--- a/Engine/Source/Runtime/Engine/Classes/Engine/Level.h
+++ b/Engine/Source/Runtime/Engine/Classes/Engine/Level.h
@@ -708,7 +708,10 @@ public:
 	std::atomic<uint32> FixupOverrideVertexColorsCount;
 
 	UPROPERTY(transient)
-	bool bLevelOkayForPlacementWhileCheckedIn;
+	bool bPromptWhenAddingToLevelBeforeCheckout;
+
+	UPROPERTY(transient)
+	bool bPromptWhenAddingToLevelOutsideBounds;
 
 	UPROPERTY()
 	EActorPackagingScheme ActorPackagingScheme;
@@ -750,6 +753,9 @@ public:
 	ENGINE_API static bool GetLevelBoundsFromPackage(FName LevelPackage, FBox& OutLevelBounds);
 	ENGINE_API static bool GetIsLevelPartitionedFromPackage(FName LevelPackage);
 	ENGINE_API static bool GetIsLevelUsingExternalActorsFromPackage(FName LevelPackage);
+
+	ENGINE_API bool GetPromptWhenAddingToLevelOutsideBounds() const;
+	ENGINE_API bool GetPromptWhenAddingToLevelBeforeCheckout() const;
 #endif
 
 private:
diff --git a/Engine/Source/Runtime/Engine/Private/Level.cpp b/Engine/Source/Runtime/Engine/Private/Level.cpp
--- a/Engine/Source/Runtime/Engine/Private/Level.cpp
+++ b/Engine/Source/Runtime/Engine/Private/Level.cpp
@@ -65,6 +65,7 @@ Level.cpp: Level-related functions
 #include "Subsystems/EditorActorSubsystem.h"
 #include "Widgets/Notifications/SNotificationList.h"
 #include "Framework/Notifications/NotificationManager.h"
+#include "Settings/LevelEditorMiscSettings.h"
 #endif
 #include "WorldPartition/WorldPartition.h"
 #include "Engine/LevelStreaming.h"
@@ -341,6 +342,8 @@ ULevel::ULevel( const FObjectInitializer& ObjectInitializer )
 	bContainsStableActorGUIDs = true;
 	PlayFromHereActor = nullptr;
 	ActorPackagingScheme = EActorPackagingScheme::Reduced;
+	bPromptWhenAddingToLevelBeforeCheckout = true;
+	bPromptWhenAddingToLevelOutsideBounds = true;
 #endif	
 	bActorClusterCreated = false;
 	bIsPartitioned = false;
@@ -2011,6 +2014,16 @@ bool ULevel::GetLevelBoundsFromPackage(FName LevelPackage, FBox& OutLevelBounds)
 		return false;
 	});
 }
+
+bool ULevel::GetPromptWhenAddingToLevelOutsideBounds() const
+{
+	return !bIsPartitioned && bPromptWhenAddingToLevelOutsideBounds && GetDefault<ULevelEditorMiscSettings>()->bPromptWhenAddingToLevelOutsideBounds;
+}
+
+bool ULevel::GetPromptWhenAddingToLevelBeforeCheckout() const
+{
+	return !bIsPartitioned && bPromptWhenAddingToLevelBeforeCheckout && GetDefault<ULevelEditorMiscSettings>()->bPromptWhenAddingToLevelBeforeCheckout;
+}
 #endif
 
 void ULevel::InvalidateModelGeometry()
diff --git a/Engine/Source/Runtime/Engine/Private/LevelInstance/LevelInstanceEditorLevelStreaming.cpp b/Engine/Source/Runtime/Engine/Private/LevelInstance/LevelInstanceEditorLevelStreaming.cpp
--- a/Engine/Source/Runtime/Engine/Private/LevelInstance/LevelInstanceEditorLevelStreaming.cpp
+++ b/Engine/Source/Runtime/Engine/Private/LevelInstance/LevelInstanceEditorLevelStreaming.cpp
@@ -86,6 +86,10 @@ void ULevelStreamingLevelInstanceEditor::SetLoadedLevel(ULevel* Level)
 
 	if (ULevel* NewLoadedLevel = GetLoadedLevel())
 	{
+		// Avoid prompts for Level Instance editing
+		NewLoadedLevel->bPromptWhenAddingToLevelBeforeCheckout = false;
+		NewLoadedLevel->bPromptWhenAddingToLevelOutsideBounds = false;
+
 		check(!NewLoadedLevel->bAlreadyMovedActors);
 		if (AWorldSettings* WorldSettings = NewLoadedLevel->GetWorldSettings())
 		{
