Fix CIS warning and addition of a console command to turn on static navmesh building for world partitionned map

diff --git a/Engine/Source/Editor/UnrealEd/Private/EditorBuildUtils.cpp b/Engine/Source/Editor/UnrealEd/Private/EditorBuildUtils.cpp
--- a/Engine/Source/Editor/UnrealEd/Private/EditorBuildUtils.cpp
+++ b/Engine/Source/Editor/UnrealEd/Private/EditorBuildUtils.cpp
@@ -47,6 +47,7 @@
 #include "Components/RuntimeVirtualTextureComponent.h"
 #include "LandscapeSubsystem.h"
 #include "ShaderCompilerCore.h"
+#include "AssetRegistryModule.h"
 
 DEFINE_LOG_CATEGORY_STATIC(LogEditorBuildUtils, Log, All);
 
@@ -73,6 +74,12 @@ bool FEditorBuildUtils::bBuildingNavigationFromUserRequest = false;
 TMap<FName, FEditorBuildUtils::FCustomBuildType> FEditorBuildUtils::CustomBuildTypes;
 FName FEditorBuildUtils::InProgressBuildId;
 
+namespace UE::EditorBuildUtils
+{
+	static bool bNavmeshAllowPartitionedBuildingFromEditor = false; // Experimental, not enabled by default yet.
+	static FAutoConsoleVariableRef AllowPartitionedBuildingFromEditor(TEXT("n.bNavmeshAllowPartitionedBuildingFromEditor"), bNavmeshAllowPartitionedBuildingFromEditor, TEXT("Enable experimental navmesh partition building."), ECVF_Default);
+}
+
 /**
  * Class that handles potentially-async Build All requests.
  */
@@ -857,9 +864,8 @@ void FEditorBuildUtils::TriggerNavigationBuilder(UWorld*& InOutWorld, FName Id)
 			bBuildingNavigationFromUserRequest = false;
 		}
 
-		const bool bAllowPartitionedBuildingFromEditor = false; // Still experimental, not enabled by default yet.
 		const FString& LongPackageName = InOutWorld->GetPackage()->GetName();
-		if (bAllowPartitionedBuildingFromEditor && ULevel::GetIsLevelPartitionedFromPackage(*LongPackageName))
+		if (UE::EditorBuildUtils::bNavmeshAllowPartitionedBuildingFromEditor && ULevel::GetIsLevelPartitionedFromPackage(*LongPackageName))
 		{
 			WorldPartitionBuildNavigation(LongPackageName);
 			InOutWorld = GEditor->GetEditorWorldContext().World();
