Cooker: Check for HasExceededMaxMemory more frequently to avoid running out of memory in between checking.
#rnx

diff --git a/Engine/Source/Editor/UnrealEd/Private/CookOnTheFlyServer.cpp b/Engine/Source/Editor/UnrealEd/Private/CookOnTheFlyServer.cpp
--- a/Engine/Source/Editor/UnrealEd/Private/CookOnTheFlyServer.cpp
+++ b/Engine/Source/Editor/UnrealEd/Private/CookOnTheFlyServer.cpp
@@ -1897,6 +1897,12 @@ void UCookOnTheFlyServer::PumpLoads(UE::Cook::FTickStackData& StackData, uint32
 		LoadPackageInQueue(PackageData, StackData.ResultFlags, NumPushed);
 		OutNumPushed += NumPushed;
 		ProcessUnsolicitedPackages(); // May add new packages into the LoadQueue
+
+		if (HasExceededMaxMemory())
+		{
+			StackData.ResultFlags |= COSR_RequiresGC;
+			return;
+		}
 	}
 }
 
@@ -3083,6 +3089,13 @@ void UCookOnTheFlyServer::PumpSaves(UE::Cook::FTickStackData& StackData, uint32
 
 	UE_SCOPED_HIERARCHICAL_COOKTIMER(SavingPackages);
 	check(IsInGameThread());
+	ON_SCOPE_EXIT
+	{
+		if (HasExceededMaxMemory())
+		{
+			StackData.ResultFlags |= COSR_RequiresGC;
+		}
+	};
 
 	// save as many packages as we can during our time slice
 	FPackageDataQueue& SaveQueue = PackageDatas->GetSaveQueue();
