Fixed crash when right-clicking a skelmesh actor in single anim mode without an animation asset
- Don't return null object ptrs from GetReferencedContentObjects
- Added nullptr checks to UEditorEngine::GetReferencedAssetsForEditorSelection

diff --git a/Engine/Source/Editor/UnrealEd/Private/EditorEngine.cpp b/Engine/Source/Editor/UnrealEd/Private/EditorEngine.cpp
--- a/Engine/Source/Editor/UnrealEd/Private/EditorEngine.cpp
+++ b/Engine/Source/Editor/UnrealEd/Private/EditorEngine.cpp
@@ -2937,11 +2937,11 @@ void UEditorEngine::GetReferencedAssetsForEditorSelection(TArray<UObject*>& Obje
 
 		// If Blueprint assets should take precedence over any other referenced asset, check if there are any blueprints in this actor's list
 		// and if so, add only those.
-		if (bIgnoreOtherAssetsIfBPReferenced && ActorObjects.ContainsByPredicate([](UObject* Obj) { return Obj->IsA(UBlueprint::StaticClass()); }))
+		if (bIgnoreOtherAssetsIfBPReferenced && ActorObjects.ContainsByPredicate([](UObject* Obj) { return Obj && Obj->IsA(UBlueprint::StaticClass()); }))
 		{
 			for (UObject* Object : ActorObjects)
 			{
-				if (Object->IsA(UBlueprint::StaticClass()))
+				if (Object && Object->IsA(UBlueprint::StaticClass()))
 				{
 					Objects.Add(Object);
 				}
diff --git a/Engine/Source/Runtime/Engine/Private/Animation/SkeletalMeshActor.cpp b/Engine/Source/Runtime/Engine/Private/Animation/SkeletalMeshActor.cpp
--- a/Engine/Source/Runtime/Engine/Private/Animation/SkeletalMeshActor.cpp
+++ b/Engine/Source/Runtime/Engine/Private/Animation/SkeletalMeshActor.cpp
@@ -163,7 +163,8 @@ bool ASkeletalMeshActor::GetReferencedContentObjects(TArray<UObject*>& Objects)
 		Objects.Add(SkeletalMeshComponent->SkeletalMesh);
 	}
 
-	if (SkeletalMeshComponent->GetAnimationMode() == EAnimationMode::Type::AnimationSingleNode)
+	if (SkeletalMeshComponent->GetAnimationMode() == EAnimationMode::Type::AnimationSingleNode &&
+		SkeletalMeshComponent->AnimationData.AnimToPlay)
 	{
 		Objects.Add(SkeletalMeshComponent->AnimationData.AnimToPlay);
 	}
