Cooker: Decouple diff tools implementations from CookOnTheFlyServer. Move LinkerDiff into a CookedPackageWriter class. Move the implementation details of DiffOnly inside some new callbacks on ICookedPackageWriter.
#rnx

diff --git a/Engine/Source/Editor/UnrealEd/Private/CookOnTheFlyServer.cpp b/Engine/Source/Editor/UnrealEd/Private/CookOnTheFlyServer.cpp
--- a/Engine/Source/Editor/UnrealEd/Private/CookOnTheFlyServer.cpp
+++ b/Engine/Source/Editor/UnrealEd/Private/CookOnTheFlyServer.cpp
@@ -111,8 +111,6 @@
 #include "UObject/Class.h"
 #include "UObject/ConstructorHelpers.h"
 #include "UObject/GarbageCollection.h"
-#include "UObject/LinkerDiff.h"
-#include "UObject/LinkerSave.h"
 #include "UObject/MetaData.h"
 #include "UObject/ObjectSaveContext.h"
 #include "UObject/Package.h"
@@ -4011,41 +4009,10 @@ bool UCookOnTheFlyServer::HasRecompileShaderRequests() const
 class FDiffModeCookServerUtils
 {
 public:
-	/**
-	 *  Mode for the linker diff test
-	 */
-	enum ELinkerDiffMode
-	{
-		LDM_None = 0,
-		LDM_Algo = 1,
-		LDM_Consistent = 2,
-	};
-
-private:
-	/** Misc / common settings */
-	bool bInitialized = false;
-	bool bDiffEnabled = false;
-	ELinkerDiffMode LinkerDiffMode = LDM_None;
-	FString PackageFilter;
-
-	/** DumpObjList settings */
-	bool bDumpObjList = false;
-	FString DumpObjListParams;
-
-	/** DumpObjects settings */
-	bool bDumpObjects = false;
-	bool bDumpObjectsSorted = false;
-
-	/* LinkerDiff settings */
-	IConsoleVariable* EnableNewSave = nullptr;
-	int32 CurrentEnableNewSaveValue = 0;
-
-public:
-
 	void InitializePackageWriter(ICookedPackageWriter*& CookedPackageWriter)
 	{
 		Initialize();
-		if (!IsRunningCookDiff() && !IsRunningCookLinkerDiff())
+		if (!bDiffEnabled && !bLinkerDiffEnabled)
 		{
 			return;
 		}
@@ -4053,76 +4020,24 @@ public:
 		ICookedPackageWriter::FCookCapabilities Capabilities = CookedPackageWriter->GetCookCapabilities();
 		if (!Capabilities.bDiffModeSupported)
 		{
-			const TCHAR* CommandLineArg = IsRunningCookDiff() ? TEXT("-DIFFONLY") : TEXT("-LINKERDIFF");
+			const TCHAR* CommandLineArg = bDiffEnabled ? TEXT("-DIFFONLY") : TEXT("-LINKERDIFF");
 			UE_LOG(LogCook, Fatal, TEXT("%s was enabled, but -iostore is also enabled and iostore PackageWriters do not support %s."),
 				CommandLineArg, CommandLineArg);
 		}
 
-		if (IsRunningCookDiff())
+		if (bDiffEnabled)
 		{
 			// Wrap the incoming writer inside a FDiffPackageWriter
 			CookedPackageWriter = new FDiffPackageWriter(TUniquePtr<ICookedPackageWriter>(CookedPackageWriter));
 		}
-	}
-
-	bool IsRunningCookDiff() const
-	{
-		return bDiffEnabled;
-	}
-
-	bool IsRunningCookLinkerDiff() const
-	{
-		return LinkerDiffMode != LDM_None;
-	}
-
-	ELinkerDiffMode GetLinkerDiffMode() const
-	{
-		return LinkerDiffMode;
-	}
-
-	void ProcessPackage(UPackage* InPackage)
-	{
-		ConditionallyDumpObjList(InPackage);
-		ConditionallyDumpObjects(InPackage);
-	}
-
-	void LinkerDiffSetupOther()
-	{
-		// If the linker diff is comparing the two save algo, switch the cvar to the other algo before saving again,
-		// Otherwise the linker diff mode is tracking if the save is consistent across multiple save
-		if (LinkerDiffMode == FDiffModeCookServerUtils::LDM_Algo)
-		{
-			// see CVarEnablePackageNewSave definition in SavePackageUtilities.cpp for value meaning
-			EnableNewSave->Set(CurrentEnableNewSaveValue & 1 ? 0 : 1);
-		}
-	}
-
-	void LinkerDiffSetupCurrent()
-	{
-		if (LinkerDiffMode == FDiffModeCookServerUtils::LDM_Algo)
-		{
-			EnableNewSave->Set(CurrentEnableNewSaveValue);
-		}
-	}
-
-	void LinkerDiffFinish(FSavePackageResultStruct& OtherResult, FSavePackageResultStruct& CurrentResult)
-	{
-		if (LinkerDiffMode == FDiffModeCookServerUtils::LDM_Algo)
-		{
-			EnableNewSave->Set(CurrentEnableNewSaveValue);
-		}
-
-		if (OtherResult.LinkerSave && CurrentResult.LinkerSave)
+		else
 		{
-			FLinkerDiff LinkerDiff = FLinkerDiff::CompareLinkers(OtherResult.LinkerSave.Get(), CurrentResult.LinkerSave.Get());
-			LinkerDiff.PrintDiff(*GWarn);
+			check(bLinkerDiffEnabled);
+			CookedPackageWriter = new FLinkerDiffPackageWriter(TUniquePtr<ICookedPackageWriter>(CookedPackageWriter));
 		}
-		OtherResult.LinkerSave.Reset();
-		CurrentResult.LinkerSave.Reset();
 	}
 
 private:
-
 	void Initialize()
 	{
 		if (bInitialized)
@@ -4131,150 +4046,19 @@ private:
 		}
 
 		bDiffEnabled = FParse::Param(FCommandLine::Get(), TEXT("DIFFONLY"));
-
-		FString DiffMode;
-		if (FParse::Value(FCommandLine::Get(), TEXT("-LINKERDIFF="), DiffMode))
+		FString Value;
+		bLinkerDiffEnabled = FParse::Value(FCommandLine::Get(), TEXT("-LINKERDIFF="), Value);
+		if (bDiffEnabled && bLinkerDiffEnabled)
 		{
-			if (DiffMode == TEXT("1") || DiffMode == TEXT("algo"))
-			{
-				LinkerDiffMode = LDM_Algo;
-			}
-			else if (DiffMode == TEXT("2") || DiffMode == TEXT("consistent"))
-			{
-				LinkerDiffMode = LDM_Consistent;
-			}
-
-			EnableNewSave = IConsoleManager::Get().FindConsoleVariable(TEXT("SavePackage.EnableNewSave"));
-			CurrentEnableNewSaveValue = EnableNewSave->GetInt();
+			UE_LOG(LogCook, Fatal, TEXT("-DiffOnly and -LinkerDiff are mutually exclusive."));
 		}
 
-		ParseCmds();
 		bInitialized = true;
 	}
 
-	void RemoveParam(FString& InOutParams, const TCHAR* InParamToRemove)
-	{
-		int32 ParamIndex = InOutParams.Find(InParamToRemove);
-		if (ParamIndex >= 0)
-		{
-			int32 NextParamIndex = InOutParams.Find(TEXT(" -"), ESearchCase::CaseSensitive, ESearchDir::FromStart, ParamIndex + 1);
-			if (NextParamIndex < ParamIndex)
-			{
-				NextParamIndex = InOutParams.Len();
-			}
-			InOutParams = InOutParams.Mid(0, ParamIndex) + InOutParams.Mid(NextParamIndex);
-		}
-	}
-	void ParseDumpObjList(FString InParams)
-	{
-		const TCHAR* PackageFilterParam = TEXT("-packagefilter=");
-		FParse::Value(*InParams, PackageFilterParam, PackageFilter);
-		RemoveParam(InParams, PackageFilterParam);
-
-		// Add support for more parameters here
-		// After all parameters have been parsed and removed, pass the remaining string as objlist params
-		DumpObjListParams = InParams;
-	}
-	void ParseDumpObjects(FString InParams)
-	{
-		const TCHAR* PackageFilterParam = TEXT("-packagefilter=");
-		FParse::Value(*InParams, PackageFilterParam, PackageFilter);
-		RemoveParam(InParams, PackageFilterParam);
-
-		const TCHAR* SortParam = TEXT("sort");
-		bDumpObjectsSorted = FParse::Param(*InParams, SortParam);
-		RemoveParam(InParams, SortParam);
-	}
-
-	void ParseCmds()
-	{
-		const TCHAR* DumpObjListParam = TEXT("dumpobjlist");
-		const TCHAR* DumpObjectsParam = TEXT("dumpobjects");
-
-		FString CmdsText;
-		if (FParse::Value(FCommandLine::Get(), TEXT("-diffcmds="), CmdsText, false))
-		{
-			CmdsText = CmdsText.TrimQuotes();
-			TArray<FString> CmdsList;
-			CmdsText.ParseIntoArray(CmdsList, TEXT(","));
-			for (FString Cmd : CmdsList)
-			{
-				if (Cmd.StartsWith(DumpObjListParam))
-				{
-					bDumpObjList = true;
-					ParseDumpObjList(*Cmd + FCString::Strlen(DumpObjListParam));
-				}
-				else if (Cmd.StartsWith(DumpObjectsParam))
-				{
-					bDumpObjects = true;
-					ParseDumpObjects(*Cmd + FCString::Strlen(DumpObjectsParam));
-				}
-			}
-		}
-	}
-	bool FilterPackageName(UPackage* InPackage, const FString& InWildcard)
-	{
-		bool bInclude = false;
-		FString PackageName = InPackage->GetName();
-		if (PackageName.MatchesWildcard(InWildcard))
-		{
-			bInclude = true;
-		}
-		else if (FPackageName::GetShortName(PackageName).MatchesWildcard(InWildcard))
-		{
-			bInclude = true;
-		}
-		else if (InPackage->LinkerLoad)
-		{
-			FString Filename(InPackage->LinkerLoad->GetPackagePath().GetLocalFullPath());
-			bInclude = Filename.MatchesWildcard(InWildcard);
-		}
-		return bInclude;
-	}
-	void ConditionallyDumpObjList(UPackage* InPackage)
-	{
-		if (bDumpObjList)
-		{
-			if (FilterPackageName(InPackage, PackageFilter))
-			{
-				FString ObjListExec = TEXT("OBJ LIST ");
-				ObjListExec += DumpObjListParams;
-
-				TGuardValue<ELogTimes::Type> GuardLogTimes(GPrintLogTimes, ELogTimes::None);
-				TGuardValue<bool> GuardLogCategory(GPrintLogCategory, false);
-				TGuardValue<bool> GuardPrintLogVerbosity(GPrintLogVerbosity, false);
-
-				GEngine->Exec(nullptr, *ObjListExec);
-			}
-		}
-	}
-	void ConditionallyDumpObjects(UPackage* InPackage)
-	{
-		if (bDumpObjects)
-		{
-			if (FilterPackageName(InPackage, PackageFilter))
-			{
-				TArray<FString> AllObjects;
-				for (FThreadSafeObjectIterator It; It; ++It)
-				{
-					AllObjects.Add(*It->GetFullName());
-				}
-				if (bDumpObjectsSorted)
-				{
-					AllObjects.Sort();
-				}
-
-				TGuardValue<ELogTimes::Type> GuardLogTimes(GPrintLogTimes, ELogTimes::None);
-				TGuardValue<bool> GuardLogCategory(GPrintLogCategory, false);
-				TGuardValue<bool> GuardPrintLogVerbosity(GPrintLogVerbosity, false);
-
-				for (const FString& Obj : AllObjects)
-				{
-					UE_LOG(LogCook, Display, TEXT("%s"), *Obj);
-				}
-			}
-		}
-	}
+	bool bInitialized = false;
+	bool bDiffEnabled = false;
+	bool bLinkerDiffEnabled = false;
 };
 
 UE_TRACE_EVENT_BEGIN(UE_CUSTOM_COOKTIMER_LOG, SaveCookedPackage, NoSync)
@@ -4319,62 +4103,21 @@ void UCookOnTheFlyServer::SaveCookedPackage(UE::Cook::FSaveCookedPackageContext&
 			SaveArgs.TargetPlatform = TargetPlatform;
 			SaveArgs.bSlowTask = false;
 			SaveArgs.SavePackageContext = Context.SavePackageContext;
-			try
+
+			Context.PackageWriter->UpdateSaveArguments(SaveArgs);
+			do
 			{
-				if (DiffModeHelper->IsRunningCookDiff())
+				try
 				{
-					DiffModeHelper->ProcessPackage(Package);
-
-					// When looking for deterministic cook issues, first serialize the package to memory and do a simple diff with the existing package
-					FDiffPackageWriter* DiffPackageWriter = static_cast<FDiffPackageWriter*>(Context.PackageWriter);
-					Context.SavePackageResult = GEditor->Save(Package, Context.World, *Context.PlatFilename, SaveArgs);
-					if (Context.SavePackageResult == ESavePackageResult::Success && DiffPackageWriter->IsDifferenceFound())
-					{
-						// If the simple memory diff was not identical, collect callstacks for all Serialize calls and dump differences to log
-						DiffPackageWriter->BeginDiffCallstack();
-
-						Context.SavePackageResult = GEditor->Save(Package, Context.World, *Context.PlatFilename,
-							SaveArgs);
-					}
-				}
-				else if (DiffModeHelper->IsRunningCookLinkerDiff())
-				{
-					// Save the package with the other save algorithm. Discard the result but keep the linker.
-					DiffModeHelper->LinkerDiffSetupOther();
-					FSavePackageResultStruct OtherResult;
-					OtherResult = GEditor->Save(Package, Context.World, *Context.PlatFilename, SaveArgs);
-
-					// The contract with the PackageWriter is that every Begin is paired with a single commit.
-					// Send the old commit and the new begin.
-					IPackageWriter::FCommitPackageInfo CommitInfo;
-					CommitInfo.bSucceeded = OtherResult == ESavePackageResult::Success;
-					CommitInfo.PackageName = Package->GetFName();
-					CommitInfo.WriteOptions = IPackageWriter::EWriteOptions::None;
-					Context.PackageWriter->CommitPackage(MoveTemp(CommitInfo));
-					IPackageWriter::FBeginPackageInfo BeginInfo;
-					BeginInfo.PackageName = Package->GetFName();
-					BeginInfo.LooseFilePath = Context.PlatFilename;
-					Context.PackageWriter->BeginPackage(BeginInfo);
-
-					// Resave the package with the current save algorithm.
-					// Commit the result in FinishPlatform later, and compare the linkers now.
-					DiffModeHelper->LinkerDiffSetupCurrent();
-					FSavePackageResultStruct NewResult;
 					Context.SavePackageResult = GEditor->Save(Package, Context.World, *Context.PlatFilename, SaveArgs);
-
-					DiffModeHelper->LinkerDiffFinish(OtherResult, Context.SavePackageResult);
 				}
-				else
+				catch (std::exception&)
 				{
-					Context.SavePackageResult = GEditor->Save(Package, Context.World, *Context.PlatFilename, SaveArgs);
+					UE_LOG(LogCook, Warning, TEXT("Tried to save package %s for target platform %s but threw an exception"),
+						*Package->GetName(), *TargetPlatform->PlatformName());
+					Context.SavePackageResult = ESavePackageResult::Error;
 				}
-			}
-			catch (std::exception&)
-			{
-				UE_LOG(LogCook, Warning, TEXT("Tried to save package %s for target platform %s but threw an exception"),
-					*Package->GetName(), *TargetPlatform->PlatformName());
-				Context.SavePackageResult = ESavePackageResult::Error;
-			}
+			} while (Context.PackageWriter->IsAnotherSaveNeeded(Context.SavePackageResult, SaveArgs));
 
 			// If package was actually saved check with asset manager to make sure it wasn't excluded for being a
 			// development or never cook package. But skip sending the warnings from this check if it was editor-only.
@@ -4432,10 +4175,6 @@ void FSaveCookedPackageContext::SetupPackage()
 	bKeepEditorOnlyPackages |= COTFS.IsCookFlagSet(ECookInitializationFlags::Iterative);
 	SaveFlags |= bKeepEditorOnlyPackages ? SAVE_KeepEditorOnlyCookedPackages : SAVE_None;
 	SaveFlags |= COTFS.CookByTheBookOptions ? SAVE_ComputeHash : SAVE_None;
-	if (COTFS.DiffModeHelper->IsRunningCookLinkerDiff())
-	{
-		SaveFlags |= SAVE_CompareLinker;
-	}
 
 	// Use SandboxFile to do path conversion to properly handle sandbox paths (outside of standard paths in particular).
 	Filename = COTFS.ConvertToFullSandboxPath(*Filename, true);
diff --git a/Engine/Source/Editor/UnrealEd/Private/Cooker/DiffPackageWriter.cpp b/Engine/Source/Editor/UnrealEd/Private/Cooker/DiffPackageWriter.cpp
--- a/Engine/Source/Editor/UnrealEd/Private/Cooker/DiffPackageWriter.cpp
+++ b/Engine/Source/Editor/UnrealEd/Private/Cooker/DiffPackageWriter.cpp
@@ -2,15 +2,27 @@
 
 #include "Cooker/DiffPackageWriter.h"
 
+#include "Containers/Array.h"
 #include "Containers/Map.h"
+#include "CookOnTheSide/CookLog.h"
+#include "CoreGlobals.h"
+#include "Engine/Engine.h"
 #include "Misc/CommandLine.h"
 #include "Misc/ConfigCacheIni.h"
+#include "Misc/CString.h"
+#include "Misc/FeedbackContext.h"
 #include "Misc/Parse.h"
+#include "Misc/OutputDevice.h"
 #include "ProfilingDebugging/CookStats.h"
 #include "Serialization/ArchiveStackTrace.h"
 #include "Serialization/LargeMemoryWriter.h"
 #include "Templates/UniquePtr.h"
+#include "Templates/UnrealTemplate.h"
+#include "UObject/LinkerDiff.h"
+#include "UObject/LinkerSave.h"
+#include "UObject/SavePackage.h"
 #include "UObject/UObjectGlobals.h"
+#include "UObject/UObjectIterator.h"
 
 FDiffPackageWriter::FDiffPackageWriter(TUniquePtr<ICookedPackageWriter>&& InInner)
 	: Inner(MoveTemp(InInner))
@@ -31,6 +43,72 @@ FDiffPackageWriter::FDiffPackageWriter(TUniquePtr<ICookedPackageWriter>&& InInne
 	{
 		bIgnoreHeaderDiffs = FParse::Param(FCommandLine::Get(), TEXT("IgnoreHeaderDiffs"));
 	}
+
+	ParseCmds();
+}
+
+void FDiffPackageWriter::ParseCmds()
+{
+	const TCHAR* DumpObjListParam = TEXT("dumpobjlist");
+	const TCHAR* DumpObjectsParam = TEXT("dumpobjects");
+
+	FString CmdsText;
+	if (FParse::Value(FCommandLine::Get(), TEXT("-diffcmds="), CmdsText, false))
+	{
+		CmdsText = CmdsText.TrimQuotes();
+		TArray<FString> CmdsList;
+		CmdsText.ParseIntoArray(CmdsList, TEXT(","));
+		for (const FString& Cmd : CmdsList)
+		{
+			if (Cmd.StartsWith(DumpObjListParam))
+			{
+				bDumpObjList = true;
+				ParseDumpObjList(*Cmd + FCString::Strlen(DumpObjListParam));
+			}
+			else if (Cmd.StartsWith(DumpObjectsParam))
+			{
+				bDumpObjects = true;
+				ParseDumpObjects(*Cmd + FCString::Strlen(DumpObjectsParam));
+			}
+		}
+	}
+}
+
+void FDiffPackageWriter::ParseDumpObjList(FString InParams)
+{
+	const TCHAR* PackageFilterParam = TEXT("-packagefilter=");
+	FParse::Value(*InParams, PackageFilterParam, PackageFilter);
+	RemoveParam(InParams, PackageFilterParam);
+
+	// Add support for more parameters here
+	// After all parameters have been parsed and removed, pass the remaining string as objlist params
+	DumpObjListParams = InParams;
+}
+
+void FDiffPackageWriter::ParseDumpObjects(FString InParams)
+{
+	const TCHAR* PackageFilterParam = TEXT("-packagefilter=");
+	FParse::Value(*InParams, PackageFilterParam, PackageFilter);
+	RemoveParam(InParams, PackageFilterParam);
+
+	const TCHAR* SortParam = TEXT("sort");
+	bDumpObjectsSorted = FParse::Param(*InParams, SortParam);
+	RemoveParam(InParams, SortParam);
+}
+
+void FDiffPackageWriter::RemoveParam(FString& InOutParams, const TCHAR* InParamToRemove)
+{
+	int32 ParamIndex = InOutParams.Find(InParamToRemove);
+	if (ParamIndex >= 0)
+	{
+		int32 NextParamIndex = InOutParams.Find(TEXT(" -"),
+			ESearchCase::CaseSensitive, ESearchDir::FromStart, ParamIndex + 1);
+		if (NextParamIndex < ParamIndex)
+		{
+			NextParamIndex = InOutParams.Len();
+		}
+		InOutParams = InOutParams.Mid(0, ParamIndex) + InOutParams.Mid(NextParamIndex);
+	}
 }
 
 void FDiffPackageWriter::BeginPackage(const FBeginPackageInfo& Info)
@@ -40,22 +118,11 @@ void FDiffPackageWriter::BeginPackage(const FBeginPackageInfo& Info)
 	DiffMap.Reset();
 
 	BeginInfo = Info;
+	ConditionallyDumpObjList();
+	ConditionallyDumpObjects();
 	Inner->BeginPackage(Info);
 }
 
-void FDiffPackageWriter::BeginDiffCallstack()
-{
-	bDiffCallstack = true;
-
-	// The contract with the Inner is that Begin is paired with a single commit; send the old commit and the new begin
-	FCommitPackageInfo CommitInfo;
-	CommitInfo.bSucceeded = true;
-	CommitInfo.PackageName = BeginInfo.PackageName;
-	CommitInfo.WriteOptions = EWriteOptions::None;
-	Inner->CommitPackage(MoveTemp(CommitInfo));
-	Inner->BeginPackage(BeginInfo);
-}
-
 TFuture<FMD5Hash> FDiffPackageWriter::CommitPackage(FCommitPackageInfo&& Info)
 {
 	if (bDiffCallstack && bSaveForDiff)
@@ -118,4 +185,205 @@ TUniquePtr<FLargeMemoryWriter> FDiffPackageWriter::CreateLinkerArchive(FName Pac
 		return TUniquePtr<FLargeMemoryWriter>(new FArchiveStackTrace(Asset, *PackageName.ToString(),
 			false /* bInCollectCallstacks */));
 	}
-}
\ No newline at end of file
+}
+
+bool FDiffPackageWriter::IsAnotherSaveNeeded(FSavePackageResultStruct& PreviousResult, FSavePackageArgs& SaveArgs)
+{
+	checkf(!Inner->IsAnotherSaveNeeded(PreviousResult, SaveArgs),
+		TEXT("DiffPackageWriter does not support an Inner that needs multiple saves."));
+	// When looking for deterministic cook issues, first serialize the package to memory and do a simple diff with the
+	// existing package. If the simple memory diff was not identical, collect callstacks for all Serialize calls and
+	// dump differences to log
+	if (!bHasStartedSecondSave)
+	{
+		bHasStartedSecondSave = true;
+		if (PreviousResult.Result == ESavePackageResult::Success && bIsDifferent)
+		{
+			bDiffCallstack = true;
+
+			// The contract with the Inner is that Begin is paired with a single commit;
+			// send the old commit and the new begin
+			FCommitPackageInfo CommitInfo;
+			CommitInfo.bSucceeded = true;
+			CommitInfo.PackageName = BeginInfo.PackageName;
+			CommitInfo.WriteOptions = EWriteOptions::None;
+			Inner->CommitPackage(MoveTemp(CommitInfo));
+			Inner->BeginPackage(BeginInfo);
+
+			return true;
+		}
+		else
+		{
+			return false;
+		}
+	}
+	else
+	{
+		return false;
+	}
+}
+
+bool FDiffPackageWriter::FilterPackageName(const FString& InWildcard)
+{
+	bool bInclude = false;
+	FString PackageName = BeginInfo.PackageName.ToString();
+	if (PackageName.MatchesWildcard(InWildcard))
+	{
+		bInclude = true;
+	}
+	else if (FPackageName::GetShortName(PackageName).MatchesWildcard(InWildcard))
+	{
+		bInclude = true;
+	}
+	else
+	{
+		const FString& Filename = BeginInfo.LooseFilePath;
+		bInclude = Filename.MatchesWildcard(InWildcard);
+	}
+	return bInclude;
+}
+
+void FDiffPackageWriter::ConditionallyDumpObjList()
+{
+	if (bDumpObjList)
+	{
+		if (FilterPackageName(PackageFilter))
+		{
+			FString ObjListExec = TEXT("OBJ LIST ");
+			ObjListExec += DumpObjListParams;
+
+			TGuardValue<ELogTimes::Type> GuardLogTimes(GPrintLogTimes, ELogTimes::None);
+			TGuardValue<bool> GuardLogCategory(GPrintLogCategory, false);
+			TGuardValue<bool> GuardPrintLogVerbosity(GPrintLogVerbosity, false);
+
+			GEngine->Exec(nullptr, *ObjListExec);
+		}
+	}
+}
+
+void FDiffPackageWriter::ConditionallyDumpObjects()
+{
+	if (bDumpObjects)
+	{
+		if (FilterPackageName(PackageFilter))
+		{
+			TArray<FString> AllObjects;
+			for (FThreadSafeObjectIterator It; It; ++It)
+			{
+				AllObjects.Add(*It->GetFullName());
+			}
+			if (bDumpObjectsSorted)
+			{
+				AllObjects.Sort();
+			}
+
+			TGuardValue<ELogTimes::Type> GuardLogTimes(GPrintLogTimes, ELogTimes::None);
+			TGuardValue<bool> GuardLogCategory(GPrintLogCategory, false);
+			TGuardValue<bool> GuardPrintLogVerbosity(GPrintLogVerbosity, false);
+
+			for (const FString& Obj : AllObjects)
+			{
+				UE_LOG(LogCook, Display, TEXT("%s"), *Obj);
+			}
+		}
+	}
+}
+
+FLinkerDiffPackageWriter::FLinkerDiffPackageWriter(TUniquePtr<ICookedPackageWriter>&& InInner)
+	: Inner(MoveTemp(InInner))
+{
+	FString DiffModeText;
+	FParse::Value(FCommandLine::Get(), TEXT("-LINKERDIFF="), DiffModeText);
+	DiffMode = EDiffMode::LDM_Algo;
+	if (DiffModeText == TEXT("1") || DiffModeText == TEXT("algo"))
+	{
+		DiffMode = EDiffMode::LDM_Algo;
+	}
+	else if (DiffModeText == TEXT("2") || DiffModeText == TEXT("consistent"))
+	{
+		DiffMode = EDiffMode::LDM_Consistent;
+	}
+
+	EnableNewSave = IConsoleManager::Get().FindConsoleVariable(TEXT("SavePackage.EnableNewSave"));
+	CurrentEnableNewSaveValue = EnableNewSave->GetInt();
+}
+
+void FLinkerDiffPackageWriter::BeginPackage(const FBeginPackageInfo& Info)
+{
+	BeginInfo = Info;
+	Inner->BeginPackage(Info);
+	SetupOtherAlgorithm();
+}
+
+void FLinkerDiffPackageWriter::UpdateSaveArguments(FSavePackageArgs& SaveArgs)
+{
+	SaveArgs.SaveFlags |= SAVE_CompareLinker;
+	Inner->UpdateSaveArguments(SaveArgs);
+}
+
+bool FLinkerDiffPackageWriter::IsAnotherSaveNeeded(FSavePackageResultStruct& PreviousResult,
+	FSavePackageArgs& SaveArgs)
+{
+	checkf(!Inner->IsAnotherSaveNeeded(PreviousResult, SaveArgs),
+		TEXT("LinkerDiffPackageWriter does not support an Inner that needs multiple saves."));
+	if (!bHasStartedSecondSave)
+	{
+		bHasStartedSecondSave = true;
+		OtherResult = MoveTemp(PreviousResult);
+
+		// Resave the package with the current save algorithm.
+		SetupCurrentAlgorithm();
+
+		// The contract with the Inner is that every Begin is paired with a single commit.
+		// Send the old commit and the new begin.
+		IPackageWriter::FCommitPackageInfo CommitInfo;
+		CommitInfo.bSucceeded = OtherResult == ESavePackageResult::Success;
+		CommitInfo.PackageName = BeginInfo.PackageName;
+		CommitInfo.WriteOptions = IPackageWriter::EWriteOptions::None;
+		Inner->CommitPackage(MoveTemp(CommitInfo));
+		Inner->BeginPackage(BeginInfo);
+		return true;
+	}
+	else
+	{
+		CompareResults(PreviousResult);
+
+		if (DiffMode == EDiffMode::LDM_Algo)
+		{
+			EnableNewSave->Set(CurrentEnableNewSaveValue);
+		}
+		OtherResult.LinkerSave.Reset();
+		PreviousResult.LinkerSave.Reset();
+
+		return false;
+	}
+}
+
+void FLinkerDiffPackageWriter::SetupOtherAlgorithm()
+{
+	// If mode is comparing the two save algorithms, switch the cvar to the other algo before saving again,
+	// Otherwise the linker diff mode is tracking if the save is consistent across multiple save
+	if (DiffMode == EDiffMode::LDM_Algo)
+	{
+		// see CVarEnablePackageNewSave definition in SavePackageUtilities.cpp for value meaning
+		EnableNewSave->Set(CurrentEnableNewSaveValue & 1 ? 0 : 1);
+	}
+}
+
+void FLinkerDiffPackageWriter::SetupCurrentAlgorithm()
+{
+	if (DiffMode == EDiffMode::LDM_Algo)
+	{
+		EnableNewSave->Set(CurrentEnableNewSaveValue);
+	}
+}
+
+void FLinkerDiffPackageWriter::CompareResults(FSavePackageResultStruct& CurrentResult)
+{
+	if (OtherResult.LinkerSave && CurrentResult.LinkerSave)
+	{
+		FLinkerDiff LinkerDiff =
+			FLinkerDiff::CompareLinkers(OtherResult.LinkerSave.Get(), CurrentResult.LinkerSave.Get());
+		LinkerDiff.PrintDiff(*GWarn);
+	}
+}
diff --git a/Engine/Source/Editor/UnrealEd/Private/Cooker/DiffPackageWriter.h b/Engine/Source/Editor/UnrealEd/Private/Cooker/DiffPackageWriter.h
--- a/Engine/Source/Editor/UnrealEd/Private/Cooker/DiffPackageWriter.h
+++ b/Engine/Source/Editor/UnrealEd/Private/Cooker/DiffPackageWriter.h
@@ -6,7 +6,7 @@
 #include "Serialization/PackageWriter.h"
 #include "AssetRegistry/AssetRegistryState.h"
 
-/** A CookedPackageWriter that diffs the output from the current cook with the file that was saved in the previous cook. */
+/** A CookedPackageWriter that diffs output from the current cook with the file that was saved in the previous cook. */
 class FDiffPackageWriter : public ICookedPackageWriter
 {
 public:
@@ -21,8 +21,10 @@ public:
 	}
 	virtual void BeginPackage(const FBeginPackageInfo& Info) override;
 	virtual TFuture<FMD5Hash> CommitPackage(FCommitPackageInfo&& Info) override;
-	virtual void WritePackageData(const FPackageInfo& Info, FLargeMemoryWriter& ExportsArchive, const TArray<FFileRegion>& FileRegions) override;
-	virtual void WriteBulkData(const FBulkDataInfo& Info, const FIoBuffer& BulkData, const TArray<FFileRegion>& FileRegions) override
+	virtual void WritePackageData(const FPackageInfo& Info, FLargeMemoryWriter& ExportsArchive,
+		const TArray<FFileRegion>& FileRegions) override;
+	virtual void WriteBulkData(const FBulkDataInfo& Info, const FIoBuffer& BulkData,
+		const TArray<FFileRegion>& FileRegions) override
 	{
 		Inner->WriteBulkData(Info, BulkData, FileRegions);
 	}
@@ -30,7 +32,8 @@ public:
 	{
 		Inner->WriteAdditionalFile(Info, FileData);
 	}
-	virtual void WriteLinkerAdditionalData(const FLinkerAdditionalDataInfo& Info, const FIoBuffer& Data, const TArray<FFileRegion>& FileRegions) override
+	virtual void WriteLinkerAdditionalData(const FLinkerAdditionalDataInfo& Info, const FIoBuffer& Data,
+		const TArray<FFileRegion>& FileRegions) override
 	{
 		Inner->WriteLinkerAdditionalData(Info, Data, FileRegions);
 	}
@@ -91,25 +94,151 @@ public:
 	{
 		Inner->MarkPackagesUpToDate(UpToDatePackages);
 	}
-
-	// FDiffPackageWriter
-	/** Return whether a difference was found in the first save */
-	bool IsDifferenceFound() const
+	virtual void UpdateSaveArguments(FSavePackageArgs& SaveArgs) override
 	{
-		return bIsDifferent;
+		Inner->UpdateSaveArguments(SaveArgs);
 	}
-
-	/** Prepare the Inner PackageWriter for a second save, and set settings for callstack diffing. */
-	void BeginDiffCallstack();
+	virtual bool IsAnotherSaveNeeded(FSavePackageResultStruct& PreviousResult, FSavePackageArgs& SaveArgs) override;
 
 private:
+	void ParseCmds();
+	void ParseDumpObjList(FString InParams);
+	void ParseDumpObjects(FString InParams);
+	void RemoveParam(FString& InOutParams, const TCHAR* InParamToRemove);
+	bool FilterPackageName(const FString& InWildcard);
+	void ConditionallyDumpObjList();
+	void ConditionallyDumpObjects();
+
 	FArchiveDiffMap DiffMap;
 	FBeginPackageInfo BeginInfo;
 	TUniquePtr<ICookedPackageWriter> Inner;
+	FString DumpObjListParams;
+	FString PackageFilter;
 	int32 MaxDiffsToLog = 5;
 	bool bSaveForDiff = false;
 	bool bIgnoreHeaderDiffs = false;
 	bool bIsDifferent = false;
 	bool bDiffCallstack = false;
+	bool bHasStartedSecondSave = false;
+	bool bDumpObjList = false;
+	bool bDumpObjects = false;
+	bool bDumpObjectsSorted = false;
 };
 
+/** A CookedPackageWriter that runs the save twice with different methods (Save1 vs Save2). */
+class FLinkerDiffPackageWriter : public ICookedPackageWriter
+{
+public:
+	FLinkerDiffPackageWriter(TUniquePtr<ICookedPackageWriter>&& InInner);
+
+	// IPackageWriter
+	virtual FCapabilities GetCapabilities() const override
+	{
+		return Inner->GetCapabilities();
+	}
+	virtual void BeginPackage(const FBeginPackageInfo& Info) override;
+	virtual TFuture<FMD5Hash> CommitPackage(FCommitPackageInfo&& Info) override
+	{
+		return Inner->CommitPackage(MoveTemp(Info));
+	}
+	virtual void WritePackageData(const FPackageInfo& Info, FLargeMemoryWriter& ExportsArchive,
+		const TArray<FFileRegion>& FileRegions) override
+	{
+		Inner->WritePackageData(Info, ExportsArchive, FileRegions);
+	}
+	virtual void WriteBulkData(const FBulkDataInfo& Info, const FIoBuffer& BulkData,
+		const TArray<FFileRegion>& FileRegions) override
+	{
+		Inner->WriteBulkData(Info, BulkData, FileRegions);
+	}
+	virtual void WriteAdditionalFile(const FAdditionalFileInfo& Info, const FIoBuffer& FileData) override
+	{
+		Inner->WriteAdditionalFile(Info, FileData);
+	}
+	virtual void WriteLinkerAdditionalData(const FLinkerAdditionalDataInfo& Info, const FIoBuffer& Data,
+		const TArray<FFileRegion>& FileRegions) override
+	{
+		Inner->WriteLinkerAdditionalData(Info, Data, FileRegions);
+	}
+	virtual void AddToExportsSize(int64& ExportsSize) override
+	{
+		Inner->AddToExportsSize(ExportsSize);
+	}
+	virtual TUniquePtr<FLargeMemoryWriter> CreateLinkerArchive(FName PackageName, UObject* Asset) override
+	{
+		return Inner->CreateLinkerArchive(PackageName, Asset);
+	}
+	virtual bool IsPreSaveCompleted() const override
+	{
+		return Inner->IsPreSaveCompleted();
+	}
+
+	// ICookedPackageWriter
+	virtual FCookCapabilities GetCookCapabilities() const override
+	{
+		FCookCapabilities Result = Inner->GetCookCapabilities();
+		Result.bDiffModeSupported = false; // LinkerDiffPackageWriter can not be an inner of DiffPackageWriter
+		return Result;
+	}
+	virtual FDateTime GetPreviousCookTime() const
+	{
+		return Inner->GetPreviousCookTime();
+	}
+	virtual void Initialize(const FCookInfo& Info) override
+	{
+		Inner->Initialize(Info);
+	}
+	virtual void BeginCook() override
+	{
+		Inner->BeginCook();
+	}
+	virtual void EndCook() override
+	{
+		Inner->EndCook();
+	}
+	virtual void Flush() override
+	{
+		Inner->Flush();
+	}
+	virtual TUniquePtr<FAssetRegistryState> LoadPreviousAssetRegistry() override
+	{
+		return Inner->LoadPreviousAssetRegistry();
+	}
+	virtual FCbObject GetOplogAttachment(FName PackageName, FUtf8StringView AttachmentKey) override
+	{
+		return Inner->GetOplogAttachment(PackageName, AttachmentKey);
+	}
+	virtual void RemoveCookedPackages(TArrayView<const FName> PackageNamesToRemove) override
+	{
+		Inner->RemoveCookedPackages(PackageNamesToRemove);
+	}
+	virtual void RemoveCookedPackages() override
+	{
+		Inner->RemoveCookedPackages();
+	}
+	virtual void MarkPackagesUpToDate(TArrayView<const FName> UpToDatePackages) override
+	{
+		Inner->MarkPackagesUpToDate(UpToDatePackages);
+	}
+	virtual void UpdateSaveArguments(FSavePackageArgs& SaveArgs) override;
+	virtual bool IsAnotherSaveNeeded(FSavePackageResultStruct& PreviousResult, FSavePackageArgs& SaveArgs) override;
+
+private:
+	enum class EDiffMode : uint8
+	{
+		LDM_Algo,
+		LDM_Consistent,
+	};
+
+	void SetupOtherAlgorithm();
+	void SetupCurrentAlgorithm();
+	void CompareResults(FSavePackageResultStruct& CurrentResult);
+
+	FSavePackageResultStruct OtherResult;
+	FBeginPackageInfo BeginInfo;
+	TUniquePtr<ICookedPackageWriter> Inner;
+	IConsoleVariable* EnableNewSave = nullptr;
+	int32 CurrentEnableNewSaveValue = 0;
+	EDiffMode DiffMode = EDiffMode::LDM_Algo;
+	bool bHasStartedSecondSave = false;
+};
\ No newline at end of file
diff --git a/Engine/Source/Runtime/Core/Public/Serialization/PackageWriter.h b/Engine/Source/Runtime/Core/Public/Serialization/PackageWriter.h
--- a/Engine/Source/Runtime/Core/Public/Serialization/PackageWriter.h
+++ b/Engine/Source/Runtime/Core/Public/Serialization/PackageWriter.h
@@ -17,6 +17,8 @@ class FLargeMemoryWriter;
 class ICookedPackageWriter;
 class IPackageStoreWriter;
 struct FPackageStoreEntryResource;
+struct FSavePackageArgs;
+struct FSavePackageResultStruct;
 
 /** Interface for SavePackage to write packages to storage. */
 class IPackageWriter
@@ -279,7 +281,15 @@ public:
 		// The subclass must override this method if it returns bDiffModeSupported=true in GetCookCapabilities
 		unimplemented();
 	}
-
+	/** Modify the SaveArgs if required before the first Save. Used for diffing. */
+	virtual void UpdateSaveArguments(FSavePackageArgs& SaveArgs)
+	{
+	}
+	/** Report whether an additional save is needed and set up for it if so. Used for diffing. */
+	virtual bool IsAnotherSaveNeeded(FSavePackageResultStruct& PreviousResult, FSavePackageArgs& SaveArgs)
+	{
+		return false;
+	}
 	/** Downcast function for ICookedPackageWriters that implement the IPackageStoreWriter inherited interface. */
 	virtual IPackageStoreWriter* AsPackageStoreWriter()
 	{
