TargetDomain::FetchCookAttachments: Remove string allocations
#rnx

diff --git a/Engine/Source/Editor/UnrealEd/Private/TargetDomain/TargetDomainUtils.cpp b/Engine/Source/Editor/UnrealEd/Private/TargetDomain/TargetDomainUtils.cpp
--- a/Engine/Source/Editor/UnrealEd/Private/TargetDomain/TargetDomainUtils.cpp
+++ b/Engine/Source/Editor/UnrealEd/Private/TargetDomain/TargetDomainUtils.cpp
@@ -299,17 +299,17 @@ void FetchCookAttachments(TArrayView<FName> PackageNames, const ITargetPlatform*
 
 		for (FCbFieldView DepObj : DependenciesObj["builddependencies"])
 		{
-			if (FString DependencyName(DepObj.AsString()); !DependencyName.IsEmpty())
+			if (FUtf8StringView DependencyName(DepObj.AsString()); !DependencyName.IsEmpty())
 			{
-				Result.BuildDependencies.Add(FName(*DependencyName));
+				Result.BuildDependencies.Add(FName(DependencyName));
 			}
 		}
 
 		for (FCbFieldView DepObj : DependenciesObj["runtimeonlydependencies"])
 		{
-			if (FString DependencyName(DepObj.AsString()); !DependencyName.IsEmpty())
+			if (FUtf8StringView DependencyName(DepObj.AsString()); !DependencyName.IsEmpty())
 			{
-				Result.RuntimeOnlyDependencies.Add(FName(*DependencyName));
+				Result.RuntimeOnlyDependencies.Add(FName(DependencyName));
 			}
 		}
 
