Reduce cook time by avoiding slow array scan on critical path
  - Saves 4 minutes on a 2h24 cook

diff --git a/Engine/Source/Editor/UnrealEd/Private/Cooker/CookPackageData.cpp b/Engine/Source/Editor/UnrealEd/Private/Cooker/CookPackageData.cpp
--- a/Engine/Source/Editor/UnrealEd/Private/Cooker/CookPackageData.cpp
+++ b/Engine/Source/Editor/UnrealEd/Private/Cooker/CookPackageData.cpp
@@ -652,7 +652,8 @@ namespace UE::Cook
 			check(Algo::Find(PackageDatas.GetLoadReadyQueue(), this) != nullptr);
 			break;
 		case EPackageState::Save:
-			check(Algo::Find(PackageDatas.GetSaveQueue(), this) != nullptr);
+			// The save queue is huge and often pushed at end. Check last element first and then scan.
+			check(PackageDatas.GetSaveQueue().Num() && (PackageDatas.GetSaveQueue().Last() == this || Algo::Find(PackageDatas.GetSaveQueue(), this)));
 			break;
 		default:
 			check(false);
