Fix 3m19s of IncrementalDestroyGarbage due to linear search in TArray and constant shrinking due to RemoveAll

diff --git a/Engine/Source/Editor/UnrealEd/Private/CookOnTheFlyServer.cpp b/Engine/Source/Editor/UnrealEd/Private/CookOnTheFlyServer.cpp
--- a/Engine/Source/Editor/UnrealEd/Private/CookOnTheFlyServer.cpp
+++ b/Engine/Source/Editor/UnrealEd/Private/CookOnTheFlyServer.cpp
@@ -3005,12 +3005,12 @@ void UCookOnTheFlyServer::ProcessUnsolicitedPackages(TArray<FName>* OutDiscovere
 	{
 		UE_SCOPED_HIERARCHICAL_COOKTIMER(PostLoadPackageFixup);
 
-		TArray<UE::Cook::FPackageWithInstigator> NewPackages = PackageTracker->GetNewPackages();
+		TMap<UPackage*, UE::Cook::FInstigator> NewPackages = PackageTracker->GetNewPackages();
 
-		for (FPackageWithInstigator& PackageWithInstigator : NewPackages)
+		for (auto& PackageWithInstigator : NewPackages)
 		{
-			UPackage* Package = PackageWithInstigator.Package;
-			FInstigator& Instigator = PackageWithInstigator.Instigator;
+			UPackage* Package = PackageWithInstigator.Key;
+			FInstigator& Instigator = PackageWithInstigator.Value;
 			if (!IsCookByTheBookMode() || !CookByTheBookOptions->bSkipSoftReferences)
 			{
 				PostLoadPackageFixup(Package, OutDiscoveredPackageNames, OutInstigators);
@@ -4272,6 +4272,8 @@ void FSaveCookedPackageContext::SetupPlatform(const ITargetPlatform* InTargetPla
 
 void FSaveCookedPackageContext::FinishPlatform()
 {
+	TRACE_CPUPROFILER_EVENT_SCOPE(FSaveCookedPackageContext::FinishPlatform);
+
 	bool bSuccessful = SavePackageResult.IsSuccessful();
 	bool bLocalReferencedOnlyByEditorOnlyData = SavePackageResult == ESavePackageResult::ReferencedOnlyByEditorOnlyData;
 
diff --git a/Engine/Source/Editor/UnrealEd/Private/Cooker/PackageTracker.cpp b/Engine/Source/Editor/UnrealEd/Private/Cooker/PackageTracker.cpp
--- a/Engine/Source/Editor/UnrealEd/Private/Cooker/PackageTracker.cpp
+++ b/Engine/Source/Editor/UnrealEd/Private/Cooker/PackageTracker.cpp
@@ -74,7 +74,7 @@ namespace UE::Cook
 		NewPackages.Reserve(LoadedPackages.Num());
 		for (UPackage* Package : LoadedPackages)
 		{
-			NewPackages.Add(FPackageWithInstigator{ FInstigator(EInstigator::StartupPackage), Package });
+			NewPackages.Add(Package, FInstigator(EInstigator::StartupPackage));
 		}
 
 		GUObjectArray.AddUObjectDeleteListener(this);
@@ -87,7 +87,7 @@ namespace UE::Cook
 		GUObjectArray.RemoveUObjectCreateListener(this);
 	}
 
-	TArray<FPackageWithInstigator> FPackageTracker::GetNewPackages()
+	TMap<UPackage*, FInstigator> FPackageTracker::GetNewPackages()
 	{
 		return MoveTemp(NewPackages);
 	}
@@ -106,12 +106,10 @@ namespace UE::Cook
 				}
 
 				LoadedPackages.Add(Package);
-				NewPackages.Add(FPackageWithInstigator
-					{
+				NewPackages.Add(Package,
 						FInstigator(EInstigator::Unsolicited,
 							LoadingPackageData ? LoadingPackageData->GetPackageName() : NAME_None)
-						, Package
-					});
+					);
 			}
 		}
 	}
@@ -123,10 +121,7 @@ namespace UE::Cook
 			UPackage* Package = const_cast<UPackage*>(static_cast<const UPackage*>(Object));
 
 			LoadedPackages.Remove(Package);
-			NewPackages.RemoveAll([Package](const FPackageWithInstigator& Existing)
-				{
-					return Existing.Package == Package;
-				});
+			NewPackages.Remove(Package);
 			PostLoadFixupPackages.Remove(Package);
 		}
 	}
diff --git a/Engine/Source/Editor/UnrealEd/Private/Cooker/PackageTracker.h b/Engine/Source/Editor/UnrealEd/Private/Cooker/PackageTracker.h
--- a/Engine/Source/Editor/UnrealEd/Private/Cooker/PackageTracker.h
+++ b/Engine/Source/Editor/UnrealEd/Private/Cooker/PackageTracker.h
@@ -24,12 +24,6 @@ namespace Cook
 {
 	struct FPackageDatas;
 
-	struct FPackageWithInstigator
-	{
-		FInstigator Instigator;
-		UPackage* Package;
-	};
-
 	template<typename Type>
 	struct FThreadSafeQueue
 	{
@@ -181,7 +175,7 @@ namespace Cook
 		~FPackageTracker();
 
 		/* Returns all packages that have been loaded since the last time GetNewPackages was called */
-		TArray<FPackageWithInstigator> GetNewPackages();
+		TMap<UPackage*, FInstigator> GetNewPackages();
 
 		virtual void NotifyUObjectCreated(const class UObjectBase* Object, int32 Index) override;
 		virtual void NotifyUObjectDeleted(const class UObjectBase* Object, int32 Index) override;
@@ -197,7 +191,7 @@ namespace Cook
 		TFastPointerSet<UPackage*> LoadedPackages;
 
 		// This list contains the UPackages loaded since last call to GetNewPackages
-		TArray<FPackageWithInstigator> NewPackages;
+		TMap<UPackage*, FInstigator> NewPackages;
 
 		/** The package currently being loaded at CookOnTheFlyServer's direct request. Used to determine which load dependencies were not preloaded. */
 		FPackageData* LoadingPackageData = nullptr;
