Change handling of named UV sets in FBX source files, so that they cannot generate a null entry in the UV channel array.

diff --git a/Engine/Source/Editor/UnrealEd/Private/Fbx/FbxStaticMeshImport.cpp b/Engine/Source/Editor/UnrealEd/Private/Fbx/FbxStaticMeshImport.cpp
--- a/Engine/Source/Editor/UnrealEd/Private/Fbx/FbxStaticMeshImport.cpp
+++ b/Engine/Source/Editor/UnrealEd/Private/Fbx/FbxStaticMeshImport.cpp
@@ -222,30 +222,31 @@ struct FFBXUVs
 
 
 		// If the the UV sets are named using the following format (UVChannel_X; where X ranges from 1 to 4)
-		// we will re-order them based on these names.  Any UV sets that do not follow this naming convention
-		// will be slotted into available spaces.
+		// we will re-order them based on these names, provided there are enough UV channels in total to do so.
+		// Any UV sets that do not follow this naming convention will be slotted into available spaces.
+		//
+		// Examples:
+		// {UVChannel_3, UVChannel_1, RandomName} will be reordered to {UVChannel_1, RandomName, UVChannel_3}.
+		// {UVChannel_3, UVChannel_1} will be reordered to {UVChannel?1, UVChannel_3}
+		//     - note that UVChannel_3 cannot be placed in the appropriate slot as there are not enough sets.
 		if(UVSets.Num())
 		{
-			for(int32 ChannelNumIdx = 0; ChannelNumIdx < 4; ChannelNumIdx++)
+			const int32 MaxPossibleUVIndex = FMath::Min(UVSets.Num(), 4);
+			for(int32 ChannelNumIdx = 0; ChannelNumIdx < MaxPossibleUVIndex; ChannelNumIdx++)
 			{
 				FString ChannelName = FString::Printf( TEXT("UVChannel_%d"), ChannelNumIdx+1 );
 				int32 SetIdx = UVSets.Find( ChannelName );
 
 				// If the specially formatted UVSet name appears in the list and it is in the wrong spot,
-				// we will swap it into the correct spot.
+				// we will swap it into the correct spot, provided there are enough UV channels
 				if( SetIdx != INDEX_NONE && SetIdx != ChannelNumIdx )
 				{
-					// If we are going to swap to a position that is outside the bounds of the
-					// array, then we pad out to that spot with empty data.
-					for(int32 ArrSize = UVSets.Num(); ArrSize < ChannelNumIdx+1; ArrSize++)
-					{
-						UVSets.Add ( FString(TEXT("")) );
-					}
 					//Swap the entry into the appropriate spot.
 					UVSets.Swap( SetIdx, ChannelNumIdx );
 				}
 			}
 		}
+
 	}
 
 	void Phase2(UnFbx::FFbxImporter* FbxImporter, FbxMesh* Mesh)
