EditorDomain: Improve performance of AppendDigest:
1) Change hash algorithm from FCbWriter (which collects memory in an array) to FBlake3 (which only keeps the memory necessary for the digest).
2) Change deduplication of Guids to deduplicate integer handles to the guids (cuts the deduplication time in half).
#rnx

diff --git a/Engine/Source/Editor/UnrealEd/Private/EditorDomain/EditorDomainUtils.cpp b/Engine/Source/Editor/UnrealEd/Private/EditorDomain/EditorDomainUtils.cpp
--- a/Engine/Source/Editor/UnrealEd/Private/EditorDomain/EditorDomainUtils.cpp
+++ b/Engine/Source/Editor/UnrealEd/Private/EditorDomain/EditorDomainUtils.cpp
@@ -14,6 +14,7 @@
 #include "DerivedDataRequestOwner.h"
 #include "Editor.h"
 #include "HAL/FileManager.h"
+#include "Hash/Blake3.h"
 #include "Memory/SharedBuffer.h"
 #include "Misc/CoreDelegates.h"
 #include "Misc/FileHelper.h"
@@ -25,6 +26,7 @@
 #include "Serialization/PackageWriterToSharedBuffer.h"
 #include "String/ParseTokens.h"
 #include "TargetDomain/TargetDomainUtils.h"
+#include "Templates/Function.h"
 #include "UObject/CoreRedirects.h"
 #include "UObject/ObjectVersion.h"
 #include "UObject/Package.h"
@@ -68,11 +70,11 @@ TArray<FName> GGlobalConstructClasses;
 TSet<FName> GTargetDomainClassBlockList;
 bool GTargetDomainClassUseAllowList = true;
 bool GTargetDomainClassEmptyAllowList = false;
-TArray<FGuid> GGlobalAddedCustomVersions;
+TArray<int32> GGlobalAddedCustomVersions;
 bool bGGlobalAddedCustomVersionsInitialized = false;
 
 // Change to a new guid when EditorDomain needs to be invalidated
-const TCHAR* EditorDomainVersion = TEXT("CE45C7BCF04246A4AEA6136DE532DF85");
+const TCHAR* EditorDomainVersion = TEXT("30E58214A4A84D638FAA8826B81338A1");
 
 // Identifier of the CacheBuckets for EditorDomain tables
 const TCHAR* EditorDomainPackageBucketName = TEXT("EditorDomainPackage");
@@ -95,27 +97,111 @@ static bool GetEditorDomainSaveUnversioned()
 	return bEditorDomainSaveUnversioned;
 }
 
-EPackageDigestResult AppendPackageDigest(FCbWriter& Writer, EDomainUse& OutEditorDomainUse, FString& OutErrorMessage,
+/**
+ * Thread-safe cache to compress CustomVersion Guids into integer handles, to reduce the cost of removing duplicates
+ * when lists of CustomVersion Guids are merged
+*/
+class FKnownCustomVersions
+{
+public:
+	/** Find or if necessary add the handle for each Guid; append them to the output handles. */
+	static void FindOrAddHandles(TArray<int32>& OutHandles, TConstArrayView<FGuid> InGuids);
+	/** Find or if necessary add the handle for each Guid; append them to the output handles. */
+	static void FindOrAddHandles(TArray<int32>& OutHandles, int32 NumGuids, TFunctionRef<const FGuid& (int32)> GetGuid);
+	/** Find the guid for each handle. Handles must be values returned from a FindHandle function. */
+	static void FindGuidsChecked(TArray<FGuid>& OutGuids, TConstArrayView<int32> Handles);
+
+private:
+	static FRWLock Lock;
+	static TMap<FGuid, int32> GuidToHandle;
+	static TArray<FGuid> Guids;
+};
+
+FRWLock FKnownCustomVersions::Lock;
+TMap<FGuid, int32> FKnownCustomVersions::GuidToHandle;
+TArray<FGuid> FKnownCustomVersions::Guids;
+
+void FKnownCustomVersions::FindOrAddHandles(TArray<int32>& OutHandles, TConstArrayView<FGuid> InGuids)
+{
+	FindOrAddHandles(OutHandles, InGuids.Num(), [InGuids](int32 Index) -> const FGuid& { return InGuids[Index]; });
+}
+
+void FKnownCustomVersions::FindOrAddHandles(TArray<int32>& OutHandles, int32 NumGuids, TFunctionRef<const FGuid& (int32)> GetGuid)
+{
+	// Avoid a WriteLock in most cases by finding-only the incoming guids and writing their handle to the output
+	// For any Guids that are not found, add a placeholder handle and store the missing guid and its index in
+	// the output in a list to iterate over later.
+	TArray<TPair<FGuid, int32>> UnknownGuids;
+	{
+		FReadScopeLock ReadLock(Lock);
+		OutHandles.Reserve(OutHandles.Num() + NumGuids);
+		for (int32 Index = 0; Index < NumGuids; ++Index)
+		{
+			const FGuid& Guid = GetGuid(Index);
+			int32* CustomVersionHandle = GuidToHandle.Find(GetGuid(Index));
+			if (CustomVersionHandle)
+			{
+				OutHandles.Add(*CustomVersionHandle);
+			}
+			else
+			{
+				UnknownGuids.Reserve(NumGuids);
+				UnknownGuids.Add(TPair<FGuid, int32>{ Guid, OutHandles.Num() });
+				OutHandles.Add(INDEX_NONE);
+			}
+		}
+	}
+
+	if (UnknownGuids.Num())
+	{
+		// Add the missing guids under the writelock and write their handle over the placeholders in the output.
+		FWriteScopeLock WriteLock(Lock);
+		int32 NumKnownGuids = Guids.Num();
+		for (TPair<FGuid, int32>& Pair : UnknownGuids)
+		{
+			int32& ExistingIndex = GuidToHandle.FindOrAdd(Pair.Key, NumKnownGuids);
+			if (ExistingIndex == NumKnownGuids)
+			{
+				Guids.Add(Pair.Key);
+				++NumKnownGuids;
+			}
+			OutHandles[Pair.Value] = ExistingIndex;
+		}
+	}
+}
+
+void FKnownCustomVersions::FindGuidsChecked(TArray<FGuid>& OutGuids, TConstArrayView<int32> Handles)
+{
+	OutGuids.Reserve(OutGuids.Num() + Handles.Num());
+	FReadScopeLock ReadLock(Lock);
+	for (int32 Handle : Handles)
+	{
+		check(0 <= Handle && Handle < Guids.Num());
+		OutGuids.Add(Guids[Handle]);
+	}
+}
+
+EPackageDigestResult AppendPackageDigest(FBlake3& Writer, EDomainUse& OutEditorDomainUse, FString& OutErrorMessage,
 	const FAssetPackageData& PackageData, FName PackageName, TArray<FGuid>* OutCustomVersions)
 {
 	OutEditorDomainUse = EDomainUse::LoadEnabled | EDomainUse::SaveEnabled;
 	
 	FPackageFileVersion CurrentFileVersionUE = GPackageFileUEVersion;
 	int32 CurrentFileVersionLicenseeUE = GPackageFileLicenseeUEVersion;
-	Writer << EditorDomainVersion;
-	Writer << GetEditorDomainSaveUnversioned();
-	PRAGMA_DISABLE_DEPRECATION_WARNINGS
-	Writer << const_cast<FGuid&>(PackageData.PackageGuid);
-	PRAGMA_ENABLE_DEPRECATION_WARNINGS
-	Writer << CurrentFileVersionUE;
-	Writer << CurrentFileVersionLicenseeUE;
-	TArray<FGuid> CustomVersionGuidBuffer;
-	TArray<FGuid>& CustomVersionGuids(OutCustomVersions ? *OutCustomVersions : CustomVersionGuidBuffer);
-	CustomVersionGuids.Reset(PackageData.GetCustomVersions().Num());
-	for (const UE::AssetRegistry::FPackageCustomVersion& PackageVersion : PackageData.GetCustomVersions())
-	{
-		CustomVersionGuids.Add(PackageVersion.Key);
-	}
+	Writer.Update(EditorDomainVersion, FCString::Strlen(EditorDomainVersion)*sizeof(TCHAR));
+	uint8 EditorDomainSaveUnversioned = GetEditorDomainSaveUnversioned() ? 1 : 0;
+	Writer.Update(&EditorDomainSaveUnversioned, sizeof(EditorDomainSaveUnversioned));
+	PRAGMA_DISABLE_DEPRECATION_WARNINGS;
+	Writer.Update(&PackageData.PackageGuid, sizeof(PackageData.PackageGuid));
+	PRAGMA_ENABLE_DEPRECATION_WARNINGS;
+	Writer.Update(&CurrentFileVersionUE, sizeof(CurrentFileVersionUE));
+	Writer.Update(&CurrentFileVersionLicenseeUE, sizeof(CurrentFileVersionLicenseeUE));
+	TArray<int32> CustomVersionHandles;
+	// Reserve 10 custom versions per class times 100 classes per package times twice (once in package, once in class)
+	CustomVersionHandles.Reserve(10*100*2);
+	TConstArrayView<UE::AssetRegistry::FPackageCustomVersion> PackageVersions = PackageData.GetCustomVersions();
+	FKnownCustomVersions::FindOrAddHandles(CustomVersionHandles,
+		PackageVersions.Num(), [PackageVersions](int32 Index) -> const FGuid& { return PackageVersions[Index].Key;});
 
 	int32 NextClass = 0;
 	FClassDigestMap& ClassDigests = GetClassDigests();
@@ -145,25 +231,31 @@ EPackageDigestResult AppendPackageDigest(FCbWriter& Writer, EDomainUse& OutEdito
 			}
 			if (ExistingData->bNative)
 			{
-				Writer << ExistingData->SchemaHash;
+				Writer.Update(&ExistingData->SchemaHash, sizeof(ExistingData->SchemaHash));
 			}
-			CustomVersionGuids.Append(ExistingData->CustomVersionGuids);
+			CustomVersionHandles.Append(ExistingData->CustomVersionHandles);
 			EnumSetFlagsAnd(OutEditorDomainUse, EDomainUse::LoadEnabled | EDomainUse::SaveEnabled,
 				OutEditorDomainUse, ExistingData->EditorDomainUse);
 		}
 	}
 
 	InitializeGlobalAddedCustomVersions();
-	CustomVersionGuids.Append(GGlobalAddedCustomVersions);
+	CustomVersionHandles.Append(GGlobalAddedCustomVersions);
+	CustomVersionHandles.Sort();
+	CustomVersionHandles.SetNum(Algo::Unique(CustomVersionHandles));
+
+	TArray<FGuid> CustomVersionGuidBuffer;
+	TArray<FGuid>& CustomVersionGuids(OutCustomVersions ? *OutCustomVersions : CustomVersionGuidBuffer);
+	FKnownCustomVersions::FindGuidsChecked(CustomVersionGuids, CustomVersionHandles);
 	CustomVersionGuids.Sort();
-	CustomVersionGuids.SetNum(Algo::Unique(CustomVersionGuids));
+
 	for (const FGuid& CustomVersionGuid : CustomVersionGuids)
 	{
-		Writer << const_cast<FGuid&>(CustomVersionGuid);
+		Writer.Update(&CustomVersionGuid, sizeof(CustomVersionGuid));
 		TOptional<FCustomVersion> CurrentVersion = FCurrentCustomVersions::Get(CustomVersionGuid);
 		if (CurrentVersion.IsSet())
 		{
-			Writer << CurrentVersion->Version;
+			Writer.Update(&CurrentVersion->Version, sizeof(CurrentVersion->Version));
 		}
 		else
 		{
@@ -291,7 +383,7 @@ FClassDigestData* FPrecacheClassDigest::GetRecursive(FName ClassName, bool bAllo
 	{
 		DigestData->bNative = false;
 		DigestData->SchemaHash.Reset();
-		DigestData->CustomVersionGuids.Empty();
+		DigestData->CustomVersionHandles.Empty();
 		FStringView UnusedClassOfClassName;
 		FStringView ClassPackageName;
 		FStringView ClassObjectName;
@@ -328,11 +420,11 @@ FClassDigestData* FPrecacheClassDigest::GetRecursive(FName ClassName, bool bAllo
 		// GetCustomVersions can create the ClassDefaultObject, which can trigger LoadPackage, which
 		// can reenter this function recursively. We have to drop the lock to prevent a deadlock.
 		FUnlockScope UnlockScope(ClassDigestsMap.Lock);
-		DigestData->CustomVersionGuids = GetCustomVersions(*StructAsClass);
+		FKnownCustomVersions::FindOrAddHandles(DigestData->CustomVersionHandles, GetCustomVersions(*StructAsClass));
 	}
 	else
 	{
-		DigestData->CustomVersionGuids.Reset();
+		DigestData->CustomVersionHandles.Reset();
 	}
 
 	// Propagate values from the parent
@@ -383,7 +475,7 @@ FClassDigestData* FPrecacheClassDigest::GetRecursive(FName ClassName, bool bAllo
 	}
 	if (!ConstructClasses.IsEmpty())
 	{
-		TArray<FGuid> ConstructCustomVersions; // Not a class scratch variable because we use it across recursive calls
+		TArray<int32> ConstructCustomVersions; // Not a class scratch variable because we use it across recursive calls
 		for (FName ConstructClass : ConstructClasses)
 		{
 			FClassDigestData* ConstructClassDigest = GetRecursive(ConstructClass, true /* bAllowRedirects */);
@@ -403,14 +495,14 @@ FClassDigestData* FPrecacheClassDigest::GetRecursive(FName ClassName, bool bAllo
 						TEXT("Cycle detected in Editor.ini:[EditorDomain]:PostLoadConstructClasses of class %s. This is unexpected, but not a problem."),
 						*ClassName.ToString());
 				}
-				ConstructCustomVersions.Append(ConstructClassDigest->CustomVersionGuids);
+				ConstructCustomVersions.Append(ConstructClassDigest->CustomVersionHandles);
 			}
 		}
 		// The map has possibly been modified so we need to recalculate the address of ClassName's DigestData
 		DigestData = &ClassDigests.FindChecked(ClassName);
-		DigestData->CustomVersionGuids.Append(MoveTemp(ConstructCustomVersions));
-		Algo::Sort(DigestData->CustomVersionGuids);
-		DigestData->CustomVersionGuids.SetNum(Algo::Unique(DigestData->CustomVersionGuids));
+		DigestData->CustomVersionHandles.Append(MoveTemp(ConstructCustomVersions));
+		Algo::Sort(DigestData->CustomVersionHandles);
+		DigestData->CustomVersionHandles.SetNum(Algo::Unique(DigestData->CustomVersionHandles));
 	}
 
 	DigestData->bConstructionComplete = true;
@@ -461,7 +553,7 @@ void InitializeGlobalAddedCustomVersions()
 		}
 		else
 		{
-			GGlobalAddedCustomVersions.Append(ExistingData->CustomVersionGuids);
+			GGlobalAddedCustomVersions.Append(ExistingData->CustomVersionHandles);
 		}
 	}
 	Algo::Sort(GGlobalAddedCustomVersions);
@@ -792,18 +884,18 @@ EPackageDigestResult GetPackageDigest(IAssetRegistry& AssetRegistry, FName Packa
 	FPackageDigest& OutPackageDigest, EDomainUse& OutEditorDomainUse, FString& OutErrorMessage,
 	TArray<FGuid>* OutCustomVersions)
 {
-	FCbWriter Builder;
+	FBlake3 Builder;
 	EPackageDigestResult Result = AppendPackageDigest(AssetRegistry, PackageName, Builder, OutEditorDomainUse, 
 		OutErrorMessage, OutCustomVersions);
 	if (Result == EPackageDigestResult::Success)
 	{
-		OutPackageDigest = Builder.Save().GetRangeHash();
+		OutPackageDigest = Builder.Finalize();
 	}
 	return Result;
 }
 
 EPackageDigestResult AppendPackageDigest(IAssetRegistry& AssetRegistry, FName PackageName,
-	FCbWriter& Builder, EDomainUse& OutEditorDomainUse, FString& OutErrorMessage, TArray<FGuid>* OutCustomVersions)
+	FBlake3& Builder, EDomainUse& OutEditorDomainUse, FString& OutErrorMessage, TArray<FGuid>* OutCustomVersions)
 {
 	AssetRegistry.WaitForPackage(PackageName.ToString());
 	TOptional<FAssetPackageData> PackageData = AssetRegistry.GetAssetPackageDataCopy(PackageName);
@@ -1138,7 +1230,16 @@ bool TrySavePackage(UPackage* Package)
 				FClassDigestMap& ClassDigests = GetClassDigests();
 				FReadScopeLock ClassDigestsScopeLock(ClassDigests.Lock);
 				FClassDigestData* ExistingData = ClassDigests.Map.Find(CulpritClassName);
-				bConstructedClassDeclaresTheVersion = ExistingData && ExistingData->CustomVersionGuids.Contains(CustomVersionGuid);
+				if (!ExistingData)
+				{
+					bConstructedClassDeclaresTheVersion = false;
+				}
+				else
+				{
+					TArray<FGuid> ClassCustomVersionGuids;
+					FKnownCustomVersions::FindGuidsChecked(ClassCustomVersionGuids, ExistingData->CustomVersionHandles);
+					bConstructedClassDeclaresTheVersion = ClassCustomVersionGuids.Contains(CustomVersionGuid);
+				}
 			}
 			if (bConstructedClassDeclaresTheVersion)
 			{
@@ -1275,17 +1376,17 @@ void PutBulkDataList(FName PackageName, FSharedBuffer Buffer)
 	Owner.KeepAlive();
 }
 
-FIoHash GetPackageAndGuidDigest(FCbWriter& Builder, const FGuid& BulkDataId)
+FIoHash GetPackageAndGuidDigest(FBlake3& Builder, const FGuid& BulkDataId)
 {
-	Builder << BulkDataId;
-	return Builder.Save().GetRangeHash();
+	Builder.Update(&BulkDataId, sizeof(BulkDataId));
+	return Builder.Finalize();
 }
 
 void GetBulkDataPayloadId(FName PackageName, const FGuid& BulkDataId, UE::DerivedData::IRequestOwner& Owner,
 	TUniqueFunction<void(FSharedBuffer Buffer)>&& Callback)
 {
 	FString ErrorMessage;
-	FCbWriter Builder;
+	FBlake3 Builder;
 	EDomainUse EditorDomainUse;
 	EPackageDigestResult FindHashResult = AppendPackageDigest(*IAssetRegistry::Get(), PackageName, Builder, EditorDomainUse, ErrorMessage);
 	switch (FindHashResult)
@@ -1319,7 +1420,7 @@ void GetBulkDataPayloadId(FName PackageName, const FGuid& BulkDataId, UE::Derive
 void PutBulkDataPayloadId(FName PackageName, const FGuid& BulkDataId, FSharedBuffer Buffer)
 {
 	FString ErrorMessage;
-	FCbWriter Builder;
+	FBlake3 Builder;
 	EDomainUse EditorDomainUse;
 	EPackageDigestResult FindHashResult = AppendPackageDigest(*IAssetRegistry::Get(), PackageName, Builder, EditorDomainUse, ErrorMessage);
 	switch (FindHashResult)
diff --git a/Engine/Source/Editor/UnrealEd/Private/EditorDomain/EditorDomainUtils.h b/Engine/Source/Editor/UnrealEd/Private/EditorDomain/EditorDomainUtils.h
--- a/Engine/Source/Editor/UnrealEd/Private/EditorDomain/EditorDomainUtils.h
+++ b/Engine/Source/Editor/UnrealEd/Private/EditorDomain/EditorDomainUtils.h
@@ -19,7 +19,7 @@
 DECLARE_LOG_CATEGORY_EXTERN(LogEditorDomainSave, Log, All);
 
 class FAssetPackageData;
-class FCbWriter;
+class FBlake3;
 class FPackagePath;
 class FSharedBuffer;
 class IAssetRegistry;
@@ -51,7 +51,7 @@ struct FClassDigestData
 {
 public:
 	FBlake3Hash SchemaHash;
-	TArray<FGuid> CustomVersionGuids;
+	TArray<int32> CustomVersionHandles;
 	bool bNative = false;
 	/** EditorDomainEnabled allows everything and uses only a blocklist, so DomainUse by default is enabled. */
 	EDomainUse EditorDomainUse = EDomainUse::LoadEnabled | EDomainUse::SaveEnabled;
@@ -79,7 +79,7 @@ EPackageDigestResult GetPackageDigest(IAssetRegistry& AssetRegistry, FName Packa
 	TArray<FGuid>* OutCustomVersions = nullptr);
 /** Appends the fields to calculate the packagedigest; call Builder.Save().GetRangeHash() to get digest. */
 EPackageDigestResult AppendPackageDigest(IAssetRegistry& AssetRegistry, FName PackageName,
-	FCbWriter& Builder, EDomainUse& OutEditorDomainUse, FString& OutErrorMessage,
+	FBlake3& Builder, EDomainUse& OutEditorDomainUse, FString& OutErrorMessage,
 	TArray<FGuid>* OutCustomVersions = nullptr);
 
 /** For any ClassNames not already in ClassDigests, look up their UStruct and add them. */
diff --git a/Engine/Source/Editor/UnrealEd/Private/TargetDomain/TargetDomainUtils.cpp b/Engine/Source/Editor/UnrealEd/Private/TargetDomain/TargetDomainUtils.cpp
--- a/Engine/Source/Editor/UnrealEd/Private/TargetDomain/TargetDomainUtils.cpp
+++ b/Engine/Source/Editor/UnrealEd/Private/TargetDomain/TargetDomainUtils.cpp
@@ -82,7 +82,7 @@ bool TryCreateKey(FName PackageName, TConstArrayView<FName> SortedBuildDependenc
 		if (OutErrorMessage) *OutErrorMessage = TEXT("AssetRegistry is unavailable.");
 		return false;
 	}
-	FCbWriter KeyBuilder;
+	FBlake3 KeyBuilder;
 	UE::EditorDomain::EPackageDigestResult Result;
 	FString ErrorMessage;
 	UE::EditorDomain::EDomainUse EditorDomainUse;
@@ -109,7 +109,7 @@ bool TryCreateKey(FName PackageName, TConstArrayView<FName> SortedBuildDependenc
 
 	if (OutHash)
 	{
-		*OutHash = KeyBuilder.Save().GetRangeHash();
+		*OutHash = KeyBuilder.Finalize();
 	}
 	return true;
 }
