Add GLevelEditorModeToolsIsValid() method to test before accessing GLevelEditorModeTools()

diff --git a/Engine/Source/Editor/UnrealEd/Private/UnrealEdGlobals.cpp b/Engine/Source/Editor/UnrealEd/Private/UnrealEdGlobals.cpp
--- a/Engine/Source/Editor/UnrealEd/Private/UnrealEdGlobals.cpp
+++ b/Engine/Source/Editor/UnrealEd/Private/UnrealEdGlobals.cpp
@@ -55,9 +55,14 @@ namespace Internal
 	}
 };
 
+bool GLevelEditorModeToolsIsValid()
+{
+	return Internal::GetGlobalModeManager().IsValid();
+}
+
 FEditorModeTools& GLevelEditorModeTools()
 {
-	check(Internal::GetGlobalModeManager().IsValid());
+	check(GLevelEditorModeToolsIsValid());
 	return *Internal::GetGlobalModeManager().Get();
 }
 
diff --git a/Engine/Source/Editor/UnrealEd/Public/Editor.h b/Engine/Source/Editor/UnrealEd/Public/Editor.h
--- a/Engine/Source/Editor/UnrealEd/Public/Editor.h
+++ b/Engine/Source/Editor/UnrealEd/Public/Editor.h
@@ -655,6 +655,11 @@ UNREALED_API bool IsUniqueObjectName( const FName& InName, UObject* Outer, FText
  */
 UNREALED_API class FEditorModeTools& GLevelEditorModeTools();
 
+/**
+ * Checks if FEditorModeTools is valid
+ */
+UNREALED_API bool GLevelEditorModeToolsIsValid();
+
 namespace EditorUtilities
 {
 	/**
