EditorMode:
- Guard against legacy mode getting GCed while Deactivating modes
- ForEachMode iterates over array copy in case original array is modified while iterating
- New IsSelectionDisallowed virtual allows mode to refuse selection of actor regardless of other modes

diff --git a/Engine/Source/Editor/UnrealEd/Private/EditorModeManager.cpp b/Engine/Source/Editor/UnrealEd/Private/EditorModeManager.cpp
--- a/Engine/Source/Editor/UnrealEd/Private/EditorModeManager.cpp
+++ b/Engine/Source/Editor/UnrealEd/Private/EditorModeManager.cpp
@@ -40,6 +40,7 @@
 #include "EngineUtils.h"
 #include "Tools/AssetEditorContextObject.h"
 #include "ContextObjectStore.h"
+#include "UObject/GCObjectScopeGuard.h"
 
 #include "Elements/Interfaces/TypedElementWorldInterface.h"
 
@@ -248,7 +249,13 @@ bool FEditorModeTools::IsSelectionAllowed(AActor* InActor, const bool bInSelecte
 
 	for (const UEdMode* Mode : ActiveScriptableModes)
 	{
-		bSelectionAllowed |= Mode->IsSelectionAllowed(InActor, bInSelected);
+		// Exclusive ability for a mode to disable selection
+		if (Mode->IsSelectionDisallowed(InActor, bInSelected))
+		{
+			return false;
+		}
+				
+		bSelectionAllowed |= Mode->IsSelectionAllowed(InActor, bInSelected);		
 	}
 
 	return bSelectionAllowed;
@@ -418,14 +425,14 @@ EEditAction::Type FEditorModeTools::GetActionEditPaste()
 
 void FEditorModeTools::DeactivateOtherVisibleModes(FEditorModeID InMode)
 {
-	TArray<UEdMode*> TempModes(ActiveScriptableModes);
-	for (const UEdMode* Mode : TempModes)
-	{
-		if (Mode->GetID() != InMode && Mode->GetModeInfo().IsVisible())
+	ForEachEdMode([this, InMode](UEdMode* Mode)
 		{
-			DeactivateMode(Mode->GetID());
-		}
-	}
+			if (Mode->GetID() != InMode && Mode->GetModeInfo().IsVisible())
+			{
+				DeactivateMode(Mode->GetID());
+			}
+			return true;
+		});
 }
 
 bool FEditorModeTools::IsSnapRotationEnabled() const
@@ -591,7 +598,9 @@ void FEditorModeTools::DrawBrackets(FEditorViewportClient* ViewportClient, FView
 
 void FEditorModeTools::ForEachEdMode(TFunctionRef<bool(UEdMode*)> InCalllback) const
 {
-	for (UEdMode* Mode : ActiveScriptableModes)
+	// Copy Array in case callback deactivates a mode
+	TArray<UEdMode*> ActiveModes(ActiveScriptableModes);
+	for (UEdMode* Mode : ActiveModes)
 	{
 		if (Mode)
 		{
@@ -842,16 +851,21 @@ void FEditorModeTools::ActivateMode(FEditorModeID InID, bool bToggle)
 		// Just return and leave the mode list unmodified
 		return;
 	}
-
-	// Remove anything that isn't compatible with this mode
-	const bool bIsVisibleMode = ScriptableMode->GetModeInfo().IsVisible();
-	for (int32 ModeIndex = ActiveScriptableModes.Num() - 1; ModeIndex >= 0; ModeIndex--)
+	
 	{
-		UEdMode* Mode = ActiveScriptableModes[ModeIndex];
-		const bool bModesAreCompatible = ScriptableMode->IsCompatibleWith(Mode->GetID()) || Mode->IsCompatibleWith(ScriptableMode->GetID());
-		if (!bModesAreCompatible || (bIsVisibleMode && Mode->GetModeInfo().IsVisible()))
+		// Make sure ScriptableMode doesn't get GCed while Deactivating modes
+		FGCObjectScopeGuard ScriptModeGuard(ScriptableMode);
+
+		// Remove anything that isn't compatible with this mode
+		const bool bIsVisibleMode = ScriptableMode->GetModeInfo().IsVisible();
+		for (int32 ModeIndex = ActiveScriptableModes.Num() - 1; ModeIndex >= 0; ModeIndex--)
 		{
-			DeactivateMode(Mode->GetID());
+			UEdMode* Mode = ActiveScriptableModes[ModeIndex];
+			const bool bModesAreCompatible = ScriptableMode->IsCompatibleWith(Mode->GetID()) || Mode->IsCompatibleWith(ScriptableMode->GetID());
+			if (!bModesAreCompatible || (bIsVisibleMode && Mode->GetModeInfo().IsVisible()))
+			{
+				DeactivateMode(Mode->GetID());
+			}
 		}
 	}
 
@@ -864,7 +878,7 @@ void FEditorModeTools::ActivateMode(FEditorModeID InID, bool bToggle)
 	}
 
 	const bool bIsEnteringMode = true;
-	BroadcastEditorModeIDChanged(ScriptableMode->GetID(), bIsEnteringMode);
+	BroadcastEditorModeIDChanged(InID, bIsEnteringMode);
 
 	PendingDeactivateModes.Remove(InID);
 	RecycledScriptableModes.Remove(InID);
diff --git a/Engine/Source/Editor/UnrealEd/Public/Tools/UEdMode.h b/Engine/Source/Editor/UnrealEd/Public/Tools/UEdMode.h
--- a/Engine/Source/Editor/UnrealEd/Public/Tools/UEdMode.h
+++ b/Engine/Source/Editor/UnrealEd/Public/Tools/UEdMode.h
@@ -156,6 +156,9 @@ public:
 	/** Check to see if an actor can be selected in this mode - no side effects */
 	virtual bool IsSelectionAllowed(AActor* InActor, bool bInSelection) const { return true; }
 
+	/** Check to see if an actor selection is exclusively dissallowed by a mode -- no side effects */
+	virtual bool IsSelectionDisallowed(AActor* InActor, bool bInSelection) const { return false; }
+
 	/** Returns the editor mode identifier. */
 	FEditorModeID GetID() const { return Info.ID; }
 
