PackageTools::UnloadPackages param to force unload of dirty packages
FPluginUtils::UnloadPlugins will unmount plugins even if some of its packages cannot be unloaded

diff --git a/Engine/Plugins/Developer/PluginUtils/Source/PluginUtils/Private/PluginUtils.cpp b/Engine/Plugins/Developer/PluginUtils/Source/PluginUtils/Private/PluginUtils.cpp
--- a/Engine/Plugins/Developer/PluginUtils/Source/PluginUtils/Private/PluginUtils.cpp
+++ b/Engine/Plugins/Developer/PluginUtils/Source/PluginUtils/Private/PluginUtils.cpp
@@ -656,16 +656,15 @@ bool FPluginUtils::UnloadPlugins(const TConstArrayView<TSharedRef<IPlugin>> Plug
 
 			if (PackagesToUnload.Num() > 0)
 			{
-				FText FailReason;
-				if (!UPackageTools::UnloadPackages(PackagesToUnload.Array(), FailReason))
+				FText ErrorMsg;
+				UPackageTools::UnloadPackages(PackagesToUnload.Array(), ErrorMsg, /*bUnloadDirtyPackages=*/true);
+
+				// Log an error if some packages fail to unload, but unmount the plugins anyway.
+				// @note UnloadPackages returned bool indicates whether some packages were unloaded
+				// To tell whether all packages were successfully unloaded we must check the ErrorMsg output param
+				if (!ErrorMsg.IsEmpty())
 				{
-					// If some packages fail to unload, bail out
-					UE_LOG(LogPluginUtils, Error, TEXT("Plugins cannot be unloaded because some packages cannot be unloaded: %s"), *FailReason.ToString());
-					if (OutFailReason)
-					{
-						*OutFailReason = MoveTemp(FailReason);
-					}
-					return false;
+					UE_LOG(LogPluginUtils, Error, TEXT("%s"), *ErrorMsg.ToString());
 				}
 			}
 		}
diff --git a/Engine/Plugins/Developer/PluginUtils/Source/PluginUtils/Public/PluginUtils.h b/Engine/Plugins/Developer/PluginUtils/Source/PluginUtils/Public/PluginUtils.h
--- a/Engine/Plugins/Developer/PluginUtils/Source/PluginUtils/Public/PluginUtils.h
+++ b/Engine/Plugins/Developer/PluginUtils/Source/PluginUtils/Public/PluginUtils.h
@@ -170,33 +170,41 @@ public:
 
 	/**
 	 * Unload assets from the specified plugin and unmount it
+	 * @note Only works on content-only plugins; plugins with code modules cannot be safely unloaded
+	 * @warning Dirty assets that need to be saved will be unloaded anyway
 	 * @param Plugin Plugin to unload
 	 * @param OutFailReason Outputs the reason of the failure if any
-	 * @return true on success
+	 * @return Whether the plugin was successfully unloaded
 	 */
 	static bool UnloadPlugin(const TSharedRef<IPlugin>& Plugin, FText* OutFailReason = nullptr);
 
 	/**
 	 * Unload assets from the specified plugin and unmount it
+	 * @note Only works on content-only plugins; plugins with code modules cannot be safely unloaded
+	 * @warning Dirty assets that need to be saved will be unloaded anyway
 	 * @param PluginName Name of the plugin to unload
 	 * @param OutFailReason Outputs the reason of the failure if any
-	 * @return true on success
+	 * @return Whether the plugin was successfully unloaded
 	 */
 	static bool UnloadPlugin(const FString& PluginName, FText* OutFailReason = nullptr);
 
 	/**
 	 * Unload assets from the specified plugins and unmount them
+	 * @note Only works on content-only plugins; plugins with code modules cannot be safely unloaded
+	 * @warning Dirty assets that need to be saved will be unloaded anyway
 	 * @param Plugins Plugins to unload
 	 * @param OutFailReason Outputs the reason of the failure if any
-	 * @return true on success
+	 * @return Whether all plugins were successfully unloaded
 	 */
 	static bool UnloadPlugins(const TConstArrayView<TSharedRef<IPlugin>> Plugins, FText* OutFailReason = nullptr);
 
 	/**
 	 * Unload assets from the specified plugins and unmount them
+	 * @note Only works on content-only plugins; plugins with code modules cannot be safely unloaded
+	 * @warning Dirty assets that need to be saved will be unloaded anyway
 	 * @param PluginNames Names of the plugins to unload
 	 * @param OutFailReason Outputs the reason of the failure if any
-	 * @return true on success
+	 * @return Whether all plugins were successfully unloaded
 	 */
 	static bool UnloadPlugins(const TConstArrayView<FString> PluginNames, FText* OutFailReason = nullptr);
 
diff --git a/Engine/Source/Editor/UnrealEd/Private/PackageTools.cpp b/Engine/Source/Editor/UnrealEd/Private/PackageTools.cpp
--- a/Engine/Source/Editor/UnrealEd/Private/PackageTools.cpp
+++ b/Engine/Source/Editor/UnrealEd/Private/PackageTools.cpp
@@ -64,7 +64,7 @@
 DEFINE_LOG_CATEGORY_STATIC(LogPackageTools, Log, All);
 
 /** State passed to RestoreStandaloneOnReachableObjects. */
-TArray<UPackage*>* UPackageTools::PackagesBeingUnloaded = nullptr;
+TSet<UPackage*>* UPackageTools::PackagesBeingUnloaded = nullptr;
 TSet<UObject*> UPackageTools::ObjectsThatHadFlagsCleared;
 FDelegateHandle UPackageTools::ReachabilityCallbackHandle;
 
@@ -270,8 +270,7 @@ UPackageTools::UPackageTools(const FObjectInitializer& ObjectInitializer)
 		return bResult;
 	}
 
-
-	bool UPackageTools::UnloadPackages( const TArray<UPackage*>& TopLevelPackages, FText& OutErrorMessage )
+	bool UPackageTools::UnloadPackages(const TArray<UPackage*>& TopLevelPackages, FText& OutErrorMessage, bool bUnloadDirtyPackages)
 	{
 		// Early out if no package is provided
 		if (TopLevelPackages.IsEmpty())
@@ -282,23 +281,27 @@ UPackageTools::UPackageTools(const FObjectInitializer& ObjectInitializer)
 		bool bResult = false;
 
 		// Get outermost packages, in case groups were selected.
-		TArray<UPackage*> PackagesToUnload;
+		TSet<UPackage*> PackagesToUnload;
 
 		// Split the set of selected top level packages into packages which are dirty (and thus cannot be unloaded)
 		// and packages that are not dirty (and thus can be unloaded).
 		TArray<UPackage*> DirtyPackages;
-		for ( int32 PackageIndex = 0 ; PackageIndex < TopLevelPackages.Num() ; ++PackageIndex )
+		for (UPackage* TopLevelPackage : TopLevelPackages)
 		{
-			UPackage* Package = TopLevelPackages[PackageIndex];
-			if( Package != NULL )
+			if (TopLevelPackage)
 			{
-				if ( Package->IsDirty() )
+				if (!bUnloadDirtyPackages && TopLevelPackage->IsDirty())
 				{
-					DirtyPackages.Add( Package );
+					DirtyPackages.Add(TopLevelPackage);
 				}
 				else
 				{
-					PackagesToUnload.AddUnique( Package->GetOutermost() ? Package->GetOutermost() : Package );
+					UPackage* PackageToUnload = TopLevelPackage->GetOutermost();
+					if (!PackageToUnload)
+					{
+						PackageToUnload = TopLevelPackage;
+					}
+					PackagesToUnload.Add(PackageToUnload);
 				}
 			}
 		}
@@ -307,13 +310,13 @@ UPackageTools::UPackageTools(const FObjectInitializer& ObjectInitializer)
 		if ( DirtyPackages.Num() > 0 )
 		{
 			FString DirtyPackagesList;
-			for ( int32 PackageIndex = 0 ; PackageIndex < DirtyPackages.Num() ; ++PackageIndex )
+			for (UPackage* DirtyPackage : DirtyPackages)
 			{
-				DirtyPackagesList += FString::Printf( TEXT("\n    %s"), *DirtyPackages[PackageIndex]->GetName() );
+				DirtyPackagesList += FString::Printf(TEXT("\n    %s"), *(DirtyPackage->GetName()));
 			}
 
 			FFormatNamedArguments Args;
-			Args.Add( TEXT("DirtyPackages"),FText::FromString( DirtyPackagesList ) );
+			Args.Add(TEXT("DirtyPackages"), FText::FromString(DirtyPackagesList));
 
 			OutErrorMessage = FText::Format( NSLOCTEXT("UnrealEd", "UnloadDirtyPackagesList", "The following assets have been modified and cannot be unloaded:{DirtyPackages}\nSaving these assets will allow them to be unloaded."), Args );
 		}
@@ -352,7 +355,7 @@ UPackageTools::UPackageTools(const FObjectInitializer& ObjectInitializer)
 			}
 		}
 
-		if ( PackagesToUnload.Num() > 0 )
+		if (PackagesToUnload.Num() > 0)
 		{
 			const FScopedBusyCursor BusyCursor;
 
@@ -369,21 +372,20 @@ UPackageTools::UPackageTools(const FObjectInitializer& ObjectInitializer)
 
 			// First add all packages to unload to the root set so they don't get garbage collected while we are operating on them
 			TArray<UPackage*> PackagesAddedToRoot;
-			for ( int32 PackageIndex = 0 ; PackageIndex < PackagesToUnload.Num() ; ++PackageIndex )
+			for (UPackage* PackageToUnload : PackagesToUnload)
 			{
-				UPackage* Pkg = PackagesToUnload[PackageIndex];
-				if ( !Pkg->IsRooted() )
+				if (!PackageToUnload->IsRooted())
 				{
-					Pkg->AddToRoot();
-					PackagesAddedToRoot.Add(Pkg);
+					PackageToUnload->AddToRoot();
+					PackagesAddedToRoot.Add(PackageToUnload);
 				}
 			}
 
 			// Now try to clean up assets in all packages to unload.
-			for ( int32 PackageIndex = 0 ; PackageIndex < PackagesToUnload.Num() ; ++PackageIndex )
+			int32 PackageIndex = 0;
+			for (UPackage* PackageBeingUnloaded : PackagesToUnload)
 			{
-				UPackage* PackageBeingUnloaded = PackagesToUnload[PackageIndex];
-				GWarn->StatusUpdate( PackageIndex, PackagesToUnload.Num(), FText::Format(NSLOCTEXT("UnrealEd", "Unloadingf", "Unloading {0}..."), FText::FromString(PackageBeingUnloaded->GetName()) ) );
+				GWarn->StatusUpdate(PackageIndex++, PackagesToUnload.Num(), FText::Format(NSLOCTEXT("UnrealEd", "Unloadingf", "Unloading {0}..."), FText::FromString(PackageBeingUnloaded->GetName()) ) );
 
 				// Flush all pending render commands, as unloading the package may invalidate render resources.
 				FlushRenderingCommands();
@@ -464,9 +466,9 @@ UPackageTools::UPackageTools(const FObjectInitializer& ObjectInitializer)
 			PackagesBeingUnloaded = nullptr;
 			
 			// Now remove from root all the packages we added earlier so they may be GCed if possible
-			for ( int32 PackageIndex = 0 ; PackageIndex < PackagesAddedToRoot.Num() ; ++PackageIndex )
+			for (UPackage* PackageAddedToRoot : PackagesAddedToRoot)
 			{
-				PackagesAddedToRoot[PackageIndex]->RemoveFromRoot();
+				PackageAddedToRoot->RemoveFromRoot();
 			}
 			PackagesAddedToRoot.Empty();
 
@@ -478,7 +480,7 @@ UPackageTools::UPackageTools(const FObjectInitializer& ObjectInitializer)
 			// Clear the standalone flag on metadata objects that are going to be GC'd below.
 			// This resolves the circular dependency between metadata and packages.
 			TArray<TWeakObjectPtr<UMetaData>> PackageMetaDataWithClearedStandaloneFlag;
-			for ( UPackage* PackageToUnload : PackagesToUnload )
+			for (UPackage* PackageToUnload : PackagesToUnload)
 			{
 				UMetaData* PackageMetaData = PackageToUnload ? PackageToUnload->MetaData : nullptr;
 				if ( PackageMetaData && PackageMetaData->HasAnyFlags(RF_Standalone) )
diff --git a/Engine/Source/Editor/UnrealEd/Public/PackageTools.h b/Engine/Source/Editor/UnrealEd/Public/PackageTools.h
--- a/Engine/Source/Editor/UnrealEd/Public/PackageTools.h
+++ b/Engine/Source/Editor/UnrealEd/Public/PackageTools.h
@@ -85,12 +85,13 @@ public:
 	/**
 	 * Helper function that attempts to unload the specified top-level packages.
 	 *
-	 * @param	PackagesToUnload	the list of packages that should be unloaded
-	 * @param	OutErrorMessage		An error message specifying any problems with unloading packages
+	 * @param	PackagesToUnload		the list of packages that should be unloaded
+	 * @param	OutErrorMessage			An error message specifying any problems with unloading packages
+	 * @param	bUnloadDirtyPackages	Whether to unload packages that are dirty (that need to be saved)
 	 *
 	 * @return	true if the set of loaded packages was changed
 	 */
-	static bool UnloadPackages( const TArray<UPackage*>& PackagesToUnload, FText& OutErrorMessage );
+	static bool UnloadPackages( const TArray<UPackage*>& PackagesToUnload, FText& OutErrorMessage, bool bUnloadDirtyPackages = false);
 
 	/**
 	 * Helper function that attempts to reload the specified top-level packages.
@@ -194,7 +195,7 @@ private:
 
 	static void HandlePackageReloaded(const EPackageReloadPhase InPackageReloadPhase, FPackageReloadedEvent* InPackageReloadedEvent);
 
-	static TArray<UPackage*>* PackagesBeingUnloaded;
+	static TSet<UPackage*>* PackagesBeingUnloaded;
 	static TSet<UObject*> ObjectsThatHadFlagsCleared;
 	static FDelegateHandle ReachabilityCallbackHandle;
 };
