ActorElement: Fix Focus not including underlying actors

diff --git a/Engine/Source/Editor/UnrealEd/Private/Elements/Actor/ActorElementEditorViewportInteractionCustomization.cpp b/Engine/Source/Editor/UnrealEd/Private/Elements/Actor/ActorElementEditorViewportInteractionCustomization.cpp
--- a/Engine/Source/Editor/UnrealEd/Private/Elements/Actor/ActorElementEditorViewportInteractionCustomization.cpp
+++ b/Engine/Source/Editor/UnrealEd/Private/Elements/Actor/ActorElementEditorViewportInteractionCustomization.cpp
@@ -86,19 +86,35 @@ bool FActorElementEditorViewportInteractionCustomization::GetFocusBounds( const
 	}
 	else if (USceneComponent* RootComponent = Actor->GetRootComponent())
 	{
-		TArray<USceneComponent*> SceneComponents;
-		RootComponent->GetChildrenComponents(true, SceneComponents);
-		SceneComponents.Add(RootComponent);
+		TSet<AActor*> Actors;
+		Actor->EditorGetUnderlyingActors(Actors);
+		Actors.Add(Actor);
+		TSet<USceneComponent*> AllSceneComponents;
+
+		for (AActor* CurrentActor : Actors)
+		{
+			if (USceneComponent* CurrentRootComponent = CurrentActor->GetRootComponent())
+			{
+				bool bIsAlreadyInSet = false;
+				AllSceneComponents.Add(CurrentRootComponent, &bIsAlreadyInSet);
+				if (!bIsAlreadyInSet)
+				{
+					TArray<USceneComponent*> SceneComponents;
+					CurrentRootComponent->GetChildrenComponents(true, SceneComponents);
+					AllSceneComponents.Append(SceneComponents);
+				}
+			}
+		}
 
 		bool bHasAtLeastOnePrimitiveComponent = false;
-		for (USceneComponent* SceneComponent : SceneComponents)
+		for (USceneComponent* SceneComponent : AllSceneComponents)
 		{
 			UPrimitiveComponent* PrimitiveComponent = Cast<UPrimitiveComponent>(SceneComponent);
 
 			if (PrimitiveComponent && PrimitiveComponent->IsRegistered())
 			{
 				// Some components can have huge bounds but are not visible.  Ignore these components unless it is the only component on the actor 
-				const bool bIgnore = SceneComponents.Num() > 1 && PrimitiveComponent->IgnoreBoundsForEditorFocus();
+				const bool bIgnore = AllSceneComponents.Num() > 1 && PrimitiveComponent->IgnoreBoundsForEditorFocus();
 
 				if (!bIgnore)
 				{
