Fix Unhide/Hide all selected commands with Parent actor selection

diff --git a/Engine/Source/Editor/UnrealEd/Private/EditorActor.cpp b/Engine/Source/Editor/UnrealEd/Private/EditorActor.cpp
--- a/Engine/Source/Editor/UnrealEd/Private/EditorActor.cpp
+++ b/Engine/Source/Editor/UnrealEd/Private/EditorActor.cpp
@@ -1497,7 +1497,7 @@ void UUnrealEdEngine::edactHideUnselected( UWorld* InWorld )
 	for( FActorIterator It(InWorld); It; ++It )
 	{
 		AActor* Actor = *It;
-		if( !FActorEditorUtils::IsABuilderBrush(Actor) && !Actor->IsSelected() && !Actor->IsHiddenEd() )
+		if( !FActorEditorUtils::IsABuilderBrush(Actor) && !Actor->IsActorOrSelectionParentSelected() && !Actor->IsHiddenEd() )
 		{
 			// Save the actor to the transaction buffer to support undo/redo, but do
 			// not call Modify, as we do not want to dirty the actor's package and
@@ -1598,7 +1598,7 @@ void UUnrealEdEngine::edactHideSelectedStartup( UWorld* InWorld )
 				FBspSurf& CurSurface = *SurfaceIterator;
 				
 				// Set the BSP surface to hide at editor startup, if it's not already set that way
-				const bool bSelected = CurSurface.Actor->IsSelected() || (CurSurface.PolyFlags & PF_Selected);
+				const bool bSelected = CurSurface.Actor->IsActorOrSelectionParentSelected() || (CurSurface.PolyFlags & PF_Selected);
 				if (bSelected && !CurSurface.IsHiddenEdAtStartup() && !CurSurface.IsHiddenEd())
 				{
 					CurLevelModel.Modify();
@@ -1689,7 +1689,7 @@ void UUnrealEdEngine::edactUnHideSelectedStartup( UWorld* InWorld )
 				FBspSurf& CurSurface = *SurfaceIterator;
 
 				// Mark the selected BSP surface as showing at editor startup if it was currently set to be hidden
-				const bool bSelected = CurSurface.Actor->IsSelected() || (CurSurface.PolyFlags & PF_Selected);
+				const bool bSelected = CurSurface.Actor->IsActorOrSelectionParentSelected() || (CurSurface.PolyFlags & PF_Selected);
 				if (bSelected && CurSurface.IsHiddenEdAtStartup())
 				{
 					CurLevelModel.Modify();
@@ -1714,7 +1714,7 @@ void UUnrealEdEngine::edactUnhideSelected( UWorld* InWorld )
 		checkSlow( Actor->IsA(AActor::StaticClass()) );
 		
 		// Don't consider already visible actors or the builder brush
-		if ( !FActorEditorUtils::IsABuilderBrush(Actor) && Actor->IsHiddenEd() )
+		if ( !FActorEditorUtils::IsABuilderBrush(Actor) && Actor->IsActorOrSelectionParentSelected() && Actor->IsHiddenEd() )
 		{
 			ActorsToShow.Add( Actor );
 		}
diff --git a/Engine/Source/Editor/UnrealEd/Private/LevelViewportClickHandlers.cpp b/Engine/Source/Editor/UnrealEd/Private/LevelViewportClickHandlers.cpp
--- a/Engine/Source/Editor/UnrealEd/Private/LevelViewportClickHandlers.cpp
+++ b/Engine/Source/Editor/UnrealEd/Private/LevelViewportClickHandlers.cpp
@@ -235,8 +235,7 @@ namespace LevelViewportClickHandlers
 				if( bAllowSelectionChange && GEditor->CanSelectActor(Actor, true, true) )
 				{
 					// If the actor the user clicked on was already selected, then we won't bother clearing the selection
-					const AActor* RootSelectionParent = Actor->GetRootSelectionParent();
-					if( !Actor->IsSelected() && (RootSelectionParent == nullptr || !RootSelectionParent->IsSelected()))
+					if(!Actor->IsActorOrSelectionParentSelected())
 					{
 						GEditor->SelectNone( false, true );
 						bNeedViewportRefresh = true;
diff --git a/Engine/Source/Runtime/Engine/Classes/GameFramework/Actor.h b/Engine/Source/Runtime/Engine/Classes/GameFramework/Actor.h
--- a/Engine/Source/Runtime/Engine/Classes/GameFramework/Actor.h
+++ b/Engine/Source/Runtime/Engine/Classes/GameFramework/Actor.h
@@ -2683,6 +2683,9 @@ public:
 	/** Returns top most selection parent */
 	virtual AActor* GetRootSelectionParent() const;
 
+	/** Returns if actor or selection parent is selected */
+	bool IsActorOrSelectionParentSelected() const;
+
 	/** Push Selection to actor */
 	virtual void PushSelectionToProxies();
 
diff --git a/Engine/Source/Runtime/Engine/Private/Actor.cpp b/Engine/Source/Runtime/Engine/Private/Actor.cpp
--- a/Engine/Source/Runtime/Engine/Private/Actor.cpp
+++ b/Engine/Source/Runtime/Engine/Private/Actor.cpp
@@ -4734,6 +4734,11 @@ AActor* AActor::GetRootSelectionParent() const
 	return Parent;
 }
 
+bool AActor::IsActorOrSelectionParentSelected() const
+{
+	return IsSelected() || (GetSelectionParent() && GetSelectionParent()->IsActorOrSelectionParentSelected());
+}
+
 void AActor::PushSelectionToProxies()
 {
 	TInlineComponentArray<UActorComponent*> Components;
diff --git a/Engine/Source/Runtime/Engine/Private/Components/PrimitiveComponent.cpp b/Engine/Source/Runtime/Engine/Private/Components/PrimitiveComponent.cpp
--- a/Engine/Source/Runtime/Engine/Private/Components/PrimitiveComponent.cpp
+++ b/Engine/Source/Runtime/Engine/Private/Components/PrimitiveComponent.cpp
@@ -1558,14 +1558,7 @@ bool UPrimitiveComponent::ShouldRenderSelected() const
 	{
 		if (const AActor* Owner = GetOwner())
 		{
-			if (Owner->IsSelected())
-			{
-				return true;
-			}
-			else if (const AActor* RootSelection = Owner->GetRootSelectionParent())
-			{
-				return RootSelection->IsSelected();
-			}
+			return Owner->IsActorOrSelectionParentSelected();
 		}
 	}
 	return false;
