Use FGCObjectScopeGuard instead of manually adding world builder object to root.

diff --git a/Engine/Source/Editor/UnrealEd/Private/Commandlets/WorldPartitionBuilderCommandlet.cpp b/Engine/Source/Editor/UnrealEd/Private/Commandlets/WorldPartitionBuilderCommandlet.cpp
--- a/Engine/Source/Editor/UnrealEd/Private/Commandlets/WorldPartitionBuilderCommandlet.cpp
+++ b/Engine/Source/Editor/UnrealEd/Private/Commandlets/WorldPartitionBuilderCommandlet.cpp
@@ -1,7 +1,6 @@
 // Copyright Epic Games, Inc. All Rights Reserved.
 
 #include "Commandlets/WorldPartitionBuilderCommandlet.h"
-#include "WorldPartition/WorldPartitionBuilder.h"
 
 #include "CoreMinimal.h"
 #include "EngineUtils.h"
@@ -11,7 +10,8 @@
 #include "HAL/PlatformFileManager.h"
 #include "ProfilingDebugging/ScopedTimers.h"
 #include "WorldPartition/WorldPartition.h"
-
+#include "WorldPartition/WorldPartitionBuilder.h"
+#include "UObject/GCObjectScopeGuard.h"
 #include "Trace/Trace.h"
 
 DEFINE_LOG_CATEGORY_STATIC(LogWorldPartitionBuilderCommandlet, All, All);
@@ -95,9 +95,11 @@ int32 UWorldPartitionBuilderCommandlet::Main(const FString& Params)
 		return false;
 	}
 
-	Builder->AddToRoot();
-	bool bResult = Builder->RunBuilder(World);
-	Builder->RemoveFromRoot();
+	bool bResult;
+	{
+		FGCObjectScopeGuard BuilderGuard(Builder);
+		bResult = Builder->RunBuilder(World);
+	}
 
 	// Save configuration file
 	if (!FPlatformFileManager::Get().GetPlatformFile().FileExists(*WorldConfigFilename) ||
