Fix sequencer curve import that was broken by cl: 17911258.
When the sequencer import curve code was merge with the animsequence import curve code. The issue here was the sequencer was expecting the delta Y value in the rich curve key tangents instead of the derivative like it suppose to be. With this fix the sequencer expect the derivative.

All fbx curve automation test pass with the fix. I can also import correctly both camera track from the jira attachment in sequencer.
#rnx

diff --git a/Engine/Source/Editor/MovieSceneTools/Private/MovieSceneToolHelpers.cpp b/Engine/Source/Editor/MovieSceneTools/Private/MovieSceneToolHelpers.cpp
--- a/Engine/Source/Editor/MovieSceneTools/Private/MovieSceneToolHelpers.cpp
+++ b/Engine/Source/Editor/MovieSceneTools/Private/MovieSceneToolHelpers.cpp
@@ -1223,21 +1223,7 @@ void ImportTransformChannelToBezierChannel(const FRichCurve& Source, ChannelType
 	{
 		const FRichCurveKey Key = Source.GetKey(*SourceIt);
 		float ArriveTangent = Key.ArriveTangent;
-		FKeyHandle PrevKeyHandle = Source.GetPreviousKey(*SourceIt);
-		if (Source.IsKeyHandleValid(PrevKeyHandle))
-		{
-			const FRichCurveKey PrevKey = Source.GetKey(PrevKeyHandle);
-			ArriveTangent = ArriveTangent / (Key.Time - PrevKey.Time);
-
-		}
 		float LeaveTangent = Key.LeaveTangent;
-		FKeyHandle NextKeyHandle = Source.GetNextKey(*SourceIt);
-		if (Source.IsKeyHandleValid(NextKeyHandle))
-		{
-			const FRichCurveKey NextKey = Source.GetKey(NextKeyHandle);
-			LeaveTangent = LeaveTangent / (NextKey.Time - Key.Time);
-		}
-
 		if (bNegateTangents)
 		{
 			ArriveTangent = -ArriveTangent;
@@ -2062,12 +2048,12 @@ static bool ImportFBXTransformToChannels(FString NodeName, const UMovieSceneUser
 		}
 
 		ImportTransformChannelToDouble(Translation[0], DoubleChannels[0], FrameRate, false, false, StartFrame);
-		ImportTransformChannelToDouble(Translation[1], DoubleChannels[1], FrameRate, true, false, StartFrame);
+		ImportTransformChannelToDouble(Translation[1], DoubleChannels[1], FrameRate, false, false, StartFrame);
 		ImportTransformChannelToDouble(Translation[2], DoubleChannels[2], FrameRate, false, false, StartFrame);
 
 		ImportTransformChannelToDouble(EulerRotation[0], DoubleChannels[3], FrameRate, false,false, StartFrame);
-		ImportTransformChannelToDouble(EulerRotation[1], DoubleChannels[4], FrameRate, true, false, StartFrame);
-		ImportTransformChannelToDouble(EulerRotation[2], DoubleChannels[5], FrameRate, true, false, StartFrame);
+		ImportTransformChannelToDouble(EulerRotation[1], DoubleChannels[4], FrameRate, false, false, StartFrame);
+		ImportTransformChannelToDouble(EulerRotation[2], DoubleChannels[5], FrameRate, false, false, StartFrame);
 
 		if (DoubleChannels.Num() > 6) //noscale
 		{
@@ -2749,12 +2735,12 @@ bool ImportFBXTransform(FString NodeName, FGuid ObjectBinding, UnFbx::FFbxCurves
 	Channels[8]->SetDefault(Scale3D.Z);
 
 	ImportTransformChannelToDouble(Translation[0],   Channels[0], FrameRate, false, true);
-	ImportTransformChannelToDouble(Translation[1],   Channels[1], FrameRate, true, true);
+	ImportTransformChannelToDouble(Translation[1],   Channels[1], FrameRate, false, true);
 	ImportTransformChannelToDouble(Translation[2],   Channels[2], FrameRate, false, true);
 
 	ImportTransformChannelToDouble(EulerRotation[0], Channels[3], FrameRate, false, true);
-	ImportTransformChannelToDouble(EulerRotation[1], Channels[4], FrameRate, true, true);
-	ImportTransformChannelToDouble(EulerRotation[2], Channels[5], FrameRate, true, true);
+	ImportTransformChannelToDouble(EulerRotation[1], Channels[4], FrameRate, false, true);
+	ImportTransformChannelToDouble(EulerRotation[2], Channels[5], FrameRate, false, true);
 
 	ImportTransformChannelToDouble(Scale[0],         Channels[6], FrameRate, false, true);
 	ImportTransformChannelToDouble(Scale[1],         Channels[7], FrameRate, false, true);
diff --git a/Engine/Source/Editor/UnrealEd/Private/Fbx/FbxCurveDataImport.cpp b/Engine/Source/Editor/UnrealEd/Private/Fbx/FbxCurveDataImport.cpp
--- a/Engine/Source/Editor/UnrealEd/Private/Fbx/FbxCurveDataImport.cpp
+++ b/Engine/Source/Editor/UnrealEd/Private/Fbx/FbxCurveDataImport.cpp
@@ -241,7 +241,7 @@ namespace UnFbx {
 			FbxTimeSpan AnimTimeSpan(FBXSDK_TIME_ZERO, FBXSDK_TIME_INFINITE);
 			//Sequencer use auto set tangents
 			const bool bAutoSetTangents = true;
-			UnFbx::FFbxImporter::ImportCurve(FbxCurve, RichCurve, AnimTimeSpan, (bNegative ? -Scale : Scale), bAutoSetTangents);
+			UnFbx::FFbxImporter::ImportCurve(FbxCurve, RichCurve, AnimTimeSpan, bNegative, Scale, bAutoSetTangents);
 		}
 	}
 
@@ -254,7 +254,7 @@ namespace UnFbx {
 			//Send a neutral timespan 0 to infinite, the ImportCurve, offset the keytime by doing: (KeyTime - TimeSpan.start), setting TimeSpan.start at zero will not affect the keys time value.
 			FbxTimeSpan AnimTimeSpan(FBXSDK_TIME_ZERO, FBXSDK_TIME_INFINITE);
 			const bool bAutoSetTangents = false;
-			UnFbx::FFbxImporter::ImportCurve(FbxCurve, RichCurve, AnimTimeSpan, (bNegative ? -Scale : Scale), bAutoSetTangents);
+			UnFbx::FFbxImporter::ImportCurve(FbxCurve, RichCurve, AnimTimeSpan, bNegative, Scale, bAutoSetTangents);
 		}
 	}
 
diff --git a/Engine/Source/Editor/UnrealEd/Private/FbxAnimUtils.cpp b/Engine/Source/Editor/UnrealEd/Private/FbxAnimUtils.cpp
--- a/Engine/Source/Editor/UnrealEd/Private/FbxAnimUtils.cpp
+++ b/Engine/Source/Editor/UnrealEd/Private/FbxAnimUtils.cpp
@@ -162,8 +162,8 @@ namespace FbxAnimUtils
 										{
 											FRichCurve& RichCurve = InOutCurveTable->AddRichCurve(*ChannelName);
 											RichCurve.Reset();
-
-											FbxImporter->ImportCurve(Curve, RichCurve, AnimTimeSpan, 0.01f);
+											constexpr bool bNegative = false;
+											FbxImporter->ImportCurve(Curve, RichCurve, AnimTimeSpan, bNegative, 0.01f);
 										}
 									}
 								}
@@ -180,8 +180,8 @@ namespace FbxAnimUtils
 							{
 								FRichCurve& RichCurve = InOutCurveTable->AddRichCurve(*InCurveName);
 								RichCurve.Reset();
-
-								FbxImporter->ImportCurve(InCurve, RichCurve, AnimTimeSpan, 1.0f);
+								constexpr bool bNegative = false;
+								FbxImporter->ImportCurve(InCurve, RichCurve, AnimTimeSpan, bNegative, 1.0f);
 							});
 							return true;
 						}
diff --git a/Engine/Source/Editor/UnrealEd/Private/SkeletalMeshEdit.cpp b/Engine/Source/Editor/UnrealEd/Private/SkeletalMeshEdit.cpp
--- a/Engine/Source/Editor/UnrealEd/Private/SkeletalMeshEdit.cpp
+++ b/Engine/Source/Editor/UnrealEd/Private/SkeletalMeshEdit.cpp
@@ -961,20 +961,22 @@ bool UnFbx::FFbxImporter::ValidateAnimStack(TArray<FbxNode*>& SortedLinks, TArra
 	return bValidAnimStack;
 }
 
-bool UnFbx::FFbxImporter::ImportCurve(const FbxAnimCurve* FbxCurve, FRichCurve& RichCurve, const FbxTimeSpan &AnimTimeSpan, const float ValueScale/*=1.f*/, const bool bAutoSetTangents /*= true*/)
+bool UnFbx::FFbxImporter::ImportCurve(const FbxAnimCurve* FbxCurve, FRichCurve& RichCurve, const FbxTimeSpan &AnimTimeSpan, const bool bNegative/*=false*/, const float ValueScale/*=1.f*/, const bool bAutoSetTangents /*= true*/)
 {
 	const float DefaultCurveWeight = FbxAnimCurveDef::sDEFAULT_WEIGHT;
 
 	if ( FbxCurve )
 	{
+		//We use the non const to query the left and right derivative of the key, for whatever reason those FBX API functions are not const
+		FbxAnimCurve* NonConstFbxCurve = const_cast<FbxAnimCurve*>(FbxCurve);
 		int32 KeyCount = FbxCurve->KeyGetCount();
-		
+		const float AdjustedValueScale = (bNegative ? -ValueScale : ValueScale);
 		for ( int32 KeyIndex=0; KeyIndex < KeyCount; ++KeyIndex )
 		{
 			FbxAnimCurveKey Key = FbxCurve->KeyGet(KeyIndex);
 			FbxTime KeyTime = Key.GetTime() - AnimTimeSpan.GetStart();
 			const float KeyTimeValue = static_cast<float>(KeyTime.GetSecondDouble());
-			float Value = Key.GetValue() * ValueScale;
+			float Value = Key.GetValue() * AdjustedValueScale;
 			FKeyHandle NewKeyHandle = RichCurve.AddKey(KeyTimeValue, Value, false);
 
 			const bool bIncludeOverrides = true;
@@ -986,17 +988,17 @@ bool UnFbx::FFbxImporter::ImportCurve(const FbxAnimCurve* FbxCurve, FRichCurve&
 			ERichCurveTangentMode NewTangentMode = RCTM_Auto;
 			ERichCurveTangentWeightMode NewTangentWeightMode = RCTWM_WeightedNone;
 
-			float RightTangent = 0.f; 
-			float LeftTangent = 0.f;
+			float RightTangent = NonConstFbxCurve->KeyGetRightDerivative(KeyIndex) * AdjustedValueScale;
+			float LeftTangent = NonConstFbxCurve->KeyGetLeftDerivative(KeyIndex) * AdjustedValueScale;
 			float RightTangentWeight = 0.0f;
 			float LeftTangentWeight = 0.0f; //This one is dependent on the previous key.
+			bool bLeftWeightActive = false;
+			bool bRightWeightActive = false;
 
 			const bool bPreviousKeyValid = KeyIndex > 0;
 			const bool bNextKeyValid = KeyIndex < KeyCount - 1;
 			float PreviousValue = 0.0f;
 			float PreviousKeyTimeValue = 0.0f;
-			bool bLeftWeightActive = false;
-			bool bRightWeightActive = false;
 			float NextValue = 0.0f;
 			float NextKeyTimeValue = 0.0f;
 			if (bPreviousKeyValid)
@@ -1004,7 +1006,7 @@ bool UnFbx::FFbxImporter::ImportCurve(const FbxAnimCurve* FbxCurve, FRichCurve&
 				FbxAnimCurveKey PreviousKey = FbxCurve->KeyGet(KeyIndex - 1);
 				FbxTime PreviousKeyTime = PreviousKey.GetTime() - AnimTimeSpan.GetStart();
 				PreviousKeyTimeValue = static_cast<float>(PreviousKeyTime.GetSecondDouble());
-				PreviousValue = PreviousKey.GetValue() * ValueScale;
+				PreviousValue = PreviousKey.GetValue() * AdjustedValueScale;
 				//The left tangent is driven by the previous key. If the previous key have a the NextLeftweight or both flag weighted mode, it mean the next key is weighted on the left side
 				bLeftWeightActive = (PreviousKey.GetTangentWeightMode() & FbxAnimCurveDef::eWeightedNextLeft) > 0;
 				if (bLeftWeightActive)
@@ -1017,7 +1019,7 @@ bool UnFbx::FFbxImporter::ImportCurve(const FbxAnimCurve* FbxCurve, FRichCurve&
 				FbxAnimCurveKey NextKey = FbxCurve->KeyGet(KeyIndex + 1);
 				FbxTime NextKeyTime = NextKey.GetTime() - AnimTimeSpan.GetStart();
 				NextKeyTimeValue = static_cast<float>(NextKeyTime.GetSecondDouble());
-				NextValue = NextKey.GetValue() * ValueScale;
+				NextValue = NextKey.GetValue() * AdjustedValueScale;
 
 				bRightWeightActive = (KeyTangentWeightMode & FbxAnimCurveDef::eWeightedRight) > 0;
 				if (bRightWeightActive)
@@ -1096,28 +1098,8 @@ bool UnFbx::FFbxImporter::ImportCurve(const FbxAnimCurve* FbxCurve, FRichCurve&
 						}
 						else
 						{
-							//Spline mode is the tangent are equal to the slope between previous and next control point
-							if (bPreviousKeyValid && bNextKeyValid)
-							{
-								bIsComputedTangent = true;
-								LeftTangent = (NextValue - PreviousValue) / FMath::Max<float>(KINDA_SMALL_NUMBER, NextKeyTimeValue - PreviousKeyTimeValue);
-								RightTangent = LeftTangent;
-							}
-							else
-							{
-								//Last key left tangent is the slope between last key and the previous key
-								if (bPreviousKeyValid)
-								{
-									bIsComputedTangent = true;
-									LeftTangent = (Value - PreviousValue) / FMath::Max<float>(KINDA_SMALL_NUMBER, KeyTimeValue - PreviousKeyTimeValue);
-								}
-								//First key right tangent is the slope between first key and the next key
-								if (bNextKeyValid)
-								{
-									bIsComputedTangent = true;
-									RightTangent = (NextValue - Value) / FMath::Max<float>(KINDA_SMALL_NUMBER, NextKeyTimeValue - KeyTimeValue);
-								}
-							}
+							//Spline tangent key must be User mode since we want to keep the tangents provide by the fbx key left and right derivatives
+							NewTangentMode = RCTM_User;
 						}
 					}
 					
@@ -1128,23 +1110,6 @@ bool UnFbx::FFbxImporter::ImportCurve(const FbxAnimCurve* FbxCurve, FRichCurve&
 						//To force flat tangent we need to set the tangent mode to user
 						NewTangentMode = RCTM_User;
 					}
-					else if (bIsComputedTangent)
-					{
-						NewTangentMode = RCTM_User;
-					}
-					else
-					{
-						RightTangent = Key.GetDataFloat(FbxAnimCurveDef::eRightSlope) * ValueScale;
-						if (bPreviousKeyValid)
-						{
-							FbxAnimCurveKey PrevKey = FbxCurve->KeyGet(KeyIndex-1);
-							LeftTangent = PrevKey.GetDataFloat(FbxAnimCurveDef::eNextLeftSlope) * ValueScale;
-						}
-						else
-						{
-							LeftTangent = 0.f;
-						}
-					}
 
 				}
 				break;
@@ -1162,7 +1127,7 @@ bool UnFbx::FFbxImporter::ImportCurve(const FbxAnimCurve* FbxCurve, FRichCurve&
 			if (NewTangentMode != RCTM_Auto)
 			{
 				const bool bEqualTangents = FMath::IsNearlyEqual(LeftTangent, RightTangent);
-				//instead we manually check to see if the tangents are the same or different, if different then broken.
+				//If tangents are different then broken.
 				if (bEqualTangents)
 				{
 					NewTangentMode = RCTM_User;
@@ -1385,7 +1350,8 @@ bool UnFbx::FFbxImporter::ImportCurveToAnimSequence(class UAnimSequence * Target
 		Controller.UpdateCurveNamesFromSkeleton(Skeleton, ERawCurveTrackTypes::RCT_Float);
 
 		FRichCurve RichCurve;
-		if (ImportCurve(FbxCurve, RichCurve, AnimTimeSpan, ValueScale))
+		constexpr bool bNegative = false;
+		if (ImportCurve(FbxCurve, RichCurve, AnimTimeSpan, bNegative, ValueScale))
 		{
 			if (ImportOptions->bRemoveRedundantKeys)
 			{
diff --git a/Engine/Source/Editor/UnrealEd/Public/FbxImporter.h b/Engine/Source/Editor/UnrealEd/Public/FbxImporter.h
--- a/Engine/Source/Editor/UnrealEd/Public/FbxImporter.h
+++ b/Engine/Source/Editor/UnrealEd/Public/FbxImporter.h
@@ -1101,7 +1101,7 @@ public:
 	/**
 	 * Import FbxCurve to Curve
 	 */
-	static bool ImportCurve(const FbxAnimCurve* FbxCurve, FRichCurve& RichCurve, const FbxTimeSpan &AnimTimeSpan, const float ValueScale = 1.f, const bool bAutoSetTangents = true);
+	static bool ImportCurve(const FbxAnimCurve* FbxCurve, FRichCurve& RichCurve, const FbxTimeSpan &AnimTimeSpan, const bool bNegative = false, const float ValueScale = 1.f, const bool bAutoSetTangents = true);
 
 	/**
 	 * Merge all layers of one AnimStack to one layer.
