CookPackageData: Fix whitespace - the block within a namespace should not be indented.
#rnx

diff --git a/Engine/Source/Editor/UnrealEd/Private/Cooker/CookPackageData.cpp b/Engine/Source/Editor/UnrealEd/Private/Cooker/CookPackageData.cpp
--- a/Engine/Source/Editor/UnrealEd/Private/Cooker/CookPackageData.cpp
+++ b/Engine/Source/Editor/UnrealEd/Private/Cooker/CookPackageData.cpp
@@ -31,2122 +31,2122 @@
 namespace UE::Cook
 {
 
-	float GPollAsyncPeriod = .100f;
-	static FAutoConsoleVariableRef CVarPollAsyncPeriod(
-		TEXT("cook.PollAsyncPeriod"),
-		GPollAsyncPeriod,
-		TEXT("Minimum time in seconds between PollPendingCookedPlatformDatas."),
-		ECVF_Default);
+float GPollAsyncPeriod = .100f;
+static FAutoConsoleVariableRef CVarPollAsyncPeriod(
+	TEXT("cook.PollAsyncPeriod"),
+	GPollAsyncPeriod,
+	TEXT("Minimum time in seconds between PollPendingCookedPlatformDatas."),
+	ECVF_Default);
 	
-	//////////////////////////////////////////////////////////////////////////
-	// FPackageData
-	FPackageData::FPlatformData::FPlatformData()
-		: bRequested(false), bCookAttempted(false), bCookSucceeded(false), bExplored(false)
-	{
-	}
+//////////////////////////////////////////////////////////////////////////
+// FPackageData
+FPackageData::FPlatformData::FPlatformData()
+	: bRequested(false), bCookAttempted(false), bCookSucceeded(false), bExplored(false)
+{
+}
+
+FPackageData::FPackageData(FPackageDatas& PackageDatas, const FName& InPackageName, const FName& InFileName)
+	: GeneratedOwner(nullptr), PackageName(InPackageName), FileName(InFileName), PackageDatas(PackageDatas)
+	, Instigator(EInstigator::NotYetRequested), bIsUrgent(0)
+	, bIsVisited(0), bIsPreloadAttempted(0)
+	, bIsPreloaded(0), bHasSaveCache(0), bHasBeginPrepareSaveFailed(0), bCookedPlatformDataStarted(0)
+	, bCookedPlatformDataCalled(0), bCookedPlatformDataComplete(0), bMonitorIsCooked(0)
+	, bInitializedGeneratorSave(0), bCompletedGeneration(0), bGenerated(0)
+{
+	SetState(EPackageState::Idle);
+	SendToState(EPackageState::Idle, ESendFlags::QueueAdd);
+}
 
-	FPackageData::FPackageData(FPackageDatas& PackageDatas, const FName& InPackageName, const FName& InFileName)
-		: GeneratedOwner(nullptr), PackageName(InPackageName), FileName(InFileName), PackageDatas(PackageDatas)
-		, Instigator(EInstigator::NotYetRequested), bIsUrgent(0)
-		, bIsVisited(0), bIsPreloadAttempted(0)
-		, bIsPreloaded(0), bHasSaveCache(0), bHasBeginPrepareSaveFailed(0), bCookedPlatformDataStarted(0)
-		, bCookedPlatformDataCalled(0), bCookedPlatformDataComplete(0), bMonitorIsCooked(0)
-		, bInitializedGeneratorSave(0), bCompletedGeneration(0), bGenerated(0)
-	{
-		SetState(EPackageState::Idle);
-		SendToState(EPackageState::Idle, ESendFlags::QueueAdd);
-	}
+FPackageData::~FPackageData()
+{
+	// ClearReferences should have been called earlier, but call it here in case it was missed
+	ClearReferences();
+	// We need to send OnLastCookedPlatformRemoved message to the monitor, so call SetPlatformsNotCooked
+	SetPlatformsNotCooked();
+	// Update the monitor's counters and call exit functions
+	SendToState(EPackageState::Idle, ESendFlags::QueueNone);
+}
+
+void FPackageData::ClearReferences()
+{
+	DestroyGeneratorPackage();
+}
 
-	FPackageData::~FPackageData()
-	{
-		// ClearReferences should have been called earlier, but call it here in case it was missed
-		ClearReferences();
-		// We need to send OnLastCookedPlatformRemoved message to the monitor, so call SetPlatformsNotCooked
-		SetPlatformsNotCooked();
-		// Update the monitor's counters and call exit functions
-		SendToState(EPackageState::Idle, ESendFlags::QueueNone);
-	}
+const FName& FPackageData::GetPackageName() const
+{
+	return PackageName;
+}
+
+const FName& FPackageData::GetFileName() const
+{
+	return FileName;
+}
+
+void FPackageData::SetFileName(const FName& InFileName)
+{
+	FileName = InFileName;
+}
 
-	void FPackageData::ClearReferences()
+int32 FPackageData::GetNumRequestedPlatforms() const
+{
+	int32 Result = 0;
+	for (const TPair<const ITargetPlatform*, FPlatformData>& Pair : PlatformDatas)
 	{
-		DestroyGeneratorPackage();
+		Result += Pair.Value.bRequested ? 1 : 0;
 	}
+	return Result;
+}
 
-	const FName& FPackageData::GetPackageName() const
+void FPackageData::SetPlatformsRequested(TConstArrayView<const ITargetPlatform*> TargetPlatforms, bool bRequested)
+{
+	for (const ITargetPlatform* TargetPlatform : TargetPlatforms)
 	{
-		return PackageName;
+		PlatformDatas.FindOrAdd(TargetPlatform).bRequested = true;
 	}
+}
 
-	const FName& FPackageData::GetFileName() const
+void FPackageData::ClearRequestedPlatforms()
+{
+	for (TPair<const ITargetPlatform*, FPlatformData>& Pair : PlatformDatas)
 	{
-		return FileName;
+		Pair.Value.bRequested = false;
 	}
+}
 
-	void FPackageData::SetFileName(const FName& InFileName)
+bool FPackageData::HasAllRequestedPlatforms(const TArrayView<const ITargetPlatform* const>& Platforms) const
+{
+	if (Platforms.Num() == 0)
 	{
-		FileName = InFileName;
+		return true;
 	}
-
-	int32 FPackageData::GetNumRequestedPlatforms() const
+	if (PlatformDatas.Num() == 0)
 	{
-		int32 Result = 0;
-		for (const TPair<const ITargetPlatform*, FPlatformData>& Pair : PlatformDatas)
-		{
-			Result += Pair.Value.bRequested ? 1 : 0;
-		}
-		return Result;
+		return false;
 	}
 
-	void FPackageData::SetPlatformsRequested(TConstArrayView<const ITargetPlatform*> TargetPlatforms, bool bRequested)
+	for (const ITargetPlatform* QueryPlatform : Platforms)
 	{
-		for (const ITargetPlatform* TargetPlatform : TargetPlatforms)
+		const FPlatformData* PlatformData = PlatformDatas.Find(QueryPlatform);
+		if (!PlatformData || !PlatformData->bRequested)
 		{
-			PlatformDatas.FindOrAdd(TargetPlatform).bRequested = true;
+			return false;
 		}
 	}
+	return true;
+}
 
-	void FPackageData::ClearRequestedPlatforms()
+bool FPackageData::AreAllRequestedPlatformsCooked(bool bAllowFailedCooks) const
+{
+	for (const TPair<const ITargetPlatform*, FPlatformData>& Pair : PlatformDatas)
 	{
-		for (TPair<const ITargetPlatform*, FPlatformData>& Pair : PlatformDatas)
+		if (Pair.Value.bRequested && (!Pair.Value.bCookAttempted || (!bAllowFailedCooks && !Pair.Value.bCookSucceeded)))
 		{
-			Pair.Value.bRequested = false;
+			return false;
 		}
 	}
+	return true;
+}
 
-	bool FPackageData::HasAllRequestedPlatforms(const TArrayView<const ITargetPlatform* const>& Platforms) const
+bool FPackageData::AreAllRequestedPlatformsExplored() const
+{
+	for (const TPair<const ITargetPlatform*, FPlatformData>& Pair : PlatformDatas)
 	{
-		if (Platforms.Num() == 0)
-		{
-			return true;
-		}
-		if (PlatformDatas.Num() == 0)
+		if (Pair.Value.bRequested && !Pair.Value.bExplored)
 		{
 			return false;
 		}
-
-		for (const ITargetPlatform* QueryPlatform : Platforms)
-		{
-			const FPlatformData* PlatformData = PlatformDatas.Find(QueryPlatform);
-			if (!PlatformData || !PlatformData->bRequested)
-			{
-				return false;
-			}
-		}
-		return true;
 	}
+	return true;
+}
 
-	bool FPackageData::AreAllRequestedPlatformsCooked(bool bAllowFailedCooks) const
+bool FPackageData::HasAllExploredPlatforms(const TArrayView<const ITargetPlatform* const>& Platforms) const
+{
+	if (Platforms.Num() == 0)
 	{
-		for (const TPair<const ITargetPlatform*, FPlatformData>& Pair : PlatformDatas)
-		{
-			if (Pair.Value.bRequested && (!Pair.Value.bCookAttempted || (!bAllowFailedCooks && !Pair.Value.bCookSucceeded)))
-			{
-				return false;
-			}
-		}
 		return true;
 	}
-
-	bool FPackageData::AreAllRequestedPlatformsExplored() const
+	if (PlatformDatas.Num() == 0)
 	{
-		for (const TPair<const ITargetPlatform*, FPlatformData>& Pair : PlatformDatas)
-		{
-			if (Pair.Value.bRequested && !Pair.Value.bExplored)
-			{
-				return false;
-			}
-		}
-		return true;
+		return false;
 	}
 
-	bool FPackageData::HasAllExploredPlatforms(const TArrayView<const ITargetPlatform* const>& Platforms) const
+	for (const ITargetPlatform* QueryPlatform : Platforms)
 	{
-		if (Platforms.Num() == 0)
-		{
-			return true;
-		}
-		if (PlatformDatas.Num() == 0)
+		const FPlatformData* PlatformData = FindPlatformData(QueryPlatform);
+		if (!PlatformData || !PlatformData->bExplored)
 		{
 			return false;
 		}
-
-		for (const ITargetPlatform* QueryPlatform : Platforms)
-		{
-			const FPlatformData* PlatformData = FindPlatformData(QueryPlatform);
-			if (!PlatformData || !PlatformData->bExplored)
-			{
-				return false;
-			}
-		}
-		return true;
 	}
+	return true;
+}
 
-	void FPackageData::SetIsUrgent(bool Value)
+void FPackageData::SetIsUrgent(bool Value)
+{
+	bool OldValue = static_cast<bool>(bIsUrgent);
+	if (OldValue != Value)
 	{
-		bool OldValue = static_cast<bool>(bIsUrgent);
-		if (OldValue != Value)
-		{
-			bIsUrgent = Value != 0;
-			PackageDatas.GetMonitor().OnUrgencyChanged(*this);
-		}
+		bIsUrgent = Value != 0;
+		PackageDatas.GetMonitor().OnUrgencyChanged(*this);
 	}
+}
 
-	void FPackageData::UpdateRequestData(const TConstArrayView<const ITargetPlatform*> InRequestedPlatforms,
-		bool bInIsUrgent, FCompletionCallback&& InCompletionCallback, FInstigator&& InInstigator, bool bAllowUpdateUrgency)
+void FPackageData::UpdateRequestData(const TConstArrayView<const ITargetPlatform*> InRequestedPlatforms,
+	bool bInIsUrgent, FCompletionCallback&& InCompletionCallback, FInstigator&& InInstigator, bool bAllowUpdateUrgency)
+{
+	if (IsInProgress())
 	{
-		if (IsInProgress())
-		{
-			AddCompletionCallback(MoveTemp(InCompletionCallback));
+		AddCompletionCallback(MoveTemp(InCompletionCallback));
 
-			bool bUrgencyChanged = false;
-			if (bInIsUrgent && !GetIsUrgent())
-			{
-				bUrgencyChanged = true;
-				SetIsUrgent(true);
-			}
+		bool bUrgencyChanged = false;
+		if (bInIsUrgent && !GetIsUrgent())
+		{
+			bUrgencyChanged = true;
+			SetIsUrgent(true);
+		}
 
-			if (!HasAllRequestedPlatforms(InRequestedPlatforms))
-			{
-				// Send back to the Request state (canceling any current operations) and then add the new platforms
-				if (GetState() != EPackageState::Request)
-				{
-					SendToState(EPackageState::Request, ESendFlags::QueueAddAndRemove);
-				}
-				SetPlatformsRequested(InRequestedPlatforms, true);
-			}
-			else if (bUrgencyChanged && bAllowUpdateUrgency)
+		if (!HasAllRequestedPlatforms(InRequestedPlatforms))
+		{
+			// Send back to the Request state (canceling any current operations) and then add the new platforms
+			if (GetState() != EPackageState::Request)
 			{
-				SendToState(GetState(), ESendFlags::QueueAddAndRemove);
+				SendToState(EPackageState::Request, ESendFlags::QueueAddAndRemove);
 			}
+			SetPlatformsRequested(InRequestedPlatforms, true);
 		}
-		else if (InRequestedPlatforms.Num() > 0)
+		else if (bUrgencyChanged && bAllowUpdateUrgency)
 		{
-			SetRequestData(InRequestedPlatforms, bInIsUrgent, MoveTemp(InCompletionCallback), MoveTemp(Instigator));
-			SendToState(EPackageState::Request, ESendFlags::QueueAddAndRemove);
+			SendToState(GetState(), ESendFlags::QueueAddAndRemove);
 		}
 	}
-
-	void FPackageData::SetRequestData(const TArrayView<const ITargetPlatform* const>& InRequestedPlatforms,
-		bool bInIsUrgent, FCompletionCallback&& InCompletionCallback, FInstigator&& InInstigator)
+	else if (InRequestedPlatforms.Num() > 0)
 	{
-		check(!CompletionCallback);
-		check(GetNumRequestedPlatforms() == 0)
-		check(!bIsUrgent);
-
-		check(InRequestedPlatforms.Num() != 0);
-		SetPlatformsRequested(InRequestedPlatforms, true);
-		SetIsUrgent(bInIsUrgent);
-		AddCompletionCallback(MoveTemp(InCompletionCallback));
-		if (Instigator.Category == EInstigator::NotYetRequested)
-		{
-			Instigator = MoveTemp(InInstigator);
-			PackageDatas.DebugInstigator(*this);
-		}
+		SetRequestData(InRequestedPlatforms, bInIsUrgent, MoveTemp(InCompletionCallback), MoveTemp(Instigator));
+		SendToState(EPackageState::Request, ESendFlags::QueueAddAndRemove);
 	}
+}
+
+void FPackageData::SetRequestData(const TArrayView<const ITargetPlatform* const>& InRequestedPlatforms,
+	bool bInIsUrgent, FCompletionCallback&& InCompletionCallback, FInstigator&& InInstigator)
+{
+	check(!CompletionCallback);
+	check(GetNumRequestedPlatforms() == 0)
+	check(!bIsUrgent);
 
-	void FPackageData::ClearInProgressData()
+	check(InRequestedPlatforms.Num() != 0);
+	SetPlatformsRequested(InRequestedPlatforms, true);
+	SetIsUrgent(bInIsUrgent);
+	AddCompletionCallback(MoveTemp(InCompletionCallback));
+	if (Instigator.Category == EInstigator::NotYetRequested)
 	{
-		ClearRequestedPlatforms();
-		SetIsUrgent(false);
-		CompletionCallback = FCompletionCallback();
+		Instigator = MoveTemp(InInstigator);
+		PackageDatas.DebugInstigator(*this);
 	}
+}
+
+void FPackageData::ClearInProgressData()
+{
+	ClearRequestedPlatforms();
+	SetIsUrgent(false);
+	CompletionCallback = FCompletionCallback();
+}
 
-	void FPackageData::SetPlatformsCooked(const TConstArrayView<const ITargetPlatform*> TargetPlatforms,
-		const TConstArrayView<bool> Succeeded)
+void FPackageData::SetPlatformsCooked(const TConstArrayView<const ITargetPlatform*> TargetPlatforms,
+	const TConstArrayView<bool> Succeeded)
+{
+	check(TargetPlatforms.Num() == Succeeded.Num());
+	for (int32 n = 0; n < TargetPlatforms.Num(); ++n)
 	{
-		check(TargetPlatforms.Num() == Succeeded.Num());
-		for (int32 n = 0; n < TargetPlatforms.Num(); ++n)
-		{
-			SetPlatformCooked(TargetPlatforms[n], Succeeded[n]);
-		}
+		SetPlatformCooked(TargetPlatforms[n], Succeeded[n]);
 	}
+}
 
-	void FPackageData::SetPlatformsCooked(const TConstArrayView<const ITargetPlatform*> TargetPlatforms,
-		bool bSucceeded)
+void FPackageData::SetPlatformsCooked(const TConstArrayView<const ITargetPlatform*> TargetPlatforms,
+	bool bSucceeded)
+{
+	for (const ITargetPlatform* TargetPlatform : TargetPlatforms)
 	{
-		for (const ITargetPlatform* TargetPlatform : TargetPlatforms)
-		{
-			SetPlatformCooked(TargetPlatform, bSucceeded);
-		}
+		SetPlatformCooked(TargetPlatform, bSucceeded);
 	}
+}
 
-	void FPackageData::SetPlatformCooked(const ITargetPlatform* TargetPlatform, bool bSucceeded)
+void FPackageData::SetPlatformCooked(const ITargetPlatform* TargetPlatform, bool bSucceeded)
+{
+	bool bHasAnyOthers = false;
+	bool bModified = false;
+	bool bExists = false;
+	for (TPair<const ITargetPlatform*, FPlatformData>& Pair : PlatformDatas)
 	{
-		bool bHasAnyOthers = false;
-		bool bModified = false;
-		bool bExists = false;
-		for (TPair<const ITargetPlatform*, FPlatformData>& Pair : PlatformDatas)
-		{
-			if (Pair.Key == TargetPlatform)
-			{
-				bExists = true;
-				bModified = bModified | (Pair.Value.bCookAttempted == false);
-				Pair.Value.bCookAttempted = true;
-				Pair.Value.bCookSucceeded = bSucceeded;
-			}
-			else
-			{
-				bHasAnyOthers = bHasAnyOthers | (Pair.Value.bCookAttempted != false);
-			}
-		}
-		if (!bExists)
+		if (Pair.Key == TargetPlatform)
 		{
-			FPlatformData& Value = PlatformDatas.FindOrAdd(TargetPlatform);
-			Value.bCookAttempted = true;
-			Value.bCookSucceeded = bSucceeded;
-			bModified = true;
+			bExists = true;
+			bModified = bModified | (Pair.Value.bCookAttempted == false);
+			Pair.Value.bCookAttempted = true;
+			Pair.Value.bCookSucceeded = bSucceeded;
 		}
-		if (bModified && !bHasAnyOthers)
+		else
 		{
-			PackageDatas.GetMonitor().OnFirstCookedPlatformAdded(*this);
+			bHasAnyOthers = bHasAnyOthers | (Pair.Value.bCookAttempted != false);
 		}
 	}
+	if (!bExists)
+	{
+		FPlatformData& Value = PlatformDatas.FindOrAdd(TargetPlatform);
+		Value.bCookAttempted = true;
+		Value.bCookSucceeded = bSucceeded;
+		bModified = true;
+	}
+	if (bModified && !bHasAnyOthers)
+	{
+		PackageDatas.GetMonitor().OnFirstCookedPlatformAdded(*this);
+	}
+}
 
-	void FPackageData::SetPlatformsNotCooked(const TConstArrayView<const ITargetPlatform*> TargetPlatforms)
+void FPackageData::SetPlatformsNotCooked(const TConstArrayView<const ITargetPlatform*> TargetPlatforms)
+{
+	for (const ITargetPlatform* TargetPlatform : TargetPlatforms)
 	{
-		for (const ITargetPlatform* TargetPlatform : TargetPlatforms)
-		{
-			SetPlatformNotCooked(TargetPlatform);
-		}
+		SetPlatformNotCooked(TargetPlatform);
+	}
+}
+
+void FPackageData::SetPlatformsNotCooked()
+{
+	bool bModified = false;
+	for (TPair<const ITargetPlatform*, FPlatformData>& Pair : PlatformDatas)
+	{
+		bModified = bModified | (Pair.Value.bCookAttempted != false);
+		Pair.Value.bCookAttempted = false;
+		Pair.Value.bCookSucceeded = false;
+	}
+	if (bModified)
+	{
+		PackageDatas.GetMonitor().OnLastCookedPlatformRemoved(*this);
 	}
+}
 
-	void FPackageData::SetPlatformsNotCooked()
+void FPackageData::SetPlatformNotCooked(const ITargetPlatform* TargetPlatform)
+{
+	bool bHasAnyOthers = false;
+	bool bModified = false;
+	for (TPair<const ITargetPlatform*, FPlatformData>& Pair : PlatformDatas)
 	{
-		bool bModified = false;
-		for (TPair<const ITargetPlatform*, FPlatformData>& Pair : PlatformDatas)
+		if (Pair.Key == TargetPlatform)
 		{
 			bModified = bModified | (Pair.Value.bCookAttempted != false);
 			Pair.Value.bCookAttempted = false;
 			Pair.Value.bCookSucceeded = false;
 		}
-		if (bModified)
+		else
 		{
-			PackageDatas.GetMonitor().OnLastCookedPlatformRemoved(*this);
+			bHasAnyOthers = bHasAnyOthers | (Pair.Value.bCookAttempted != false);
 		}
 	}
-
-	void FPackageData::SetPlatformNotCooked(const ITargetPlatform* TargetPlatform)
+	if (bModified && !bHasAnyOthers)
 	{
-		bool bHasAnyOthers = false;
-		bool bModified = false;
-		for (TPair<const ITargetPlatform*, FPlatformData>& Pair : PlatformDatas)
-		{
-			if (Pair.Key == TargetPlatform)
-			{
-				bModified = bModified | (Pair.Value.bCookAttempted != false);
-				Pair.Value.bCookAttempted = false;
-				Pair.Value.bCookSucceeded = false;
-			}
-			else
-			{
-				bHasAnyOthers = bHasAnyOthers | (Pair.Value.bCookAttempted != false);
-			}
-		}
-		if (bModified && !bHasAnyOthers)
-		{
-			PackageDatas.GetMonitor().OnLastCookedPlatformRemoved(*this);
-		}
+		PackageDatas.GetMonitor().OnLastCookedPlatformRemoved(*this);
 	}
+}
 
-	const TSortedMap<const ITargetPlatform*, FPackageData::FPlatformData>& FPackageData::GetPlatformDatas() const
-	{
-		return PlatformDatas;
-	}
+const TSortedMap<const ITargetPlatform*, FPackageData::FPlatformData>& FPackageData::GetPlatformDatas() const
+{
+	return PlatformDatas;
+}
 
-	FPackageData::FPlatformData& FPackageData::FindOrAddPlatformData(const ITargetPlatform* TargetPlatform)
-	{
-		return PlatformDatas.FindOrAdd(TargetPlatform);
-	}
+FPackageData::FPlatformData& FPackageData::FindOrAddPlatformData(const ITargetPlatform* TargetPlatform)
+{
+	return PlatformDatas.FindOrAdd(TargetPlatform);
+}
 
-	FPackageData::FPlatformData* FPackageData::FindPlatformData(const ITargetPlatform* TargetPlatform)
-	{
-		return PlatformDatas.Find(TargetPlatform);
-	}
+FPackageData::FPlatformData* FPackageData::FindPlatformData(const ITargetPlatform* TargetPlatform)
+{
+	return PlatformDatas.Find(TargetPlatform);
+}
 
-	const FPackageData::FPlatformData* FPackageData::FindPlatformData(const ITargetPlatform* TargetPlatform) const
-	{
-		return PlatformDatas.Find(TargetPlatform);
-	}
+const FPackageData::FPlatformData* FPackageData::FindPlatformData(const ITargetPlatform* TargetPlatform) const
+{
+	return PlatformDatas.Find(TargetPlatform);
+}
 
-	bool FPackageData::HasAnyCookedPlatform() const
-	{
-		return Algo::AnyOf(PlatformDatas,
-			[](const TPair<const ITargetPlatform*, FPlatformData>& Pair) { return Pair.Value.bCookAttempted; });
-	}
+bool FPackageData::HasAnyCookedPlatform() const
+{
+	return Algo::AnyOf(PlatformDatas,
+		[](const TPair<const ITargetPlatform*, FPlatformData>& Pair) { return Pair.Value.bCookAttempted; });
+}
 
-	bool FPackageData::HasAnyCookedPlatforms(const TArrayView<const ITargetPlatform* const>& Platforms,
-		bool bIncludeFailed) const
+bool FPackageData::HasAnyCookedPlatforms(const TArrayView<const ITargetPlatform* const>& Platforms,
+	bool bIncludeFailed) const
+{
+	if (PlatformDatas.Num() == 0)
 	{
-		if (PlatformDatas.Num() == 0)
-		{
-			return false;
-		}
-
-		for (const ITargetPlatform* QueryPlatform : Platforms)
-		{
-			if (HasCookedPlatform(QueryPlatform, bIncludeFailed))
-			{
-				return true;
-			}
-		}
 		return false;
 	}
 
-	bool FPackageData::HasAllCookedPlatforms(const TArrayView<const ITargetPlatform* const>& Platforms,
-		bool bIncludeFailed) const
+	for (const ITargetPlatform* QueryPlatform : Platforms)
 	{
-		if (Platforms.Num() == 0)
+		if (HasCookedPlatform(QueryPlatform, bIncludeFailed))
 		{
 			return true;
 		}
-		if (PlatformDatas.Num() == 0)
-		{
-			return false;
-		}
+	}
+	return false;
+}
 
-		for (const ITargetPlatform* QueryPlatform : Platforms)
-		{
-			if (!HasCookedPlatform(QueryPlatform, bIncludeFailed))
-			{
-				return false;
-			}
-		}
+bool FPackageData::HasAllCookedPlatforms(const TArrayView<const ITargetPlatform* const>& Platforms,
+	bool bIncludeFailed) const
+{
+	if (Platforms.Num() == 0)
+	{
 		return true;
 	}
-
-	bool FPackageData::HasCookedPlatform(const ITargetPlatform* Platform, bool bIncludeFailed) const
+	if (PlatformDatas.Num() == 0)
 	{
-		ECookResult Result = GetCookResults(Platform);
-		return (Result == ECookResult::Succeeded) | ((Result == ECookResult::Failed) & (bIncludeFailed != 0));
+		return false;
 	}
 
-	ECookResult FPackageData::GetCookResults(const ITargetPlatform* Platform) const
+	for (const ITargetPlatform* QueryPlatform : Platforms)
 	{
-		const FPlatformData* PlatformData = PlatformDatas.Find(Platform);
-		if (PlatformData && PlatformData->bCookAttempted)
+		if (!HasCookedPlatform(QueryPlatform, bIncludeFailed))
 		{
-			return PlatformData->bCookSucceeded ? ECookResult::Succeeded : ECookResult::Failed;
+			return false;
 		}
-		return ECookResult::Unseen;
 	}
+	return true;
+}
+
+bool FPackageData::HasCookedPlatform(const ITargetPlatform* Platform, bool bIncludeFailed) const
+{
+	ECookResult Result = GetCookResults(Platform);
+	return (Result == ECookResult::Succeeded) | ((Result == ECookResult::Failed) & (bIncludeFailed != 0));
+}
 
-	UPackage* FPackageData::GetPackage() const
+ECookResult FPackageData::GetCookResults(const ITargetPlatform* Platform) const
+{
+	const FPlatformData* PlatformData = PlatformDatas.Find(Platform);
+	if (PlatformData && PlatformData->bCookAttempted)
 	{
-		return Package.Get();
+		return PlatformData->bCookSucceeded ? ECookResult::Succeeded : ECookResult::Failed;
 	}
+	return ECookResult::Unseen;
+}
+
+UPackage* FPackageData::GetPackage() const
+{
+	return Package.Get();
+}
+
+void FPackageData::SetPackage(UPackage* InPackage)
+{
+	Package = InPackage;
+}
+
+EPackageState FPackageData::GetState() const
+{
+	return static_cast<EPackageState>(State);
+}
 
-	void FPackageData::SetPackage(UPackage* InPackage)
+/** Boilerplate-reduction struct that defines all multi-state properties and sets them based on the given state. */
+struct FStateProperties
+{
+	EPackageStateProperty Properties;
+	explicit FStateProperties(EPackageState InState)
 	{
-		Package = InPackage;
+		switch (InState)
+		{
+		case EPackageState::Idle:
+			Properties = EPackageStateProperty::None;
+			break;
+		case EPackageState::Request:
+			Properties = EPackageStateProperty::InProgress;
+			break;
+		case EPackageState::LoadPrepare:
+			Properties = EPackageStateProperty::InProgress | EPackageStateProperty::Loading;
+			break;
+		case EPackageState::LoadReady:
+			Properties = EPackageStateProperty::InProgress | EPackageStateProperty::Loading;
+			break;
+		// TODO_SaveQueue: When we add state PrepareForSave, it will also have bHasPackage = true, 
+		case EPackageState::Save:
+			Properties = EPackageStateProperty::InProgress | EPackageStateProperty::HasPackage;
+			break;
+		default:
+			check(false);
+			Properties = EPackageStateProperty::None;
+			break;
+		}
 	}
+};
 
-	EPackageState FPackageData::GetState() const
+void FPackageData::SendToState(EPackageState NextState, ESendFlags SendFlags)
+{
+	EPackageState OldState = GetState();
+	switch (OldState)
 	{
-		return static_cast<EPackageState>(State);
+	case EPackageState::Idle:
+		OnExitIdle();
+		break;
+	case EPackageState::Request:
+		if (!!(SendFlags & ESendFlags::QueueRemove))
+		{
+			ensure(PackageDatas.GetRequestQueue().Remove(this) == 1);
+		}
+		OnExitRequest();
+		break;
+	case EPackageState::LoadPrepare:
+		if (!!(SendFlags & ESendFlags::QueueRemove))
+		{
+			ensure(PackageDatas.GetLoadPrepareQueue().Remove(this) == 1);
+		}
+		OnExitLoadPrepare();
+		break;
+	case EPackageState::LoadReady:
+		if (!!(SendFlags & ESendFlags::QueueRemove))
+		{
+			ensure(PackageDatas.GetLoadReadyQueue().Remove(this) == 1);
+		}
+		OnExitLoadReady();
+		break;
+	case EPackageState::Save:
+		if (!!(SendFlags & ESendFlags::QueueRemove))
+		{
+			ensure(PackageDatas.GetSaveQueue().Remove(this) == 1);
+		}
+		OnExitSave();
+		break;
+	default:
+		check(false);
+		break;
 	}
 
-	/** Boilerplate-reduction struct that defines all multi-state properties and sets them based on the given state. */
-	struct FStateProperties
+	FStateProperties OldProperties(OldState);
+	FStateProperties NewProperties(NextState);
+	// Exit state properties from highest to lowest; enter state properties from lowest to highest.
+	// This ensures that properties that rely on earlier properties are constructed later and torn down earlier
+	// than the earlier properties.
+	for (EPackageStateProperty Iterator = EPackageStateProperty::Max;
+		Iterator >= EPackageStateProperty::Min;
+		Iterator = static_cast<EPackageStateProperty>(static_cast<uint32>(Iterator) >> 1))
 	{
-		EPackageStateProperty Properties;
-		explicit FStateProperties(EPackageState InState)
+		if (((OldProperties.Properties & Iterator) != EPackageStateProperty::None) &
+			((NewProperties.Properties & Iterator) == EPackageStateProperty::None))
 		{
-			switch (InState)
+			switch (Iterator)
 			{
-			case EPackageState::Idle:
-				Properties = EPackageStateProperty::None;
+			case EPackageStateProperty::InProgress:
+				OnExitInProgress();
 				break;
-			case EPackageState::Request:
-				Properties = EPackageStateProperty::InProgress;
+			case EPackageStateProperty::Loading:
+				OnExitLoading();
+				break;
+			case EPackageStateProperty::HasPackage:
+				OnExitHasPackage();
+				break;
+			default:
+				check(false);
 				break;
-			case EPackageState::LoadPrepare:
-				Properties = EPackageStateProperty::InProgress | EPackageStateProperty::Loading;
+			}
+		}
+	}
+	for (EPackageStateProperty Iterator = EPackageStateProperty::Min;
+		Iterator <= EPackageStateProperty::Max;
+		Iterator = static_cast<EPackageStateProperty>(static_cast<uint32>(Iterator) << 1))
+	{
+		if (((OldProperties.Properties & Iterator) == EPackageStateProperty::None) &
+			((NewProperties.Properties & Iterator) != EPackageStateProperty::None))
+		{
+			switch (Iterator)
+			{
+			case EPackageStateProperty::InProgress:
+				OnEnterInProgress();
 				break;
-			case EPackageState::LoadReady:
-				Properties = EPackageStateProperty::InProgress | EPackageStateProperty::Loading;
+			case EPackageStateProperty::Loading:
+				OnEnterLoading();
 				break;
-			// TODO_SaveQueue: When we add state PrepareForSave, it will also have bHasPackage = true, 
-			case EPackageState::Save:
-				Properties = EPackageStateProperty::InProgress | EPackageStateProperty::HasPackage;
+			case EPackageStateProperty::HasPackage:
+				OnEnterHasPackage();
 				break;
 			default:
 				check(false);
-				Properties = EPackageStateProperty::None;
 				break;
 			}
 		}
-	};
+	}
+
 
-	void FPackageData::SendToState(EPackageState NextState, ESendFlags SendFlags)
+	SetState(NextState);
+	switch (NextState)
 	{
-		EPackageState OldState = GetState();
-		switch (OldState)
+	case EPackageState::Idle:
+		OnEnterIdle();
+		break;
+	case EPackageState::Request:
+		OnEnterRequest();
+		if (((SendFlags & ESendFlags::QueueAdd) != ESendFlags::QueueNone))
 		{
-		case EPackageState::Idle:
-			OnExitIdle();
-			break;
-		case EPackageState::Request:
-			if (!!(SendFlags & ESendFlags::QueueRemove))
+			PackageDatas.GetRequestQueue().AddRequest(this);
+		}
+		break;
+	case EPackageState::LoadPrepare:
+		OnEnterLoadPrepare();
+		if ((SendFlags & ESendFlags::QueueAdd) != ESendFlags::QueueNone)
+		{
+			if (GetIsUrgent())
 			{
-				ensure(PackageDatas.GetRequestQueue().Remove(this) == 1);
+				PackageDatas.GetLoadPrepareQueue().AddFront(this);
 			}
-			OnExitRequest();
-			break;
-		case EPackageState::LoadPrepare:
-			if (!!(SendFlags & ESendFlags::QueueRemove))
+			else
 			{
-				ensure(PackageDatas.GetLoadPrepareQueue().Remove(this) == 1);
+				PackageDatas.GetLoadPrepareQueue().Add(this);
 			}
-			OnExitLoadPrepare();
-			break;
-		case EPackageState::LoadReady:
-			if (!!(SendFlags & ESendFlags::QueueRemove))
+		}
+		break;
+	case EPackageState::LoadReady:
+		OnEnterLoadReady();
+		if ((SendFlags & ESendFlags::QueueAdd) != ESendFlags::QueueNone)
+		{
+			if (GetIsUrgent())
 			{
-				ensure(PackageDatas.GetLoadReadyQueue().Remove(this) == 1);
+				PackageDatas.GetLoadReadyQueue().AddFront(this);
 			}
-			OnExitLoadReady();
-			break;
-		case EPackageState::Save:
-			if (!!(SendFlags & ESendFlags::QueueRemove))
+			else
 			{
-				ensure(PackageDatas.GetSaveQueue().Remove(this) == 1);
+				PackageDatas.GetLoadReadyQueue().Add(this);
 			}
-			OnExitSave();
-			break;
-		default:
-			check(false);
-			break;
 		}
-
-		FStateProperties OldProperties(OldState);
-		FStateProperties NewProperties(NextState);
-		// Exit state properties from highest to lowest; enter state properties from lowest to highest.
-		// This ensures that properties that rely on earlier properties are constructed later and torn down earlier
-		// than the earlier properties.
-		for (EPackageStateProperty Iterator = EPackageStateProperty::Max;
-			Iterator >= EPackageStateProperty::Min;
-			Iterator = static_cast<EPackageStateProperty>(static_cast<uint32>(Iterator) >> 1))
+		break;
+	case EPackageState::Save:
+		OnEnterSave();
+		if (((SendFlags & ESendFlags::QueueAdd) != ESendFlags::QueueNone))
 		{
-			if (((OldProperties.Properties & Iterator) != EPackageStateProperty::None) &
-				((NewProperties.Properties & Iterator) == EPackageStateProperty::None))
+			if (GetIsUrgent())
 			{
-				switch (Iterator)
-				{
-				case EPackageStateProperty::InProgress:
-					OnExitInProgress();
-					break;
-				case EPackageStateProperty::Loading:
-					OnExitLoading();
-					break;
-				case EPackageStateProperty::HasPackage:
-					OnExitHasPackage();
-					break;
-				default:
-					check(false);
-					break;
-				}
+				PackageDatas.GetSaveQueue().AddFront(this);
 			}
-		}
-		for (EPackageStateProperty Iterator = EPackageStateProperty::Min;
-			Iterator <= EPackageStateProperty::Max;
-			Iterator = static_cast<EPackageStateProperty>(static_cast<uint32>(Iterator) << 1))
-		{
-			if (((OldProperties.Properties & Iterator) == EPackageStateProperty::None) &
-				((NewProperties.Properties & Iterator) != EPackageStateProperty::None))
+			else
 			{
-				switch (Iterator)
-				{
-				case EPackageStateProperty::InProgress:
-					OnEnterInProgress();
-					break;
-				case EPackageStateProperty::Loading:
-					OnEnterLoading();
-					break;
-				case EPackageStateProperty::HasPackage:
-					OnEnterHasPackage();
-					break;
-				default:
-					check(false);
-					break;
-				}
+				PackageDatas.GetSaveQueue().Add(this);
 			}
 		}
+		break;
+	default:
+		check(false);
+		break;
+	}
 
+	PackageDatas.GetMonitor().OnStateChanged(*this, OldState);
+}
 
-		SetState(NextState);
-		switch (NextState)
-		{
-		case EPackageState::Idle:
-			OnEnterIdle();
-			break;
-		case EPackageState::Request:
-			OnEnterRequest();
-			if (((SendFlags & ESendFlags::QueueAdd) != ESendFlags::QueueNone))
-			{
-				PackageDatas.GetRequestQueue().AddRequest(this);
-			}
-			break;
-		case EPackageState::LoadPrepare:
-			OnEnterLoadPrepare();
-			if ((SendFlags & ESendFlags::QueueAdd) != ESendFlags::QueueNone)
-			{
-				if (GetIsUrgent())
-				{
-					PackageDatas.GetLoadPrepareQueue().AddFront(this);
-				}
-				else
-				{
-					PackageDatas.GetLoadPrepareQueue().Add(this);
-				}
-			}
-			break;
-		case EPackageState::LoadReady:
-			OnEnterLoadReady();
-			if ((SendFlags & ESendFlags::QueueAdd) != ESendFlags::QueueNone)
-			{
-				if (GetIsUrgent())
-				{
-					PackageDatas.GetLoadReadyQueue().AddFront(this);
-				}
-				else
-				{
-					PackageDatas.GetLoadReadyQueue().Add(this);
-				}
-			}
-			break;
-		case EPackageState::Save:
-			OnEnterSave();
-			if (((SendFlags & ESendFlags::QueueAdd) != ESendFlags::QueueNone))
-			{
-				if (GetIsUrgent())
-				{
-					PackageDatas.GetSaveQueue().AddFront(this);
-				}
-				else
-				{
-					PackageDatas.GetSaveQueue().Add(this);
-				}
-			}
-			break;
-		default:
-			check(false);
-			break;
-		}
-
-		PackageDatas.GetMonitor().OnStateChanged(*this, OldState);
-	}
+void FPackageData::CheckInContainer() const
+{
+	switch (GetState())
+	{
+	case EPackageState::Idle:
+		break;
+	case EPackageState::Request:
+		check(PackageDatas.GetRequestQueue().Contains(this));
+		break;
+	case EPackageState::LoadPrepare:
+		check(PackageDatas.GetLoadPrepareQueue().Contains(this));
+		break;
+	case EPackageState::LoadReady:
+		check(Algo::Find(PackageDatas.GetLoadReadyQueue(), this) != nullptr);
+		break;
+	case EPackageState::Save:
+		// The save queue is huge and often pushed at end. Check last element first and then scan.
+		check(PackageDatas.GetSaveQueue().Num() && (PackageDatas.GetSaveQueue().Last() == this || Algo::Find(PackageDatas.GetSaveQueue(), this)));
+		break;
+	default:
+		check(false);
+		break;
+	}
+}
+
+bool FPackageData::IsInProgress() const
+{
+	return IsInStateProperty(EPackageStateProperty::InProgress);
+}
 
-	void FPackageData::CheckInContainer() const
-	{
-		switch (GetState())
-		{
-		case EPackageState::Idle:
-			break;
-		case EPackageState::Request:
-			check(PackageDatas.GetRequestQueue().Contains(this));
-			break;
-		case EPackageState::LoadPrepare:
-			check(PackageDatas.GetLoadPrepareQueue().Contains(this));
-			break;
-		case EPackageState::LoadReady:
-			check(Algo::Find(PackageDatas.GetLoadReadyQueue(), this) != nullptr);
-			break;
-		case EPackageState::Save:
-			// The save queue is huge and often pushed at end. Check last element first and then scan.
-			check(PackageDatas.GetSaveQueue().Num() && (PackageDatas.GetSaveQueue().Last() == this || Algo::Find(PackageDatas.GetSaveQueue(), this)));
-			break;
-		default:
-			check(false);
-			break;
-		}
-	}
+bool FPackageData::IsInStateProperty(EPackageStateProperty Property) const
+{
+	return (FStateProperties(GetState()).Properties & Property) != EPackageStateProperty::None;
+}
 
-	bool FPackageData::IsInProgress() const
-	{
-		return IsInStateProperty(EPackageStateProperty::InProgress);
-	}
+void FPackageData::OnEnterIdle()
+{
+	// Note that this might be on construction of the PackageData
+}
 
-	bool FPackageData::IsInStateProperty(EPackageStateProperty Property) const
+void FPackageData::OnExitIdle()
+{
+	if (PackageDatas.GetLogDiscoveredPackages())
 	{
-		return (FStateProperties(GetState()).Properties & Property) != EPackageStateProperty::None;
+		UE_LOG(LogCook, Warning, TEXT("Missing dependency: Package %s discovered after initial dependency search."), *WriteToString<256>(PackageName));
 	}
+}
 
-	void FPackageData::OnEnterIdle()
-	{
-		// Note that this might be on construction of the PackageData
-	}
+void FPackageData::OnEnterRequest()
+{
+	// It is not valid to enter the request state without requested platforms; it indicates a bug due to e.g.
+	// calling SendToState without UpdateRequestData from Idle
+	check(GetNumRequestedPlatforms() > 0);
+}
 
-	void FPackageData::OnExitIdle()
-	{
-		if (PackageDatas.GetLogDiscoveredPackages())
-		{
-			UE_LOG(LogCook, Warning, TEXT("Missing dependency: Package %s discovered after initial dependency search."), *WriteToString<256>(PackageName));
-		}
-	}
+void FPackageData::OnExitRequest()
+{
+}
 
-	void FPackageData::OnEnterRequest()
-	{
-		// It is not valid to enter the request state without requested platforms; it indicates a bug due to e.g.
-		// calling SendToState without UpdateRequestData from Idle
-		check(GetNumRequestedPlatforms() > 0);
-	}
+void FPackageData::OnEnterLoadPrepare()
+{
+}
 
-	void FPackageData::OnExitRequest()
-	{
-	}
+void FPackageData::OnExitLoadPrepare()
+{
+}
 
-	void FPackageData::OnEnterLoadPrepare()
-	{
-	}
+void FPackageData::OnEnterLoadReady()
+{
+}
 
-	void FPackageData::OnExitLoadPrepare()
-	{
-	}
+void FPackageData::OnExitLoadReady()
+{
+}
 
-	void FPackageData::OnEnterLoadReady()
-	{
-	}
+void FPackageData::OnEnterSave()
+{
+	check(GetPackage() != nullptr && GetPackage()->IsFullyLoaded());
 
-	void FPackageData::OnExitLoadReady()
-	{
-	}
+	check(!GetHasBeginPrepareSaveFailed());
+	CheckObjectCacheEmpty();
+	CheckCookedPlatformDataEmpty();
+}
 
-	void FPackageData::OnEnterSave()
-	{
-		check(GetPackage() != nullptr && GetPackage()->IsFullyLoaded());
+void FPackageData::OnExitSave()
+{
+	PackageDatas.GetCookOnTheFlyServer().ReleaseCookedPlatformData(*this, false /* bCompletedSave */);
+	ClearObjectCache();
+	SetHasBeginPrepareSaveFailed(false);
+}
 
-		check(!GetHasBeginPrepareSaveFailed());
-		CheckObjectCacheEmpty();
-		CheckCookedPlatformDataEmpty();
-	}
+void FPackageData::OnEnterInProgress()
+{
+	PackageDatas.GetMonitor().OnInProgressChanged(*this, true);
+}
 
-	void FPackageData::OnExitSave()
+void FPackageData::OnExitInProgress()
+{
+	PackageDatas.GetMonitor().OnInProgressChanged(*this, false);
+	UE::Cook::FCompletionCallback LocalCompletionCallback(MoveTemp(GetCompletionCallback()));
+	if (LocalCompletionCallback)
 	{
-		PackageDatas.GetCookOnTheFlyServer().ReleaseCookedPlatformData(*this, false /* bCompletedSave */);
-		ClearObjectCache();
-		SetHasBeginPrepareSaveFailed(false);
+		LocalCompletionCallback(this);
 	}
+	ClearInProgressData();
+}
 
-	void FPackageData::OnEnterInProgress()
-	{
-		PackageDatas.GetMonitor().OnInProgressChanged(*this, true);
-	}
+void FPackageData::OnEnterLoading()
+{
+	CheckPreloadEmpty();
+}
 
-	void FPackageData::OnExitInProgress()
-	{
-		PackageDatas.GetMonitor().OnInProgressChanged(*this, false);
-		UE::Cook::FCompletionCallback LocalCompletionCallback(MoveTemp(GetCompletionCallback()));
-		if (LocalCompletionCallback)
-		{
-			LocalCompletionCallback(this);
-		}
-		ClearInProgressData();
-	}
+void FPackageData::OnExitLoading()
+{
+	ClearPreload();
+}
 
-	void FPackageData::OnEnterLoading()
-	{
-		CheckPreloadEmpty();
-	}
+void FPackageData::OnEnterHasPackage()
+{
+}
 
-	void FPackageData::OnExitLoading()
-	{
-		ClearPreload();
-	}
+void FPackageData::OnExitHasPackage()
+{
+	SetPackage(nullptr);
+}
 
-	void FPackageData::OnEnterHasPackage()
-	{
-	}
+void FPackageData::SetState(EPackageState NextState)
+{
+	State = static_cast<uint32>(NextState);
+}
+
+FCompletionCallback& FPackageData::GetCompletionCallback()
+{
+	return CompletionCallback;
+}
 
-	void FPackageData::OnExitHasPackage()
+void FPackageData::AddCompletionCallback(FCompletionCallback&& InCompletionCallback)
+{
+	if (InCompletionCallback)
 	{
-		SetPackage(nullptr);
+		// We don't yet have a mechanism for calling two completion callbacks.
+		// CompletionCallbacks only come from external requests, and it should not be possible to request twice,
+		// so a failed check here shouldn't happen.
+		check(!CompletionCallback);
+		CompletionCallback = MoveTemp(InCompletionCallback);
 	}
+}
 
-	void FPackageData::SetState(EPackageState NextState)
+bool FPackageData::TryPreload()
+{
+	check(IsInStateProperty(EPackageStateProperty::Loading));
+	if (GetIsPreloadAttempted())
 	{
-		State = static_cast<uint32>(NextState);
+		return true;
 	}
-
-	FCompletionCallback& FPackageData::GetCompletionCallback()
+	if (FindObjectFast<UPackage>(nullptr, GetPackageName()))
 	{
-		return CompletionCallback;
+		// If the package has already loaded, then there is no point in further preloading
+		ClearPreload();
+		SetIsPreloadAttempted(true);
+		return true;
 	}
-
-	void FPackageData::AddCompletionCallback(FCompletionCallback&& InCompletionCallback)
+	if (IsGenerated())
 	{
-		if (InCompletionCallback)
-		{
-			// We don't yet have a mechanism for calling two completion callbacks.
-			// CompletionCallbacks only come from external requests, and it should not be possible to request twice,
-			// so a failed check here shouldn't happen.
-			check(!CompletionCallback);
-			CompletionCallback = MoveTemp(InCompletionCallback);
-		}
+		// Deferred populate generated packages are loaded from their generator, not from disk
+		ClearPreload();
+		SetIsPreloadAttempted(true);
+		return true;
 	}
-
-	bool FPackageData::TryPreload()
+	if (!PreloadableFile.Get())
 	{
-		check(IsInStateProperty(EPackageStateProperty::Loading));
-		if (GetIsPreloadAttempted())
+		if (FEditorDomain* EditorDomain = FEditorDomain::Get())
 		{
-			return true;
-		}
-		if (FindObjectFast<UPackage>(nullptr, GetPackageName()))
-		{
-			// If the package has already loaded, then there is no point in further preloading
-			ClearPreload();
-			SetIsPreloadAttempted(true);
-			return true;
-		}
-		if (IsGenerated())
-		{
-			// Deferred populate generated packages are loaded from their generator, not from disk
-			ClearPreload();
-			SetIsPreloadAttempted(true);
-			return true;
+			EditorDomain->PrecachePackageDigest(GetPackageName());
 		}
-		if (!PreloadableFile.Get())
-		{
-			if (FEditorDomain* EditorDomain = FEditorDomain::Get())
+		TStringBuilder<NAME_SIZE> FileNameString;
+		GetFileName().ToString(FileNameString);
+		PreloadableFile.Set(MakeShared<FPreloadableArchive>(FileNameString.ToString()), *this);
+		PreloadableFile.Get()->InitializeAsync([this]()
 			{
-				EditorDomain->PrecachePackageDigest(GetPackageName());
-			}
-			TStringBuilder<NAME_SIZE> FileNameString;
-			GetFileName().ToString(FileNameString);
-			PreloadableFile.Set(MakeShared<FPreloadableArchive>(FileNameString.ToString()), *this);
-			PreloadableFile.Get()->InitializeAsync([this]()
+				TStringBuilder<NAME_SIZE> FileNameString;
+				// Note this async callback has an read of this->GetFilename and a write of PreloadableFileOpenResult
+				// outside of a critical section. This read and write is allowed because GetFilename does
+				// not change until this is destructed, and the destructor does not run and other threads do not read
+				// or write PreloadableFileOpenResult until after PreloadableFile.Get() has finished initialization
+				// and this callback is therefore complete.
+				// The code that accomplishes that waiting is in TryPreload (IsInitialized) and ClearPreload (ReleaseCache)
+				this->GetFileName().ToString(FileNameString);
+				FPackagePath PackagePath = FPackagePath::FromLocalPath(FileNameString);
+				FOpenPackageResult Result = IPackageResourceManager::Get().OpenReadPackage(PackagePath);
+				if (Result.Archive)
 				{
-					TStringBuilder<NAME_SIZE> FileNameString;
-					// Note this async callback has an read of this->GetFilename and a write of PreloadableFileOpenResult
-					// outside of a critical section. This read and write is allowed because GetFilename does
-					// not change until this is destructed, and the destructor does not run and other threads do not read
-					// or write PreloadableFileOpenResult until after PreloadableFile.Get() has finished initialization
-					// and this callback is therefore complete.
-					// The code that accomplishes that waiting is in TryPreload (IsInitialized) and ClearPreload (ReleaseCache)
-					this->GetFileName().ToString(FileNameString);
-					FPackagePath PackagePath = FPackagePath::FromLocalPath(FileNameString);
-					FOpenPackageResult Result = IPackageResourceManager::Get().OpenReadPackage(PackagePath);
-					if (Result.Archive)
-					{
-						this->PreloadableFileOpenResult.CopyMetaData(Result);
-					}
-					return Result.Archive.Release();
-				},
-				FPreloadableFile::Flags::PreloadHandle | FPreloadableFile::Flags::Prime);
-		}
-		const TSharedPtr<FPreloadableArchive>& FilePtr = PreloadableFile.Get();
-		if (!FilePtr->IsInitialized())
-		{
-			if (GetIsUrgent())
-			{
-				// For urgent requests, wait on them to finish preloading rather than letting them run asynchronously
-				// and coming back to them later
-				FilePtr->WaitForInitialization();
-				check(FilePtr->IsInitialized());
-			}
-			else
-			{
-				return false;
-			}
-		}
-		if (FilePtr->TotalSize() < 0)
+					this->PreloadableFileOpenResult.CopyMetaData(Result);
+				}
+				return Result.Archive.Release();
+			},
+			FPreloadableFile::Flags::PreloadHandle | FPreloadableFile::Flags::Prime);
+	}
+	const TSharedPtr<FPreloadableArchive>& FilePtr = PreloadableFile.Get();
+	if (!FilePtr->IsInitialized())
+	{
+		if (GetIsUrgent())
 		{
-			UE_LOG(LogCook, Warning, TEXT("Failed to find file when preloading %s."), *GetFileName().ToString());
-			SetIsPreloadAttempted(true);
-			PreloadableFile.Reset(*this);
-			PreloadableFileOpenResult = FOpenPackageResult();
-			return true;
+			// For urgent requests, wait on them to finish preloading rather than letting them run asynchronously
+			// and coming back to them later
+			FilePtr->WaitForInitialization();
+			check(FilePtr->IsInitialized());
 		}
-
-		TStringBuilder<NAME_SIZE> FileNameString;
-		GetFileName().ToString(FileNameString);
-		if (!IPackageResourceManager::TryRegisterPreloadableArchive(FPackagePath::FromLocalPath(FileNameString),
-			FilePtr, PreloadableFileOpenResult))
+		else
 		{
-			UE_LOG(LogCook, Warning, TEXT("Failed to register %s for preload."), *GetFileName().ToString());
-			SetIsPreloadAttempted(true);
-			PreloadableFile.Reset(*this);
-			PreloadableFileOpenResult = FOpenPackageResult();
-			return true;
+			return false;
 		}
+	}
+	if (FilePtr->TotalSize() < 0)
+	{
+		UE_LOG(LogCook, Warning, TEXT("Failed to find file when preloading %s."), *GetFileName().ToString());
+		SetIsPreloadAttempted(true);
+		PreloadableFile.Reset(*this);
+		PreloadableFileOpenResult = FOpenPackageResult();
+		return true;
+	}
 
-		SetIsPreloaded(true);
+	TStringBuilder<NAME_SIZE> FileNameString;
+	GetFileName().ToString(FileNameString);
+	if (!IPackageResourceManager::TryRegisterPreloadableArchive(FPackagePath::FromLocalPath(FileNameString),
+		FilePtr, PreloadableFileOpenResult))
+	{
+		UE_LOG(LogCook, Warning, TEXT("Failed to register %s for preload."), *GetFileName().ToString());
 		SetIsPreloadAttempted(true);
+		PreloadableFile.Reset(*this);
+		PreloadableFileOpenResult = FOpenPackageResult();
 		return true;
 	}
 
-	void FPackageData::FTrackedPreloadableFilePtr::Set(TSharedPtr<FPreloadableArchive>&& InPtr, FPackageData& Owner)
+	SetIsPreloaded(true);
+	SetIsPreloadAttempted(true);
+	return true;
+}
+
+void FPackageData::FTrackedPreloadableFilePtr::Set(TSharedPtr<FPreloadableArchive>&& InPtr, FPackageData& Owner)
+{
+	Reset(Owner);
+	if (InPtr)
 	{
-		Reset(Owner);
-		if (InPtr)
-		{
-			Ptr = MoveTemp(InPtr);
-			Owner.PackageDatas.GetMonitor().OnPreloadAllocatedChanged(Owner, true);
-		}
+		Ptr = MoveTemp(InPtr);
+		Owner.PackageDatas.GetMonitor().OnPreloadAllocatedChanged(Owner, true);
 	}
+}
 
-	void FPackageData::FTrackedPreloadableFilePtr::Reset(FPackageData& Owner)
+void FPackageData::FTrackedPreloadableFilePtr::Reset(FPackageData& Owner)
+{
+	if (Ptr)
 	{
-		if (Ptr)
-		{
-			Owner.PackageDatas.GetMonitor().OnPreloadAllocatedChanged(Owner, false);
-			Ptr.Reset();
-		}
+		Owner.PackageDatas.GetMonitor().OnPreloadAllocatedChanged(Owner, false);
+		Ptr.Reset();
 	}
+}
 
-	void FPackageData::ClearPreload()
+void FPackageData::ClearPreload()
+{
+	const TSharedPtr<FPreloadableArchive>& FilePtr = PreloadableFile.Get();
+	if (GetIsPreloaded())
 	{
-		const TSharedPtr<FPreloadableArchive>& FilePtr = PreloadableFile.Get();
-		if (GetIsPreloaded())
-		{
-			check(FilePtr);
-			TStringBuilder<NAME_SIZE> FileNameString;
-			GetFileName().ToString(FileNameString);
-			if (IPackageResourceManager::UnRegisterPreloadableArchive(FPackagePath::FromLocalPath(FileNameString)))
-			{
-				UE_LOG(LogCook, Display, TEXT("PreloadableFile was created for %s but never used. This is wasteful and bad for cook performance."),
-					*PackageName.ToString());
-			}
-			FilePtr->ReleaseCache(); // ReleaseCache to conserve memory if the Linker still has a pointer to it
-		}
-		else
+		check(FilePtr);
+		TStringBuilder<NAME_SIZE> FileNameString;
+		GetFileName().ToString(FileNameString);
+		if (IPackageResourceManager::UnRegisterPreloadableArchive(FPackagePath::FromLocalPath(FileNameString)))
 		{
-			check(!FilePtr || !FilePtr->IsCacheAllocated());
+			UE_LOG(LogCook, Display, TEXT("PreloadableFile was created for %s but never used. This is wasteful and bad for cook performance."),
+				*PackageName.ToString());
 		}
-
-		PreloadableFile.Reset(*this);
-		PreloadableFileOpenResult = FOpenPackageResult();
-		SetIsPreloaded(false);
-		SetIsPreloadAttempted(false);
+		FilePtr->ReleaseCache(); // ReleaseCache to conserve memory if the Linker still has a pointer to it
 	}
-
-	void FPackageData::CheckPreloadEmpty()
+	else
 	{
-		check(!GetIsPreloadAttempted());
-		check(!PreloadableFile.Get());
-		check(!GetIsPreloaded());
+		check(!FilePtr || !FilePtr->IsCacheAllocated());
 	}
 
-	TArray<FWeakObjectPtr>& FPackageData::GetCachedObjectsInOuter()
-	{
-		return CachedObjectsInOuter;
-	}
+	PreloadableFile.Reset(*this);
+	PreloadableFileOpenResult = FOpenPackageResult();
+	SetIsPreloaded(false);
+	SetIsPreloadAttempted(false);
+}
+
+void FPackageData::CheckPreloadEmpty()
+{
+	check(!GetIsPreloadAttempted());
+	check(!PreloadableFile.Get());
+	check(!GetIsPreloaded());
+}
+
+TArray<FWeakObjectPtr>& FPackageData::GetCachedObjectsInOuter()
+{
+	return CachedObjectsInOuter;
+}
 
-	void FPackageData::CheckObjectCacheEmpty() const
+void FPackageData::CheckObjectCacheEmpty() const
+{
+	check(CachedObjectsInOuter.Num() == 0);
+	check(!GetHasSaveCache());
+}
+
+void FPackageData::CreateObjectCache()
+{
+	if (GetHasSaveCache())
 	{
-		check(CachedObjectsInOuter.Num() == 0);
-		check(!GetHasSaveCache());
+		return;
 	}
 
-	void FPackageData::CreateObjectCache()
+	UPackage* LocalPackage = GetPackage();
+	if (LocalPackage && LocalPackage->IsFullyLoaded())
 	{
-		if (GetHasSaveCache())
+		PackageName = LocalPackage->GetFName();
+		TArray<UObject*> ObjectsInOuter;
+		GetObjectsWithOuter(LocalPackage, ObjectsInOuter);
+		CachedObjectsInOuter.Reset(ObjectsInOuter.Num());
+		for (UObject* Object : ObjectsInOuter)
 		{
-			return;
-		}
-
-		UPackage* LocalPackage = GetPackage();
-		if (LocalPackage && LocalPackage->IsFullyLoaded())
-		{
-			PackageName = LocalPackage->GetFName();
-			TArray<UObject*> ObjectsInOuter;
-			GetObjectsWithOuter(LocalPackage, ObjectsInOuter);
-			CachedObjectsInOuter.Reset(ObjectsInOuter.Num());
-			for (UObject* Object : ObjectsInOuter)
+			FWeakObjectPtr ObjectWeakPointer(Object);
+			// ignore pending kill objects; they will not be serialized out so we don't need to call
+			// BeginCacheForCookedPlatformData on them
+			if (!ObjectWeakPointer.Get())
 			{
-				FWeakObjectPtr ObjectWeakPointer(Object);
-				// ignore pending kill objects; they will not be serialized out so we don't need to call
-				// BeginCacheForCookedPlatformData on them
-				if (!ObjectWeakPointer.Get())
-				{
-					continue;
-				}
-				CachedObjectsInOuter.Emplace(MoveTemp(ObjectWeakPointer));
+				continue;
 			}
-			SetHasSaveCache(true);
-		}
-		else
-		{
-			check(false);
+			CachedObjectsInOuter.Emplace(MoveTemp(ObjectWeakPointer));
 		}
+		SetHasSaveCache(true);
 	}
-
-	void FPackageData::ClearObjectCache()
+	else
 	{
-		CachedObjectsInOuter.Empty();
-		SetHasSaveCache(false);
+		check(false);
 	}
+}
 
-	const int32& FPackageData::GetNumPendingCookedPlatformData() const
-	{
-		return NumPendingCookedPlatformData;
-	}
+void FPackageData::ClearObjectCache()
+{
+	CachedObjectsInOuter.Empty();
+	SetHasSaveCache(false);
+}
 
-	int32& FPackageData::GetNumPendingCookedPlatformData()
-	{
-		return NumPendingCookedPlatformData;
-	}
+const int32& FPackageData::GetNumPendingCookedPlatformData() const
+{
+	return NumPendingCookedPlatformData;
+}
 
-	const int32& FPackageData::GetCookedPlatformDataNextIndex() const
-	{
-		return CookedPlatformDataNextIndex;
-	}
+int32& FPackageData::GetNumPendingCookedPlatformData()
+{
+	return NumPendingCookedPlatformData;
+}
 
-	int32& FPackageData::GetCookedPlatformDataNextIndex()
-	{
-		return CookedPlatformDataNextIndex;
-	}
+const int32& FPackageData::GetCookedPlatformDataNextIndex() const
+{
+	return CookedPlatformDataNextIndex;
+}
 
-	void FPackageData::CheckCookedPlatformDataEmpty() const
-	{
-		check(GetCookedPlatformDataNextIndex() == 0);
-		check(!GetCookedPlatformDataStarted());
-		check(!GetCookedPlatformDataCalled());
-		check(!GetCookedPlatformDataComplete());
-	}
+int32& FPackageData::GetCookedPlatformDataNextIndex()
+{
+	return CookedPlatformDataNextIndex;
+}
 
-	void FPackageData::ClearCookedPlatformData()
-	{
-		CookedPlatformDataNextIndex = 0;
-		// Note that GetNumPendingCookedPlatformData is not cleared; it persists across Saves and CookSessions
-		SetCookedPlatformDataStarted(false);
-		SetCookedPlatformDataCalled(false);
-		SetCookedPlatformDataComplete(false);
-	}
+void FPackageData::CheckCookedPlatformDataEmpty() const
+{
+	check(GetCookedPlatformDataNextIndex() == 0);
+	check(!GetCookedPlatformDataStarted());
+	check(!GetCookedPlatformDataCalled());
+	check(!GetCookedPlatformDataComplete());
+}
 
-	void FPackageData::ResetGenerationProgress()
-	{
-		SetInitializedGeneratorSave(false);
-		SetCompletedGeneration(false);
-	}
+void FPackageData::ClearCookedPlatformData()
+{
+	CookedPlatformDataNextIndex = 0;
+	// Note that GetNumPendingCookedPlatformData is not cleared; it persists across Saves and CookSessions
+	SetCookedPlatformDataStarted(false);
+	SetCookedPlatformDataCalled(false);
+	SetCookedPlatformDataComplete(false);
+}
+
+void FPackageData::ResetGenerationProgress()
+{
+	SetInitializedGeneratorSave(false);
+	SetCompletedGeneration(false);
+}
 
-	void FPackageData::OnRemoveSessionPlatform(const ITargetPlatform* Platform)
-	{
-		PlatformDatas.Remove(Platform);
-	}
+void FPackageData::OnRemoveSessionPlatform(const ITargetPlatform* Platform)
+{
+	PlatformDatas.Remove(Platform);
+}
+
+bool FPackageData::HasReferencedObjects() const
+{
+	return Package != nullptr || CachedObjectsInOuter.Num() > 0;
+}
 
-	bool FPackageData::HasReferencedObjects() const
+void FPackageData::RemapTargetPlatforms(const TMap<ITargetPlatform*, ITargetPlatform*>& Remap)
+{
+	typedef TSortedMap<const ITargetPlatform*, FPlatformData> MapType;
+	MapType NewPlatformDatas;
+	NewPlatformDatas.Reserve(PlatformDatas.Num());
+	for (TPair<const ITargetPlatform*, FPlatformData>& ExistingPair : PlatformDatas)
 	{
-		return Package != nullptr || CachedObjectsInOuter.Num() > 0;
+		ITargetPlatform* NewKey = Remap[ExistingPair.Key];
+		NewPlatformDatas.FindOrAdd(NewKey) = MoveTemp(ExistingPair.Value);
 	}
 
-	void FPackageData::RemapTargetPlatforms(const TMap<ITargetPlatform*, ITargetPlatform*>& Remap)
+	// The save state (and maybe more in the future) depend on the order of the request platforms remaining
+	// unchanged, due to CookedPlatformDataNextIndex. If we change that order due to the remap, we need to
+	// demote back to request.
+	if (IsInProgress() && GetState() != EPackageState::Request)
 	{
-		typedef TSortedMap<const ITargetPlatform*, FPlatformData> MapType;
-		MapType NewPlatformDatas;
-		NewPlatformDatas.Reserve(PlatformDatas.Num());
-		for (TPair<const ITargetPlatform*, FPlatformData>& ExistingPair : PlatformDatas)
-		{
-			ITargetPlatform* NewKey = Remap[ExistingPair.Key];
-			NewPlatformDatas.FindOrAdd(NewKey) = MoveTemp(ExistingPair.Value);
-		}
-
-		// The save state (and maybe more in the future) depend on the order of the request platforms remaining
-		// unchanged, due to CookedPlatformDataNextIndex. If we change that order due to the remap, we need to
-		// demote back to request.
-		if (IsInProgress() && GetState() != EPackageState::Request)
+		bool bDemote = true;
+		MapType::TConstIterator OldIter = PlatformDatas.CreateConstIterator();
+		MapType::TConstIterator NewIter = NewPlatformDatas.CreateConstIterator();
+		for (; OldIter; ++OldIter, ++NewIter)
 		{
-			bool bDemote = true;
-			MapType::TConstIterator OldIter = PlatformDatas.CreateConstIterator();
-			MapType::TConstIterator NewIter = NewPlatformDatas.CreateConstIterator();
-			for (; OldIter; ++OldIter, ++NewIter)
+			if (OldIter.Key() != NewIter.Key())
 			{
-				if (OldIter.Key() != NewIter.Key())
-				{
-					bDemote = true;
-				}
-			}
-			if (bDemote)
-			{
-				SendToState(EPackageState::Request, ESendFlags::QueueAddAndRemove);
+				bDemote = true;
 			}
 		}
-		PlatformDatas = MoveTemp(NewPlatformDatas);
-	}
-
-	bool FPackageData::IsSaveInvalidated() const
-	{
-		if (GetState() != EPackageState::Save)
+		if (bDemote)
 		{
-			return false;
+			SendToState(EPackageState::Request, ESendFlags::QueueAddAndRemove);
 		}
-
-		return GetPackage() == nullptr || !GetPackage()->IsFullyLoaded() ||
-			Algo::AnyOf(CachedObjectsInOuter, [](const FWeakObjectPtr& WeakPtr)
-				{
-					// TODO: Keep track of which objects were public, and only invalidate the save if the object
-					// that has been deleted or marked pending kill was public
-					// Until we make that change, we will unnecessarily invalidate and demote some packages after a
-					// garbage collect
-					return WeakPtr.Get() == nullptr;
-				});
 	}
+	PlatformDatas = MoveTemp(NewPlatformDatas);
+}
 
-	void FPackageData::SetGeneratedOwner(FGeneratorPackage* InGeneratedOwner)
+bool FPackageData::IsSaveInvalidated() const
+{
+	if (GetState() != EPackageState::Save)
 	{
-		check(IsGenerated());
-		check(!(GeneratedOwner && InGeneratedOwner));
-		GeneratedOwner = InGeneratedOwner;
+		return false;
 	}
 
-	bool FPackageData::GeneratorPackageRequiresGC() const
-	{
-		// We consider that if a FPackageData has valid GeneratorPackage helper object,
-		// this means that COTFS's process of generating packages was not completed 
-		// either due to an error or because it has exceeded a maximum memory threshold. 
-		return IsGenerating() && !GetHasBeginPrepareSaveFailed();
-	}
+	return GetPackage() == nullptr || !GetPackage()->IsFullyLoaded() ||
+		Algo::AnyOf(CachedObjectsInOuter, [](const FWeakObjectPtr& WeakPtr)
+			{
+				// TODO: Keep track of which objects were public, and only invalidate the save if the object
+				// that has been deleted or marked pending kill was public
+				// Until we make that change, we will unnecessarily invalidate and demote some packages after a
+				// garbage collect
+				return WeakPtr.Get() == nullptr;
+			});
+}
+
+void FPackageData::SetGeneratedOwner(FGeneratorPackage* InGeneratedOwner)
+{
+	check(IsGenerated());
+	check(!(GeneratedOwner && InGeneratedOwner));
+	GeneratedOwner = InGeneratedOwner;
+}
 
-	UE::Cook::FGeneratorPackage* FPackageData::CreateGeneratorPackage(const UObject* InSplitDataObject,
-		ICookPackageSplitter* InCookPackageSplitterInstance)
-	{
-		check(!GetGeneratorPackage());
-		GeneratorPackage.Reset(new UE::Cook::FGeneratorPackage(*this, InSplitDataObject,
-			InCookPackageSplitterInstance));
-		return GetGeneratorPackage();
-	}
+bool FPackageData::GeneratorPackageRequiresGC() const
+{
+	// We consider that if a FPackageData has valid GeneratorPackage helper object,
+	// this means that COTFS's process of generating packages was not completed 
+	// either due to an error or because it has exceeded a maximum memory threshold. 
+	return IsGenerating() && !GetHasBeginPrepareSaveFailed();
+}
+
+UE::Cook::FGeneratorPackage* FPackageData::CreateGeneratorPackage(const UObject* InSplitDataObject,
+	ICookPackageSplitter* InCookPackageSplitterInstance)
+{
+	check(!GetGeneratorPackage());
+	GeneratorPackage.Reset(new UE::Cook::FGeneratorPackage(*this, InSplitDataObject,
+		InCookPackageSplitterInstance));
+	return GetGeneratorPackage();
+}
 	
-	//////////////////////////////////////////////////////////////////////////
-	// FGeneratorPackage
+//////////////////////////////////////////////////////////////////////////
+// FGeneratorPackage
 
-	FGeneratorPackage::FGeneratorPackage(UE::Cook::FPackageData& InOwner, const UObject* InSplitDataObject,
-		ICookPackageSplitter* InCookPackageSplitterInstance)
-	: Owner(InOwner)
-	, SplitDataObjectName(*InSplitDataObject->GetFullName())
-	{
-		check(InCookPackageSplitterInstance);
-		CookPackageSplitterInstance.Reset(InCookPackageSplitterInstance);
-	}
+FGeneratorPackage::FGeneratorPackage(UE::Cook::FPackageData& InOwner, const UObject* InSplitDataObject,
+	ICookPackageSplitter* InCookPackageSplitterInstance)
+: Owner(InOwner)
+, SplitDataObjectName(*InSplitDataObject->GetFullName())
+{
+	check(InCookPackageSplitterInstance);
+	CookPackageSplitterInstance.Reset(InCookPackageSplitterInstance);
+}
 
-	FGeneratorPackage::~FGeneratorPackage()
-	{
-		ClearGeneratedPackages();
-	}
+FGeneratorPackage::~FGeneratorPackage()
+{
+	ClearGeneratedPackages();
+}
 
-	void FGeneratorPackage::ClearGeneratedPackages()
+void FGeneratorPackage::ClearGeneratedPackages()
+{
+	for (FGeneratedStruct& GeneratedStruct : PackagesToGenerate)
 	{
-		for (FGeneratedStruct& GeneratedStruct : PackagesToGenerate)
+		if (GeneratedStruct.PackageData)
 		{
-			if (GeneratedStruct.PackageData)
-			{
-				check(GeneratedStruct.PackageData->GetGeneratedOwner() == this);
-				GeneratedStruct.PackageData->SetGeneratedOwner(nullptr);
-				GeneratedStruct.PackageData = nullptr;
-			}
+			check(GeneratedStruct.PackageData->GetGeneratedOwner() == this);
+			GeneratedStruct.PackageData->SetGeneratedOwner(nullptr);
+			GeneratedStruct.PackageData = nullptr;
 		}
 	}
+}
 
-	bool FGeneratorPackage::TryGenerateList(UObject* OwnerObject, const FPackageNameCache& PackageNameCache,
-		FPackageDatas& PackageDatas)
-	{
-		UPackage* OwnerPackage = Owner.GetPackage();
-		check(OwnerPackage);
-		TArray<ICookPackageSplitter::FGeneratedPackage> GeneratorDatas =
-			CookPackageSplitterInstance->GetGenerateList(OwnerPackage, OwnerObject);
-		PackagesToGenerate.Reset(GeneratorDatas.Num());
-		for (ICookPackageSplitter::FGeneratedPackage& SplitterData : GeneratorDatas)
-		{
-			FGeneratedStruct& GeneratedStruct = PackagesToGenerate.Emplace_GetRef();
-			GeneratedStruct.RelativePath = MoveTemp(SplitterData.RelativePath);
-			GeneratedStruct.Dependencies = MoveTemp(SplitterData.Dependencies);
-			FString PackageName = FPaths::RemoveDuplicateSlashes(FString::Printf(TEXT("%s/%s/%s"),
-				*this->Owner.GetPackageName().ToString(), GeneratedPackageSubPath, *GeneratedStruct.RelativePath));
-
-			if (!SplitterData.GetCreateAsMap().IsSet())
-			{
-				UE_LOG(LogCook, Error, TEXT("PackageSplitter did not specify whether CreateAsMap is true for generated package. Splitter=%s, Generated=%s."),
-					*this->GetSplitDataObjectName().ToString(), *PackageName);
-				return false;
-			}
-			GeneratedStruct.bCreateAsMap = *SplitterData.GetCreateAsMap();
+bool FGeneratorPackage::TryGenerateList(UObject* OwnerObject, const FPackageNameCache& PackageNameCache,
+	FPackageDatas& PackageDatas)
+{
+	UPackage* OwnerPackage = Owner.GetPackage();
+	check(OwnerPackage);
+	TArray<ICookPackageSplitter::FGeneratedPackage> GeneratorDatas =
+		CookPackageSplitterInstance->GetGenerateList(OwnerPackage, OwnerObject);
+	PackagesToGenerate.Reset(GeneratorDatas.Num());
+	for (ICookPackageSplitter::FGeneratedPackage& SplitterData : GeneratorDatas)
+	{
+		FGeneratedStruct& GeneratedStruct = PackagesToGenerate.Emplace_GetRef();
+		GeneratedStruct.RelativePath = MoveTemp(SplitterData.RelativePath);
+		GeneratedStruct.Dependencies = MoveTemp(SplitterData.Dependencies);
+		FString PackageName = FPaths::RemoveDuplicateSlashes(FString::Printf(TEXT("%s/%s/%s"),
+			*this->Owner.GetPackageName().ToString(), GeneratedPackageSubPath, *GeneratedStruct.RelativePath));
+
+		if (!SplitterData.GetCreateAsMap().IsSet())
+		{
+			UE_LOG(LogCook, Error, TEXT("PackageSplitter did not specify whether CreateAsMap is true for generated package. Splitter=%s, Generated=%s."),
+				*this->GetSplitDataObjectName().ToString(), *PackageName);
+			return false;
+		}
+		GeneratedStruct.bCreateAsMap = *SplitterData.GetCreateAsMap();
 
-			const FName PackageFName(*PackageName);
-			const FName FileName = PackageNameCache.GetCachedStandardFileName(PackageFName, false /* bRequireExists */,
-				GeneratedStruct.bCreateAsMap);
-			if (FileName.IsNone())
-			{
-				UE_LOG(LogCook, Error, TEXT("PackageSplitter could not find mounted filename for generated packagepath. Splitter=%s, Generated=%s."),
-					*this->GetSplitDataObjectName().ToString(), *PackageName);
-				return false;
-			}
-			UE::Cook::FPackageData& PackageData = PackageDatas.FindOrAddPackageData(PackageFName, FileName);
-			if (IFileManager::Get().FileExists(*FileName.ToString()))
-			{
-				UE_LOG(LogCook, Warning, TEXT("PackageSplitter specified a generated package that already exists in the workspace domain. Splitter=%s, Generated=%s."),
-					*this->GetSplitDataObjectName().ToString(), *PackageName);
-				return false;
-			}
-			GeneratedStruct.PackageData = &PackageData;
-			PackageData.SetGenerated(true);
-			// No package should be generated by two different splitters. If an earlier run of this splitter generated
-			// the package, the package's owner should have been reset to null when we called ClearGeneratedPackages
-			// between then and now
-			check(PackageData.GetGeneratedOwner() == nullptr); 
-			PackageData.SetGeneratedOwner(this);
+		const FName PackageFName(*PackageName);
+		const FName FileName = PackageNameCache.GetCachedStandardFileName(PackageFName, false /* bRequireExists */,
+			GeneratedStruct.bCreateAsMap);
+		if (FileName.IsNone())
+		{
+			UE_LOG(LogCook, Error, TEXT("PackageSplitter could not find mounted filename for generated packagepath. Splitter=%s, Generated=%s."),
+				*this->GetSplitDataObjectName().ToString(), *PackageName);
+			return false;
 		}
-		RemainingToPopulate = GeneratorDatas.Num();
-		return true;
+		UE::Cook::FPackageData& PackageData = PackageDatas.FindOrAddPackageData(PackageFName, FileName);
+		if (IFileManager::Get().FileExists(*FileName.ToString()))
+		{
+			UE_LOG(LogCook, Warning, TEXT("PackageSplitter specified a generated package that already exists in the workspace domain. Splitter=%s, Generated=%s."),
+				*this->GetSplitDataObjectName().ToString(), *PackageName);
+			return false;
+		}
+		GeneratedStruct.PackageData = &PackageData;
+		PackageData.SetGenerated(true);
+		// No package should be generated by two different splitters. If an earlier run of this splitter generated
+		// the package, the package's owner should have been reset to null when we called ClearGeneratedPackages
+		// between then and now
+		check(PackageData.GetGeneratedOwner() == nullptr); 
+		PackageData.SetGeneratedOwner(this);
 	}
+	RemainingToPopulate = GeneratorDatas.Num();
+	return true;
+}
 
-	FGeneratorPackage::FGeneratedStruct* FGeneratorPackage::FindGeneratedStruct(FPackageData* PackageData)
+FGeneratorPackage::FGeneratedStruct* FGeneratorPackage::FindGeneratedStruct(FPackageData* PackageData)
+{
+	for (FGeneratedStruct& GeneratedStruct : PackagesToGenerate)
 	{
-		for (FGeneratedStruct& GeneratedStruct : PackagesToGenerate)
+		if (GeneratedStruct.PackageData == PackageData)
 		{
-			if (GeneratedStruct.PackageData == PackageData)
-			{
-				return &GeneratedStruct;
-			}
+			return &GeneratedStruct;
 		}
-		return nullptr;
 	}
+	return nullptr;
+}
 
-	UObject* FGeneratorPackage::FindSplitDataObject() const
-	{
-		FString ObjectPath = GetSplitDataObjectName().ToString();
+UObject* FGeneratorPackage::FindSplitDataObject() const
+{
+	FString ObjectPath = GetSplitDataObjectName().ToString();
 
-		// SplitDataObjectName is a FullObjectPath; strip off the leading <ClassName> in
-		// "<ClassName> <Package>.<Object>:<SubObject>"
-		int32 ClassDelimiterIndex = -1;
-		if (ObjectPath.FindChar(' ', ClassDelimiterIndex))
-		{
-			ObjectPath.RightChopInline(ClassDelimiterIndex + 1);
-		}
-		return FindObject<UObject>(nullptr, *ObjectPath);
+	// SplitDataObjectName is a FullObjectPath; strip off the leading <ClassName> in
+	// "<ClassName> <Package>.<Object>:<SubObject>"
+	int32 ClassDelimiterIndex = -1;
+	if (ObjectPath.FindChar(' ', ClassDelimiterIndex))
+	{
+		ObjectPath.RightChopInline(ClassDelimiterIndex + 1);
 	}
+	return FindObject<UObject>(nullptr, *ObjectPath);
+}
 
-	void FGeneratorPackage::PostGarbageCollect()
+void FGeneratorPackage::PostGarbageCollect()
+{
+	if (!bGeneratedList)
+	{
+		return;
+	}
+	if (Owner.GetState() == EPackageState::Save)
 	{
-		if (!bGeneratedList)
+		// UCookOnTheFlyServer::PreCollectGarbage adds references for the Generator package and all its public
+		// objects, so it should still be loaded
+		if (!Owner.GetPackage() || !FindSplitDataObject())
 		{
-			return;
+			UE_LOG(LogCook, Error, TEXT("PackageSplitter object was deleted by garbage collection while generation was still ongoing. This will break the generation.")
+				TEXT("\n\tSplitter=%s."), *GetSplitDataObjectName().ToString());
 		}
-		if (Owner.GetState() == EPackageState::Save)
+	}
+	else
+	{
+		// After the Generator Package is saved, we drop its referenced and it can be garbage collected
+		// If we have any packages left to populate, our splitter contract requires that it be garbage collected;
+		// we promise that the package is not partially GC'd during calls to TryPopulateGeneratedPackage
+		// The splitter can opt-out of this contract and keep it referenced itself if it desires.
+		UPackage* OwnerPackage = FindObject<UPackage>(nullptr, *Owner.GetPackageName().ToString());
+		if (OwnerPackage)
 		{
-			// UCookOnTheFlyServer::PreCollectGarbage adds references for the Generator package and all its public
-			// objects, so it should still be loaded
-			if (!Owner.GetPackage() || !FindSplitDataObject())
+			if (RemainingToPopulate > 0 &&
+				!CookPackageSplitterInstance->UseInternalReferenceToAvoidGarbageCollect())
 			{
-				UE_LOG(LogCook, Error, TEXT("PackageSplitter object was deleted by garbage collection while generation was still ongoing. This will break the generation.")
-					TEXT("\n\tSplitter=%s."), *GetSplitDataObjectName().ToString());
+				UE_LOG(LogCook, Error, TEXT("PackageSplitter found the Generator package still in memory after it should have been deleted by GC.")
+					TEXT("\n\tThis is unexpected since garbage has been collected and the package should have been unreferenced so it should have been collected, and will break population of Generated packages.")
+					TEXT("\n\tSplitter=%s"), *GetSplitDataObjectName().ToString());
+				EReferenceChainSearchMode SearchMode = EReferenceChainSearchMode::Shortest
+					| EReferenceChainSearchMode::PrintAllResults
+					| EReferenceChainSearchMode::FullChain;
+				FReferenceChainSearch RefChainSearch(OwnerPackage, SearchMode);
 			}
 		}
 		else
 		{
-			// After the Generator Package is saved, we drop its referenced and it can be garbage collected
-			// If we have any packages left to populate, our splitter contract requires that it be garbage collected;
-			// we promise that the package is not partially GC'd during calls to TryPopulateGeneratedPackage
-			// The splitter can opt-out of this contract and keep it referenced itself if it desires.
-			UPackage* OwnerPackage = FindObject<UPackage>(nullptr, *Owner.GetPackageName().ToString());
-			if (OwnerPackage)
-			{
-				if (RemainingToPopulate > 0 &&
-					!CookPackageSplitterInstance->UseInternalReferenceToAvoidGarbageCollect())
-				{
-					UE_LOG(LogCook, Error, TEXT("PackageSplitter found the Generator package still in memory after it should have been deleted by GC.")
-						TEXT("\n\tThis is unexpected since garbage has been collected and the package should have been unreferenced so it should have been collected, and will break population of Generated packages.")
-						TEXT("\n\tSplitter=%s"), *GetSplitDataObjectName().ToString());
-					EReferenceChainSearchMode SearchMode = EReferenceChainSearchMode::Shortest
-						| EReferenceChainSearchMode::PrintAllResults
-						| EReferenceChainSearchMode::FullChain;
-					FReferenceChainSearch RefChainSearch(OwnerPackage, SearchMode);
-				}
-			}
-			else
-			{
-				bWasOwnerReloaded = true;
-			}
+			bWasOwnerReloaded = true;
 		}
+	}
 
-		bool bHasIssuedWarning = false;
-		for (FGeneratedStruct& GeneratedStruct : PackagesToGenerate)
+	bool bHasIssuedWarning = false;
+	for (FGeneratedStruct& GeneratedStruct : PackagesToGenerate)
+	{
+		GeneratedStruct.bHasCreatedPackage = false;
+		if (!GeneratedStruct.bHasSaved && !bHasIssuedWarning)
 		{
-			GeneratedStruct.bHasCreatedPackage = false;
-			if (!GeneratedStruct.bHasSaved && !bHasIssuedWarning)
+			if (FindObject<UPackage>(nullptr, *GeneratedStruct.PackageData->GetPackageName().ToString()))
 			{
-				if (FindObject<UPackage>(nullptr, *GeneratedStruct.PackageData->GetPackageName().ToString()))
-				{
-					UE_LOG(LogCook, Warning, TEXT("PackageSplitter found a package it generated that was not removed from memory during garbage collection. This will cause errors later during population.")
-						TEXT("\n\tSplitter=%s, Generated=%s."), *GetSplitDataObjectName().ToString(), *GeneratedStruct.PackageData->GetPackageName().ToString());
-					bHasIssuedWarning = true; // Only issue the warning once per GC
-				}
+				UE_LOG(LogCook, Warning, TEXT("PackageSplitter found a package it generated that was not removed from memory during garbage collection. This will cause errors later during population.")
+					TEXT("\n\tSplitter=%s, Generated=%s."), *GetSplitDataObjectName().ToString(), *GeneratedStruct.PackageData->GetPackageName().ToString());
+				bHasIssuedWarning = true; // Only issue the warning once per GC
 			}
 		}
 	}
+}
 
-	UPackage* FGeneratorPackage::CreateGeneratedUPackage(FGeneratorPackage::FGeneratedStruct& GeneratedStruct,
-		const UPackage* OwnerPackage, const TCHAR* GeneratedPackageName)
-	{
-		UPackage* GeneratedPackage = CreatePackage(GeneratedPackageName);
-		PRAGMA_DISABLE_DEPRECATION_WARNINGS
-		GeneratedPackage->SetGuid(OwnerPackage->GetGuid());
-		PRAGMA_ENABLE_DEPRECATION_WARNINGS
-		GeneratedPackage->SetPersistentGuid(OwnerPackage->GetPersistentGuid());
-		GeneratedStruct.bHasCreatedPackage = true;
-		return GeneratedPackage;
-	}
-
-	void FGeneratorPackage::SetGeneratedSaved(FPackageData& PackageData)
+UPackage* FGeneratorPackage::CreateGeneratedUPackage(FGeneratorPackage::FGeneratedStruct& GeneratedStruct,
+	const UPackage* OwnerPackage, const TCHAR* GeneratedPackageName)
+{
+	UPackage* GeneratedPackage = CreatePackage(GeneratedPackageName);
+	PRAGMA_DISABLE_DEPRECATION_WARNINGS
+	GeneratedPackage->SetGuid(OwnerPackage->GetGuid());
+	PRAGMA_ENABLE_DEPRECATION_WARNINGS
+	GeneratedPackage->SetPersistentGuid(OwnerPackage->GetPersistentGuid());
+	GeneratedStruct.bHasCreatedPackage = true;
+	return GeneratedPackage;
+}
+
+void FGeneratorPackage::SetGeneratedSaved(FPackageData& PackageData)
+{
+	FGeneratedStruct* GeneratedStruct = FindGeneratedStruct(&PackageData);
+	if (!GeneratedStruct)
 	{
-		FGeneratedStruct* GeneratedStruct = FindGeneratedStruct(&PackageData);
-		if (!GeneratedStruct)
-		{
-			UE_LOG(LogCook, Warning, TEXT("PackageSplitter called SetGeneratedSaved on a package that does not belong to the splitter.")
-				TEXT("\n\tSplitter=%s, Generated=%s."), *GetSplitDataObjectName().ToString(),
-				*PackageData.GetPackageName().ToString());
-			return;
-		}
-		if (GeneratedStruct->bHasSaved)
-		{
-			return;
-		}
-		GeneratedStruct->bHasSaved = true;
-		--RemainingToPopulate;
-		check(RemainingToPopulate >= 0);
+		UE_LOG(LogCook, Warning, TEXT("PackageSplitter called SetGeneratedSaved on a package that does not belong to the splitter.")
+			TEXT("\n\tSplitter=%s, Generated=%s."), *GetSplitDataObjectName().ToString(),
+			*PackageData.GetPackageName().ToString());
+		return;
 	}
-
-	bool FGeneratorPackage::IsComplete() const
+	if (GeneratedStruct->bHasSaved)
 	{
-		return bGeneratedList && RemainingToPopulate == 0;
+		return;
 	}
+	GeneratedStruct->bHasSaved = true;
+	--RemainingToPopulate;
+	check(RemainingToPopulate >= 0);
+}
 
-	void FGeneratorPackage::GetIntermediateMountPoint(FString& OutPackagePath, FString& OutLocalFilePath) const
-	{
-		const FString OwnerShortName = FPackageName::GetShortName(Owner.GetPackageName().ToString());
-		OutPackagePath = FPaths::RemoveDuplicateSlashes(FString::Printf(TEXT("/%s%s/"),
-			*OwnerShortName, UE::Cook::GeneratedPackageSubPath));
-		OutLocalFilePath = FPaths::RemoveDuplicateSlashes(FString::Printf(TEXT("%s/Cooked/%s/%s/"),
-			*FPaths::ProjectIntermediateDir(), *OwnerShortName, UE::Cook::GeneratedPackageSubPath));
-
-	}
-	FString FGeneratorPackage::GetIntermediateLocalPath(const FGeneratorPackage::FGeneratedStruct& GeneratedStruct) const
-	{
-		FString UnusedPackagePath;
-		FString MountLocalFilePath;
-		GetIntermediateMountPoint(UnusedPackagePath, MountLocalFilePath);
-		FString Extension = FPaths::GetExtension(GeneratedStruct.PackageData->GetFileName().ToString(),
-			true /* bIncludeDot */);
-		return FPaths::RemoveDuplicateSlashes(FString::Printf(TEXT("%s/%s%s"),
-			*MountLocalFilePath, *GeneratedStruct.RelativePath, *Extension));
-	}
+bool FGeneratorPackage::IsComplete() const
+{
+	return bGeneratedList && RemainingToPopulate == 0;
+}
 
-	//////////////////////////////////////////////////////////////////////////
-	// FPendingCookedPlatformData
+void FGeneratorPackage::GetIntermediateMountPoint(FString& OutPackagePath, FString& OutLocalFilePath) const
+{
+	const FString OwnerShortName = FPackageName::GetShortName(Owner.GetPackageName().ToString());
+	OutPackagePath = FPaths::RemoveDuplicateSlashes(FString::Printf(TEXT("/%s%s/"),
+		*OwnerShortName, UE::Cook::GeneratedPackageSubPath));
+	OutLocalFilePath = FPaths::RemoveDuplicateSlashes(FString::Printf(TEXT("%s/Cooked/%s/%s/"),
+		*FPaths::ProjectIntermediateDir(), *OwnerShortName, UE::Cook::GeneratedPackageSubPath));
+
+}
+FString FGeneratorPackage::GetIntermediateLocalPath(const FGeneratorPackage::FGeneratedStruct& GeneratedStruct) const
+{
+	FString UnusedPackagePath;
+	FString MountLocalFilePath;
+	GetIntermediateMountPoint(UnusedPackagePath, MountLocalFilePath);
+	FString Extension = FPaths::GetExtension(GeneratedStruct.PackageData->GetFileName().ToString(),
+		true /* bIncludeDot */);
+	return FPaths::RemoveDuplicateSlashes(FString::Printf(TEXT("%s/%s%s"),
+		*MountLocalFilePath, *GeneratedStruct.RelativePath, *Extension));
+}
+
+//////////////////////////////////////////////////////////////////////////
+// FPendingCookedPlatformData
+
+
+FPendingCookedPlatformData::FPendingCookedPlatformData(UObject* InObject, const ITargetPlatform* InTargetPlatform,
+	FPackageData& InPackageData, bool bInNeedsResourceRelease, UCookOnTheFlyServer& InCookOnTheFlyServer)
+	: Object(InObject), TargetPlatform(InTargetPlatform), PackageData(InPackageData)
+	, CookOnTheFlyServer(InCookOnTheFlyServer),	CancelManager(nullptr), ClassName(InObject->GetClass()->GetFName())
+	, bHasReleased(false), bNeedsResourceRelease(bInNeedsResourceRelease)
+{
+	check(InObject);
+	PackageData.GetNumPendingCookedPlatformData() += 1;
+}
+
+FPendingCookedPlatformData::FPendingCookedPlatformData(FPendingCookedPlatformData&& Other)
+	: Object(Other.Object), TargetPlatform(Other.TargetPlatform), PackageData(Other.PackageData)
+	, CookOnTheFlyServer(Other.CookOnTheFlyServer), CancelManager(Other.CancelManager), ClassName(Other.ClassName)
+	, bHasReleased(Other.bHasReleased), bNeedsResourceRelease(Other.bNeedsResourceRelease)
+{
+	Other.Object = nullptr;
+}
 
+FPendingCookedPlatformData::~FPendingCookedPlatformData()
+{
+	Release();
+}
 
-	FPendingCookedPlatformData::FPendingCookedPlatformData(UObject* InObject, const ITargetPlatform* InTargetPlatform,
-		FPackageData& InPackageData, bool bInNeedsResourceRelease, UCookOnTheFlyServer& InCookOnTheFlyServer)
-		: Object(InObject), TargetPlatform(InTargetPlatform), PackageData(InPackageData)
-		, CookOnTheFlyServer(InCookOnTheFlyServer),	CancelManager(nullptr), ClassName(InObject->GetClass()->GetFName())
-		, bHasReleased(false), bNeedsResourceRelease(bInNeedsResourceRelease)
+bool FPendingCookedPlatformData::PollIsComplete()
+{
+	if (bHasReleased)
 	{
-		check(InObject);
-		PackageData.GetNumPendingCookedPlatformData() += 1;
+		return true;
 	}
 
-	FPendingCookedPlatformData::FPendingCookedPlatformData(FPendingCookedPlatformData&& Other)
-		: Object(Other.Object), TargetPlatform(Other.TargetPlatform), PackageData(Other.PackageData)
-		, CookOnTheFlyServer(Other.CookOnTheFlyServer), CancelManager(Other.CancelManager), ClassName(Other.ClassName)
-		, bHasReleased(Other.bHasReleased), bNeedsResourceRelease(Other.bNeedsResourceRelease)
+	UObject* LocalObject = Object.Get();
+	if (!LocalObject)
 	{
-		Other.Object = nullptr;
+		Release();
+		return true;
 	}
-
-	FPendingCookedPlatformData::~FPendingCookedPlatformData()
+	UE_TRACK_REFERENCING_PACKAGE_SCOPED(LocalObject->GetPackage(), PackageAccessTrackingOps::NAME_CookerBuildObject);
+	if (LocalObject->IsCachedCookedPlatformDataLoaded(TargetPlatform))
 	{
 		Release();
+		return true;
 	}
-
-	bool FPendingCookedPlatformData::PollIsComplete()
+	else
 	{
-		if (bHasReleased)
-		{
-			return true;
-		}
-
-		UObject* LocalObject = Object.Get();
-		if (!LocalObject)
-		{
-			Release();
-			return true;
-		}
-		UE_TRACK_REFERENCING_PACKAGE_SCOPED(LocalObject->GetPackage(), PackageAccessTrackingOps::NAME_CookerBuildObject);
-		if (LocalObject->IsCachedCookedPlatformDataLoaded(TargetPlatform))
-		{
-			Release();
-			return true;
-		}
-		else
-		{
 #if DEBUG_COOKONTHEFLY
-			UE_LOG(LogCook, Display, TEXT("Object %s isn't cached yet"), *LocalObject->GetFullName());
+		UE_LOG(LogCook, Display, TEXT("Object %s isn't cached yet"), *LocalObject->GetFullName());
 #endif
-			/*if ( LocalObject->IsA(UMaterial::StaticClass()) )
+		/*if ( LocalObject->IsA(UMaterial::StaticClass()) )
+		{
+			if (GShaderCompilingManager->HasShaderJobs() == false)
 			{
-				if (GShaderCompilingManager->HasShaderJobs() == false)
-				{
-					UE_LOG(LogCook, Warning, TEXT("Shader compiler is in a bad state!  Shader %s is finished compile but shader compiling manager did not notify shader.  "),
-						*LocalObject->GetPathName());
-				}
-			}*/
-			return false;
-		}
+				UE_LOG(LogCook, Warning, TEXT("Shader compiler is in a bad state!  Shader %s is finished compile but shader compiling manager did not notify shader.  "),
+					*LocalObject->GetPathName());
+			}
+		}*/
+		return false;
 	}
+}
 
-	void FPendingCookedPlatformData::Release()
+void FPendingCookedPlatformData::Release()
+{
+	if (bHasReleased)
 	{
-		if (bHasReleased)
-		{
-			return;
-		}
-
-		if (bNeedsResourceRelease)
-		{
-			int32* CurrentAsyncCache = CookOnTheFlyServer.CurrentAsyncCacheForType.Find(ClassName);
-			// bNeedsRelease should not have been set if the AsyncCache does not have an entry for the class
-			check(CurrentAsyncCache != nullptr);
-			*CurrentAsyncCache += 1;
-		}
-
-		PackageData.GetNumPendingCookedPlatformData() -= 1;
-		check(PackageData.GetNumPendingCookedPlatformData() >= 0);
-		if (CancelManager)
-		{
-			CancelManager->Release(*this);
-			CancelManager = nullptr;
-		}
+		return;
+	}
 
-		Object = nullptr;
-		bHasReleased = true;
+	if (bNeedsResourceRelease)
+	{
+		int32* CurrentAsyncCache = CookOnTheFlyServer.CurrentAsyncCacheForType.Find(ClassName);
+		// bNeedsRelease should not have been set if the AsyncCache does not have an entry for the class
+		check(CurrentAsyncCache != nullptr);
+		*CurrentAsyncCache += 1;
 	}
 
-	void FPendingCookedPlatformData::RemapTargetPlatforms(const TMap<ITargetPlatform*, ITargetPlatform*>& Remap)
+	PackageData.GetNumPendingCookedPlatformData() -= 1;
+	check(PackageData.GetNumPendingCookedPlatformData() >= 0);
+	if (CancelManager)
 	{
-		TargetPlatform = Remap[TargetPlatform];
+		CancelManager->Release(*this);
+		CancelManager = nullptr;
 	}
 
+	Object = nullptr;
+	bHasReleased = true;
+}
+
+void FPendingCookedPlatformData::RemapTargetPlatforms(const TMap<ITargetPlatform*, ITargetPlatform*>& Remap)
+{
+	TargetPlatform = Remap[TargetPlatform];
+}
+
 
-	//////////////////////////////////////////////////////////////////////////
-	// FPendingCookedPlatformDataCancelManager
+//////////////////////////////////////////////////////////////////////////
+// FPendingCookedPlatformDataCancelManager
 
 
-	void FPendingCookedPlatformDataCancelManager::Release(FPendingCookedPlatformData& Data)
+void FPendingCookedPlatformDataCancelManager::Release(FPendingCookedPlatformData& Data)
+{
+	--NumPendingPlatforms;
+	if (NumPendingPlatforms <= 0)
 	{
-		--NumPendingPlatforms;
-		if (NumPendingPlatforms <= 0)
+		check(NumPendingPlatforms == 0);
+		UObject* LocalObject = Data.Object.Get();
+		if (LocalObject)
 		{
-			check(NumPendingPlatforms == 0);
-			UObject* LocalObject = Data.Object.Get();
-			if (LocalObject)
-			{
-				LocalObject->ClearAllCachedCookedPlatformData();
-			}
-			delete this;
+			LocalObject->ClearAllCachedCookedPlatformData();
 		}
+		delete this;
 	}
+}
 
 
-	//////////////////////////////////////////////////////////////////////////
-	// FPackageDataMonitor
-	FPackageDataMonitor::FPackageDataMonitor()
-	{
-		FMemory::Memset(NumUrgentInState, 0);
-	}
+//////////////////////////////////////////////////////////////////////////
+// FPackageDataMonitor
+FPackageDataMonitor::FPackageDataMonitor()
+{
+	FMemory::Memset(NumUrgentInState, 0);
+}
 
-	int32 FPackageDataMonitor::GetNumUrgent() const
+int32 FPackageDataMonitor::GetNumUrgent() const
+{
+	int32 NumUrgent = 0;
+	for (EPackageState State = EPackageState::Min;
+		State <= EPackageState::Max;
+		State = static_cast<EPackageState>(static_cast<uint32>(State) + 1))
 	{
-		int32 NumUrgent = 0;
-		for (EPackageState State = EPackageState::Min;
-			State <= EPackageState::Max;
-			State = static_cast<EPackageState>(static_cast<uint32>(State) + 1))
-		{
-			NumUrgent += NumUrgentInState[static_cast<uint32>(State) - static_cast<uint32>(EPackageState::Min)];
-		}
-		return NumUrgent;
+		NumUrgent += NumUrgentInState[static_cast<uint32>(State) - static_cast<uint32>(EPackageState::Min)];
 	}
+	return NumUrgent;
+}
 
-	int32 FPackageDataMonitor::GetNumUrgent(EPackageState InState) const
-	{
-		check(EPackageState::Min <= InState && InState <= EPackageState::Max);
-		return NumUrgentInState[static_cast<uint32>(InState) - static_cast<uint32>(EPackageState::Min)];
-	}
+int32 FPackageDataMonitor::GetNumUrgent(EPackageState InState) const
+{
+	check(EPackageState::Min <= InState && InState <= EPackageState::Max);
+	return NumUrgentInState[static_cast<uint32>(InState) - static_cast<uint32>(EPackageState::Min)];
+}
 
-	int32 FPackageDataMonitor::GetNumPreloadAllocated() const
-	{
-		return NumPreloadAllocated;
-	}
+int32 FPackageDataMonitor::GetNumPreloadAllocated() const
+{
+	return NumPreloadAllocated;
+}
 
-	int32 FPackageDataMonitor::GetNumInProgress() const
-	{
-		return NumInProgress;
-	}
+int32 FPackageDataMonitor::GetNumInProgress() const
+{
+	return NumInProgress;
+}
 
-	int32 FPackageDataMonitor::GetNumCooked() const
-	{
-		return NumCooked;
-	}
+int32 FPackageDataMonitor::GetNumCooked() const
+{
+	return NumCooked;
+}
 
-	void FPackageDataMonitor::OnInProgressChanged(FPackageData& PackageData, bool bInProgress)
-	{
-		NumInProgress += bInProgress ? 1 : -1;
-		check(NumInProgress >= 0);
-	}
+void FPackageDataMonitor::OnInProgressChanged(FPackageData& PackageData, bool bInProgress)
+{
+	NumInProgress += bInProgress ? 1 : -1;
+	check(NumInProgress >= 0);
+}
 
-	void FPackageDataMonitor::OnPreloadAllocatedChanged(FPackageData& PackageData, bool bPreloadAllocated)
-	{
-		NumPreloadAllocated += bPreloadAllocated ? 1 : -1;
-		check(NumPreloadAllocated >= 0);
-	}
+void FPackageDataMonitor::OnPreloadAllocatedChanged(FPackageData& PackageData, bool bPreloadAllocated)
+{
+	NumPreloadAllocated += bPreloadAllocated ? 1 : -1;
+	check(NumPreloadAllocated >= 0);
+}
 
-	void FPackageDataMonitor::OnFirstCookedPlatformAdded(FPackageData& PackageData)
+void FPackageDataMonitor::OnFirstCookedPlatformAdded(FPackageData& PackageData)
+{
+	if (!PackageData.GetMonitorIsCooked())
 	{
-		if (!PackageData.GetMonitorIsCooked())
-		{
-			++NumCooked;
-			PackageData.SetMonitorIsCooked(true);
-		}
+		++NumCooked;
+		PackageData.SetMonitorIsCooked(true);
 	}
+}
 
-	void FPackageDataMonitor::OnLastCookedPlatformRemoved(FPackageData& PackageData)
+void FPackageDataMonitor::OnLastCookedPlatformRemoved(FPackageData& PackageData)
+{
+	if (PackageData.GetMonitorIsCooked())
 	{
-		if (PackageData.GetMonitorIsCooked())
-		{
-			--NumCooked;
-			PackageData.SetMonitorIsCooked(false);
-		}
+		--NumCooked;
+		PackageData.SetMonitorIsCooked(false);
 	}
+}
 
-	void FPackageDataMonitor::OnUrgencyChanged(FPackageData& PackageData)
-	{
-		int32 Delta = PackageData.GetIsUrgent() ? 1 : -1;
-		TrackUrgentRequests(PackageData.GetState(), Delta);
-	}
+void FPackageDataMonitor::OnUrgencyChanged(FPackageData& PackageData)
+{
+	int32 Delta = PackageData.GetIsUrgent() ? 1 : -1;
+	TrackUrgentRequests(PackageData.GetState(), Delta);
+}
 
-	void FPackageDataMonitor::OnStateChanged(FPackageData& PackageData, EPackageState OldState)
+void FPackageDataMonitor::OnStateChanged(FPackageData& PackageData, EPackageState OldState)
+{
+	if (!PackageData.GetIsUrgent())
 	{
-		if (!PackageData.GetIsUrgent())
-		{
-			return;
-		}
-
-		TrackUrgentRequests(OldState, -1);
-		TrackUrgentRequests(PackageData.GetState(), 1);
+		return;
 	}
 
-	void FPackageDataMonitor::TrackUrgentRequests(EPackageState State, int32 Delta)
-	{
-		check(EPackageState::Min <= State && State <= EPackageState::Max);
-		NumUrgentInState[static_cast<uint32>(State) - static_cast<uint32>(EPackageState::Min)] += Delta;
-		check(NumUrgentInState[static_cast<uint32>(State) - static_cast<uint32>(EPackageState::Min)] >= 0);
-	}
+	TrackUrgentRequests(OldState, -1);
+	TrackUrgentRequests(PackageData.GetState(), 1);
+}
 
+void FPackageDataMonitor::TrackUrgentRequests(EPackageState State, int32 Delta)
+{
+	check(EPackageState::Min <= State && State <= EPackageState::Max);
+	NumUrgentInState[static_cast<uint32>(State) - static_cast<uint32>(EPackageState::Min)] += Delta;
+	check(NumUrgentInState[static_cast<uint32>(State) - static_cast<uint32>(EPackageState::Min)] >= 0);
+}
 
-	//////////////////////////////////////////////////////////////////////////
-	// FPackageDatas
 
+//////////////////////////////////////////////////////////////////////////
+// FPackageDatas
 
-	FPackageDatas::FPackageDatas(UCookOnTheFlyServer& InCookOnTheFlyServer)
-		: CookOnTheFlyServer(InCookOnTheFlyServer)
-		, LastPollAsyncTime(0)
-	{
-	}
 
-	void FPackageDatas::BeginCook()
+FPackageDatas::FPackageDatas(UCookOnTheFlyServer& InCookOnTheFlyServer)
+	: CookOnTheFlyServer(InCookOnTheFlyServer)
+	, LastPollAsyncTime(0)
+{
+}
+
+void FPackageDatas::BeginCook()
+{
+	FString FileOrPackageName;
+	ShowInstigatorPackageData = nullptr;
+	if (FParse::Value(FCommandLine::Get(), TEXT("-CookShowInstigator="), FileOrPackageName))
 	{
-		FString FileOrPackageName;
-		ShowInstigatorPackageData = nullptr;
-		if (FParse::Value(FCommandLine::Get(), TEXT("-CookShowInstigator="), FileOrPackageName))
+		FString LocalPath;
+		FString PackageName;
+		if (!FPackageName::TryConvertToMountedPath(FileOrPackageName, &LocalPath, &PackageName, nullptr, nullptr, nullptr))
 		{
-			FString LocalPath;
-			FString PackageName;
-			if (!FPackageName::TryConvertToMountedPath(FileOrPackageName, &LocalPath, &PackageName, nullptr, nullptr, nullptr))
+			UE_LOG(LogCook, Fatal, TEXT("-CookShowInstigator argument %s is not a mounted filename or packagename"),
+				*FileOrPackageName);
+		}
+		else
+		{
+			FName PackageFName(*PackageName);
+			FName FileName = PackageNameCache.GetCachedStandardFileName(PackageFName);
+			if (FileName.IsNone())
 			{
-				UE_LOG(LogCook, Fatal, TEXT("-CookShowInstigator argument %s is not a mounted filename or packagename"),
+				UE_LOG(LogCook, Fatal, TEXT("-CookShowInstigator argument %s could not be found on disk"),
 					*FileOrPackageName);
 			}
 			else
 			{
-				FName PackageFName(*PackageName);
-				FName FileName = PackageNameCache.GetCachedStandardFileName(PackageFName);
-				if (FileName.IsNone())
-				{
-					UE_LOG(LogCook, Fatal, TEXT("-CookShowInstigator argument %s could not be found on disk"),
-						*FileOrPackageName);
-				}
-				else
-				{
-					ShowInstigatorPackageData = &FindOrAddPackageData(PackageFName, FileName);
-				}
+				ShowInstigatorPackageData = &FindOrAddPackageData(PackageFName, FileName);
 			}
 		}
 	}
+}
 
-	FPackageDatas::~FPackageDatas()
-	{
-		Clear();
-	}
+FPackageDatas::~FPackageDatas()
+{
+	Clear();
+}
+
+FString FPackageDatas::GetReferencerName() const
+{
+	return TEXT("FPackageDatas");
+}
+
+void FPackageDatas::AddReferencedObjects(FReferenceCollector& Collector)
+{
+	return CookOnTheFlyServer.CookerAddReferencedObjects(Collector);
+}
+
+const FPackageNameCache& FPackageDatas::GetPackageNameCache() const
+{
+	return PackageNameCache;
+}
+
+FPackageDataMonitor& FPackageDatas::GetMonitor()
+{
+	return Monitor;
+}
+
+UCookOnTheFlyServer& FPackageDatas::GetCookOnTheFlyServer()
+{
+	return CookOnTheFlyServer;
+}
+
+FRequestQueue& FPackageDatas::GetRequestQueue()
+{
+	return RequestQueue;
+}
+
+FPackageDataQueue& FPackageDatas::GetSaveQueue()
+{
+	return SaveQueue;
+}
 
-	FString FPackageDatas::GetReferencerName() const
+FPackageData& FPackageDatas::FindOrAddPackageData(const FName& PackageName, const FName& NormalizedFileName)
+{
+	FPackageData** PackageDataMapAddr = PackageNameToPackageData.Find(PackageName);
+	if (PackageDataMapAddr != nullptr)
 	{
-		return TEXT("FPackageDatas");
+		FPackageData** FileNameMapAddr = FileNameToPackageData.Find(NormalizedFileName);
+		check(FileNameMapAddr && *FileNameMapAddr == *PackageDataMapAddr);
+		return **PackageDataMapAddr;
 	}
 
-	void FPackageDatas::AddReferencedObjects(FReferenceCollector& Collector)
+	checkf(FileNameToPackageData.Find(NormalizedFileName) == nullptr,
+		TEXT("Package \"%s\" and package \"%s\" share the same filename \"%s\"."),
+		*PackageName.ToString(), *(*FileNameToPackageData.Find(NormalizedFileName))->GetPackageName().ToString(),
+		*NormalizedFileName.ToString());
+	return CreatePackageData(PackageName, NormalizedFileName);
+}
+
+FPackageData* FPackageDatas::FindPackageDataByPackageName(const FName& PackageName)
+{
+	if (PackageName.IsNone())
 	{
-		return CookOnTheFlyServer.CookerAddReferencedObjects(Collector);
+		return nullptr;
 	}
 
-	const FPackageNameCache& FPackageDatas::GetPackageNameCache() const
+	FPackageData** PackageDataMapAddr = PackageNameToPackageData.Find(PackageName);
+	return PackageDataMapAddr ? *PackageDataMapAddr : nullptr;
+}
+
+FPackageData* FPackageDatas::TryAddPackageDataByPackageName(const FName& PackageName)
+{
+	if (PackageName.IsNone())
 	{
-		return PackageNameCache;
+		return nullptr;
 	}
 
-	FPackageDataMonitor& FPackageDatas::GetMonitor()
+	FPackageData** PackageDataMapAddr = PackageNameToPackageData.Find(PackageName);
+	if (PackageDataMapAddr != nullptr)
 	{
-		return Monitor;
+		return *PackageDataMapAddr;
 	}
 
-	UCookOnTheFlyServer& FPackageDatas::GetCookOnTheFlyServer()
+	FName FileName = PackageNameCache.GetCachedStandardFileName(PackageName);
+	if (FileName.IsNone())
 	{
-		return CookOnTheFlyServer;
+		// This can happen if PackageName is a script package
+		return nullptr;
 	}
+	checkf(FileNameToPackageData.Find(FileName) == nullptr,
+		TEXT("Package \"%s\" and package \"%s\" share the same filename \"%s\"."),
+		*PackageName.ToString(), *(*FileNameToPackageData.Find(FileName))->GetPackageName().ToString(),
+		*FileName.ToString()); 
+	return &CreatePackageData(PackageName, FileName);
+}
+
+FPackageData& FPackageDatas::AddPackageDataByPackageNameChecked(const FName& PackageName)
+{
+	FPackageData* PackageData = TryAddPackageDataByPackageName(PackageName);
+	check(PackageData);
+	return *PackageData;
+}
 
-	FRequestQueue& FPackageDatas::GetRequestQueue()
+FPackageData* FPackageDatas::FindPackageDataByFileName(const FName& InFileName)
+{
+	FName FileName(FPackageNameCache::GetStandardFileName(InFileName));
+	if (FileName.IsNone())
 	{
-		return RequestQueue;
+		return nullptr;
 	}
 
-	FPackageDataQueue& FPackageDatas::GetSaveQueue()
+	FPackageData** PackageDataMapAddr = FileNameToPackageData.Find(FileName);
+	return PackageDataMapAddr ? *PackageDataMapAddr : nullptr;
+}
+
+FPackageData* FPackageDatas::TryAddPackageDataByFileName(const FName& InFileName)
+{
+	FName FileName(FPackageNameCache::GetStandardFileName(InFileName));
+	if (FileName.IsNone())
 	{
-		return SaveQueue;
+		return nullptr;
 	}
 
-	FPackageData& FPackageDatas::FindOrAddPackageData(const FName& PackageName, const FName& NormalizedFileName)
+	FPackageData** PackageDataMapAddr = FileNameToPackageData.Find(FileName);
+	if (PackageDataMapAddr != nullptr)
 	{
-		FPackageData** PackageDataMapAddr = PackageNameToPackageData.Find(PackageName);
-		if (PackageDataMapAddr != nullptr)
-		{
-			FPackageData** FileNameMapAddr = FileNameToPackageData.Find(NormalizedFileName);
-			check(FileNameMapAddr && *FileNameMapAddr == *PackageDataMapAddr);
-			return **PackageDataMapAddr;
-		}
-
-		checkf(FileNameToPackageData.Find(NormalizedFileName) == nullptr,
-			TEXT("Package \"%s\" and package \"%s\" share the same filename \"%s\"."),
-			*PackageName.ToString(), *(*FileNameToPackageData.Find(NormalizedFileName))->GetPackageName().ToString(),
-			*NormalizedFileName.ToString());
-		return CreatePackageData(PackageName, NormalizedFileName);
+		return *PackageDataMapAddr;
 	}
 
-	FPackageData* FPackageDatas::FindPackageDataByPackageName(const FName& PackageName)
+	const FName* PackageName = PackageNameCache.GetCachedPackageNameFromStandardFileName(FileName);
+	if (!PackageName)
 	{
-		if (PackageName.IsNone())
-		{
-			return nullptr;
-		}
-
-		FPackageData** PackageDataMapAddr = PackageNameToPackageData.Find(PackageName);
-		return PackageDataMapAddr ? *PackageDataMapAddr : nullptr;
+		return nullptr;
 	}
+	return &CreatePackageData(*PackageName, FileName);
+}
 
-	FPackageData* FPackageDatas::TryAddPackageDataByPackageName(const FName& PackageName)
+FPackageData& FPackageDatas::CreatePackageData(FName PackageName, FName FileName)
+{
+	if (PackageName.IsNone())
 	{
-		if (PackageName.IsNone())
-		{
-			return nullptr;
-		}
-
-		FPackageData** PackageDataMapAddr = PackageNameToPackageData.Find(PackageName);
-		if (PackageDataMapAddr != nullptr)
-		{
-			return *PackageDataMapAddr;
-		}
-
-		FName FileName = PackageNameCache.GetCachedStandardFileName(PackageName);
-		if (FileName.IsNone())
-		{
-			// This can happen if PackageName is a script package
-			return nullptr;
-		}
-		checkf(FileNameToPackageData.Find(FileName) == nullptr,
-			TEXT("Package \"%s\" and package \"%s\" share the same filename \"%s\"."),
-			*PackageName.ToString(), *(*FileNameToPackageData.Find(FileName))->GetPackageName().ToString(),
-			*FileName.ToString()); 
-		return &CreatePackageData(PackageName, FileName);
+		check(!FileName.IsNone());
+		const FName* FoundPackageName = PackageNameCache.GetCachedPackageNameFromStandardFileName(FileName);
+		check(FoundPackageName);
+		PackageName = *FoundPackageName;
+		check(!PackageName.IsNone());
 	}
-
-	FPackageData& FPackageDatas::AddPackageDataByPackageNameChecked(const FName& PackageName)
+	else if (FileName.IsNone())
 	{
-		FPackageData* PackageData = TryAddPackageDataByPackageName(PackageName);
-		check(PackageData);
-		return *PackageData;
+		FileName = PackageNameCache.GetCachedStandardFileName(PackageName);
+		check(!FileName.IsNone());
 	}
 
-	FPackageData* FPackageDatas::FindPackageDataByFileName(const FName& InFileName)
-	{
-		FName FileName(FPackageNameCache::GetStandardFileName(InFileName));
-		if (FileName.IsNone())
-		{
-			return nullptr;
-		}
+	FPackageData* PackageData = new FPackageData(*this, PackageName, FileName);
+	PackageDatas.Add(PackageData);
+	PackageNameToPackageData.Add(PackageName, PackageData);
+	FileNameToPackageData.Add(FileName, PackageData);
+	return *PackageData;
+}
 
-		FPackageData** PackageDataMapAddr = FileNameToPackageData.Find(FileName);
-		return PackageDataMapAddr ? *PackageDataMapAddr : nullptr;
-	}
+FPackageData& FPackageDatas::AddPackageDataByFileNameChecked(const FName& FileName)
+{
+	FPackageData* PackageData = TryAddPackageDataByFileName(FileName);
+	check(PackageData);
+	return *PackageData;
+}
+
+void FPackageDatas::AddExistingPackageDatasForPlatform(const TArray<FName>& ExistingPackages,
+	const ITargetPlatform* TargetPlatform)
+{
+	TArray<FPackageData*> NewPackageDataObjects;
+	NewPackageDataObjects.AddZeroed(ExistingPackages.Num());
 
-	FPackageData* FPackageDatas::TryAddPackageDataByFileName(const FName& InFileName)
+	// parallelize the read-only operations (and write NewPackageDataObjects by index which has no threading issues)
+	ParallelFor(ExistingPackages.Num(),
+		[&ExistingPackages, TargetPlatform, &NewPackageDataObjects, this](int Index)
 	{
-		FName FileName(FPackageNameCache::GetStandardFileName(InFileName));
+		const FName& PackageFilename = ExistingPackages[Index];
+
+		FName FileName(FPackageNameCache::GetStandardFileName(PackageFilename));
 		if (FileName.IsNone())
 		{
-			return nullptr;
+			return;
 		}
 
 		FPackageData** PackageDataMapAddr = FileNameToPackageData.Find(FileName);
+		FPackageData* PackageData = nullptr;
 		if (PackageDataMapAddr != nullptr)
 		{
-			return *PackageDataMapAddr;
+			PackageData = *PackageDataMapAddr;
 		}
-
-		const FName* PackageName = PackageNameCache.GetCachedPackageNameFromStandardFileName(FileName);
-		if (!PackageName)
+		else
 		{
-			return nullptr;
+			// we expect all of these to be in PackageNameCache from earlier in the AR load process
+			const FName* PackageName = PackageNameCache.FindExistingCachedPackageNameFromStandardFileName(FileName);
+			checkf(PackageName != nullptr, TEXT("AddExistingPackageDatasForPlatform should only be used when the PackageCache has already been filled out with the packages"));
+
+			// create the package data and remember it for updating caches after the the ParallelFor
+			PackageData = new FPackageData(*this, *PackageName, FileName);
+			NewPackageDataObjects[Index] = PackageData;
 		}
-		return &CreatePackageData(*PackageName, FileName);
-	}
+		PackageData->SetPlatformCooked(TargetPlatform, true /* Succeeded */);
+	});
 
-	FPackageData& FPackageDatas::CreatePackageData(FName PackageName, FName FileName)
+	// update cache for all newly created objects (taken from CreatePackageData)
+	for (FPackageData* PackageData : NewPackageDataObjects)
 	{
-		if (PackageName.IsNone())
-		{
-			check(!FileName.IsNone());
-			const FName* FoundPackageName = PackageNameCache.GetCachedPackageNameFromStandardFileName(FileName);
-			check(FoundPackageName);
-			PackageName = *FoundPackageName;
-			check(!PackageName.IsNone());
-		}
-		else if (FileName.IsNone())
+		if (PackageData)
 		{
-			FileName = PackageNameCache.GetCachedStandardFileName(PackageName);
-			check(!FileName.IsNone());
+			PackageDatas.Add(PackageData);
+			PackageNameToPackageData.Add(PackageData->PackageName, PackageData);
+			FileNameToPackageData.Add(PackageData->FileName, PackageData);
 		}
-
-		FPackageData* PackageData = new FPackageData(*this, PackageName, FileName);
-		PackageDatas.Add(PackageData);
-		PackageNameToPackageData.Add(PackageName, PackageData);
-		FileNameToPackageData.Add(FileName, PackageData);
-		return *PackageData;
 	}
+}
 
-	FPackageData& FPackageDatas::AddPackageDataByFileNameChecked(const FName& FileName)
+FPackageData* FPackageDatas::UpdateFileName(const FName& PackageName)
+{
+	if (!PackageNameCache.HasCacheForPackageName(PackageName))
 	{
-		FPackageData* PackageData = TryAddPackageDataByFileName(FileName);
-		check(PackageData);
-		return *PackageData;
+		return nullptr;
 	}
 
-	void FPackageDatas::AddExistingPackageDatasForPlatform(const TArray<FName>& ExistingPackages,
-		const ITargetPlatform* TargetPlatform)
-	{
-		TArray<FPackageData*> NewPackageDataObjects;
-		NewPackageDataObjects.AddZeroed(ExistingPackages.Num());
-
-		// parallelize the read-only operations (and write NewPackageDataObjects by index which has no threading issues)
-		ParallelFor(ExistingPackages.Num(),
-			[&ExistingPackages, TargetPlatform, &NewPackageDataObjects, this](int Index)
-		{
-			const FName& PackageFilename = ExistingPackages[Index];
-
-			FName FileName(FPackageNameCache::GetStandardFileName(PackageFilename));
-			if (FileName.IsNone())
-			{
-				return;
-			}
-
-			FPackageData** PackageDataMapAddr = FileNameToPackageData.Find(FileName);
-			FPackageData* PackageData = nullptr;
-			if (PackageDataMapAddr != nullptr)
-			{
-				PackageData = *PackageDataMapAddr;
-			}
-			else
-			{
-				// we expect all of these to be in PackageNameCache from earlier in the AR load process
-				const FName* PackageName = PackageNameCache.FindExistingCachedPackageNameFromStandardFileName(FileName);
-				checkf(PackageName != nullptr, TEXT("AddExistingPackageDatasForPlatform should only be used when the PackageCache has already been filled out with the packages"));
+	FName OldFileName = PackageNameCache.GetCachedStandardFileName(PackageName);
+	PackageNameCache.ClearPackageFileNameCacheForPackage(PackageName);
+	FName NewFileName = PackageNameCache.GetCachedStandardFileName(PackageName);
 
-				// create the package data and remember it for updating caches after the the ParallelFor
-				PackageData = new FPackageData(*this, *PackageName, FileName);
-				NewPackageDataObjects[Index] = PackageData;
-			}
-			PackageData->SetPlatformCooked(TargetPlatform, true /* Succeeded */);
-		});
-
-		// update cache for all newly created objects (taken from CreatePackageData)
-		for (FPackageData* PackageData : NewPackageDataObjects)
-		{
-			if (PackageData)
-			{
-				PackageDatas.Add(PackageData);
-				PackageNameToPackageData.Add(PackageData->PackageName, PackageData);
-				FileNameToPackageData.Add(PackageData->FileName, PackageData);
-			}
-		}
+	FPackageData** PackageDataAddr = PackageNameToPackageData.Find(PackageName);
+	if (!PackageDataAddr)
+	{
+		check(OldFileName.IsNone() || !FileNameToPackageData.Find(OldFileName));
+		return nullptr;
 	}
+	FPackageData* PackageData = *PackageDataAddr;
 
-	FPackageData* FPackageDatas::UpdateFileName(const FName& PackageName)
+	if (OldFileName == NewFileName)
 	{
-		if (!PackageNameCache.HasCacheForPackageName(PackageName))
-		{
-			return nullptr;
-		}
-
-		FName OldFileName = PackageNameCache.GetCachedStandardFileName(PackageName);
-		PackageNameCache.ClearPackageFileNameCacheForPackage(PackageName);
-		FName NewFileName = PackageNameCache.GetCachedStandardFileName(PackageName);
-
-		FPackageData** PackageDataAddr = PackageNameToPackageData.Find(PackageName);
-		if (!PackageDataAddr)
-		{
-			check(OldFileName.IsNone() || !FileNameToPackageData.Find(OldFileName));
-			return nullptr;
-		}
-		FPackageData* PackageData = *PackageDataAddr;
-
-		if (OldFileName == NewFileName)
-		{
-			return PackageData;
-		}
-
-		if (!OldFileName.IsNone())
-		{
-			PackageDataAddr = FileNameToPackageData.Find(OldFileName);
-			check(PackageDataAddr == nullptr || *PackageDataAddr == PackageData);
-			FileNameToPackageData.Remove(OldFileName);
-		}
-		PackageData->SetFileName(NewFileName);
-		if (!NewFileName.IsNone())
-		{
-			check(FileNameToPackageData.Find(NewFileName) == nullptr);
-			FileNameToPackageData.Add(NewFileName, PackageData);
-		}
-
 		return PackageData;
 	}
 
-	void FPackageDatas::RegisterFileNameAlias(FPackageData& PackageData, FName FileName)
+	if (!OldFileName.IsNone())
 	{
-		FileName = FPackageNameCache::GetStandardFileName(FileName);
-		if (FileName.IsNone())
-		{
-			return;
-		}
-
-		FPackageData*& PackageDataMapAddr = FileNameToPackageData.FindOrAdd(FileName);
-		check(PackageDataMapAddr == nullptr || PackageDataMapAddr == &PackageData);
-		PackageDataMapAddr = &PackageData;
+		PackageDataAddr = FileNameToPackageData.Find(OldFileName);
+		check(PackageDataAddr == nullptr || *PackageDataAddr == PackageData);
+		FileNameToPackageData.Remove(OldFileName);
 	}
-
-	int32 FPackageDatas::GetNumCooked()
+	PackageData->SetFileName(NewFileName);
+	if (!NewFileName.IsNone())
 	{
-		return Monitor.GetNumCooked();
+		check(FileNameToPackageData.Find(NewFileName) == nullptr);
+		FileNameToPackageData.Add(NewFileName, PackageData);
 	}
 
-	void FPackageDatas::GetCookedFileNamesForPlatform(const ITargetPlatform* Platform, TArray<FName>& CookedFiles,
-		bool bGetFailedCookedPackages, bool bGetSuccessfulCookedPackages)
+	return PackageData;
+}
+
+void FPackageDatas::RegisterFileNameAlias(FPackageData& PackageData, FName FileName)
+{
+	FileName = FPackageNameCache::GetStandardFileName(FileName);
+	if (FileName.IsNone())
 	{
-		for (const FPackageData* PackageData : PackageDatas)
-		{
-			ECookResult CookResults = PackageData->GetCookResults(Platform);
-			if (((CookResults == ECookResult::Succeeded) & (bGetSuccessfulCookedPackages != 0)) |
-				((CookResults == ECookResult::Failed) & (bGetFailedCookedPackages != 0)))
-			{
-				CookedFiles.Add(PackageData->GetFileName());
-			}
-		}
+		return;
 	}
 
-	void FPackageDatas::Clear()
+	FPackageData*& PackageDataMapAddr = FileNameToPackageData.FindOrAdd(FileName);
+	check(PackageDataMapAddr == nullptr || PackageDataMapAddr == &PackageData);
+	PackageDataMapAddr = &PackageData;
+}
+
+int32 FPackageDatas::GetNumCooked()
+{
+	return Monitor.GetNumCooked();
+}
+
+void FPackageDatas::GetCookedFileNamesForPlatform(const ITargetPlatform* Platform, TArray<FName>& CookedFiles,
+	bool bGetFailedCookedPackages, bool bGetSuccessfulCookedPackages)
+{
+	for (const FPackageData* PackageData : PackageDatas)
 	{
-		PendingCookedPlatformDatas.Empty(); // These destructors will dereference PackageDatas
-		RequestQueue.Empty();
-		SaveQueue.Empty();
-		PackageNameToPackageData.Empty();
-		FileNameToPackageData.Empty();
-		for (FPackageData* PackageData : PackageDatas)
-		{
-			PackageData->ClearReferences();
-		}
-		for (FPackageData* PackageData : PackageDatas)
+		ECookResult CookResults = PackageData->GetCookResults(Platform);
+		if (((CookResults == ECookResult::Succeeded) & (bGetSuccessfulCookedPackages != 0)) |
+			((CookResults == ECookResult::Failed) & (bGetFailedCookedPackages != 0)))
 		{
-			delete PackageData;
+			CookedFiles.Add(PackageData->GetFileName());
 		}
-		PackageDatas.Empty();
-		ShowInstigatorPackageData = nullptr;
 	}
+}
 
-	void FPackageDatas::ClearCookedPlatforms()
+void FPackageDatas::Clear()
+{
+	PendingCookedPlatformDatas.Empty(); // These destructors will dereference PackageDatas
+	RequestQueue.Empty();
+	SaveQueue.Empty();
+	PackageNameToPackageData.Empty();
+	FileNameToPackageData.Empty();
+	for (FPackageData* PackageData : PackageDatas)
 	{
-		for (FPackageData* PackageData : PackageDatas)
-		{
-			PackageData->SetPlatformsNotCooked();
-		}
+		PackageData->ClearReferences();
 	}
-
-	void FPackageDatas::OnRemoveSessionPlatform(const ITargetPlatform* TargetPlatform)
+	for (FPackageData* PackageData : PackageDatas)
 	{
-		for (FPackageData* PackageData : PackageDatas)
-		{
-			PackageData->OnRemoveSessionPlatform(TargetPlatform);
-		}
+		delete PackageData;
 	}
+	PackageDatas.Empty();
+	ShowInstigatorPackageData = nullptr;
+}
 
-	TArray<FPendingCookedPlatformData>& FPackageDatas::GetPendingCookedPlatformDatas()
+void FPackageDatas::ClearCookedPlatforms()
+{
+	for (FPackageData* PackageData : PackageDatas)
 	{
-		return PendingCookedPlatformDatas;
+		PackageData->SetPlatformsNotCooked();
 	}
+}
 
-	void FPackageDatas::PollPendingCookedPlatformDatas()
+void FPackageDatas::OnRemoveSessionPlatform(const ITargetPlatform* TargetPlatform)
+{
+	for (FPackageData* PackageData : PackageDatas)
 	{
-		if (PendingCookedPlatformDatas.Num() == 0)
-		{
-			return;
-		}
-
-		// ProcessAsyncResults and IsCachedCookedPlatformDataLoaded can be expensive to call
-		// Cap the frequency at which we call them.
-		double CurrentTime = FPlatformTime::Seconds();
-		if (CurrentTime < LastPollAsyncTime + GPollAsyncPeriod)
-		{
-			return;
-		}
-		LastPollAsyncTime = CurrentTime;
-
-		GShaderCompilingManager->ProcessAsyncResults(true /* bLimitExecutionTime */,
-			false /* bBlockOnGlobalShaderCompletion */);
-		FAssetCompilingManager::Get().ProcessAsyncTasks(true);
-
-		FPendingCookedPlatformData* Datas = PendingCookedPlatformDatas.GetData();
-		for (int Index = 0; Index < PendingCookedPlatformDatas.Num();)
-		{
-			if (Datas[Index].PollIsComplete())
-			{
-				PendingCookedPlatformDatas.RemoveAtSwap(Index, 1 /* Count */, false /* bAllowShrinking */);
-			}
-			else
-			{
-				++Index;
-			}
-		}
+		PackageData->OnRemoveSessionPlatform(TargetPlatform);
 	}
+}
+
+TArray<FPendingCookedPlatformData>& FPackageDatas::GetPendingCookedPlatformDatas()
+{
+	return PendingCookedPlatformDatas;
+}
 
-	TArray<FPackageData*>::RangedForIteratorType FPackageDatas::begin()
+void FPackageDatas::PollPendingCookedPlatformDatas()
+{
+	if (PendingCookedPlatformDatas.Num() == 0)
 	{
-		return PackageDatas.begin();
+		return;
 	}
 
-	TArray<FPackageData*>::RangedForIteratorType FPackageDatas::end()
+	// ProcessAsyncResults and IsCachedCookedPlatformDataLoaded can be expensive to call
+	// Cap the frequency at which we call them.
+	double CurrentTime = FPlatformTime::Seconds();
+	if (CurrentTime < LastPollAsyncTime + GPollAsyncPeriod)
 	{
-		return PackageDatas.end();
+		return;
 	}
+	LastPollAsyncTime = CurrentTime;
 
-	void FPackageDatas::RemapTargetPlatforms(const TMap<ITargetPlatform*, ITargetPlatform*>& Remap)
+	GShaderCompilingManager->ProcessAsyncResults(true /* bLimitExecutionTime */,
+		false /* bBlockOnGlobalShaderCompletion */);
+	FAssetCompilingManager::Get().ProcessAsyncTasks(true);
+
+	FPendingCookedPlatformData* Datas = PendingCookedPlatformDatas.GetData();
+	for (int Index = 0; Index < PendingCookedPlatformDatas.Num();)
 	{
-		for (FPackageData* PackageData : PackageDatas)
+		if (Datas[Index].PollIsComplete())
 		{
-			PackageData->RemapTargetPlatforms(Remap);
+			PendingCookedPlatformDatas.RemoveAtSwap(Index, 1 /* Count */, false /* bAllowShrinking */);
 		}
-		for (FPendingCookedPlatformData& CookedPlatformData : PendingCookedPlatformDatas)
+		else
 		{
-			CookedPlatformData.RemapTargetPlatforms(Remap);
+			++Index;
 		}
 	}
+}
 
-	void FPackageDatas::DebugInstigator(FPackageData& PackageData)
-	{
-		if (ShowInstigatorPackageData != &PackageData)
-		{
-			return;
-		}
+TArray<FPackageData*>::RangedForIteratorType FPackageDatas::begin()
+{
+	return PackageDatas.begin();
+}
 
-		TArray<FInstigator> Chain = CookOnTheFlyServer.GetInstigatorChain(PackageData.GetPackageName());
-		TStringBuilder<256> ChainText;
-		if (Chain.Num() == 0)
-		{
-			ChainText << TEXT("<NoInstigator>");
-		}
-		bool bFirst = true;
-		for (FInstigator& Instigator : Chain)
-		{
-			if (!bFirst) ChainText << TEXT(" <- ");
-			ChainText << TEXT("{ ") << Instigator.ToString() << TEXT(" }");
-			bFirst = false;
-		}
-		UE_LOG(LogCook, Display, TEXT("Instigator chain of %s: %s"), *PackageData.GetPackageName().ToString(), ChainText.ToString());
-	}
+TArray<FPackageData*>::RangedForIteratorType FPackageDatas::end()
+{
+	return PackageDatas.end();
+}
 
-	void FRequestQueue::Empty()
+void FPackageDatas::RemapTargetPlatforms(const TMap<ITargetPlatform*, ITargetPlatform*>& Remap)
+{
+	for (FPackageData* PackageData : PackageDatas)
 	{
-		NormalRequests.Empty();
-		UrgentRequests.Empty();
+		PackageData->RemapTargetPlatforms(Remap);
 	}
-
-	bool FRequestQueue::IsEmpty() const
+	for (FPendingCookedPlatformData& CookedPlatformData : PendingCookedPlatformDatas)
 	{
-		return Num() == 0;
+		CookedPlatformData.RemapTargetPlatforms(Remap);
 	}
+}
 
-	uint32 FRequestQueue::Num() const
+void FPackageDatas::DebugInstigator(FPackageData& PackageData)
+{
+	if (ShowInstigatorPackageData != &PackageData)
 	{
-		uint32 Count = UnclusteredRequests.Num() + ReadyRequestsNum();
-		for (const FRequestCluster& RequestCluster : RequestClusters)
-		{
-			Count += RequestCluster.NumPackageDatas();
-		}
-		return Count;
+		return;
 	}
 
-	bool FRequestQueue::Contains(const FPackageData* InPackageData) const
+	TArray<FInstigator> Chain = CookOnTheFlyServer.GetInstigatorChain(PackageData.GetPackageName());
+	TStringBuilder<256> ChainText;
+	if (Chain.Num() == 0)
 	{
-		FPackageData* PackageData = const_cast<FPackageData*>(InPackageData);
-		if (UnclusteredRequests.Contains(PackageData) || NormalRequests.Contains(PackageData) || UrgentRequests.Contains(PackageData))
-		{
-			return true;
-		}
-		for (const FRequestCluster& RequestCluster : RequestClusters)
-		{
-			if (RequestCluster.Contains(PackageData))
-			{
-				return true;
-			}
-		}
-		return false;
+		ChainText << TEXT("<NoInstigator>");
 	}
-
-	uint32 FRequestQueue::RemoveRequest(FPackageData* PackageData)
+	bool bFirst = true;
+	for (FInstigator& Instigator : Chain)
 	{
-		uint32 OriginalNum = Num();
-		UnclusteredRequests.Remove(PackageData);
-		NormalRequests.Remove(PackageData);
-		UrgentRequests.Remove(PackageData);
-		for (FRequestCluster& RequestCluster : RequestClusters)
-		{
-			RequestCluster.RemovePackageData(PackageData);
-		}
-		uint32 Result = OriginalNum - Num();
-		check(Result == 0 || Result == 1);
-		return Result;
+		if (!bFirst) ChainText << TEXT(" <- ");
+		ChainText << TEXT("{ ") << Instigator.ToString() << TEXT(" }");
+		bFirst = false;
 	}
+	UE_LOG(LogCook, Display, TEXT("Instigator chain of %s: %s"), *PackageData.GetPackageName().ToString(), ChainText.ToString());
+}
 
-	uint32 FRequestQueue::Remove(FPackageData* PackageData)
-	{
-		return RemoveRequest(PackageData);
-	}
+void FRequestQueue::Empty()
+{
+	NormalRequests.Empty();
+	UrgentRequests.Empty();
+}
+
+bool FRequestQueue::IsEmpty() const
+{
+	return Num() == 0;
+}
 
-	bool FRequestQueue::IsReadyRequestsEmpty() const
+uint32 FRequestQueue::Num() const
+{
+	uint32 Count = UnclusteredRequests.Num() + ReadyRequestsNum();
+	for (const FRequestCluster& RequestCluster : RequestClusters)
 	{
-		return ReadyRequestsNum() == 0;
+		Count += RequestCluster.NumPackageDatas();
 	}
+	return Count;
+}
 
-	uint32 FRequestQueue::ReadyRequestsNum() const
+bool FRequestQueue::Contains(const FPackageData* InPackageData) const
+{
+	FPackageData* PackageData = const_cast<FPackageData*>(InPackageData);
+	if (UnclusteredRequests.Contains(PackageData) || NormalRequests.Contains(PackageData) || UrgentRequests.Contains(PackageData))
 	{
-		return UrgentRequests.Num() + NormalRequests.Num();
+		return true;
 	}
-
-	FPackageData* FRequestQueue::PopReadyRequest()
+	for (const FRequestCluster& RequestCluster : RequestClusters)
 	{
-		for (FPackageData* PackageData : UrgentRequests)
+		if (RequestCluster.Contains(PackageData))
 		{
-			UrgentRequests.Remove(PackageData);
-			return PackageData;
-		}
-		for (FPackageData* PackageData : NormalRequests)
-		{
-			NormalRequests.Remove(PackageData);
-			return PackageData;
+			return true;
 		}
-		return nullptr;
 	}
+	return false;
+}
 
-	void FRequestQueue::AddRequest(FPackageData* PackageData, bool bForceUrgent)
+uint32 FRequestQueue::RemoveRequest(FPackageData* PackageData)
+{
+	uint32 OriginalNum = Num();
+	UnclusteredRequests.Remove(PackageData);
+	NormalRequests.Remove(PackageData);
+	UrgentRequests.Remove(PackageData);
+	for (FRequestCluster& RequestCluster : RequestClusters)
 	{
-		if (!PackageData->AreAllRequestedPlatformsExplored())
-		{
-			UnclusteredRequests.Add(PackageData);
-		}
-		else
-		{
-			AddReadyRequest(PackageData, bForceUrgent);
-		}
+		RequestCluster.RemovePackageData(PackageData);
 	}
+	uint32 Result = OriginalNum - Num();
+	check(Result == 0 || Result == 1);
+	return Result;
+}
+
+uint32 FRequestQueue::Remove(FPackageData* PackageData)
+{
+	return RemoveRequest(PackageData);
+}
+
+bool FRequestQueue::IsReadyRequestsEmpty() const
+{
+	return ReadyRequestsNum() == 0;
+}
+
+uint32 FRequestQueue::ReadyRequestsNum() const
+{
+	return UrgentRequests.Num() + NormalRequests.Num();
+}
 
-	void FRequestQueue::AddReadyRequest(FPackageData* PackageData, bool bForceUrgent)
+FPackageData* FRequestQueue::PopReadyRequest()
+{
+	for (FPackageData* PackageData : UrgentRequests)
 	{
-		if (bForceUrgent || PackageData->GetIsUrgent())
-		{
-			UrgentRequests.Add(PackageData);
-		}
-		else
-		{
-			NormalRequests.Add(PackageData);
-		}
+		UrgentRequests.Remove(PackageData);
+		return PackageData;
 	}
-
-	bool FLoadPrepareQueue::IsEmpty()
+	for (FPackageData* PackageData : NormalRequests)
 	{
-		return Num() == 0;
+		NormalRequests.Remove(PackageData);
+		return PackageData;
 	}
+	return nullptr;
+}
 
-	int32 FLoadPrepareQueue::Num() const
+void FRequestQueue::AddRequest(FPackageData* PackageData, bool bForceUrgent)
+{
+	if (!PackageData->AreAllRequestedPlatformsExplored())
 	{
-		return PreloadingQueue.Num() + EntryQueue.Num();
+		UnclusteredRequests.Add(PackageData);
 	}
-
-	FPackageData* FLoadPrepareQueue::PopFront()
+	else
 	{
-		if (!PreloadingQueue.IsEmpty())
-		{
-			return PreloadingQueue.PopFrontValue();
-		}
-		else
-		{
-			return EntryQueue.PopFrontValue();
-		}
+		AddReadyRequest(PackageData, bForceUrgent);
 	}
+}
 
-	void FLoadPrepareQueue::Add(FPackageData* PackageData)
+void FRequestQueue::AddReadyRequest(FPackageData* PackageData, bool bForceUrgent)
+{
+	if (bForceUrgent || PackageData->GetIsUrgent())
 	{
-		EntryQueue.Add(PackageData);
+		UrgentRequests.Add(PackageData);
 	}
-
-	void FLoadPrepareQueue::AddFront(FPackageData* PackageData)
+	else
 	{
-		PreloadingQueue.AddFront(PackageData);
+		NormalRequests.Add(PackageData);
 	}
+}
+
+bool FLoadPrepareQueue::IsEmpty()
+{
+	return Num() == 0;
+}
+
+int32 FLoadPrepareQueue::Num() const
+{
+	return PreloadingQueue.Num() + EntryQueue.Num();
+}
 
-	bool FLoadPrepareQueue::Contains(const FPackageData* PackageData) const
+FPackageData* FLoadPrepareQueue::PopFront()
+{
+	if (!PreloadingQueue.IsEmpty())
 	{
-		return (Algo::Find(PreloadingQueue, PackageData) != nullptr) ||
-			(Algo::Find(EntryQueue, PackageData) != nullptr);
+		return PreloadingQueue.PopFrontValue();
 	}
-
-	uint32 FLoadPrepareQueue::Remove(FPackageData* PackageData)
+	else
 	{
-		return PreloadingQueue.Remove(PackageData) + EntryQueue.Remove(PackageData);
+		return EntryQueue.PopFrontValue();
 	}
+}
+
+void FLoadPrepareQueue::Add(FPackageData* PackageData)
+{
+	EntryQueue.Add(PackageData);
+}
+
+void FLoadPrepareQueue::AddFront(FPackageData* PackageData)
+{
+	PreloadingQueue.AddFront(PackageData);
+}
+
+bool FLoadPrepareQueue::Contains(const FPackageData* PackageData) const
+{
+	return (Algo::Find(PreloadingQueue, PackageData) != nullptr) ||
+		(Algo::Find(EntryQueue, PackageData) != nullptr);
+}
 
-	FPoppedPackageDataScope::FPoppedPackageDataScope(FPackageData& InPackageData)
+uint32 FLoadPrepareQueue::Remove(FPackageData* PackageData)
+{
+	return PreloadingQueue.Remove(PackageData) + EntryQueue.Remove(PackageData);
+}
+
+FPoppedPackageDataScope::FPoppedPackageDataScope(FPackageData& InPackageData)
 #if COOK_CHECKSLOW_PACKAGEDATA
-		: PackageData(InPackageData)
+	: PackageData(InPackageData)
 #endif
-	{
-	}
+{
+}
 
 #if COOK_CHECKSLOW_PACKAGEDATA
-	FPoppedPackageDataScope::~FPoppedPackageDataScope()
-	{
-		PackageData.CheckInContainer();
-	}
+FPoppedPackageDataScope::~FPoppedPackageDataScope()
+{
+	PackageData.CheckInContainer();
+}
 #endif
 
 }
\ No newline at end of file
