CookPackageData: Fix whitespace - the block within a namespace should not be indented.
#rnx

diff --git a/Engine/Source/Editor/UnrealEd/Private/Cooker/CookPackageData.h b/Engine/Source/Editor/UnrealEd/Private/Cooker/CookPackageData.h
--- a/Engine/Source/Editor/UnrealEd/Private/Cooker/CookPackageData.h
+++ b/Engine/Source/Editor/UnrealEd/Private/Cooker/CookPackageData.h
@@ -32,1000 +32,1002 @@ class UPackage;
 
 namespace UE::Cook
 {
-	class FPackageDataQueue;
-	class FRequestCluster;
-	struct FGeneratorPackage;
-	struct FPackageDataMonitor;
-	struct FPendingCookedPlatformDataCancelManager;
 
-	extern const TCHAR* GeneratedPackageSubPath;
+class FPackageDataQueue;
+class FRequestCluster;
+struct FGeneratorPackage;
+struct FPackageDataMonitor;
+struct FPendingCookedPlatformDataCancelManager;
 
-	/** Flags specifying the behavior of FPackageData::SendToState */
-	enum class ESendFlags
+extern const TCHAR* GeneratedPackageSubPath;
+
+/** Flags specifying the behavior of FPackageData::SendToState */
+enum class ESendFlags
+{
+	/**
+	 * PackageData will not be removed from queue for its old state and will not be added to queue for its new state.
+	 * Caller is responsible for remove and add.
+	 */
+	QueueNone = 0x0,
+	/**
+	 * PackageData will be removed from the queue of its old state.
+	 * If this flag is missing, caller is responsible for removing from the old state's queue.
+	 */
+	QueueRemove = 0x1,
+	/**
+	 * PackageData will be added to queue for its next state.
+	 * If this flag is missing, caller is responsible for adding to queue.
+	 */
+	QueueAdd = 0x2,
+	/**
+	 * PackageData will be removed from the queue of its old state and added to the queue of its new state.
+	 * This may be wasteful, if the caller can add or remove more efficiently.
+	 */
+	QueueAddAndRemove = QueueAdd | QueueRemove,
+};
+ENUM_CLASS_FLAGS(ESendFlags);
+
+/**
+ * Contains all the information the cooker uses for a package, during request, load, or save.  Once allocated, this
+ * structure is never deallocated or moved for a given package; it is deallocated only when the CookOnTheFlyServer
+ * is destroyed.
+ *
+ * Requested platforms - Requested platforms are the Platforms that will be Saved during cooking.
+ *   Platforms are only requested on the PackageData when the Package is InProgress (e.g. is in the RequestQueue,
+ *   LoadReadyQueue, or other state containers). Once a Package finishes cooking, the platforms are set to
+ *   unrequested, and may be set again later if the Package is requested for another platform. If multiple requests
+ *   occur at the same time, their requested platforms are merged. Since modifying the requested platforms can put
+ *   an InProgress package into an invalid state, direct write access is private; use UpdateRequestData or
+ *   SetRequestData to write.
+
+ * Cooked platforms - Cooked platforms are the platforms that have already been saved in the lifetime of the 
+ *   CookOnTheFlyServer; note this extends outside of the current CookOnTheFlyServer session.
+ *   Cooked platforms also store a flag indicating whether the cook for the given platform was successful or not.
+ *   Cooked platforms can be added to a PackageData for reasons other than normal Save, such as when a Package is
+ *   detected as not applicable for a Platform and its "Cooked" operation is therefore a noop.
+ *   Cooked platforms can be cleared for a PackageData if the Package is modified, e.g. during editor operations
+ *   when the CookOnTheFlyServer is run from the editor.
+ *
+ * Package - The package pointer corresponding to this PackageData.
+ *   Contract for the lifetime of this field is that it is only set after the Package has passed through the load
+ *   state, and it is cleared when the package returns to idle or is demoted to an earlier state.
+ *   At all other times it is nullptr.
+ *
+ * State - the state of this PackageData in the CookOnTheFlyServer's current session. See the definition of
+ *   EPackageState for a description of each state. Contract for the value of this field includes membership in the
+ *   corresponding queue such as SaveQueue, and the presence or absence of state-specific data. Since modifying the
+ *   State can put the Package into an invalid state, direct write access is private; SendToState handles enforcing
+ *   the contracts.
+
+ * PendingCookedPlatformData - a counter for how many objects in the package have had
+ *   BeginCacheForCookedPlatformData called and have not yet completed. This is used to block the package from
+ *   saving until all objects have finished their cache. It is also used to block the package from starting new
+ *   BeginCacheForCookedPlatformData calls until all pending calls from a previous canceled save have completed.
+ *   The lifetime of this counter extends for the lifetime of the PackageData; it is shared across
+ *   CookOnTheFlyServer sessions.
+ *
+ * CookedPlatformDataNextIndex - Index for the next (Object, RequestedPlatform) pair that needs to have
+ *   BeginCacheForCookedPlatformData called on it for the current PackageSave. This field is only non-zero during
+ *   the save state; it is cleared when successfully or unsucessfully leaving the save state. The field is an
+ *   object-major index into the two-dimensional array given by
+ *       {Each Object in GetCachedObjectsInOuter} x {each Platform that is requested}
+ *   e.g. the index into GetCachedObjectsInOuter is GetCookedPlatformDataNextIndex() / GetNumRequestedPlatforms()
+ *   and the index into GetRequestedPlatforms() is GetCookedPlatformDataNextIndex() % GetNumRequestedPlatforms()
+ *
+ * Other fields with explanation inline
+*/
+struct FPackageData
+{
+public:
+	/** Data about a platform that has been interacted with (requested, etc) by the cook of a PackageData. */
+	struct FPlatformData
 	{
-		/*
-		 * PackageData will not be removed from queue for its old state and will not be added to queue for its new state.
-		 * Caller is responsible for remove and add.
-		 */
-		QueueNone = 0x0,
-		/*
-		 * PackageData will be removed from the queue of its old state.
-		 * If this flag is missing, caller is responsible for removing from the old state's queue.
-		 */
-		QueueRemove = 0x1,
-		/*
-		 * PackageData will be added to queue for its next state.
-		 * If this flag is missing, caller is responsible for adding to queue.
-		 */
-		QueueAdd = 0x2,
-		/*
-		 * PackageData will be removed from the queue of its old state and added to the queue of its new state.
-		 * This may be wasteful, if the caller can add or remove more efficiently.
-		 */
-		QueueAddAndRemove = QueueAdd | QueueRemove,
+		FPlatformData();
+		bool bRequested : 1;
+		bool bCookAttempted : 1;
+		bool bCookSucceeded : 1;
+		bool bExplored : 1;
 	};
-	ENUM_CLASS_FLAGS(ESendFlags);
+
+	FPackageData(FPackageDatas& PackageDatas, const FName& InPackageName, const FName& InFileName);
+	~FPackageData();
 
 	/**
-	 * Contains all the information the cooker uses for a package, during request, load, or save.  Once allocated, this
-	 * structure is never deallocated or moved for a given package; it is deallocated only when the CookOnTheFlyServer
-	 * is destroyed.
-	 *
-	 * Requested platforms - Requested platforms are the Platforms that will be Saved during cooking.
-	 *   Platforms are only requested on the PackageData when the Package is InProgress (e.g. is in the RequestQueue,
-	 *   LoadReadyQueue, or other state containers). Once a Package finishes cooking, the platforms are set to
-	 *   unrequested, and may be set again later if the Package is requested for another platform. If multiple requests
-	 *   occur at the same time, their requested platforms are merged. Since modifying the requested platforms can put
-	 *   an InProgress package into an invalid state, direct write access is private; use UpdateRequestData or
-	 *   SetRequestData to write.
-
-	 * Cooked platforms - Cooked platforms are the platforms that have already been saved in the lifetime of the 
-	 *   CookOnTheFlyServer; note this extends outside of the current CookOnTheFlyServer session.
-	 *   Cooked platforms also store a flag indicating whether the cook for the given platform was successful or not.
-	 *   Cooked platforms can be added to a PackageData for reasons other than normal Save, such as when a Package is
-	 *   detected as not applicable for a Platform and its "Cooked" operation is therefore a noop.
-	 *   Cooked platforms can be cleared for a PackageData if the Package is modified, e.g. during editor operations
-	 *   when the CookOnTheFlyServer is run from the editor.
-	 *
-	 * Package - The package pointer corresponding to this PackageData.
-	 *   Contract for the lifetime of this field is that it is only set after the Package has passed through the load
-	 *   state, and it is cleared when the package returns to idle or is demoted to an earlier state.
-	 *   At all other times it is nullptr.
-	 *
-	 * State - the state of this PackageData in the CookOnTheFlyServer's current session. See the definition of
-	 *   EPackageState for a description of each state. Contract for the value of this field includes membership in the
-	 *   corresponding queue such as SaveQueue, and the presence or absence of state-specific data. Since modifying the
-	 *   State can put the Package into an invalid state, direct write access is private; SendToState handles enforcing
-	 *   the contracts.
-
- 	 * PendingCookedPlatformData - a counter for how many objects in the package have had
-	 *   BeginCacheForCookedPlatformData called and have not yet completed. This is used to block the package from
-	 *   saving until all objects have finished their cache. It is also used to block the package from starting new
-	 *   BeginCacheForCookedPlatformData calls until all pending calls from a previous canceled save have completed.
-	 *   The lifetime of this counter extends for the lifetime of the PackageData; it is shared across
-	 *   CookOnTheFlyServer sessions.
-	 *
-	 * CookedPlatformDataNextIndex - Index for the next (Object, RequestedPlatform) pair that needs to have
-	 *   BeginCacheForCookedPlatformData called on it for the current PackageSave. This field is only non-zero during
-	 *   the save state; it is cleared when successfully or unsucessfully leaving the save state. The field is an
-	 *   object-major index into the two-dimensional array given by
-	 *       {Each Object in GetCachedObjectsInOuter} x {each Platform that is requested}
-	 *   e.g. the index into GetCachedObjectsInOuter is GetCookedPlatformDataNextIndex() / GetNumRequestedPlatforms()
-	 *   and the index into GetRequestedPlatforms() is GetCookedPlatformDataNextIndex() % GetNumRequestedPlatforms()
-	 *
-	 * Other fields with explanation inline
-	*/
-	struct FPackageData
+	 * ClearReferences is called on every PackageData before any packageDatas are deleted,
+	 * so references are still valid during ClearReferences
+	 */
+	void ClearReferences();
+
+	FPackageData(const FPackageData& other) = delete;
+	FPackageData(FPackageData&& other) = delete;
+	FPackageData& operator=(const FPackageData& other) = delete;
+	FPackageData& operator=(FPackageData&& other) = delete;
+
+	/** Return a copy of Package->GetName(). It is derived from the FileName if necessary, and is never modified. */
+	const FName& GetPackageName() const;
+	/**
+	 * Return a copy of the FileName containing the package, normalized as returned from MakeStandardFilename.
+	 * This field may change from e.g. *.umap to *.uasset if LoadPackage loads a different FileName for the
+	 * requested PackageName.
+	 */
+	const FName& GetFileName() const;
+
+	/** Reset OutPlatforms and copy current set of requested platforms into it. */
+	template <typename ArrayType>
+	void GetRequestedPlatforms(ArrayType& OutPlatforms) const
 	{
-	public:
-		/** Data about a platform that has been interacted with (requested, etc) by the cook of a PackageData. */
-		struct FPlatformData
+		OutPlatforms.Reset(PlatformDatas.Num());
+		for (const TPair<const ITargetPlatform*, FPlatformData>& Pair : PlatformDatas)
 		{
-			FPlatformData();
-			bool bRequested : 1;
-			bool bCookAttempted : 1;
-			bool bCookSucceeded : 1;
-			bool bExplored : 1;
-		};
-
-		FPackageData(FPackageDatas& PackageDatas, const FName& InPackageName, const FName& InFileName);
-		~FPackageData();
-
-		/**
-		 * ClearReferences is called on every PackageData before any packageDatas are deleted,
-		 * so references are still valid during ClearReferences
-		 */
-		void ClearReferences();
-
-		FPackageData(const FPackageData& other) = delete;
-		FPackageData(FPackageData&& other) = delete;
-		FPackageData& operator=(const FPackageData& other) = delete;
-		FPackageData& operator=(FPackageData&& other) = delete;
-
-		/** Return a copy of Package->GetName(). It is derived from the FileName if necessary, and is never modified. */
-		const FName& GetPackageName() const;
-		/**
-		 * Return a copy of the FileName containing the package, normalized as returned from MakeStandardFilename.
-		 * This field may change from e.g. *.umap to *.uasset if LoadPackage loads a different FileName for the
-		 * requested PackageName.
-		 */
-		const FName& GetFileName() const;
-
-		/** Reset OutPlatforms and copy current set of requested platforms into it. */
-		template <typename ArrayType>
-		void GetRequestedPlatforms(ArrayType& OutPlatforms) const
-		{
-			OutPlatforms.Reset(PlatformDatas.Num());
-			for (const TPair<const ITargetPlatform*, FPlatformData>& Pair : PlatformDatas)
+			if (Pair.Value.bRequested)
 			{
-				if (Pair.Value.bRequested)
-				{
-					OutPlatforms.Add(Pair.Key);
-				}
+				OutPlatforms.Add(Pair.Key);
 			}
 		}
+	}
 
-		/** Get the number of currently requested platforms. */
-		int32 GetNumRequestedPlatforms() const;
-
-		/**
-		 * Return true if and only if every element of Platforms is currently requested.
-		 * Returns true if Platforms is empty.
-		 */
-		bool HasAllRequestedPlatforms(const TArrayView<const ITargetPlatform* const>& Platforms) const;
-
-		/**
-		 * Return whether all Platforms requested have also been cooked.
-		 * 
-		 * @param bAllowFailedCooks Iff false, failed cooks cause function to return false.
-		 */
-		bool AreAllRequestedPlatformsCooked(bool bAllowFailedCooks) const;
-
-		/** Return whether all Platforms requested have also been explored for transitive references. */
-		bool AreAllRequestedPlatformsExplored() const;
-
-		/** Return true if and only if every element of Platforms has been explored. Returns true if Platforms is empty. */
-		bool HasAllExploredPlatforms(const TArrayView<const ITargetPlatform* const>& Platforms) const;
-
-		/**
-		 * Get the flag for whether this InProgress PackageData has been marked as an urgent request
-		 * (e.g. because it has been requested from the game client during cookonthefly.).
-		 * Always false for PackageDatas that are not InProgress.
-		 * Since modifying the urgency can put an InProgress package into an invalid state, direct write access is
-		 * private; use UpdateRequestData or SetRequestData to write.
-		 */
-		bool GetIsUrgent() const { return static_cast<bool>(bIsUrgent); }
-		/**
-		 * Add the given request data onto the existing request data.
-		 * If the PackageData is not in progress, just calls SetRequestData; skipped if InRequestedPlatforms is empty.
-		 * If the PackageData is in progress it is demoted back to an earlier state if necessary.
-		 * RequestData is set to the union of the previous and new values.
-		 */
-		void UpdateRequestData(const TConstArrayView<const ITargetPlatform*> InRequestedPlatforms, bool bInIsUrgent,
-			FCompletionCallback&& InCompletionCallback, FInstigator&& Instigator, bool bAllowUpdateUrgency=true);
-		/*
-		 * Set the given data onto the existing request data.
-		 * It is invalid to call with empty InRequestedPlatforms.
-		 * It is invalid to call on a PackageData that is already InProgress. 
-		 */
-		void SetRequestData(const TArrayView<const ITargetPlatform* const>& InRequestedPlatforms, bool bInIsUrgent,
-			FCompletionCallback&& InCompletionCallback, FInstigator&& Instigator);
-		/**
-		 * Clear all the inprogress variables from the current PackageData. It is invalid to call this except when
-		 * the PackageData is transitioning out of InProgress.
-		 */
-		void ClearInProgressData();
-
-		/**
-		 * FindOrAdd each TargetPlatform and set its flags: CookAttempted=true, Succeeded=<given>.
-		 * In version that takes two arrays, TargetPlatforms and Succeeded must be the same length.
-		 */
-		void SetPlatformsCooked(const TConstArrayView<const ITargetPlatform*> TargetPlatforms,
-			const TConstArrayView<bool> Succeeded);
-		void SetPlatformsCooked(const TConstArrayView<const ITargetPlatform*> TargetPlatforms, bool bSucceeded);
-		void SetPlatformCooked(const ITargetPlatform* TargetPlatform, bool bSucceeded);
-		/**
-		 * FindOrAdd each TargetPlatform and set its flags: CookAttempted=false.
-		 * In Version that takes no TargetPlatform, CookAttempted is cleared from all existing platforms.
-		 */
-		void SetPlatformsNotCooked(const TConstArrayView<const ITargetPlatform*> TargetPlatforms);
-		void SetPlatformsNotCooked();
-		void SetPlatformNotCooked(const ITargetPlatform* TargetPlatform);
-
-		/** Access the information about platforms interacted with by *this. */
-		const TSortedMap<const ITargetPlatform*, FPlatformData>& GetPlatformDatas() const;
-
-		/** Add a platform if not already existing and return a writable pointer to its flags. */
-		FPlatformData& FindOrAddPlatformData(const ITargetPlatform* TargetPlatform);
-
-		/** Find a platform if it exists and return a writable pointer to its flags. */
-		FPlatformData* FindPlatformData(const ITargetPlatform* TargetPlatform);
-
-		/** Find a platform if it exists and return a writable pointer to its flags. */
-		const FPlatformData* FindPlatformData(const ITargetPlatform* TargetPlatform) const;
-
-		/** Return true if and only if at least one platform has been cooked. */
-		bool HasAnyCookedPlatform() const;
-		/**
-		 * Return true if and only if at least one element of Platforms has been cooked, and with its
-		 * succeeded flag set to true if bIncludeFailed is false. Returns false if Platforms is empty.
-		 */
-		bool HasAnyCookedPlatforms(const TArrayView<const ITargetPlatform* const>& Platforms, bool bIncludeFailed) const;
-		/**
-		 * Return true if and only if every element of Platforms has been cooked, and with its succeeded
-		 * flag set to true if bIncludeFailed is false. Returns true if Platforms is empty.
-		 */
-		bool HasAllCookedPlatforms(const TArrayView<const ITargetPlatform* const>& Platforms, bool bIncludeFailed) const;
-		/**
-		 * Return true if and only if the given Platform has been cooked, and with its succeeded flag
-		 * set to true if bIncludeFailed is false.
-		 */
-		bool HasCookedPlatform(const ITargetPlatform* Platform, bool bIncludeFailed) const;
-		/**
-		 * Return the CookResult for the given platform.  If the platform has not been cooked,
-		 * returns ECookResult::Unseen, otherwise returns Succeeded or Failed depending on its success flag.
-		 */
-		ECookResult GetCookResults(const ITargetPlatform* Platform) const;
-
-		/** Empties and then sets OutPlatforms to contain all platforms that have been requested and not cook-attempted. */
-		template <typename ArrayType>
-		void GetUncookedPlatforms(ArrayType& OutPlatforms)
-		{
-			OutPlatforms.Reset(PlatformDatas.Num());
-			for (const TPair<const ITargetPlatform*, FPlatformData>& Pair : PlatformDatas)
-			{
-				if (Pair.Value.bRequested && !Pair.Value.bCookAttempted)
-				{
-					OutPlatforms.Add(Pair.Key);
-				}
-			}
-		}
+	/** Get the number of currently requested platforms. */
+	int32 GetNumRequestedPlatforms() const;
 
-		/**
-		 * Return the package pointer. By contract it will be non-null if and only if the PackageData's state is
-		 * >= EPackageState::Load.
-		 */
-		UPackage* GetPackage() const;
-		/* Set the package pointer. Caller is responsible for maintaining the contract for this field. */
-		void SetPackage(UPackage* InPackage);
-
-		/** Return the current PackageState */
-		EPackageState GetState() const;
-		/**
-		 * Set the PackageData's state to the given state, remove and add of from the appropriate queues, and destroy,
-		 * create, and verify the appropriate state-specific data.
-
-		 * @param NextState The destination state
-		 * @param SendFlags Behavior for how the PackageData should be added/removed from the queues corresponding to
-		 *                  the new and old states. Callers may want to manage queue membership directly for better
-		 *                  performance; removing from the middle is more expensive than popping from the front.
-		 *                  See definition of ESendFlags for a description of the behavior controlled by SendFlags.
-		 */
-		void SendToState(EPackageState NextState, ESendFlags SendFlags);
-		/* Debug-only code to assert that this PackageData is contained by the container matching its current state. */
-		void CheckInContainer() const;
-		/**
-		 * Return true if and only if this PackageData is InProgress in the current CookOnTheFlyServer session.
-		 * Some data is allocated/destroyed/verified when moving in and out of InProgress.
-		 * InProgress means the CookOnTheFlyServer will in the future decide to cook, failtocook, or skip the PackageData.
-		 */
-		bool IsInProgress() const;
-
-		/** Return true if the Package's current state is in the given Property Group */
-		bool IsInStateProperty(EPackageStateProperty Property) const;
-
-		/*
-		 * CompletionCallback - A callback that will be called when this PackageData next transitions from InProgress to not
-		 * InProgress because of cook success, failure, skip, or cancel.
-		 */
-		/** Get a reference to the currently set callback, to e.g. move it into a local variable during execution. */
-		FCompletionCallback& GetCompletionCallback();
-		/**
-		 * Add the given callback into this PackageData's callback field. It is invalid to call this function with a
-		 * non-empty callback if this PackageData already has a CompletionCallback.
-		 */
-		void AddCompletionCallback(FCompletionCallback&& InCompletionCallback);
-
-		/**
-		 * Get/Set a visited flag used when searching graphs of PackageData. User of the graph is responsible for 
-		 * setting the bIsVisited flag back to empty when graph operations are done.
-		 */
-		bool GetIsVisited() const { return bIsVisited != 0; }
-		void SetIsVisited(bool bValue) { bIsVisited = static_cast<uint32>(bValue); }
-
-		/** Try to preload the file. Return true if preloading is complete (succeeded or failed or was skipped). */
-		bool TryPreload();
-		/** Get/Set whether preloading is complete (succeeded or failed or was skipped). */
-		bool GetIsPreloadAttempted() const { return bIsPreloadAttempted != 0; }
-		void SetIsPreloadAttempted(bool bValue) { bIsPreloadAttempted = static_cast<uint32>(bValue); }
-		/** Get/Set whether preloading succeeded and completed. */
-		bool GetIsPreloaded() const { return bIsPreloaded != 0; }
-		void SetIsPreloaded(bool bValue) { bIsPreloaded = static_cast<uint32>(bValue); }
-		/** Clear any allocated preload data. */
-		void ClearPreload();
-		/** Issue check statements confirming that no preload data is allocated or flags are set. */
-		void CheckPreloadEmpty();
-
-		/**
-		 * The list of objects inside the package.  Only non-empty during saving; it is populated on demand by
-		 * TryCreateObjectCache and is cleared when leaving the save state.
-		 */
-		TArray<FWeakObjectPtr>& GetCachedObjectsInOuter();
-		/** Validate that the CachedObjectsInOuter-dependent variables are empty, when entering save. */
-		void CheckObjectCacheEmpty() const;
-		/** Populate CachedObjectsInOuter if not already populated. Invalid to call except when in the save state. */
-		void CreateObjectCache();
-		/** Clear the CachedObjectsInOuter list, when e.g. leaving the save state. */
-		void ClearObjectCache();
-
-		const int32& GetNumPendingCookedPlatformData() const;
-		int32& GetNumPendingCookedPlatformData();
-		const int32& GetCookedPlatformDataNextIndex() const;
-		int32& GetCookedPlatformDataNextIndex();
-
-		/** Get/Set the flag for whether CachedObjectsInOuter is populated. Always false except during save state. */
-		bool GetHasSaveCache() const { return static_cast<bool>(bHasSaveCache); }
-		void SetHasSaveCache(bool Value) { bHasSaveCache = Value != 0; }
-		/** Get/Set the flag for whether CachedObjectsInOuter iteration has started. Always false except during save. */
-		bool GetCookedPlatformDataStarted() const { return static_cast<bool>(bCookedPlatformDataStarted); }
-		void SetCookedPlatformDataStarted(bool Value) { bCookedPlatformDataStarted = (Value != 0); }
-		/**
-		 * Get/Set the flag for whether BeginCacheForCookedPlatformData has been called on every CachedObjectsInOuter.
-		 * Some objects might still be working asynchronously and have not yet returned true from
-		 * IsCachedCookedPlatformDataLoaded, in which case GetNumPendingCookedPlatformData will be non-zero.
-		 * Always false except during the save state.
-		 */
-		bool GetCookedPlatformDataCalled() const { return static_cast<bool>(bCookedPlatformDataCalled); }
-		void SetCookedPlatformDataCalled(bool bValue) { bCookedPlatformDataCalled = bValue != 0; }
-		/**
-		 * Get/Set the flag for whether BeginCacheForCookedPlatformData has been called and
-		 * IsCachedCookedPlatformDataLoaded has subsequently returned true for every object in GetCachedObjectsInOuter.
-		 * Always false except during the save state.
-		 */
-		bool GetCookedPlatformDataComplete() const { return static_cast<bool>(bCookedPlatformDataComplete); }
-		void SetCookedPlatformDataComplete(bool bValue) { bCookedPlatformDataComplete = bValue != 0; }
-		/** Check whether savestate contracts on the PackageData were invalidated by by e.g. garbage collection. */
-		bool IsSaveInvalidated() const;
-
-		/** Get/Set the flag for whether BeginPrepareSave has been called and returned an error. */
-		bool GetHasBeginPrepareSaveFailed() const { return static_cast<bool>(bHasBeginPrepareSaveFailed); }
-		void SetHasBeginPrepareSaveFailed(bool bValue) { bHasBeginPrepareSaveFailed = bValue != 0; }
-
-		/** Validate that the BeginCacheForCookedPlatformData-dependent fields are empty, when entering save. */
-		void CheckCookedPlatformDataEmpty() const;
-		/**
-		 * Clear the BeginCacheForCookedPlatformData-dependent fields, when leaving save.
-		 * Caller must have already executed any required cancellation steps to avoid dangling pending operations.
-		 */
-		void ClearCookedPlatformData();
-
-		/** Get/Set the PackageDataMonitor flag that counts whether this PackageData has finished cooking. */
-		bool GetMonitorIsCooked() const { return static_cast<bool>(bMonitorIsCooked); }
-		void SetMonitorIsCooked(bool Value) { bMonitorIsCooked = Value != 0; }
-
-		/** Remove all request data about the given platform from all fields in this PackageData. */
-		void OnRemoveSessionPlatform(const ITargetPlatform* Platform);
-
-		/** Report whether this PackageData holds references to UObjects and would be affected by GarbageCollection. */
-		bool HasReferencedObjects() const;
-
-		/** Swap all ITargetPlatform* stored on this instance according to the mapping in @param Remap. */
-		void RemapTargetPlatforms(const TMap<ITargetPlatform*, ITargetPlatform*>& Remap);
-
-		/** Create and return a GeneratorPackage helper object for a given CookPackageSplitter. */
-		FGeneratorPackage* CreateGeneratorPackage(const UObject* InSplitDataObject,
-			ICookPackageSplitter* InCookPackageSplitterInstance);
-		/** Destroy GeneratorPackage helper object. */
-		void DestroyGeneratorPackage() { GeneratorPackage.Reset(); }
-		/** Get GeneratorPackage helper object. */
-		FGeneratorPackage* GetGeneratorPackage() const { return GeneratorPackage.Get(); }
-		/** Get whether the PackageData has done any necessary Generator steps and is ready for BeginCache calls. */
-		bool HasCompletedGeneration() const { return static_cast<bool>(bCompletedGeneration); }
-		/** Set whether the PackageData has done any necessary Generator steps and is ready for BeginCache calls. */
-		void SetCompletedGeneration(bool Value) { bCompletedGeneration = Value != 0; }
-		/** Return whether the package is a Generator package and has started but not completed generation */
-		bool IsGenerating() const
-		{
-			return GetGeneratorPackage() != nullptr && HasInitializedGeneratorSave() && !HasCompletedGeneration();
-		}
-		/** Get whether the PackageData has completed the check for whether it is a Generator. */
-		bool HasInitializedGeneratorSave() const { return static_cast<bool>(bInitializedGeneratorSave); }
-		/** Set whether the PackageData has completed the check for whether it is a Generator. */
-		void SetInitializedGeneratorSave(bool Value) { bInitializedGeneratorSave = Value != 0; }
-		/** Return whether the PackageData is a generated package created by a owning generator PackageData. */
-		bool IsGenerated() const { return static_cast<bool>(bGenerated); }
-		/** Set whether the PackageData is a generated package created by a owning generator PackageData. */
-		void SetGenerated(bool Value) { bGenerated = Value != 0; }
-		/** Set the owning generator PackageData. Only valid to call if SetGenerated(true) has been called. */
-		void SetGeneratedOwner(FGeneratorPackage* InGeneratedOwner);
-		/** Return the owning generator PackageData. Will be null if not a generated package or if orphaned. */
-		FGeneratorPackage* GetGeneratedOwner() const { return GeneratedOwner; }
-		/** Mark that the package has finished generation+saving, and should generate again if it needs resaving. */
-		void ResetGenerationProgress();
-
-		/** Return whether a GC is required by a generator package as soon as possible. */
-		bool GeneratorPackageRequiresGC() const;
-
-		/**
-		 * Return the instigator for this package. The Instigator is the first code location or
-		 * referencing package that causes the package to enter the requested state.
-		 */
-		const FInstigator& GetInstigator() const { return Instigator; }
+	/**
+	 * Return true if and only if every element of Platforms is currently requested.
+	 * Returns true if Platforms is empty.
+	 */
+	bool HasAllRequestedPlatforms(const TArrayView<const ITargetPlatform* const>& Platforms) const;
 
-	private:
-		friend struct UE::Cook::FPackageDatas;
-
-		/** FindOrAdd each TargetPlatform and set its bRequested value. */
-		void SetPlatformsRequested(const TConstArrayView<const ITargetPlatform*> TargetPlatforms, bool bRequested);
-		/** Set all platforms to not requested. */
-		void ClearRequestedPlatforms();
-		void SetIsUrgent(bool Value);
-
-		/**
-		 * Set the FileName of the file that contains the package. This member is private because FPackageDatas
-		 * keeps a map from FileName to PackageData that needs to be updated in sync with it.
-		 */
-		void SetFileName(const FName& InFileName);
-
-		/**
-		 * Set the State of this PackageDAta in the CookOnTheFlyServer's session. This member is private because it
-		 * needs to be updated in sync with other contract data.
-		 */
-		void SetState(EPackageState NextState);
+	/**
+	 * Return whether all Platforms requested have also been cooked.
+	 * 
+	 * @param bAllowFailedCooks Iff false, failed cooks cause function to return false.
+	 */
+	bool AreAllRequestedPlatformsCooked(bool bAllowFailedCooks) const;
 
-	private:
-		/**
-		 * Helper function to call the given EdgeFunction (e.g. OnExitInProgress)
-		 * when a property changes from true to false.
-		 */
-		typedef void (FPackageData::*FEdgeFunction)();
-		inline void UpdateDownEdge(bool bOld, bool bNew, const FEdgeFunction& EdgeFunction)
-		{
-			if ((bOld != bNew) & bOld)
-			{
-				(this->*EdgeFunction)();
-			}
-		}
-		/**
-		 * Helper function to call the given EdgeFunction (e.g. OnEnterInProgress)\
-		 * when a property changes from false to true.
-		 */
-		inline void UpdateUpEdge(bool bOld, bool bNew, const FEdgeFunction& EdgeFunction)
+	/** Return whether all Platforms requested have also been explored for transitive references. */
+	bool AreAllRequestedPlatformsExplored() const;
+
+	/** Return true if and only if every element of Platforms has been explored. Returns true if Platforms is empty. */
+	bool HasAllExploredPlatforms(const TArrayView<const ITargetPlatform* const>& Platforms) const;
+
+	/**
+	 * Get the flag for whether this InProgress PackageData has been marked as an urgent request
+	 * (e.g. because it has been requested from the game client during cookonthefly.).
+	 * Always false for PackageDatas that are not InProgress.
+	 * Since modifying the urgency can put an InProgress package into an invalid state, direct write access is
+	 * private; use UpdateRequestData or SetRequestData to write.
+	 */
+	bool GetIsUrgent() const { return static_cast<bool>(bIsUrgent); }
+	/**
+	 * Add the given request data onto the existing request data.
+	 * If the PackageData is not in progress, just calls SetRequestData; skipped if InRequestedPlatforms is empty.
+	 * If the PackageData is in progress it is demoted back to an earlier state if necessary.
+	 * RequestData is set to the union of the previous and new values.
+	 */
+	void UpdateRequestData(const TConstArrayView<const ITargetPlatform*> InRequestedPlatforms, bool bInIsUrgent,
+		FCompletionCallback&& InCompletionCallback, FInstigator&& Instigator, bool bAllowUpdateUrgency=true);
+	/*
+	 * Set the given data onto the existing request data.
+	 * It is invalid to call with empty InRequestedPlatforms.
+	 * It is invalid to call on a PackageData that is already InProgress. 
+	 */
+	void SetRequestData(const TArrayView<const ITargetPlatform* const>& InRequestedPlatforms, bool bInIsUrgent,
+		FCompletionCallback&& InCompletionCallback, FInstigator&& Instigator);
+	/**
+	 * Clear all the inprogress variables from the current PackageData. It is invalid to call this except when
+	 * the PackageData is transitioning out of InProgress.
+	 */
+	void ClearInProgressData();
+
+	/**
+	 * FindOrAdd each TargetPlatform and set its flags: CookAttempted=true, Succeeded=<given>.
+	 * In version that takes two arrays, TargetPlatforms and Succeeded must be the same length.
+	 */
+	void SetPlatformsCooked(const TConstArrayView<const ITargetPlatform*> TargetPlatforms,
+		const TConstArrayView<bool> Succeeded);
+	void SetPlatformsCooked(const TConstArrayView<const ITargetPlatform*> TargetPlatforms, bool bSucceeded);
+	void SetPlatformCooked(const ITargetPlatform* TargetPlatform, bool bSucceeded);
+	/**
+	 * FindOrAdd each TargetPlatform and set its flags: CookAttempted=false.
+	 * In Version that takes no TargetPlatform, CookAttempted is cleared from all existing platforms.
+	 */
+	void SetPlatformsNotCooked(const TConstArrayView<const ITargetPlatform*> TargetPlatforms);
+	void SetPlatformsNotCooked();
+	void SetPlatformNotCooked(const ITargetPlatform* TargetPlatform);
+
+	/** Access the information about platforms interacted with by *this. */
+	const TSortedMap<const ITargetPlatform*, FPlatformData>& GetPlatformDatas() const;
+
+	/** Add a platform if not already existing and return a writable pointer to its flags. */
+	FPlatformData& FindOrAddPlatformData(const ITargetPlatform* TargetPlatform);
+
+	/** Find a platform if it exists and return a writable pointer to its flags. */
+	FPlatformData* FindPlatformData(const ITargetPlatform* TargetPlatform);
+
+	/** Find a platform if it exists and return a writable pointer to its flags. */
+	const FPlatformData* FindPlatformData(const ITargetPlatform* TargetPlatform) const;
+
+	/** Return true if and only if at least one platform has been cooked. */
+	bool HasAnyCookedPlatform() const;
+	/**
+	 * Return true if and only if at least one element of Platforms has been cooked, and with its
+	 * succeeded flag set to true if bIncludeFailed is false. Returns false if Platforms is empty.
+	 */
+	bool HasAnyCookedPlatforms(const TArrayView<const ITargetPlatform* const>& Platforms, bool bIncludeFailed) const;
+	/**
+	 * Return true if and only if every element of Platforms has been cooked, and with its succeeded
+	 * flag set to true if bIncludeFailed is false. Returns true if Platforms is empty.
+	 */
+	bool HasAllCookedPlatforms(const TArrayView<const ITargetPlatform* const>& Platforms, bool bIncludeFailed) const;
+	/**
+	 * Return true if and only if the given Platform has been cooked, and with its succeeded flag
+	 * set to true if bIncludeFailed is false.
+	 */
+	bool HasCookedPlatform(const ITargetPlatform* Platform, bool bIncludeFailed) const;
+	/**
+	 * Return the CookResult for the given platform.  If the platform has not been cooked,
+	 * returns ECookResult::Unseen, otherwise returns Succeeded or Failed depending on its success flag.
+	 */
+	ECookResult GetCookResults(const ITargetPlatform* Platform) const;
+
+	/** Empties and then sets OutPlatforms to contain all platforms that have been requested and not cook-attempted. */
+	template <typename ArrayType>
+	void GetUncookedPlatforms(ArrayType& OutPlatforms)
+	{
+		OutPlatforms.Reset(PlatformDatas.Num());
+		for (const TPair<const ITargetPlatform*, FPlatformData>& Pair : PlatformDatas)
 		{
-			if ((bOld != bNew) & bNew)
+			if (Pair.Value.bRequested && !Pair.Value.bCookAttempted)
 			{
-				(this->*EdgeFunction)();
+				OutPlatforms.Add(Pair.Key);
 			}
 		}
+	}
 
-		/* Entry/Exit gates for PackageData states, used to enforce state contracts and free unneeded memory. */
-		void OnEnterIdle();
-		void OnExitIdle();
-		void OnEnterRequest();
-		void OnExitRequest();
-		void OnEnterLoadPrepare();
-		void OnExitLoadPrepare();
-		void OnEnterLoadReady();
-		void OnExitLoadReady();
-		void OnEnterSave();
-		void OnExitSave();
-		/* Entry/Exit gates for Properties shared between multiple states */
-		void OnExitInProgress();
-		void OnEnterInProgress();
-		void OnExitLoading();
-		void OnEnterLoading();
-		void OnExitHasPackage();
-		void OnEnterHasPackage();
-
-		TUniquePtr<FGeneratorPackage> GeneratorPackage;
-		FGeneratorPackage* GeneratedOwner;
-		/** Data for each platform that has been interacted with by *this. */
-		TSortedMap<const ITargetPlatform*, FPlatformData> PlatformDatas;
-		TArray<FWeakObjectPtr> CachedObjectsInOuter;
-		FCompletionCallback CompletionCallback;
-		FName PackageName;
-		FName FileName;
-		TWeakObjectPtr<UPackage> Package;
-		/** The one-per-CookOnTheFlyServer owner of this PackageData. */
-		FPackageDatas& PackageDatas;
-		/**
-		* The number of active PreloadableFiles is tracked globally; wrap the PreloadableFile in a struct that
-		* guarantees we always update the counter when changing it
-		*/
-		struct FTrackedPreloadableFilePtr
-		{
-			const TSharedPtr<FPreloadableArchive>& Get() { return Ptr; }
-			void Set(TSharedPtr<FPreloadableArchive>&& InPtr, FPackageData& PackageData);
-			void Reset(FPackageData& PackageData);
-		private:
-			TSharedPtr<FPreloadableArchive> Ptr;
-		};
-		FTrackedPreloadableFilePtr PreloadableFile;
-		int32 NumPendingCookedPlatformData = 0;
-		int32 CookedPlatformDataNextIndex = 0;
-		FOpenPackageResult PreloadableFileOpenResult;
-		FInstigator Instigator;
-
-		uint32 State : int32(EPackageState::BitCount);
-		uint32 bIsUrgent : 1;
-		uint32 bIsVisited : 1;
-		uint32 bIsPreloadAttempted : 1;
-		uint32 bIsPreloaded : 1;
-		uint32 bHasSaveCache : 1;
-		uint32 bHasBeginPrepareSaveFailed : 1;
-		uint32 bCookedPlatformDataStarted : 1;
-		uint32 bCookedPlatformDataCalled : 1;
-		uint32 bCookedPlatformDataComplete : 1;
-		uint32 bMonitorIsCooked : 1;
-		uint32 bInitializedGeneratorSave : 1;
-		uint32 bCompletedGeneration : 1;
-		uint32 bGenerated : 1;
-	};
+	/**
+	 * Return the package pointer. By contract it will be non-null if and only if the PackageData's state is
+	 * >= EPackageState::Load.
+	 */
+	UPackage* GetPackage() const;
+	/* Set the package pointer. Caller is responsible for maintaining the contract for this field. */
+	void SetPackage(UPackage* InPackage);
 
+	/** Return the current PackageState */
+	EPackageState GetState() const;
 	/**
-	 * Helper that wraps a ICookPackageSplitter, gets/caches packages to generate and provide
-	 * the necessary to iterate over all of them iteratively.
+	 * Set the PackageData's state to the given state, remove and add of from the appropriate queues, and destroy,
+	 * create, and verify the appropriate state-specific data.
+
+	 * @param NextState The destination state
+	 * @param SendFlags Behavior for how the PackageData should be added/removed from the queues corresponding to
+	 *                  the new and old states. Callers may want to manage queue membership directly for better
+	 *                  performance; removing from the middle is more expensive than popping from the front.
+	 *                  See definition of ESendFlags for a description of the behavior controlled by SendFlags.
 	 */
-	struct FGeneratorPackage
-	{
-	public:
-		struct FGeneratedStruct
-		{
-			FString RelativePath;
-			TArray<FName> Dependencies;
-			FPackageData* PackageData = nullptr;
-			bool bCreateAsMap = false;
-			bool bHasCreatedPackage = false;
-			bool bHasSaved = false;
-		};
-
-		/** Store the provided CookPackageSplitter and prepare the packages to generate. */
-		FGeneratorPackage(UE::Cook::FPackageData& InOwner, const UObject* InSplitDataObject,
-			ICookPackageSplitter* InCookPackageSplitterInstance);
-		~FGeneratorPackage();
-		/** Clear references to owned generated packages, and mark those packages as orphaned */
-		void ClearGeneratedPackages();
-
-		/** Call the Splitter's GetGenerateList and create the PackageDatas */
-		bool TryGenerateList(UObject* OwnerObject, const FPackageNameCache& PackageNameCache,
-			FPackageDatas& PackageDatas);
-
-		/** Accessor for the packages to generate */
-		TArrayView<FGeneratedStruct> GetPackagesToGenerate() { return PackagesToGenerate; }
-
-		/** Return CookPackageSplitter. */
-		ICookPackageSplitter* GetCookPackageSplitterInstance() const { return CookPackageSplitterInstance.Get(); }
-		/** Return owner FPackageData. */
-		const UE::Cook::FPackageData& GetOwner() const { return Owner; }
-		/** Return the SplitDataObject's FullObjectPath. */
-		const FName GetSplitDataObjectName() const { return SplitDataObjectName; }
-		/**
-		 * Return the packagepath and localpath to the owner's mount point used for loading generated packages,
-		 * if we need to save them to disk.
-		 */
-		void GetIntermediateMountPoint(FString& OutPackagePath, FString& OutLocalPath) const;
-		/**
-		 * Return the standard path to the temporary file (including extension) where we store the given
-		 * generated package, if we need to save it to disk for later retrieval in the cook session.
-		 */
-		FString GetIntermediateLocalPath(const FGeneratorPackage::FGeneratedStruct& GeneratedStruct) const;
-
-		/** Splitter capability for whether generated packages should populate from owner rather than load from disk. */
-		bool IsDeferredPopulate() const { return CookPackageSplitterInstance->UseDeferredPopulate(); }
-		/** Return the FGeneratedStruct with the given PackageData, or null if not found. */
-		FGeneratedStruct* FindGeneratedStruct(FPackageData* PackageData);
-		/**
-		 * Find again the split object from its name, or return null if no longer in memory.
-		 * It may have been GC'd and reloaded since the last time we used it.
-		 */
-		UObject* FindSplitDataObject() const;
-
-		/** State variables for reentrant SplitPackage calls */
-		bool HasGeneratedList() const { return bGeneratedList; }
-		void SetGeneratedList() { bGeneratedList = true; }
-		bool HasClearedOldPackages() { return bClearedOldPackages; }
-		void SetClearedOldPackages() { bClearedOldPackages = true; }
-		bool HasClearedOldPackagesWithGC() { return bClearedOldPackagesWithGC; }
-		void SetClearedOldPackagesWithGC() { bClearedOldPackagesWithGC = true; }
-		bool HasQueuedGeneratedPackages() const { return bQueuedGeneratedPackages; }
-		void SetQueuedGeneratedPackages() { bQueuedGeneratedPackages = true; }
-		int32& GetNextPopulateIndex() { return NextPopulateIndex; }
-
-		/** Callback during garbage collection */
-		void PostGarbageCollect();
-
-		/** Get/Clear whether the Owner was garbage collected since the previous call to TryPopulatePackage */
-		bool GetWasOwnerReloaded() const { return bWasOwnerReloaded; }
-		void ClearWasOwnerReloaded() { bWasOwnerReloaded = false; }
-		/** Call CreatePackage and set PackageData for deterministic generated packages, and update status. */
-		UPackage* CreateGeneratedUPackage(FGeneratedStruct& GeneratedStruct,
-			const UPackage* OwnerPackage, const TCHAR* GeneratedPackageName);
-		/** Mark that the generated package has been populated, to keep track of when *this is no longer needed. */
-		void SetGeneratedSaved(FPackageData& PackageData);
-		/** Return whether list has been generated and all generated packages have been populated */
-		bool IsComplete() const;
+	void SendToState(EPackageState NextState, ESendFlags SendFlags);
+	/* Debug-only code to assert that this PackageData is contained by the container matching its current state. */
+	void CheckInContainer() const;
+	/**
+	 * Return true if and only if this PackageData is InProgress in the current CookOnTheFlyServer session.
+	 * Some data is allocated/destroyed/verified when moving in and out of InProgress.
+	 * InProgress means the CookOnTheFlyServer will in the future decide to cook, failtocook, or skip the PackageData.
+	 */
+	bool IsInProgress() const;
 
-	private:
-		/** PackageData for the package that is being split */
-		const UE::Cook::FPackageData& Owner;
-		/** Name of the object that prompted the splitter creation */
-		FName SplitDataObjectName;
-		/** Cached CookPackageSplitter */
-		TUniquePtr<ICookPackageSplitter> CookPackageSplitterInstance;
-		/** Recorded list of packages to generate from the splitter, and data we need about them */
-		TArray<FGeneratedStruct> PackagesToGenerate;
-
-		int32 NextPopulateIndex = 0;
-		int32 RemainingToPopulate = 0;
-		bool bGeneratedList = false;
-		bool bClearedOldPackages = false;
-		bool bClearedOldPackagesWithGC = false;
-		bool bQueuedGeneratedPackages = false;
-		bool bWasOwnerReloaded = false;
-	};
+	/** Return true if the Package's current state is in the given Property Group */
+	bool IsInStateProperty(EPackageStateProperty Property) const;
 
+	/*
+	 * CompletionCallback - A callback that will be called when this PackageData next transitions from InProgress to not
+	 * InProgress because of cook success, failure, skip, or cancel.
+	 */
+	/** Get a reference to the currently set callback, to e.g. move it into a local variable during execution. */
+	FCompletionCallback& GetCompletionCallback();
+	/**
+	 * Add the given callback into this PackageData's callback field. It is invalid to call this function with a
+	 * non-empty callback if this PackageData already has a CompletionCallback.
+	 */
+	void AddCompletionCallback(FCompletionCallback&& InCompletionCallback);
 
 	/**
-	 * Stores information about the pending action in response to a single call to BeginCacheForCookedPlatformData that
-	 * was made on a given object for the given platform, when saving the given PackageData.
-	 * This instance will remain alive until the object returns true from IsCachedCookedPlatformDataLoaded.
-	 * If the PackageData's save was canceled, this struct also becomes responsible for cleanup of the cached data by
-	 * calling ClearAllCachedCookedPlatformData.
+	 * Get/Set a visited flag used when searching graphs of PackageData. User of the graph is responsible for 
+	 * setting the bIsVisited flag back to empty when graph operations are done.
 	 */
-	struct FPendingCookedPlatformData
-	{
-		FPendingCookedPlatformData(UObject* InObject, const ITargetPlatform* InTargetPlatform,
-			FPackageData& InPackageData, bool bInNeedsResourceRelease, UCookOnTheFlyServer& InCookOnTheFlyServer);
-		FPendingCookedPlatformData(FPendingCookedPlatformData&& Other);
-		FPendingCookedPlatformData(const FPendingCookedPlatformData& Other) = delete;
-		~FPendingCookedPlatformData();
-		FPendingCookedPlatformData& operator=(const FPendingCookedPlatformData& Other) = delete;
-		FPendingCookedPlatformData& operator=(const FPendingCookedPlatformData&& Other) = delete;
-
-		/**
-		 * Call IsCachedCookedPlatformDataLoaded on the object if it has not already returned true.
-		 * If IsCachedCookedPlatformDataLoaded returns true, this function releases all held resources related to the
-		 * pending call, and returns true. Otherwise takes no action and returns false.
-		 * Returns true and early exits if IsCachedCookedPlatformDataLoaded has already returned true.
-		 */
-		bool PollIsComplete();
-		/** Release all held resources related to the pending call, if they have not already been released. */
-		void Release();
-
-		/** Swap all ITargetPlatform* stored on this instance according to the mapping in @param Remap. */
-		void RemapTargetPlatforms(const TMap<ITargetPlatform*, ITargetPlatform*>& Remap);
-
-		/** The object with the pending call. */
-		FWeakObjectPtr Object;
-		/** The platform that was passed to BeginCacheForCookedPlatformData. */
-		const ITargetPlatform* TargetPlatform;
-		/** The PackageData that owns the call; the pending count needs to be updated on this PackageData. */
-		FPackageData& PackageData;
-		/** Backpointer to the CookOnTheFlyServer to allow releasing of resources for the pending call. */
-		UCookOnTheFlyServer& CookOnTheFlyServer;
-		/**
-		 * Non-null only in the case of a cancel. Used to synchronize release of shared resources used by all
-		 * FPendingCookedPlatformData for the various TargetPlatforms of a given object.
-		 */
-		FPendingCookedPlatformDataCancelManager* CancelManager;
-		/* Saved copy of the ClassName to use for resource releasing. */
-		FName ClassName;
-		/** Flag for whether we have executed the release. */
-		bool bHasReleased;
-		/**
-		 * Flag for whether the CookOnTheFlyServer requires resource tracking for the object's
-		 * BeginCacheForCookedPlatformData call.
-		 */
-		bool bNeedsResourceRelease;
-	};
+	bool GetIsVisited() const { return bIsVisited != 0; }
+	void SetIsVisited(bool bValue) { bIsVisited = static_cast<uint32>(bValue); }
+
+	/** Try to preload the file. Return true if preloading is complete (succeeded or failed or was skipped). */
+	bool TryPreload();
+	/** Get/Set whether preloading is complete (succeeded or failed or was skipped). */
+	bool GetIsPreloadAttempted() const { return bIsPreloadAttempted != 0; }
+	void SetIsPreloadAttempted(bool bValue) { bIsPreloadAttempted = static_cast<uint32>(bValue); }
+	/** Get/Set whether preloading succeeded and completed. */
+	bool GetIsPreloaded() const { return bIsPreloaded != 0; }
+	void SetIsPreloaded(bool bValue) { bIsPreloaded = static_cast<uint32>(bValue); }
+	/** Clear any allocated preload data. */
+	void ClearPreload();
+	/** Issue check statements confirming that no preload data is allocated or flags are set. */
+	void CheckPreloadEmpty();
 
 	/**
-	 * Stores information about all of the FPendingCookedPlatformData for a given object, so that resources shared by
-	 * all of the FPendingCookedPlatformData can be released after they are all released.
+	 * The list of objects inside the package.  Only non-empty during saving; it is populated on demand by
+	 * TryCreateObjectCache and is cleared when leaving the save state.
 	 */
-	struct FPendingCookedPlatformDataCancelManager
-	{
-		/** The number of FPendingCookedPlatformData for the given object that are still pending. */
-		int32 NumPendingPlatforms;
-		/** Decrement the reference count, and if it has reached 0, release the resources and delete *this. */
-		void Release(FPendingCookedPlatformData& Data);
-	};
+	TArray<FWeakObjectPtr>& GetCachedObjectsInOuter();
+	/** Validate that the CachedObjectsInOuter-dependent variables are empty, when entering save. */
+	void CheckObjectCacheEmpty() const;
+	/** Populate CachedObjectsInOuter if not already populated. Invalid to call except when in the save state. */
+	void CreateObjectCache();
+	/** Clear the CachedObjectsInOuter list, when e.g. leaving the save state. */
+	void ClearObjectCache();
+
+	const int32& GetNumPendingCookedPlatformData() const;
+	int32& GetNumPendingCookedPlatformData();
+	const int32& GetCookedPlatformDataNextIndex() const;
+	int32& GetCookedPlatformDataNextIndex();
+
+	/** Get/Set the flag for whether CachedObjectsInOuter is populated. Always false except during save state. */
+	bool GetHasSaveCache() const { return static_cast<bool>(bHasSaveCache); }
+	void SetHasSaveCache(bool Value) { bHasSaveCache = Value != 0; }
+	/** Get/Set the flag for whether CachedObjectsInOuter iteration has started. Always false except during save. */
+	bool GetCookedPlatformDataStarted() const { return static_cast<bool>(bCookedPlatformDataStarted); }
+	void SetCookedPlatformDataStarted(bool Value) { bCookedPlatformDataStarted = (Value != 0); }
+	/**
+	 * Get/Set the flag for whether BeginCacheForCookedPlatformData has been called on every CachedObjectsInOuter.
+	 * Some objects might still be working asynchronously and have not yet returned true from
+	 * IsCachedCookedPlatformDataLoaded, in which case GetNumPendingCookedPlatformData will be non-zero.
+	 * Always false except during the save state.
+	 */
+	bool GetCookedPlatformDataCalled() const { return static_cast<bool>(bCookedPlatformDataCalled); }
+	void SetCookedPlatformDataCalled(bool bValue) { bCookedPlatformDataCalled = bValue != 0; }
+	/**
+	 * Get/Set the flag for whether BeginCacheForCookedPlatformData has been called and
+	 * IsCachedCookedPlatformDataLoaded has subsequently returned true for every object in GetCachedObjectsInOuter.
+	 * Always false except during the save state.
+	 */
+	bool GetCookedPlatformDataComplete() const { return static_cast<bool>(bCookedPlatformDataComplete); }
+	void SetCookedPlatformDataComplete(bool bValue) { bCookedPlatformDataComplete = bValue != 0; }
+	/** Check whether savestate contracts on the PackageData were invalidated by by e.g. garbage collection. */
+	bool IsSaveInvalidated() const;
+
+	/** Get/Set the flag for whether BeginPrepareSave has been called and returned an error. */
+	bool GetHasBeginPrepareSaveFailed() const { return static_cast<bool>(bHasBeginPrepareSaveFailed); }
+	void SetHasBeginPrepareSaveFailed(bool bValue) { bHasBeginPrepareSaveFailed = bValue != 0; }
 
+	/** Validate that the BeginCacheForCookedPlatformData-dependent fields are empty, when entering save. */
+	void CheckCookedPlatformDataEmpty() const;
 	/**
-	 * The container class for PackageData pointers that are InProgress in a CookOnTheFlyServer. These containers
-	 * most frequently do queue push/pop operations, but also commonly need to support iteration.
+	 * Clear the BeginCacheForCookedPlatformData-dependent fields, when leaving save.
+	 * Caller must have already executed any required cancellation steps to avoid dangling pending operations.
 	 */
-	class FPackageDataQueue : public TRingBuffer<FPackageData*>
+	void ClearCookedPlatformData();
+
+	/** Get/Set the PackageDataMonitor flag that counts whether this PackageData has finished cooking. */
+	bool GetMonitorIsCooked() const { return static_cast<bool>(bMonitorIsCooked); }
+	void SetMonitorIsCooked(bool Value) { bMonitorIsCooked = Value != 0; }
+
+	/** Remove all request data about the given platform from all fields in this PackageData. */
+	void OnRemoveSessionPlatform(const ITargetPlatform* Platform);
+
+	/** Report whether this PackageData holds references to UObjects and would be affected by GarbageCollection. */
+	bool HasReferencedObjects() const;
+
+	/** Swap all ITargetPlatform* stored on this instance according to the mapping in @param Remap. */
+	void RemapTargetPlatforms(const TMap<ITargetPlatform*, ITargetPlatform*>& Remap);
+
+	/** Create and return a GeneratorPackage helper object for a given CookPackageSplitter. */
+	FGeneratorPackage* CreateGeneratorPackage(const UObject* InSplitDataObject,
+		ICookPackageSplitter* InCookPackageSplitterInstance);
+	/** Destroy GeneratorPackage helper object. */
+	void DestroyGeneratorPackage() { GeneratorPackage.Reset(); }
+	/** Get GeneratorPackage helper object. */
+	FGeneratorPackage* GetGeneratorPackage() const { return GeneratorPackage.Get(); }
+	/** Get whether the PackageData has done any necessary Generator steps and is ready for BeginCache calls. */
+	bool HasCompletedGeneration() const { return static_cast<bool>(bCompletedGeneration); }
+	/** Set whether the PackageData has done any necessary Generator steps and is ready for BeginCache calls. */
+	void SetCompletedGeneration(bool Value) { bCompletedGeneration = Value != 0; }
+	/** Return whether the package is a Generator package and has started but not completed generation */
+	bool IsGenerating() const
 	{
-		using TRingBuffer<FPackageData*>::TRingBuffer;
-	};
+		return GetGeneratorPackage() != nullptr && HasInitializedGeneratorSave() && !HasCompletedGeneration();
+	}
+	/** Get whether the PackageData has completed the check for whether it is a Generator. */
+	bool HasInitializedGeneratorSave() const { return static_cast<bool>(bInitializedGeneratorSave); }
+	/** Set whether the PackageData has completed the check for whether it is a Generator. */
+	void SetInitializedGeneratorSave(bool Value) { bInitializedGeneratorSave = Value != 0; }
+	/** Return whether the PackageData is a generated package created by a owning generator PackageData. */
+	bool IsGenerated() const { return static_cast<bool>(bGenerated); }
+	/** Set whether the PackageData is a generated package created by a owning generator PackageData. */
+	void SetGenerated(bool Value) { bGenerated = Value != 0; }
+	/** Set the owning generator PackageData. Only valid to call if SetGenerated(true) has been called. */
+	void SetGeneratedOwner(FGeneratorPackage* InGeneratedOwner);
+	/** Return the owning generator PackageData. Will be null if not a generated package or if orphaned. */
+	FGeneratorPackage* GetGeneratedOwner() const { return GeneratedOwner; }
+	/** Mark that the package has finished generation+saving, and should generate again if it needs resaving. */
+	void ResetGenerationProgress();
+
+	/** Return whether a GC is required by a generator package as soon as possible. */
+	bool GeneratorPackageRequiresGC() const;
 
 	/**
-	 * A monitor class held by an FPackageDatas to provide reporting and decision making based on aggregated-data
-	 * across all InProgress or completed FPackageData.
+	 * Return the instigator for this package. The Instigator is the first code location or
+	 * referencing package that causes the package to enter the requested state.
 	 */
-	struct FPackageDataMonitor
-	{
-	public:
-		FPackageDataMonitor();
-
-		/** Report the number of FPackageData that are in any non-idle state and need action by CookOnTheFlyServer. */
-		int32 GetNumInProgress() const;
-		int32 GetNumPreloadAllocated() const;
-		/** Report the number of packages that have cooked any platform. Used by CookCommandlet progress reporting. */
-		int32 GetNumCooked() const;
-		/**
-		 * Report the number of FPackageData that are currently marked as urgent.
-		 * Used to check if a Pump function needs to exit to handle urgent PackageData in other states.
-		 */
-		int32 GetNumUrgent() const;
-		/**
-		 * Report the number of FPackageData that are in the given state and have been marked as urgent. Only valid to
-		 * call on states that are in the InProgress set, such as Save.
-		 * Used to prioritize scheduler actions.
-		 */
-		int32 GetNumUrgent(EPackageState InState) const;
-
-		/** Callback called from FPackageData when it transitions to or from inprogress. */
-		void OnInProgressChanged(FPackageData& PackageData, bool bInProgress);
-		void OnPreloadAllocatedChanged(FPackageData& PackageData, bool bPreloadAllocated);
-		/** Callback called from FPackageData when it has set a platform to CookAttempted=true and it does not have any others cooked. */
-		void OnFirstCookedPlatformAdded(FPackageData& PackageData);
-		/** Callback called from FPackageData when it has set a platform to CookAttempted=false and it does not have any others cooked. */
-		void OnLastCookedPlatformRemoved(FPackageData& PackageData);
-		/** Callback called from FPackageData when it has changed its urgency. */
-		void OnUrgencyChanged(FPackageData& PackageData);
-		/** Callback called from FPackageData when it has changed its state. */
-		void OnStateChanged(FPackageData& PackageData, EPackageState OldState);
+	const FInstigator& GetInstigator() const { return Instigator; }
 
-	private:
-		/** Increment or decrement the NumUrgent counter for the given state. */
-		void TrackUrgentRequests(EPackageState State, int32 Delta);
+private:
+	friend struct UE::Cook::FPackageDatas;
 
-		int32 NumInProgress = 0;
-		int32 NumCooked = 0;
-		int32 NumPreloadAllocated = 0;
-		int32 NumUrgentInState[static_cast<uint32>(EPackageState::Count)];
-	};
+	/** FindOrAdd each TargetPlatform and set its bRequested value. */
+	void SetPlatformsRequested(const TConstArrayView<const ITargetPlatform*> TargetPlatforms, bool bRequested);
+	/** Set all platforms to not requested. */
+	void ClearRequestedPlatforms();
+	void SetIsUrgent(bool Value);
+
+	/**
+	 * Set the FileName of the file that contains the package. This member is private because FPackageDatas
+	 * keeps a map from FileName to PackageData that needs to be updated in sync with it.
+	 */
+	void SetFileName(const FName& InFileName);
 
+	/**
+	 * Set the State of this PackageDAta in the CookOnTheFlyServer's session. This member is private because it
+	 * needs to be updated in sync with other contract data.
+	 */
+	void SetState(EPackageState NextState);
 
+private:
 	/**
-	 * A container for FPackageDatas in the Request state. This container needs to support fast find and remove,
-	 * RequestClusters, staging for packages not yet in request clusters, and a FIFO for ready requests
-	 * using AddRequest/PopRequest that is overridden for urgent requests to push them to the front.
+	 * Helper function to call the given EdgeFunction (e.g. OnExitInProgress)
+	 * when a property changes from true to false.
 	 */
-	class FRequestQueue
+	typedef void (FPackageData::*FEdgeFunction)();
+	inline void UpdateDownEdge(bool bOld, bool bNew, const FEdgeFunction& EdgeFunction)
 	{
-	public:
-		bool IsEmpty() const;
-		uint32 Num() const;
-		uint32 Remove(FPackageData* PackageData);
-		bool Contains(const FPackageData* PackageData) const;
-		void Empty();
-
-		void AddRequest(FPackageData* PackageData, bool bForceUrgent=false);
-
-		uint32 ReadyRequestsNum() const;
-		bool IsReadyRequestsEmpty() const;
-		FPackageData* PopReadyRequest();
-		void AddReadyRequest(FPackageData* PackageData, bool bForceUrgent = false);
-		uint32 RemoveRequest(FPackageData* PackageData);
-
-		FPackageDataSet& GetUnclusteredRequests() { return UnclusteredRequests; }
-		TRingBuffer<FRequestCluster>& GetRequestClusters() { return RequestClusters; }
-	private:
-		FPackageDataSet UnclusteredRequests;
-		TRingBuffer<FRequestCluster> RequestClusters;
-		FPackageDataSet UrgentRequests;
-		FPackageDataSet NormalRequests;
-	};
-
+		if ((bOld != bNew) & bOld)
+		{
+			(this->*EdgeFunction)();
+		}
+	}
 	/**
-	 * Container for FPackageDatas in the LoadPrepare state. A FIFO with multiple substates.
+	 * Helper function to call the given EdgeFunction (e.g. OnEnterInProgress)\
+	 * when a property changes from false to true.
 	 */
-	class FLoadPrepareQueue
+	inline void UpdateUpEdge(bool bOld, bool bNew, const FEdgeFunction& EdgeFunction)
 	{
-	public:
-		bool IsEmpty();
-		int32 Num() const;
-		FPackageData* PopFront();
-		void Add(FPackageData* PackageData);
-		void AddFront(FPackageData* PackageData);
-		bool Contains(const FPackageData* PackageData) const;
-		uint32 Remove(FPackageData* PackageData);
-
-		FPackageDataQueue PreloadingQueue;
-		FPackageDataQueue EntryQueue;
-	};
-
-	/*
-	 * Class that manages the list of all PackageDatas for a CookOnTheFlyServer. PackageDatas is an associative
-	 * array for extra data about a package (e.g. the requested platforms) that is needed by the CookOnTheFlyServer.
-	 * FPackageData are allocated once and never destroyed or moved until the CookOnTheFlyServer is destroyed.
-	 * Memory on the FPackageData is allocated and deallocated as necessary for its current state.
-	 * FPackageData are mapped by PackageName and by FileName.
-	 * This class also manages all non-temporary references to FPackageData such as the SaveQueue and RequeustQueue.
-	*/
-	struct FPackageDatas : public FGCObject
+		if ((bOld != bNew) & bNew)
+		{
+			(this->*EdgeFunction)();
+		}
+	}
+
+	/* Entry/Exit gates for PackageData states, used to enforce state contracts and free unneeded memory. */
+	void OnEnterIdle();
+	void OnExitIdle();
+	void OnEnterRequest();
+	void OnExitRequest();
+	void OnEnterLoadPrepare();
+	void OnExitLoadPrepare();
+	void OnEnterLoadReady();
+	void OnExitLoadReady();
+	void OnEnterSave();
+	void OnExitSave();
+	/* Entry/Exit gates for Properties shared between multiple states */
+	void OnExitInProgress();
+	void OnEnterInProgress();
+	void OnExitLoading();
+	void OnEnterLoading();
+	void OnExitHasPackage();
+	void OnEnterHasPackage();
+
+	TUniquePtr<FGeneratorPackage> GeneratorPackage;
+	FGeneratorPackage* GeneratedOwner;
+	/** Data for each platform that has been interacted with by *this. */
+	TSortedMap<const ITargetPlatform*, FPlatformData> PlatformDatas;
+	TArray<FWeakObjectPtr> CachedObjectsInOuter;
+	FCompletionCallback CompletionCallback;
+	FName PackageName;
+	FName FileName;
+	TWeakObjectPtr<UPackage> Package;
+	/** The one-per-CookOnTheFlyServer owner of this PackageData. */
+	FPackageDatas& PackageDatas;
+	/**
+	 * The number of active PreloadableFiles is tracked globally; wrap the PreloadableFile in a struct that
+	 * guarantees we always update the counter when changing it
+	 */
+	struct FTrackedPreloadableFilePtr
 	{
-	public:
-		FPackageDatas(UCookOnTheFlyServer& InCookOnTheFlyServer);
-		~FPackageDatas();
-
-		/** Called each time BeginCook is called, for deferred intialization steps. */
-		void BeginCook();
-
-		/** FGCObject interface function - return a debug name describing this FGCObject. */
-		virtual FString GetReferencerName() const override;
-		/**
-		 * FGCObject interface function - add the objects referenced by this FGCObject to the ReferenceCollector.
-		 * This class forwards the query on to the CookOnTheFlyServer.
-		 */
-		virtual void AddReferencedObjects(FReferenceCollector& Collector) override;
-
-		/** Return the PackageNameCache that is caching FileNames on disk for the CookOnTheFlyServer. */
-		const FPackageNameCache& GetPackageNameCache() const;
-		/** Return the Monitor used to report aggregated information about FPackageDatas. */
-		FPackageDataMonitor& GetMonitor();
-		/** Return the backpointer to the CookOnTheFlyServer */
-		UCookOnTheFlyServer& GetCookOnTheFlyServer();
-		/**
-		 * Return the RequestQueue used by the CookOnTheFlyServer. The RequestQueue is the mostly-FIFO list of
-		 * PackageData that need to be cooked.
-		 */
-		FRequestQueue& GetRequestQueue();
-		/**
-		 * Return the LoadPrepareQueue used by the CookOnTheFlyServer. The LoadPrepareQueue is the dependency-ordered
-		 * list of FPackageData that need to be preloaded before they can be loaded.
-		 */
-		FLoadPrepareQueue& GetLoadPrepareQueue() { return LoadPrepareQueue; }
-		/**
-		 * Return the LoadReadyQueue used by the CookOnTheFlyServer. The LoadReadyQueue is the dependency-ordered list
-		 * of PackageData that need to be loaded.
-		 */
-		FPackageDataQueue& GetLoadReadyQueue() { return LoadReadyQueue; }
-		/**
-		 * Return the SaveQueue used by the CookOnTheFlyServer. The SaveQueue is the performance-sorted list of
-		 * PackageData that have been loaded and need to start or are only part way through saving.
-		 */
-		FPackageDataQueue& GetSaveQueue();
-
-		/**
-		Return the PackageData for the given PackageName and FileName; no validation is done on the names.
-		 * Creates the PackageData if it does not already exist.
-		 */
-		FPackageData& FindOrAddPackageData(const FName& PackageName, const FName& NormalizedFileName);
-
-		/** Return the PackageData with the given PackageName if one exists, otherwise return nullptr. */
-		FPackageData* FindPackageDataByPackageName(const FName& PackageName);
-		/**
-		 * Return a pointer to the PackageData for the given PackageName. If one does not already exist,
-		 * find its FileName on disk and create the PackageData. If no filename exists, return nullptr.
-		 */
-		FPackageData* TryAddPackageDataByPackageName(const FName& PackageName);
-		/**
-		 * Return a reference to the PackageData for the given PackageName. If one does not already exist,
-		 * find its FileName on disk and create the PackageData. Asserts if FileName does not exist;
-		 * caller is claiming it does.
-		 */
-		FPackageData& AddPackageDataByPackageNameChecked(const FName& PackageName);
-
-		/**
-		 * Return the PackageData with the given FileName (or that has been marked as having the given FileName as an
-		 * alias) if one exists, otherwise return nullptr.
-		 */
-		FPackageData* FindPackageDataByFileName(const FName& InFileName);
-		/**
-		 * Return a pointer to the PackageData for the given FileName (whether the actual name or an alias).
-		 * If one does not already exist, verify the FileName on disk and create the PackageData.
-		 * If no filename exists for the package on disk, return nullptr.
-		 */
-		FPackageData* TryAddPackageDataByFileName(const FName& InFileName);
-		/**
-		 * Return a reference to the PackageData for the given FileName (whether the actual name or an alias).
-		 * If one does not already exist, verify the FileName on disk and create the PackageData. Asserts if FileName
-		 * does not exist; caller is claiming it does.
-		 */
-		FPackageData& AddPackageDataByFileNameChecked(const FName& FileName);
-
-		/**
-		 * Optimized version of TryAddPackageDataByFileName for use by DLC when loading up the shipped packages.
-		 * This expects the PackageNameCache to have already been updated (for instance with AppendCacheResults)
-		 */
-		void AddExistingPackageDatasForPlatform(const TArray<FName>& ExistingPackages,
-			const ITargetPlatform* TargetPlatform);
-
-		/**
-		 * Try to find the PackageData for the given PackageName.
-		 * If it exists, change the PackageData's FileName if current FileName is different and update the map to it.
-		 * This is called in response to the package being moved in the editor or if we attempted to load a FileName
-		 * and got redirected to another FileName.
-		 * Returns the PackageData if it exists.
-		 */
-		FPackageData* UpdateFileName(const FName& PackageName);
-		/**
-		 * Mark that the give PackageData should also be returned for searches for the given FileName in addition to
-		 * searches for its own FileName.
-		 * This is used for e.g. mapping AssetName.umap to AssetName.uasset.
-		 */
-		void RegisterFileNameAlias(FPackageData& PackageData, FName FileName);
-
-		/** Report the number of packages that have cooked any platform. Used by cook commandlet progress reporting. */
-		int32 GetNumCooked();
-		/**
-		 * Add onto CookedFiles the list of all packages that have cooked any package, either successfully if
-		 * bGetSuccessfulCookedPackages is true and/or unsuccessfully if bGetFailedCookedPackages is true.
-		 */
-		void GetCookedFileNamesForPlatform(const ITargetPlatform* Platform, TArray<FName>& CookedFiles,
-			bool bGetFailedCookedPackages, bool bGetSuccessfulCookedPackages);
-
-		/**
-		 * Delete all PackageDatas and free all other memory used by this FPackageDatas.
-		 * For performance reasons, should only be called on destruction.
-		 */
-		void Clear();
-		/** Set all platforms to not cooked in all PackageDatas. Used to e.g. invalidate previous cooks. */
-		void ClearCookedPlatforms();
-		/** Remove all request data about the given platform from all PackageDatas and other memory used by *this. */
-		void OnRemoveSessionPlatform(const ITargetPlatform* TargetPlatform);
-
-		/** Return the container that tracks pending calls to BeginCacheForCookedPlatformData. */
-		TArray<FPendingCookedPlatformData>& GetPendingCookedPlatformDatas();
-		/**
-		 * Iterate over all elements in PendingCookedPlatformDatas and check whether they have completed,
-		 * releasing their resources and pending count if so.
-		 */
-		void PollPendingCookedPlatformDatas();
-
-		/** RangedFor methods for iterating over all FPackageData managed by this FPackageDatas */
-		TArray<FPackageData*>::RangedForIteratorType begin();
-		TArray<FPackageData*>::RangedForIteratorType end();
-
-		/** Swap all ITargetPlatform* stored on this instance according to the mapping in @param Remap. */
-		void RemapTargetPlatforms(const TMap<ITargetPlatform*, ITargetPlatform*>& Remap);
-
-		/** Get/Set whether to log newly discovered packages to debug missing transitive references. */
-		void SetLogDiscoveredPackages(bool bLog) { bLogDiscoveredPackages = bLog; }
-		bool GetLogDiscoveredPackages() const { return bLogDiscoveredPackages; }
-
-		/** Called when a PackageData assigns its instigator, for debugging. */
-		void DebugInstigator(FPackageData& PackageData);
+		const TSharedPtr<FPreloadableArchive>& Get() { return Ptr; }
+		void Set(TSharedPtr<FPreloadableArchive>&& InPtr, FPackageData& PackageData);
+		void Reset(FPackageData& PackageData);
 	private:
-		/**
-		 * Construct a new FPackageData with the given PackageName and FileName and store references to it in the maps.
-		 * New FPackageData are always created in the Idle state.
-		 */
-		FPackageData& CreatePackageData(FName PackageName, FName FileName);
-
-		/** Collection of pointers to all FPackageData that have been constructed by *this. */
-		TArray<FPackageData*> PackageDatas;
-		FPackageDataMonitor Monitor;
-		FPackageNameCache PackageNameCache;
-		TMap<FName, FPackageData*> PackageNameToPackageData;
-		TMap<FName, FPackageData*> FileNameToPackageData;
-		TArray<FPendingCookedPlatformData> PendingCookedPlatformDatas;
-		FRequestQueue RequestQueue;
-		FLoadPrepareQueue LoadPrepareQueue;
-		FPackageDataQueue LoadReadyQueue;
-		FPackageDataQueue SaveQueue;
-		UCookOnTheFlyServer& CookOnTheFlyServer;
-		FPackageData* ShowInstigatorPackageData = nullptr;
-		double LastPollAsyncTime;
-		bool bLogDiscoveredPackages = false;
+		TSharedPtr<FPreloadableArchive> Ptr;
+	};
+	FTrackedPreloadableFilePtr PreloadableFile;
+	int32 NumPendingCookedPlatformData = 0;
+	int32 CookedPlatformDataNextIndex = 0;
+	FOpenPackageResult PreloadableFileOpenResult;
+	FInstigator Instigator;
+
+	uint32 State : int32(EPackageState::BitCount);
+	uint32 bIsUrgent : 1;
+	uint32 bIsVisited : 1;
+	uint32 bIsPreloadAttempted : 1;
+	uint32 bIsPreloaded : 1;
+	uint32 bHasSaveCache : 1;
+	uint32 bHasBeginPrepareSaveFailed : 1;
+	uint32 bCookedPlatformDataStarted : 1;
+	uint32 bCookedPlatformDataCalled : 1;
+	uint32 bCookedPlatformDataComplete : 1;
+	uint32 bMonitorIsCooked : 1;
+	uint32 bInitializedGeneratorSave : 1;
+	uint32 bCompletedGeneration : 1;
+	uint32 bGenerated : 1;
+};
+
+/**
+ * Helper that wraps a ICookPackageSplitter, gets/caches packages to generate and provide
+ * the necessary to iterate over all of them iteratively.
+ */
+struct FGeneratorPackage
+{
+public:
+	struct FGeneratedStruct
+	{
+		FString RelativePath;
+		TArray<FName> Dependencies;
+		FPackageData* PackageData = nullptr;
+		bool bCreateAsMap = false;
+		bool bHasCreatedPackage = false;
+		bool bHasSaved = false;
 	};
 
+	/** Store the provided CookPackageSplitter and prepare the packages to generate. */
+	FGeneratorPackage(UE::Cook::FPackageData& InOwner, const UObject* InSplitDataObject,
+		ICookPackageSplitter* InCookPackageSplitterInstance);
+	~FGeneratorPackage();
+	/** Clear references to owned generated packages, and mark those packages as orphaned */
+	void ClearGeneratedPackages();
+
+	/** Call the Splitter's GetGenerateList and create the PackageDatas */
+	bool TryGenerateList(UObject* OwnerObject, const FPackageNameCache& PackageNameCache,
+		FPackageDatas& PackageDatas);
+
+	/** Accessor for the packages to generate */
+	TArrayView<FGeneratedStruct> GetPackagesToGenerate() { return PackagesToGenerate; }
+
+	/** Return CookPackageSplitter. */
+	ICookPackageSplitter* GetCookPackageSplitterInstance() const { return CookPackageSplitterInstance.Get(); }
+	/** Return owner FPackageData. */
+	const UE::Cook::FPackageData& GetOwner() const { return Owner; }
+	/** Return the SplitDataObject's FullObjectPath. */
+	const FName GetSplitDataObjectName() const { return SplitDataObjectName; }
 	/**
-	 * A debug-only scope class to confirm that each FPackageData removed from a container during a Pump function
-	 * is added to the container for its new state before leaving the Pump function.
+	 * Return the packagepath and localpath to the owner's mount point used for loading generated packages,
+	 * if we need to save them to disk.
 	 */
-	struct FPoppedPackageDataScope
-	{
-		explicit FPoppedPackageDataScope(FPackageData& InPackageData);
+	void GetIntermediateMountPoint(FString& OutPackagePath, FString& OutLocalPath) const;
+	/**
+	 * Return the standard path to the temporary file (including extension) where we store the given
+	 * generated package, if we need to save it to disk for later retrieval in the cook session.
+	 */
+	FString GetIntermediateLocalPath(const FGeneratorPackage::FGeneratedStruct& GeneratedStruct) const;
+
+	/** Splitter capability for whether generated packages should populate from owner rather than load from disk. */
+	bool IsDeferredPopulate() const { return CookPackageSplitterInstance->UseDeferredPopulate(); }
+	/** Return the FGeneratedStruct with the given PackageData, or null if not found. */
+	FGeneratedStruct* FindGeneratedStruct(FPackageData* PackageData);
+	/**
+	 * Find again the split object from its name, or return null if no longer in memory.
+	 * It may have been GC'd and reloaded since the last time we used it.
+	 */
+	UObject* FindSplitDataObject() const;
+
+	/** State variables for reentrant SplitPackage calls */
+	bool HasGeneratedList() const { return bGeneratedList; }
+	void SetGeneratedList() { bGeneratedList = true; }
+	bool HasClearedOldPackages() { return bClearedOldPackages; }
+	void SetClearedOldPackages() { bClearedOldPackages = true; }
+	bool HasClearedOldPackagesWithGC() { return bClearedOldPackagesWithGC; }
+	void SetClearedOldPackagesWithGC() { bClearedOldPackagesWithGC = true; }
+	bool HasQueuedGeneratedPackages() const { return bQueuedGeneratedPackages; }
+	void SetQueuedGeneratedPackages() { bQueuedGeneratedPackages = true; }
+	int32& GetNextPopulateIndex() { return NextPopulateIndex; }
+
+	/** Callback during garbage collection */
+	void PostGarbageCollect();
+
+	/** Get/Clear whether the Owner was garbage collected since the previous call to TryPopulatePackage */
+	bool GetWasOwnerReloaded() const { return bWasOwnerReloaded; }
+	void ClearWasOwnerReloaded() { bWasOwnerReloaded = false; }
+	/** Call CreatePackage and set PackageData for deterministic generated packages, and update status. */
+	UPackage* CreateGeneratedUPackage(FGeneratedStruct& GeneratedStruct,
+		const UPackage* OwnerPackage, const TCHAR* GeneratedPackageName);
+	/** Mark that the generated package has been populated, to keep track of when *this is no longer needed. */
+	void SetGeneratedSaved(FPackageData& PackageData);
+	/** Return whether list has been generated and all generated packages have been populated */
+	bool IsComplete() const;
+
+private:
+	/** PackageData for the package that is being split */
+	const UE::Cook::FPackageData& Owner;
+	/** Name of the object that prompted the splitter creation */
+	FName SplitDataObjectName;
+	/** Cached CookPackageSplitter */
+	TUniquePtr<ICookPackageSplitter> CookPackageSplitterInstance;
+	/** Recorded list of packages to generate from the splitter, and data we need about them */
+	TArray<FGeneratedStruct> PackagesToGenerate;
+
+	int32 NextPopulateIndex = 0;
+	int32 RemainingToPopulate = 0;
+	bool bGeneratedList = false;
+	bool bClearedOldPackages = false;
+	bool bClearedOldPackagesWithGC = false;
+	bool bQueuedGeneratedPackages = false;
+	bool bWasOwnerReloaded = false;
+};
+
+
+/**
+ * Stores information about the pending action in response to a single call to BeginCacheForCookedPlatformData that
+ * was made on a given object for the given platform, when saving the given PackageData.
+ * This instance will remain alive until the object returns true from IsCachedCookedPlatformDataLoaded.
+ * If the PackageData's save was canceled, this struct also becomes responsible for cleanup of the cached data by
+ * calling ClearAllCachedCookedPlatformData.
+ */
+struct FPendingCookedPlatformData
+{
+	FPendingCookedPlatformData(UObject* InObject, const ITargetPlatform* InTargetPlatform,
+		FPackageData& InPackageData, bool bInNeedsResourceRelease, UCookOnTheFlyServer& InCookOnTheFlyServer);
+	FPendingCookedPlatformData(FPendingCookedPlatformData&& Other);
+	FPendingCookedPlatformData(const FPendingCookedPlatformData& Other) = delete;
+	~FPendingCookedPlatformData();
+	FPendingCookedPlatformData& operator=(const FPendingCookedPlatformData& Other) = delete;
+	FPendingCookedPlatformData& operator=(const FPendingCookedPlatformData&& Other) = delete;
+
+	/**
+	 * Call IsCachedCookedPlatformDataLoaded on the object if it has not already returned true.
+	 * If IsCachedCookedPlatformDataLoaded returns true, this function releases all held resources related to the
+	 * pending call, and returns true. Otherwise takes no action and returns false.
+	 * Returns true and early exits if IsCachedCookedPlatformDataLoaded has already returned true.
+	 */
+	bool PollIsComplete();
+	/** Release all held resources related to the pending call, if they have not already been released. */
+	void Release();
+
+	/** Swap all ITargetPlatform* stored on this instance according to the mapping in @param Remap. */
+	void RemapTargetPlatforms(const TMap<ITargetPlatform*, ITargetPlatform*>& Remap);
+
+	/** The object with the pending call. */
+	FWeakObjectPtr Object;
+	/** The platform that was passed to BeginCacheForCookedPlatformData. */
+	const ITargetPlatform* TargetPlatform;
+	/** The PackageData that owns the call; the pending count needs to be updated on this PackageData. */
+	FPackageData& PackageData;
+	/** Backpointer to the CookOnTheFlyServer to allow releasing of resources for the pending call. */
+	UCookOnTheFlyServer& CookOnTheFlyServer;
+	/**
+	 * Non-null only in the case of a cancel. Used to synchronize release of shared resources used by all
+	 * FPendingCookedPlatformData for the various TargetPlatforms of a given object.
+	 */
+	FPendingCookedPlatformDataCancelManager* CancelManager;
+	/* Saved copy of the ClassName to use for resource releasing. */
+	FName ClassName;
+	/** Flag for whether we have executed the release. */
+	bool bHasReleased;
+	/**
+	 * Flag for whether the CookOnTheFlyServer requires resource tracking for the object's
+	 * BeginCacheForCookedPlatformData call.
+	 */
+	bool bNeedsResourceRelease;
+};
+
+/**
+ * Stores information about all of the FPendingCookedPlatformData for a given object, so that resources shared by
+ * all of the FPendingCookedPlatformData can be released after they are all released.
+ */
+struct FPendingCookedPlatformDataCancelManager
+{
+	/** The number of FPendingCookedPlatformData for the given object that are still pending. */
+	int32 NumPendingPlatforms;
+	/** Decrement the reference count, and if it has reached 0, release the resources and delete *this. */
+	void Release(FPendingCookedPlatformData& Data);
+};
+
+/**
+ * The container class for PackageData pointers that are InProgress in a CookOnTheFlyServer. These containers
+ * most frequently do queue push/pop operations, but also commonly need to support iteration.
+ */
+class FPackageDataQueue : public TRingBuffer<FPackageData*>
+{
+	using TRingBuffer<FPackageData*>::TRingBuffer;
+};
+
+/**
+ * A monitor class held by an FPackageDatas to provide reporting and decision making based on aggregated-data
+ * across all InProgress or completed FPackageData.
+ */
+struct FPackageDataMonitor
+{
+public:
+	FPackageDataMonitor();
+
+	/** Report the number of FPackageData that are in any non-idle state and need action by CookOnTheFlyServer. */
+	int32 GetNumInProgress() const;
+	int32 GetNumPreloadAllocated() const;
+	/** Report the number of packages that have cooked any platform. Used by CookCommandlet progress reporting. */
+	int32 GetNumCooked() const;
+	/**
+	 * Report the number of FPackageData that are currently marked as urgent.
+	 * Used to check if a Pump function needs to exit to handle urgent PackageData in other states.
+	 */
+	int32 GetNumUrgent() const;
+	/**
+	 * Report the number of FPackageData that are in the given state and have been marked as urgent. Only valid to
+	 * call on states that are in the InProgress set, such as Save.
+	 * Used to prioritize scheduler actions.
+	 */
+	int32 GetNumUrgent(EPackageState InState) const;
+
+	/** Callback called from FPackageData when it transitions to or from inprogress. */
+	void OnInProgressChanged(FPackageData& PackageData, bool bInProgress);
+	void OnPreloadAllocatedChanged(FPackageData& PackageData, bool bPreloadAllocated);
+	/** Callback called from FPackageData when it has set a platform to CookAttempted=true and it does not have any others cooked. */
+	void OnFirstCookedPlatformAdded(FPackageData& PackageData);
+	/** Callback called from FPackageData when it has set a platform to CookAttempted=false and it does not have any others cooked. */
+	void OnLastCookedPlatformRemoved(FPackageData& PackageData);
+	/** Callback called from FPackageData when it has changed its urgency. */
+	void OnUrgencyChanged(FPackageData& PackageData);
+	/** Callback called from FPackageData when it has changed its state. */
+	void OnStateChanged(FPackageData& PackageData, EPackageState OldState);
+
+private:
+	/** Increment or decrement the NumUrgent counter for the given state. */
+	void TrackUrgentRequests(EPackageState State, int32 Delta);
+
+	int32 NumInProgress = 0;
+	int32 NumCooked = 0;
+	int32 NumPreloadAllocated = 0;
+	int32 NumUrgentInState[static_cast<uint32>(EPackageState::Count)];
+};
+
+
+/**
+ * A container for FPackageDatas in the Request state. This container needs to support fast find and remove,
+ * RequestClusters, staging for packages not yet in request clusters, and a FIFO for ready requests
+ * using AddRequest/PopRequest that is overridden for urgent requests to push them to the front.
+ */
+class FRequestQueue
+{
+public:
+	bool IsEmpty() const;
+	uint32 Num() const;
+	uint32 Remove(FPackageData* PackageData);
+	bool Contains(const FPackageData* PackageData) const;
+	void Empty();
+
+	void AddRequest(FPackageData* PackageData, bool bForceUrgent=false);
+
+	uint32 ReadyRequestsNum() const;
+	bool IsReadyRequestsEmpty() const;
+	FPackageData* PopReadyRequest();
+	void AddReadyRequest(FPackageData* PackageData, bool bForceUrgent = false);
+	uint32 RemoveRequest(FPackageData* PackageData);
+
+	FPackageDataSet& GetUnclusteredRequests() { return UnclusteredRequests; }
+	TRingBuffer<FRequestCluster>& GetRequestClusters() { return RequestClusters; }
+private:
+	FPackageDataSet UnclusteredRequests;
+	TRingBuffer<FRequestCluster> RequestClusters;
+	FPackageDataSet UrgentRequests;
+	FPackageDataSet NormalRequests;
+};
+
+/**
+ * Container for FPackageDatas in the LoadPrepare state. A FIFO with multiple substates.
+ */
+class FLoadPrepareQueue
+{
+public:
+	bool IsEmpty();
+	int32 Num() const;
+	FPackageData* PopFront();
+	void Add(FPackageData* PackageData);
+	void AddFront(FPackageData* PackageData);
+	bool Contains(const FPackageData* PackageData) const;
+	uint32 Remove(FPackageData* PackageData);
+
+	FPackageDataQueue PreloadingQueue;
+	FPackageDataQueue EntryQueue;
+};
+
+/*
+ * Class that manages the list of all PackageDatas for a CookOnTheFlyServer. PackageDatas is an associative
+ * array for extra data about a package (e.g. the requested platforms) that is needed by the CookOnTheFlyServer.
+ * FPackageData are allocated once and never destroyed or moved until the CookOnTheFlyServer is destroyed.
+ * Memory on the FPackageData is allocated and deallocated as necessary for its current state.
+ * FPackageData are mapped by PackageName and by FileName.
+ * This class also manages all non-temporary references to FPackageData such as the SaveQueue and RequeustQueue.
+*/
+struct FPackageDatas : public FGCObject
+{
+public:
+	FPackageDatas(UCookOnTheFlyServer& InCookOnTheFlyServer);
+	~FPackageDatas();
+
+	/** Called each time BeginCook is called, for deferred intialization steps. */
+	void BeginCook();
+
+	/** FGCObject interface function - return a debug name describing this FGCObject. */
+	virtual FString GetReferencerName() const override;
+	/**
+	 * FGCObject interface function - add the objects referenced by this FGCObject to the ReferenceCollector.
+	 * This class forwards the query on to the CookOnTheFlyServer.
+	 */
+	virtual void AddReferencedObjects(FReferenceCollector& Collector) override;
+
+	/** Return the PackageNameCache that is caching FileNames on disk for the CookOnTheFlyServer. */
+	const FPackageNameCache& GetPackageNameCache() const;
+	/** Return the Monitor used to report aggregated information about FPackageDatas. */
+	FPackageDataMonitor& GetMonitor();
+	/** Return the backpointer to the CookOnTheFlyServer */
+	UCookOnTheFlyServer& GetCookOnTheFlyServer();
+	/**
+	 * Return the RequestQueue used by the CookOnTheFlyServer. The RequestQueue is the mostly-FIFO list of
+	 * PackageData that need to be cooked.
+	 */
+	FRequestQueue& GetRequestQueue();
+	/**
+	 * Return the LoadPrepareQueue used by the CookOnTheFlyServer. The LoadPrepareQueue is the dependency-ordered
+	 * list of FPackageData that need to be preloaded before they can be loaded.
+	 */
+	FLoadPrepareQueue& GetLoadPrepareQueue() { return LoadPrepareQueue; }
+	/**
+	 * Return the LoadReadyQueue used by the CookOnTheFlyServer. The LoadReadyQueue is the dependency-ordered list
+	 * of PackageData that need to be loaded.
+	 */
+	FPackageDataQueue& GetLoadReadyQueue() { return LoadReadyQueue; }
+	/**
+	 * Return the SaveQueue used by the CookOnTheFlyServer. The SaveQueue is the performance-sorted list of
+	 * PackageData that have been loaded and need to start or are only part way through saving.
+	 */
+	FPackageDataQueue& GetSaveQueue();
+
+	/**
+	Return the PackageData for the given PackageName and FileName; no validation is done on the names.
+	 * Creates the PackageData if it does not already exist.
+	 */
+	FPackageData& FindOrAddPackageData(const FName& PackageName, const FName& NormalizedFileName);
+
+	/** Return the PackageData with the given PackageName if one exists, otherwise return nullptr. */
+	FPackageData* FindPackageDataByPackageName(const FName& PackageName);
+	/**
+	 * Return a pointer to the PackageData for the given PackageName. If one does not already exist,
+	 * find its FileName on disk and create the PackageData. If no filename exists, return nullptr.
+	 */
+	FPackageData* TryAddPackageDataByPackageName(const FName& PackageName);
+	/**
+	 * Return a reference to the PackageData for the given PackageName. If one does not already exist,
+	 * find its FileName on disk and create the PackageData. Asserts if FileName does not exist;
+	 * caller is claiming it does.
+	 */
+	FPackageData& AddPackageDataByPackageNameChecked(const FName& PackageName);
+
+	/**
+	 * Return the PackageData with the given FileName (or that has been marked as having the given FileName as an
+	 * alias) if one exists, otherwise return nullptr.
+	 */
+	FPackageData* FindPackageDataByFileName(const FName& InFileName);
+	/**
+	 * Return a pointer to the PackageData for the given FileName (whether the actual name or an alias).
+	 * If one does not already exist, verify the FileName on disk and create the PackageData.
+	 * If no filename exists for the package on disk, return nullptr.
+	 */
+	FPackageData* TryAddPackageDataByFileName(const FName& InFileName);
+	/**
+	 * Return a reference to the PackageData for the given FileName (whether the actual name or an alias).
+	 * If one does not already exist, verify the FileName on disk and create the PackageData. Asserts if FileName
+	 * does not exist; caller is claiming it does.
+	 */
+	FPackageData& AddPackageDataByFileNameChecked(const FName& FileName);
+
+	/**
+	 * Optimized version of TryAddPackageDataByFileName for use by DLC when loading up the shipped packages.
+	 * This expects the PackageNameCache to have already been updated (for instance with AppendCacheResults)
+	 */
+	void AddExistingPackageDatasForPlatform(const TArray<FName>& ExistingPackages,
+		const ITargetPlatform* TargetPlatform);
+
+	/**
+	 * Try to find the PackageData for the given PackageName.
+	 * If it exists, change the PackageData's FileName if current FileName is different and update the map to it.
+	 * This is called in response to the package being moved in the editor or if we attempted to load a FileName
+	 * and got redirected to another FileName.
+	 * Returns the PackageData if it exists.
+	 */
+	FPackageData* UpdateFileName(const FName& PackageName);
+	/**
+	 * Mark that the give PackageData should also be returned for searches for the given FileName in addition to
+	 * searches for its own FileName.
+	 * This is used for e.g. mapping AssetName.umap to AssetName.uasset.
+	 */
+	void RegisterFileNameAlias(FPackageData& PackageData, FName FileName);
+
+	/** Report the number of packages that have cooked any platform. Used by cook commandlet progress reporting. */
+	int32 GetNumCooked();
+	/**
+	 * Add onto CookedFiles the list of all packages that have cooked any package, either successfully if
+	 * bGetSuccessfulCookedPackages is true and/or unsuccessfully if bGetFailedCookedPackages is true.
+	 */
+	void GetCookedFileNamesForPlatform(const ITargetPlatform* Platform, TArray<FName>& CookedFiles,
+		bool bGetFailedCookedPackages, bool bGetSuccessfulCookedPackages);
+
+	/**
+	 * Delete all PackageDatas and free all other memory used by this FPackageDatas.
+	 * For performance reasons, should only be called on destruction.
+	 */
+	void Clear();
+	/** Set all platforms to not cooked in all PackageDatas. Used to e.g. invalidate previous cooks. */
+	void ClearCookedPlatforms();
+	/** Remove all request data about the given platform from all PackageDatas and other memory used by *this. */
+	void OnRemoveSessionPlatform(const ITargetPlatform* TargetPlatform);
+
+	/** Return the container that tracks pending calls to BeginCacheForCookedPlatformData. */
+	TArray<FPendingCookedPlatformData>& GetPendingCookedPlatformDatas();
+	/**
+	 * Iterate over all elements in PendingCookedPlatformDatas and check whether they have completed,
+	 * releasing their resources and pending count if so.
+	 */
+	void PollPendingCookedPlatformDatas();
+
+	/** RangedFor methods for iterating over all FPackageData managed by this FPackageDatas */
+	TArray<FPackageData*>::RangedForIteratorType begin();
+	TArray<FPackageData*>::RangedForIteratorType end();
+
+	/** Swap all ITargetPlatform* stored on this instance according to the mapping in @param Remap. */
+	void RemapTargetPlatforms(const TMap<ITargetPlatform*, ITargetPlatform*>& Remap);
+
+	/** Get/Set whether to log newly discovered packages to debug missing transitive references. */
+	void SetLogDiscoveredPackages(bool bLog) { bLogDiscoveredPackages = bLog; }
+	bool GetLogDiscoveredPackages() const { return bLogDiscoveredPackages; }
+
+	/** Called when a PackageData assigns its instigator, for debugging. */
+	void DebugInstigator(FPackageData& PackageData);
+private:
+	/**
+	 * Construct a new FPackageData with the given PackageName and FileName and store references to it in the maps.
+	 * New FPackageData are always created in the Idle state.
+	 */
+	FPackageData& CreatePackageData(FName PackageName, FName FileName);
+
+	/** Collection of pointers to all FPackageData that have been constructed by *this. */
+	TArray<FPackageData*> PackageDatas;
+	FPackageDataMonitor Monitor;
+	FPackageNameCache PackageNameCache;
+	TMap<FName, FPackageData*> PackageNameToPackageData;
+	TMap<FName, FPackageData*> FileNameToPackageData;
+	TArray<FPendingCookedPlatformData> PendingCookedPlatformDatas;
+	FRequestQueue RequestQueue;
+	FLoadPrepareQueue LoadPrepareQueue;
+	FPackageDataQueue LoadReadyQueue;
+	FPackageDataQueue SaveQueue;
+	UCookOnTheFlyServer& CookOnTheFlyServer;
+	FPackageData* ShowInstigatorPackageData = nullptr;
+	double LastPollAsyncTime;
+	bool bLogDiscoveredPackages = false;
+};
+
+/**
+ * A debug-only scope class to confirm that each FPackageData removed from a container during a Pump function
+ * is added to the container for its new state before leaving the Pump function.
+ */
+struct FPoppedPackageDataScope
+{
+	explicit FPoppedPackageDataScope(FPackageData& InPackageData);
 
 #if COOK_CHECKSLOW_PACKAGEDATA
-		~FPoppedPackageDataScope();
+	~FPoppedPackageDataScope();
 
-		FPackageData& PackageData;
+	FPackageData& PackageData;
 #endif
-	};
+};
+
 } // namespace UE::Cook
\ No newline at end of file
