Turn off the default cook behavior of cooking all /Game packages.
Change the CookModificationDelegate to a multicast delegate.
[RN] Minor, Cooking

diff --git a/Engine/Source/Editor/UnrealEd/Private/Commandlets/CookCommandlet.cpp b/Engine/Source/Editor/UnrealEd/Private/Commandlets/CookCommandlet.cpp
--- a/Engine/Source/Editor/UnrealEd/Private/Commandlets/CookCommandlet.cpp
+++ b/Engine/Source/Editor/UnrealEd/Private/Commandlets/CookCommandlet.cpp
@@ -797,6 +797,7 @@ bool UCookCommandlet::CookByTheBook( const TArray<ITargetPlatform*>& Platforms)
 
 		// Check for -MAP=<name of map> entries
 		CmdLineMapEntries += GetSwitchValueElements(TEXT("MAP"));
+		CmdLineMapEntries += GetSwitchValueElements(TEXT("PACKAGE"));
 
 		// Check for -COOKDIR=<path to directory> entries
 		const FString CookDirPrefix = TEXT("COOKDIR=");
diff --git a/Engine/Source/Editor/UnrealEd/Private/CookOnTheFlyServer.cpp b/Engine/Source/Editor/UnrealEd/Private/CookOnTheFlyServer.cpp
--- a/Engine/Source/Editor/UnrealEd/Private/CookOnTheFlyServer.cpp
+++ b/Engine/Source/Editor/UnrealEd/Private/CookOnTheFlyServer.cpp
@@ -106,6 +106,7 @@
 #include "ShaderCompiler.h"
 #include "ShaderLibraryChunkDataGenerator.h"
 #include "String/Find.h"
+#include "String/ParseLines.h"
 #include "TargetDomain/TargetDomainUtils.h"
 #include "UnrealEdGlobals.h"
 #include "UObject/Class.h"
@@ -5900,6 +5901,50 @@ void UCookOnTheFlyServer::AddFileToCook( TArray<FName>& InOutFilesToCook,
 	}
 }
 
+const TCHAR* GCookRequestUsageMessage = TEXT(
+	"By default, the cooker does not cook any packages. Packages must be requested by one of the following methods.\n"
+	"All transitive dependencies of a requested package are also cooked. Packages can be specified by LocalFilename/Filepath\n"
+	"or by LongPackagename/LongPackagePath.\n"
+	"	RecommendedMethod:\n"
+	"		Use the AssetManager's default behavior of PrimaryAssetTypesToScan rules\n"
+	"			Engine.ini:[/Script/Engine.AssetManagerSettings]:+PrimaryAssetTypesToScan\n"
+	"	Commandline:\n"
+	"		-package=<PackageName>\n"
+	"			Request the given package.\n"
+	"		-cookdir=<PackagePath>\n"
+	"			Request all packages in the given directory.\n"
+	"		-mapinisection=<SectionNameInEditorIni>	\n"
+	"			Specify an ini section of packages to cook, in the style of AlwaysCookMaps.\n"
+	"	Ini:\n"
+	"		Editor.ini\n"
+	"			[AlwaysCookMaps]\n"
+	"				+Map=<PackageName>\n"
+	"					; Request the package on every cook. Repeatable.\n"
+	"			[AllMaps]\n"
+	"				+Map=<PackageName>\n"
+	"					; Request the package on default cooks. Not used if commandline, AlwaysCookMaps, or MapsToCook are present.\n"
+	"		Game.ini\n"
+	"			[/Script/UnrealEd.ProjectPackagingSettings]\n"
+	"				+MapsToCook=(FilePath=\"<PackageName>\")\n"
+	"					; Request the package in default cooks. Repeatable.\n"
+	"					; Not used if commandline packages or AlwaysCookMaps are present.\n"
+	"				DirectoriesToAlwaysCook=(Path=\"<PackagePath>\")\n"
+	"					; Request the array of packages in every cook. Repeatable.\n"
+	"				bCookAll=true\n"
+	"					; \n"
+	"		Engine.ini\n"
+	"			[/Script/EngineSettings.GameMapsSettings]\n"
+	"				GameDefaultMap=<PackageName>\n"
+	"				; And other default types; see GameMapsSettings.\n"
+	"	C++API\n"
+	"		FAssetManager::ModifyCook\n"
+	"			// Subclass FAssetManager (Engine.ini:[/Script/Engine.Engine]:AssetManagerClassName) and override this hook.\n"
+	"           // The default AssetManager behavior cooks all packages specified by PrimaryAssetTypesToScan rules from ini.\n"
+	"		FGameDelegates::Get().GetModifyCookDelegate()\n"
+	"			// Subscribe to this delegate during your module startup.\n"
+	"		ITargetPlatform::GetExtraPackagesToCook\n"
+	"			// Override this hook on a given TargetPlatform.\n"
+);
 void UCookOnTheFlyServer::CollectFilesToCook(TArray<FName>& FilesInPath, TMap<FName, UE::Cook::FInstigator>& Instigators,
 	const TArray<FString>& CookMaps, const TArray<FString>& InCookDirectories,
 	const TArray<FString> &IniMapSections, ECookByTheBookOptions FilesToCookFlags, const TArrayView<const ITargetPlatform* const>& TargetPlatforms,
@@ -5908,6 +5953,13 @@ void UCookOnTheFlyServer::CollectFilesToCook(TArray<FName>& FilesInPath, TMap<FN
 	UE_SCOPED_HIERARCHICAL_COOKTIMER(CollectFilesToCook);
 	using namespace UE::Cook;
 
+	if (FParse::Param(FCommandLine::Get(), TEXT("helpcookusage")))
+	{
+		UE::String::ParseLines(GCookRequestUsageMessage, [](FStringView Line)
+			{
+				UE_LOG(LogCook, Warning, TEXT("%.*s"), Line.Len(), Line.GetData());
+			});
+	}
 	UProjectPackagingSettings* PackagingSettings = Cast<UProjectPackagingSettings>(UProjectPackagingSettings::StaticClass()->GetDefaultObject());
 
 	bool bCookAll = (!!(FilesToCookFlags & ECookByTheBookOptions::CookAll)) || PackagingSettings->bCookAll;
@@ -6040,30 +6092,33 @@ void UCookOnTheFlyServer::CollectFilesToCook(TArray<FName>& FilesInPath, TMap<FN
 	if (!(FilesToCookFlags & ECookByTheBookOptions::NoGameAlwaysCookPackages))
 	{
 		UE_SCOPED_HIERARCHICAL_COOKTIMER_AND_DURATION(CookModificationDelegate, DetailedCookStats::GameCookModificationDelegateTimeSec);
-#define DEBUG_COOKMODIFICATIONDELEGATE 0
-#if DEBUG_COOKMODIFICATIONDELEGATE
-		TSet<UPackage*> LoadedPackages;
-		for (TObjectIterator<UPackage> It; It; ++It)
-		{
-			LoadedPackages.Add(*It);
-		}
-#endif
 
-		// allow the game to fill out the asset registry, as well as get a list of objects to always cook
 		TArray<FString> FilesInPathStrings;
+		PRAGMA_DISABLE_DEPRECATION_WARNINGS;
 		FGameDelegates::Get().GetCookModificationDelegate().ExecuteIfBound(FilesInPathStrings);
-
+		PRAGMA_ENABLE_DEPRECATION_WARNINGS;
 		for (const FString& FileString : FilesInPathStrings)
 		{
 			AddFileToCook(FilesInPath, Instigators, FileString, EInstigator::CookModificationDelegate);
 		}
 
-		if (UAssetManager::IsValid())
+		FModifyCookDelegate& ModifyCookDelegate = FGameDelegates::Get().GetModifyCookDelegate();
+		if (UAssetManager::IsValid() || ModifyCookDelegate.IsBound())
 		{
 			TArray<FName> PackagesToNeverCook;
 
-			UAssetManager::Get().ModifyCook(TargetPlatforms, FilesInPath, PackagesToNeverCook);
-			UpdateInstigators(EInstigator::AssetManagerModifyCook);
+			if (UAssetManager::IsValid())
+			{
+				// allow the AssetManager to fill out the asset registry, as well as get a list of objects to always cook
+				UAssetManager::Get().ModifyCook(TargetPlatforms, FilesInPath, PackagesToNeverCook);
+				UpdateInstigators(EInstigator::AssetManagerModifyCook);
+			}
+			if (ModifyCookDelegate.IsBound())
+			{
+				// allow game or plugins to fill out the asset registry, as well as get a list of objects to always cook
+				ModifyCookDelegate.Broadcast(TargetPlatforms, FilesInPath, PackagesToNeverCook);
+				UpdateInstigators(EInstigator::ModifyCookDelegate);
+			}
 
 			for (FName NeverCookPackage : PackagesToNeverCook)
 			{
@@ -6075,23 +6130,6 @@ void UCookOnTheFlyServer::CollectFilesToCook(TArray<FName>& FilesInPath, TMap<FN
 				}
 			}
 		}
-#if DEBUG_COOKMODIFICATIONDELEGATE
-		for (TObjectIterator<UPackage> It; It; ++It)
-		{
-			if ( !LoadedPackages.Contains(*It) )
-			{
-				UE_LOG(LogCook, Display, TEXT("CookModificationDelegate loaded %s"), *It->GetName());
-			}
-		}
-#endif
-
-		if (UE_LOG_ACTIVE(LogCook, Verbose) )
-		{
-			for ( const FString& FileName : FilesInPathStrings )
-			{
-				UE_LOG(LogCook, Verbose, TEXT("Cook modification delegate requested package %s"), *FileName);
-			}
-		}
 	}
 
 	for ( const FString& CurrEntry : CookMaps )
@@ -6161,7 +6199,7 @@ void UCookOnTheFlyServer::CollectFilesToCook(TArray<FName>& FilesInPath, TMap<FN
 		}
 
 		// If no packages were explicitly added by command line or game callback, add all maps
-		if (FilesInPath.Num() == InitialNum || bCookAll)
+		if (bCookAll)
 		{
 			TArray<FString> Tokens;
 			Tokens.Empty(2);
@@ -6196,6 +6234,11 @@ void UCookOnTheFlyServer::CollectFilesToCook(TArray<FName>& FilesInPath, TMap<FN
 				}
 			}
 		}
+		if (FilesInPath.Num() == InitialNum && !bCookAll)
+		{
+			LogCookerMessage(TEXT("Cooker found no package requests specified on commandlien or ini. Run cookcommandlet with -helpcookusage to see request options."),
+				EMessageSeverity::Info);
+		}
 	}
 
 	if (!(FilesToCookFlags & ECookByTheBookOptions::NoDefaultMaps))
@@ -7732,7 +7775,7 @@ void UCookOnTheFlyServer::StartCookByTheBook( const FCookByTheBookStartupOptions
 	
 	// start shader code library cooking
 	InitShaderCodeLibrary();
-    CleanShaderCodeLibraries();
+	CleanShaderCodeLibraries();
 	
 	if ( IsCookingDLC() )
 	{
diff --git a/Engine/Source/Runtime/Engine/Public/Engine/ICookInfo.h b/Engine/Source/Runtime/Engine/Public/Engine/ICookInfo.h
--- a/Engine/Source/Runtime/Engine/Public/Engine/ICookInfo.h
+++ b/Engine/Source/Runtime/Engine/Public/Engine/ICookInfo.h
@@ -37,6 +37,7 @@ namespace UE::Cook
 	CallbackMacro(StartupSoftObjectPath, true) \
 	CallbackMacro(PackagingSettingsMapToCook, true) \
 	CallbackMacro(CookModificationDelegate, true) \
+	CallbackMacro(ModifyCookDelegate, true) \
 	CallbackMacro(AssetManagerModifyCook, true) \
 	CallbackMacro(AssetManagerModifyDLCCook, true) \
 	CallbackMacro(TargetPlatformExtraPackagesToCook, true) \
diff --git a/Engine/Source/Runtime/Engine/Public/GameDelegates.h b/Engine/Source/Runtime/Engine/Public/GameDelegates.h
--- a/Engine/Source/Runtime/Engine/Public/GameDelegates.h
+++ b/Engine/Source/Runtime/Engine/Public/GameDelegates.h
@@ -20,8 +20,12 @@ enum class EGameDelegates_SaveGame : short
 	Detail,	
 };
 
-/** Delegate to modify cooking behavior - return extra packages to cook, load up the asset registry, etc */
+// UE_DEPRECATED(5.0, "Use AssetManager or FModifyCookDelegate instead.")
 DECLARE_DELEGATE_OneParam(FCookModificationDelegate, TArray<FString>& /*ExtraPackagesToCook*/);
+/** Delegate to modify cooking behavior - return extra packages to cook, load up the asset registry, etc */
+DECLARE_MULTICAST_DELEGATE_ThreeParams(FModifyCookDelegate,
+	TConstArrayView<const ITargetPlatform*> /* InTargetPlatforms */, TArray<FName>& /* InOutPackagesToCook */,
+	TArray<FName>& /* InOutPackagesToNeverCook */);
 DECLARE_DELEGATE_FiveParams(FAssignStreamingChunkDelegate, const FString& /*PackageToAdd*/, const FString& /*LastLoadedMapName*/, const TArray<int32>& /*AssetRegistryChunkIDs*/, const TArray<int32>& /*ExistingChunkIds*/, TArray<int32>& /*OutChunkIndexList*/);
 DECLARE_DELEGATE_RetVal_ThreeParams(bool, FGetPackageDependenciesForManifestGeneratorDelegate, FName /*PackageName*/, TArray<FName>& /*DependentPackageNames*/, uint8 /*DependencyType*/);
 
@@ -92,8 +96,14 @@ public:
 	DEFINE_GAME_DELEGATE(ExtendedSaveGameInfoDelegate);
 	DEFINE_GAME_DELEGATE(WebServerActionDelegate);	
 
-	// DEPRECATED, switch to subclassing AssetManager instead
-	DEFINE_GAME_DELEGATE(CookModificationDelegate);
+public:
+	UE_DEPRECATED(5.0, "Use AssetManager or ModifyCookDelegate instead.")
+	FCookModificationDelegate& GetCookModificationDelegate() { return CookModificationDelegate; }
+private:
+	FCookModificationDelegate CookModificationDelegate;
+
+	// Called at start of CookByTheBook
+	DEFINE_GAME_DELEGATE(ModifyCookDelegate);
 
 	DEFINE_GAME_DELEGATE(CookedEditorPackageManagerFactoryDelegate);
 public:
