Improved UPackageTools::SanitizePackageName to coalesce multiple contiguous forward slashes into a single forward slash, as package names may not contain double slashes.

diff --git a/Engine/Source/Editor/UnrealEd/Private/PackageTools.cpp b/Engine/Source/Editor/UnrealEd/Private/PackageTools.cpp
--- a/Engine/Source/Editor/UnrealEd/Private/PackageTools.cpp
+++ b/Engine/Source/Editor/UnrealEd/Private/PackageTools.cpp
@@ -58,6 +58,8 @@
 #include "MeshCardRepresentation.h"
 #include "AssetToolsModule.h"
 #include "Subsystems/AssetEditorSubsystem.h"
+#include "Misc/AutomationTest.h"
+
 
 #define LOCTEXT_NAMESPACE "PackageTools"
 
@@ -1243,8 +1245,27 @@ UPackageTools::UPackageTools(const FObjectInitializer& ObjectInitializer)
 	{
 		FString SanitizedName = ObjectTools::SanitizeInvalidChars(InPackageName, INVALID_LONGPACKAGE_CHARACTERS);
 
-		// Remove double-slashes
-		SanitizedName.ReplaceInline(TEXT("//"), TEXT("/"));
+		// Coalesce multiple contiguous slashes into a single slash
+		int32 CharIndex = 0;
+		while (CharIndex < SanitizedName.Len())
+		{
+			if (SanitizedName[CharIndex] == TEXT('/'))
+			{
+				int32 SlashCount = 1;
+				while (CharIndex + SlashCount < SanitizedName.Len() &&
+					   SanitizedName[CharIndex + SlashCount] == TEXT('/'))
+				{
+					SlashCount++;
+				}
+
+				if (SlashCount > 1)
+				{
+					SanitizedName.RemoveAt(CharIndex + 1, SlashCount - 1, false);
+				}
+			}
+
+			CharIndex++;
+		}
 
 		return SanitizedName;
 	}
@@ -1308,6 +1329,47 @@ UPackageTools::UPackageTools(const FObjectInitializer& ObjectInitializer)
 		return nullptr;
 	}
 
+
+#if WITH_DEV_AUTOMATION_TESTS
+
+IMPLEMENT_SIMPLE_AUTOMATION_TEST(FPackageToolsAutomationTest, "Editor.PackageTools.UnitTests", EAutomationTestFlags::EditorContext | EAutomationTestFlags::EngineFilter)
+
+bool FPackageToolsAutomationTest::RunTest(const FString& Parameters)
+{
+	struct FTest
+	{
+		FString Input;
+		FString Output;
+	};
+
+	TArray<FTest> TestStrings =
+	{
+		{ TEXT("/Game/Blah/Boo"), TEXT("/Game/Blah/Boo") },
+		{ TEXT("/Game/Blah//Double"), TEXT("/Game/Blah/Double") },
+		{ TEXT(""), TEXT("") },
+		{ TEXT("/Game/Trailing/"), TEXT("/Game/Trailing/") },
+		{ TEXT("/Game/TrailingDouble//"), TEXT("/Game/TrailingDouble/") },
+		{ TEXT("/Game/Blah///Multiple"), TEXT("/Game/Blah/Multiple") },
+		{ TEXT("/Game/Blah///Multiple///"), TEXT("/Game/Blah/Multiple/") }
+	};
+
+	bool bSuccess = true;
+	for (const FTest& TestString : TestStrings)
+	{
+		FString Result = UPackageTools::SanitizePackageName(TestString.Input);
+		if (Result != TestString.Output)
+		{
+			AddError(FString::Printf(TEXT("SanitizePackageName failed (result = %s, expected = %s)"), *Result, *TestString.Output));
+			bSuccess = false;
+		}
+	}
+
+	return bSuccess;
+}
+
+#endif // WITH_DEV_AUTOMATION_TESTS
+
+
 #undef LOCTEXT_NAMESPACE
 
 // EOF
