Add a BP editor utils function for checking if something is component based.

#rnx

diff --git a/Engine/Source/Editor/UnrealEd/Private/Kismet2/BlueprintEditorUtils.cpp b/Engine/Source/Editor/UnrealEd/Private/Kismet2/BlueprintEditorUtils.cpp
--- a/Engine/Source/Editor/UnrealEd/Private/Kismet2/BlueprintEditorUtils.cpp
+++ b/Engine/Source/Editor/UnrealEd/Private/Kismet2/BlueprintEditorUtils.cpp
@@ -3288,6 +3288,11 @@ bool FBlueprintEditorUtils::IsActorBased(const UBlueprint* Blueprint)
 	return Blueprint && Blueprint->ParentClass && Blueprint->ParentClass->IsChildOf(AActor::StaticClass());
 }
 
+bool FBlueprintEditorUtils::IsComponentBased(const UBlueprint* Blueprint)
+{
+	return Blueprint && Blueprint->ParentClass && Blueprint->ParentClass->IsChildOf(UActorComponent::StaticClass());
+}
+
 bool FBlueprintEditorUtils::IsDelegateSignatureGraph(const UEdGraph* Graph)
 {
 	if(Graph)
diff --git a/Engine/Source/Editor/UnrealEd/Public/Kismet2/BlueprintEditorUtils.h b/Engine/Source/Editor/UnrealEd/Public/Kismet2/BlueprintEditorUtils.h
--- a/Engine/Source/Editor/UnrealEd/Public/Kismet2/BlueprintEditorUtils.h
+++ b/Engine/Source/Editor/UnrealEd/Public/Kismet2/BlueprintEditorUtils.h
@@ -577,6 +577,11 @@ public:
 	 */
 	static bool IsActorBased(const UBlueprint* Blueprint);
 
+	/**
+	 * @return Whether or not this is a component-based blueprint
+	 */
+	static bool IsComponentBased(const UBlueprint* Blueprint);
+
 	/**
 	 * Whether or not this blueprint is an interface, used only for defining functions to implement
 	 *
