Fix a recent regression that could block users from changing member type/default values in the User-Defined Structure asset editor.

#rnx

diff --git a/Engine/Source/Editor/UnrealEd/Classes/UserDefinedStructure/UserDefinedStructEditorData.h b/Engine/Source/Editor/UnrealEd/Classes/UserDefinedStructure/UserDefinedStructEditorData.h
--- a/Engine/Source/Editor/UnrealEd/Classes/UserDefinedStructure/UserDefinedStructEditorData.h
+++ b/Engine/Source/Editor/UnrealEd/Classes/UserDefinedStructure/UserDefinedStructEditorData.h
@@ -134,7 +134,7 @@ public:
 	virtual void PostEditUndo() override;
 	virtual void PostEditUndo(TSharedPtr<ITransactionObjectAnnotation> TransactionAnnotation) override;
 	virtual void PostLoadSubobjects(struct FObjectInstancingGraph* OuterInstanceGraph) override;
-	virtual void Serialize(FStructuredArchive::FRecord Record) override;
+	virtual void PreSave(FObjectPreSaveContext ObjectSaveContext) override;
 	// End of UObject interface.
 
 	// FEditorUndoClient interface
diff --git a/Engine/Source/Editor/UnrealEd/Private/UserDefinedStructEditorData.cpp b/Engine/Source/Editor/UnrealEd/Private/UserDefinedStructEditorData.cpp
--- a/Engine/Source/Editor/UnrealEd/Private/UserDefinedStructEditorData.cpp
+++ b/Engine/Source/Editor/UnrealEd/Private/UserDefinedStructEditorData.cpp
@@ -3,6 +3,7 @@
 #include "UserDefinedStructure/UserDefinedStructEditorData.h"
 #include "Misc/ITransaction.h"
 #include "UObject/UnrealType.h"
+#include "UObject/ObjectSaveContext.h"
 #include "Engine/UserDefinedStruct.h"
 #include "Kismet2/StructureEditorUtils.h"
 #include "Kismet2/BlueprintEditorUtils.h"
@@ -170,16 +171,12 @@ void UUserDefinedStructEditorData::PostLoadSubobjects(FObjectInstancingGraph* Ou
 	}
 }
 
-void UUserDefinedStructEditorData::Serialize(FStructuredArchive::FRecord Record)
+void UUserDefinedStructEditorData::PreSave(FObjectPreSaveContext ObjectSaveContext)
 {
-	// Prior to saving, ensure that editor data matches up with the default instance data.
-	FArchive& UnderlyingArchive = Record.GetUnderlyingArchive();
-	if (UnderlyingArchive.IsSaving())
-	{
-		RefreshValuesFromDefaultInstance();
-	}
+	Super::PreSave(ObjectSaveContext);
 
-	Super::Serialize(Record);
+	// Prior to saving, ensure that editor data matches up with the default instance data.
+	RefreshValuesFromDefaultInstance();
 }
 
 const uint8* UUserDefinedStructEditorData::GetDefaultInstance() const
@@ -263,8 +260,14 @@ void UUserDefinedStructEditorData::RefreshValuesFromDefaultInstance()
 				const FGuid VarGuid = FStructureEditorUtils::GetGuidFromPropertyName(Property->GetFName());
 				if (FStructVariableDescription* VarDesc = VariablesDescriptions.FindByPredicate(FStructureEditorUtils::FFindByGuidHelper<FStructVariableDescription>(VarGuid)))
 				{
+					// If the default value has been changed elsewhere, don't refresh it here. This is typically the result of an edit action prior to compilation.
+					const bool bHasDefaultValueChanged = (VarDesc->DefaultValue != VarDesc->CurrentDefaultValue);
+
 					FBlueprintEditorUtils::PropertyValueToString(Property, StructData, VarDesc->CurrentDefaultValue, ScriptStruct);
-					VarDesc->DefaultValue = VarDesc->CurrentDefaultValue;
+					if (!bHasDefaultValueChanged)
+					{
+						VarDesc->DefaultValue = VarDesc->CurrentDefaultValue;
+					}
 				}
 			}
 		}
