SimpleAssetEditor: Details view can no longer be closed, as it does not really make sense to do so as it is the only thing shown.
#rnx

diff --git a/Engine/Source/Editor/UnrealEd/Private/Toolkits/SimpleAssetEditor.cpp b/Engine/Source/Editor/UnrealEd/Private/Toolkits/SimpleAssetEditor.cpp
--- a/Engine/Source/Editor/UnrealEd/Private/Toolkits/SimpleAssetEditor.cpp
+++ b/Engine/Source/Editor/UnrealEd/Private/Toolkits/SimpleAssetEditor.cpp
@@ -237,6 +237,7 @@ TSharedRef<SDockTab> FSimpleAssetEditor::SpawnPropertiesTab( const FSpawnTabArgs
 	return SNew(SDockTab)
 		.Label( LOCTEXT("GenericDetailsTitle", "Details") )
 		.TabColorScale( GetTabColorScale() )
+		.OnCanCloseTab_Lambda([]() { return false; })
 		[
 			DetailsView.ToSharedRef()
 		];
