FArchiveReplaceObjectRef (and subclasses) now pass parameters via flags enum rather than long list of bools
FArchiveReplaceObjectRefBase no longer tracks replaced references by default. GetReferencedResources will ensure if archiver was not run with TrackReplacedReferences.

diff --git a/Engine/Plugins/Enterprise/DatasmithImporter/Source/DatasmithImporter/Private/Utility/DatasmithImporterImpl.cpp b/Engine/Plugins/Enterprise/DatasmithImporter/Source/DatasmithImporter/Private/Utility/DatasmithImporterImpl.cpp
--- a/Engine/Plugins/Enterprise/DatasmithImporter/Source/DatasmithImporter/Private/Utility/DatasmithImporterImpl.cpp
+++ b/Engine/Plugins/Enterprise/DatasmithImporter/Source/DatasmithImporter/Private/Utility/DatasmithImporterImpl.cpp
@@ -400,13 +400,10 @@ void FDatasmithImporterImpl::CompileMaterial( UObject* Material, FMaterialUpdate
 
 void FDatasmithImporterImpl::FixReferencesForObject( UObject* Object, const TMap< UObject*, UObject* >& ReferencesToRemap )
 {
-	constexpr bool bNullPrivateRefs = false;
-	constexpr bool bIgnoreOuterRef = true;
-	constexpr bool bIgnoreArchetypeRef = true;
-
 	if ( ReferencesToRemap.Num() > 0 )
 	{
-		FArchiveReplaceObjectRef< UObject > ArchiveReplaceObjectRef( Object, ReferencesToRemap, bNullPrivateRefs, bIgnoreOuterRef, bIgnoreArchetypeRef );
+		constexpr EArchiveReplaceObjectFlags ReplaceFlags = (EArchiveReplaceObjectFlags::IgnoreOuterRef | EArchiveReplaceObjectFlags::IgnoreArchetypeRef);
+		FArchiveReplaceObjectRef< UObject > ArchiveReplaceObjectRef( Object, ReferencesToRemap, ReplaceFlags);
 	}
 }
 
diff --git a/Engine/Plugins/Importers/USDImporter/Source/USDStageImporter/Private/USDStageImporter.cpp b/Engine/Plugins/Importers/USDImporter/Source/USDStageImporter/Private/USDStageImporter.cpp
--- a/Engine/Plugins/Importers/USDImporter/Source/USDStageImporter/Private/USDStageImporter.cpp
+++ b/Engine/Plugins/Importers/USDImporter/Source/USDStageImporter/Private/USDStageImporter.cpp
@@ -1003,15 +1003,8 @@ namespace UsdStageImporterImpl
 				continue;
 			}
 
-			constexpr bool bNullPrivateRefs = false;
-			constexpr bool bIgnoreOuterRef = true;
-			constexpr bool bIgnoreArchetypeRef = true;
-			FArchiveReplaceObjectRef< UObject > ArchiveReplaceObjectRefInner(
-				Referencer,
-				ObjectsToRemap,
-				bNullPrivateRefs,
-				bIgnoreOuterRef,
-				bIgnoreArchetypeRef);
+			constexpr EArchiveReplaceObjectFlags ReplaceFlags = (EArchiveReplaceObjectFlags::IgnoreOuterRef | EArchiveReplaceObjectFlags::IgnoreArchetypeRef);
+			FArchiveReplaceObjectRef< UObject > ArchiveReplaceObjectRefInner(Referencer,ObjectsToRemap, ReplaceFlags);
 		}
 	}
 
diff --git a/Engine/Plugins/Tests/EditorTests/Source/EditorTests/Private/UnrealEd/BlueprintAutomationTests.cpp b/Engine/Plugins/Tests/EditorTests/Source/EditorTests/Private/UnrealEd/BlueprintAutomationTests.cpp
--- a/Engine/Plugins/Tests/EditorTests/Source/EditorTests/Private/UnrealEd/BlueprintAutomationTests.cpp
+++ b/Engine/Plugins/Tests/EditorTests/Source/EditorTests/Private/UnrealEd/BlueprintAutomationTests.cpp
@@ -896,7 +896,7 @@ bool FBlueprintCompileOnLoadTest::RunTest(const FString& BlueprintAssetPath)
 	{
 		TArray<UObject*> InitialBlueprintSubobjectsPtr;
 		GetObjectsWithOuter(InitialBlueprint, InitialBlueprintSubobjectsPtr);
-		for (auto Obj : InitialBlueprintSubobjectsPtr)
+		for (UObject* Obj : InitialBlueprintSubobjectsPtr)
 		{
 			InitialBlueprintSubobjects.Add(Obj);
 		}
@@ -907,7 +907,7 @@ bool FBlueprintCompileOnLoadTest::RunTest(const FString& BlueprintAssetPath)
 	{
 		TArray<UBlueprint*> DependentBlueprints;
 		FBlueprintEditorUtils::FindDependentBlueprints(InitialBlueprint, DependentBlueprints);
-		for (auto BP : DependentBlueprints)
+		for (UBlueprint* BP : DependentBlueprints)
 		{
 			BlueprintDependencies.Add(BP);
 		}
@@ -921,10 +921,10 @@ bool FBlueprintCompileOnLoadTest::RunTest(const FString& BlueprintAssetPath)
 		FSoftObjectPath BlueprintAsset;
 	};
 	TArray<FReplaceInnerData> ReplaceInnerData;
-	for (auto BPToUnloadWP : BlueprintDependencies)
+	for (TWeakObjectPtr<UBlueprint>& BPToUnloadWP : BlueprintDependencies)
 	{
-		auto BPToUnload = BPToUnloadWP.Get();
-		auto OldClass = BPToUnload ? *BPToUnload->GeneratedClass : NULL;
+		UBlueprint* BPToUnload = BPToUnloadWP.Get();
+		UClass* OldClass = BPToUnload ? *BPToUnload->GeneratedClass : nullptr;
 		if (OldClass)
 		{
 			FReplaceInnerData Data;
@@ -947,9 +947,9 @@ bool FBlueprintCompileOnLoadTest::RunTest(const FString& BlueprintAssetPath)
 	//UNLOAD DEPENDENCIES, all circular dependencies will be loaded again 
 	// unload the blueprint so we can reload it (to catch any differences, now  
 	// that all its dependencies should be loaded as well)
-	for (auto BPToUnloadWP : BlueprintDependencies)
+	for (TWeakObjectPtr<UBlueprint> BPToUnloadWP : BlueprintDependencies)
 	{
-		if (auto BPToUnload = BPToUnloadWP.Get())
+		if (UBlueprint* BPToUnload = BPToUnloadWP.Get())
 		{
 			FBlueprintAutomationTestUtilities::UnloadBlueprint(BPToUnload);
 		}
@@ -973,7 +973,7 @@ bool FBlueprintCompileOnLoadTest::RunTest(const FString& BlueprintAssetPath)
 	FObjectReader(InitialBlueprint, InitialLoadData);
 	{
 		TMap<UObject*, UObject*> ClassRedirects;
-		for (auto& Data : ReplaceInnerData)
+		for (FReplaceInnerData& Data : ReplaceInnerData)
 		{
 			UClass* OriginalClass = Data.Class.Get();
 			UBlueprint* NewBlueprint = Cast<UBlueprint>(Data.BlueprintAsset.ResolveObject());
@@ -984,12 +984,12 @@ bool FBlueprintCompileOnLoadTest::RunTest(const FString& BlueprintAssetPath)
 			}
 		}
 		// REPLACE OLD DATA
-		FArchiveReplaceObjectRef<UObject>(InitialBlueprint, ClassRedirects, /*bNullPrivateRefs=*/false, /*bIgnoreOuterRef=*/true, /*bIgnoreArchetypeRef=*/false);
-		for (auto SubobjWP : InitialBlueprintSubobjects)
+		FArchiveReplaceObjectRef<UObject>(InitialBlueprint, ClassRedirects, EArchiveReplaceObjectFlags::IgnoreOuterRef);
+		for (TWeakObjectPtr<UObject>& SubobjWP : InitialBlueprintSubobjects)
 		{
-			if (auto Subobj = SubobjWP.Get())
+			if (UObject* Subobj = SubobjWP.Get())
 			{
-				FArchiveReplaceObjectRef<UObject>(Subobj, ClassRedirects, /*bNullPrivateRefs=*/false, /*bIgnoreOuterRef=*/true, /*bIgnoreArchetypeRef=*/false);
+				FArchiveReplaceObjectRef<UObject>(Subobj, ClassRedirects, EArchiveReplaceObjectFlags::IgnoreOuterRef);
 			}
 		}
 	}
diff --git a/Engine/Source/Editor/Kismet/Private/BlueprintCompilationManager.cpp b/Engine/Source/Editor/Kismet/Private/BlueprintCompilationManager.cpp
--- a/Engine/Source/Editor/Kismet/Private/BlueprintCompilationManager.cpp
+++ b/Engine/Source/Editor/Kismet/Private/BlueprintCompilationManager.cpp
@@ -2374,7 +2374,7 @@ void FBlueprintCompilationManagerImpl::ReinstanceBatch(TArray<FReinstancingJob>&
 
 	for(UObject* ArchetypeReferencer : ArchetypeReferencers)
 	{
-		FArchiveReplaceObjectRef<UObject> ReplaceInCDOAr(ArchetypeReferencer, OldArchetypeToNewArchetype, false, false, false);
+		FArchiveReplaceObjectRef<UObject> ReplaceInCDOAr(ArchetypeReferencer, OldArchetypeToNewArchetype);
 	}
 }
 
diff --git a/Engine/Source/Editor/Kismet/Private/BlueprintEditorModule.cpp b/Engine/Source/Editor/Kismet/Private/BlueprintEditorModule.cpp
--- a/Engine/Source/Editor/Kismet/Private/BlueprintEditorModule.cpp
+++ b/Engine/Source/Editor/Kismet/Private/BlueprintEditorModule.cpp
@@ -107,7 +107,7 @@ void FixSubObjectReferencesPostUndoRedo(UObject* InObject)
 		}
 	}
 
-	FArchiveReplaceObjectRef<UObject> Replacer(InObject, OldToNewInstanceMap, false, false, false, false);
+	FArchiveReplaceObjectRef<UObject> Replacer(InObject, OldToNewInstanceMap);
 }
 
 void FixSubObjectReferencesPostUndoRedo(const FTransaction* Transaction)
diff --git a/Engine/Source/Editor/UnrealEd/Private/Commandlets/DiffPackagesCommandlet.cpp b/Engine/Source/Editor/UnrealEd/Private/Commandlets/DiffPackagesCommandlet.cpp
--- a/Engine/Source/Editor/UnrealEd/Private/Commandlets/DiffPackagesCommandlet.cpp
+++ b/Engine/Source/Editor/UnrealEd/Private/Commandlets/DiffPackagesCommandlet.cpp
@@ -392,7 +392,7 @@ void FixupObjectReferences(UPackage* Package, TMap<UObject*, UObject*>& ObjectRe
 	{
 		if (It->IsIn(Package))
 		{
-			FArchiveReplaceObjectRef<UObject> ReplaceAr(*It, ObjectReplacementMap, true, false, false);
+			FArchiveReplaceObjectRef<UObject> ReplaceAr(*It, ObjectReplacementMap, EArchiveReplaceObjectFlags::NullPrivateRefs);
 		}
 	}
 }
diff --git a/Engine/Source/Editor/UnrealEd/Private/Commandlets/PackageUtilities.cpp b/Engine/Source/Editor/UnrealEd/Private/Commandlets/PackageUtilities.cpp
--- a/Engine/Source/Editor/UnrealEd/Private/Commandlets/PackageUtilities.cpp
+++ b/Engine/Source/Editor/UnrealEd/Private/Commandlets/PackageUtilities.cpp
@@ -2273,7 +2273,7 @@ int32 UReplaceActorCommandlet::Main(const FString& Params)
 					{
 						TMap<UClass*, UClass*> ReplaceMap;
 						ReplaceMap.Add(ClassToReplace, ReplaceWithClass);
-						FArchiveReplaceObjectRef<UClass> ReplaceAr(OldObject, ReplaceMap, false, false, false);
+						FArchiveReplaceObjectRef<UClass> ReplaceAr(OldObject, ReplaceMap);
 						if( ReplaceAr.GetCount() > 0 )
 						{
 							UE_LOG(LogPackageUtilities, Display, TEXT("Replaced %i class references in an Object: %s"), ReplaceAr.GetCount(), *OldObject->GetName() );
@@ -2358,7 +2358,7 @@ int32 UReplaceActorCommandlet::Main(const FString& Params)
 						// check for any references to the old Actor and replace them with the new one
 						TMap<AActor*, AActor*> ReplaceMap;
 						ReplaceMap.Add(OldActor, NewActor);
-						FArchiveReplaceObjectRef<AActor> ReplaceAr(World, ReplaceMap, false, false, false);
+						FArchiveReplaceObjectRef<AActor> ReplaceAr(World, ReplaceMap);
 						if (ReplaceAr.GetCount() > 0)
 						{
 							UE_LOG(LogPackageUtilities, Display, TEXT("Replaced %i actor references in %s"), ReplaceAr.GetCount(), *It->GetName());
@@ -2370,7 +2370,7 @@ int32 UReplaceActorCommandlet::Main(const FString& Params)
 						// check for any references to the old class and replace them with the new one
 						TMap<UClass*, UClass*> ReplaceMap;
 						ReplaceMap.Add(ClassToReplace, ReplaceWithClass);
-						FArchiveReplaceObjectRef<UClass> ReplaceAr(*It, ReplaceMap, false, false, false);
+						FArchiveReplaceObjectRef<UClass> ReplaceAr(*It, ReplaceMap);
 						if (ReplaceAr.GetCount() > 0)
 						{
 							UE_LOG(LogPackageUtilities, Display, TEXT("Replaced %i class references in actor %s"), ReplaceAr.GetCount(), *It->GetName());
diff --git a/Engine/Source/Editor/UnrealEd/Private/EditorAnimUtils.cpp b/Engine/Source/Editor/UnrealEd/Private/EditorAnimUtils.cpp
--- a/Engine/Source/Editor/UnrealEd/Private/EditorAnimUtils.cpp
+++ b/Engine/Source/Editor/UnrealEd/Private/EditorAnimUtils.cpp
@@ -521,15 +521,15 @@ namespace EditorAnimUtils
 	{
 		UObject* DefaultObject = AnimBlueprint->GetAnimBlueprintGeneratedClass()->GetDefaultObject();
 
-		FArchiveReplaceObjectRef<UAnimationAsset> ReplaceAr(DefaultObject, AnimAssetReplacementMap, false, false, false);//bNullPrivateRefs, bIgnoreOuterRef, bIgnoreArchetypeRef);
-		FArchiveReplaceObjectRef<UAnimationAsset> ReplaceAr2(AnimBlueprint, AnimAssetReplacementMap, false, false, false);//bNullPrivateRefs, bIgnoreOuterRef, bIgnoreArchetypeRef);
+		FArchiveReplaceObjectRef<UAnimationAsset> ReplaceAr(DefaultObject, AnimAssetReplacementMap);
+		FArchiveReplaceObjectRef<UAnimationAsset> ReplaceAr2(AnimBlueprint, AnimAssetReplacementMap);
 
 		// Replace event graph references
 		for(UEdGraph* GraphPage : AnimBlueprint->UbergraphPages)
 		{
 			for(UEdGraphNode* Node : GraphPage->Nodes)
 			{
-				FArchiveReplaceObjectRef<UAnimationAsset> ReplaceGraphAr(Node, AnimAssetReplacementMap, false, false, false);
+				FArchiveReplaceObjectRef<UAnimationAsset> ReplaceGraphAr(Node, AnimAssetReplacementMap);
 			}
 		}
 
@@ -538,7 +538,7 @@ namespace EditorAnimUtils
 		{
 			for(UEdGraphNode* Node : GraphPage->Nodes)
 			{
-				FArchiveReplaceObjectRef<UAnimationAsset> ReplaceGraphAr(Node, AnimAssetReplacementMap, false, false, false);
+				FArchiveReplaceObjectRef<UAnimationAsset> ReplaceGraphAr(Node, AnimAssetReplacementMap);
 			}
 		}
 	}
diff --git a/Engine/Source/Editor/UnrealEd/Private/EditorEngine.cpp b/Engine/Source/Editor/UnrealEd/Private/EditorEngine.cpp
--- a/Engine/Source/Editor/UnrealEd/Private/EditorEngine.cpp
+++ b/Engine/Source/Editor/UnrealEd/Private/EditorEngine.cpp
@@ -5136,24 +5136,25 @@ void UEditorEngine::ReplaceActors(UActorFactory* Factory, const FAssetData& Asse
 
 	for (UObject* Referencer : ReferencedLevels)
 	{
-		FArchiveReplaceObjectRef<AActor> Ar(Referencer, ConvertedMap, false, true, false);
+		constexpr EArchiveReplaceObjectFlags ArFlags = (EArchiveReplaceObjectFlags::IgnoreOuterRef | EArchiveReplaceObjectFlags::TrackReplacedReferences);
+		FArchiveReplaceObjectRef<AActor> Ar(Referencer, ConvertedMap, ArFlags);
 
-	for (const auto& MapItem : Ar.GetReplacedReferences())
-	{
-		UObject* ModifiedObject = MapItem.Key;
-
-		if (!ModifiedObject->HasAnyFlags(RF_Transient) && ModifiedObject->GetOutermost() != GetTransientPackage() && !ModifiedObject->RootPackageHasAnyFlags(PKG_CompiledIn))
+		for (const TPair<UObject*, TArray<FProperty*>>& MapItem : Ar.GetReplacedReferences())
 		{
-			ModifiedObject->MarkPackageDirty();
-		}
+			UObject* ModifiedObject = MapItem.Key;
 
-		for (FProperty* Property : MapItem.Value)
-		{
-			FPropertyChangedEvent PropertyEvent(Property);
-			ModifiedObject->PostEditChangeProperty(PropertyEvent);
+			if (!ModifiedObject->HasAnyFlags(RF_Transient) && ModifiedObject->GetOutermost() != GetTransientPackage() && !ModifiedObject->RootPackageHasAnyFlags(PKG_CompiledIn))
+			{
+				ModifiedObject->MarkPackageDirty();
+			}
+
+			for (FProperty* Property : MapItem.Value)
+			{
+				FPropertyChangedEvent PropertyEvent(Property);
+				ModifiedObject->PostEditChangeProperty(PropertyEvent);
+			}
 		}
 	}
-	}
 
 	RedrawLevelEditingViewports();
 
diff --git a/Engine/Source/Editor/UnrealEd/Private/EditorObject.cpp b/Engine/Source/Editor/UnrealEd/Private/EditorObject.cpp
--- a/Engine/Source/Editor/UnrealEd/Private/EditorObject.cpp
+++ b/Engine/Source/Editor/UnrealEd/Private/EditorObject.cpp
@@ -614,7 +614,7 @@ static const TCHAR* ImportProperties(
 				{
 					ReplacementMap.Add(OldComponent, ComponentTemplate);
 				}
-				FArchiveReplaceObjectRef<UObject> ReplaceAr(SubobjectOuter, ReplacementMap, false, false, true);
+				FArchiveReplaceObjectRef<UObject> ReplaceAr(SubobjectOuter, ReplacementMap, EArchiveReplaceObjectFlags::IgnoreArchetypeRef);
 
 				// import the properties for the subobject
 				SourceText = ImportObjectProperties(
diff --git a/Engine/Source/Editor/UnrealEd/Private/Kismet2/BlueprintEditorUtils.cpp b/Engine/Source/Editor/UnrealEd/Private/Kismet2/BlueprintEditorUtils.cpp
--- a/Engine/Source/Editor/UnrealEd/Private/Kismet2/BlueprintEditorUtils.cpp
+++ b/Engine/Source/Editor/UnrealEd/Private/Kismet2/BlueprintEditorUtils.cpp
@@ -2009,7 +2009,7 @@ void FBlueprintEditorUtils::PostDuplicateBlueprint(UBlueprint* Blueprint, bool b
 				FBPCompileRequest(Blueprint, BPCompileOptions, nullptr)
 			);
 
-			FArchiveReplaceObjectRef<UObject> ReplaceTemplateRefs(NewBPGC, OldToNewMap, /*bNullPrivateRefs=*/ false, /*bIgnoreOuterRef=*/ false, /*bIgnoreArchetypeRef=*/ false);
+			FArchiveReplaceObjectRef<UObject> ReplaceTemplateRefs(NewBPGC, OldToNewMap);
 
 			// Now propagate the values from the old CDO to the new one
 			check(Blueprint->SkeletonGeneratedClass != nullptr);
diff --git a/Engine/Source/Editor/UnrealEd/Private/Kismet2/Kismet2.cpp b/Engine/Source/Editor/UnrealEd/Private/Kismet2/Kismet2.cpp
--- a/Engine/Source/Editor/UnrealEd/Private/Kismet2/Kismet2.cpp
+++ b/Engine/Source/Editor/UnrealEd/Private/Kismet2/Kismet2.cpp
@@ -388,7 +388,7 @@ void FBlueprintUnloader::ReplaceStaleRefs(UBlueprint* NewBlueprint)
 
 	for (UObject* Referencer : Referencers)
 	{
-		FArchiveReplaceObjectRef<UObject>(Referencer, Redirects, /*bNullPrivateRefs=*/false, /*bIgnoreOuterRef=*/false, /*bIgnoreArchetypeRef=*/false);
+		FArchiveReplaceObjectRef<UObject>(Referencer, Redirects);
 	}
 }
 
diff --git a/Engine/Source/Editor/UnrealEd/Private/Kismet2/KismetReinstanceUtilities.cpp b/Engine/Source/Editor/UnrealEd/Private/Kismet2/KismetReinstanceUtilities.cpp
--- a/Engine/Source/Editor/UnrealEd/Private/Kismet2/KismetReinstanceUtilities.cpp
+++ b/Engine/Source/Editor/UnrealEd/Private/Kismet2/KismetReinstanceUtilities.cpp
@@ -28,7 +28,7 @@
 #include "UObject/UObjectHash.h"
 #include "UObject/UObjectIterator.h"
 #include "Serialization/FindObjectReferencers.h"
-#include "Serialization/ArchiveReplaceObjectRef.h" // @todo replace with ArchiveReplaceObjectAndStructPropertyRef.h in Main and remove FArchiveReplaceObjectAndStructPropertyRef below
+#include "Serialization/ArchiveReplaceObjectAndStructPropertyRef.h"
 #include "BlueprintEditor.h"
 #include "Engine/Selection.h"
 #include "BlueprintEditorSettings.h"
@@ -53,96 +53,6 @@ DECLARE_CYCLE_STAT(TEXT("Reinstance Objects"), EKismetCompilerStats_ReinstanceOb
 DECLARE_CYCLE_STAT(TEXT("Refresh Dependent Blueprints In Reinstancer"), EKismetCompilerStats_RefreshDependentBlueprintsInReinstancer, STATGROUP_KismetCompiler);
 DECLARE_CYCLE_STAT(TEXT("Recreate UberGraphPersistentFrame"), EKismetCompilerStats_RecreateUberGraphPersistentFrame, STATGROUP_KismetCompiler);
 
-/*----------------------------------------------------------------------------
-	FArchiveReplaceObjectAndStructPropertyRef.
-----------------------------------------------------------------------------*/
-/**
- * Specialized version of FArchiveReplaceObjectRef that replaces references to FFields
- * that were owned by any of the old UStructs in the Replacement Map with their respective
- * new versions that belong to the new UStrtucts in the Replacement Map.
- */
-template <class T>
-class FArchiveReplaceObjectAndStructPropertyRef : public FArchiveReplaceObjectRef<T>
-{
-public:
-	/**
-	 * Initializes variables and starts the serialization search
-	 *
-	 * @param InSearchObject		The object to start the search on
-	 * @param InReplacementMap		Map of objects to find -> objects to replace them with (null zeros them)
-	 * @param bNullPrivateRefs		Whether references to non-public objects not contained within the SearchObject
-	 *								should be set to null
-	 * @param bIgnoreOuterRef		Whether we should replace Outer pointers on Objects.
-	 * @param bIgnoreArchetypeRef	Whether we should replace the ObjectArchetype reference on Objects.
-	 * @param bDelayStart			Specify true to prevent the constructor from starting the process.  Allows child classes' to do initialization stuff in their ctor
-	 */
-	FArchiveReplaceObjectAndStructPropertyRef
-	(
-		UObject* InSearchObject,
-		const TMap<T*, T*>& InReplacementMap,
-		bool bNullPrivateRefs,
-		bool bIgnoreOuterRef,
-		bool bIgnoreArchetypeRef,
-		bool bDelayStart = false,
-		bool bIgnoreClassGeneratedByRef = true
-	)
-		: FArchiveReplaceObjectRef<T>(InSearchObject, InReplacementMap, bNullPrivateRefs, bIgnoreOuterRef, bIgnoreArchetypeRef, bDelayStart, bIgnoreClassGeneratedByRef)
-	{
-	}
-
-	/**
-	 * Serializes the reference to FProperties
-	 */
-	virtual FArchive& operator<<(FField*& InField) override
-	{
-		if (InField)
-		{
-			// Some structs (like UFunctions in their bytecode) reference properties of another UStructs.
-			// In this case we need to inspect their owner and if it's one of the objects we want to replace,
-			// replace the entire property with the one matching on the struct we want to replace it with
-			UStruct* OldOwnerStruct = InField->GetOwner<UStruct>();
-			if (OldOwnerStruct)
-			{
-				T* const* ReplaceWith = (T* const*)((const TMap<UObject*, UObject*>*) & this->ReplacementMap)->Find(OldOwnerStruct);
-				if (ReplaceWith)
-				{
-					// We want to replace the property's owner but since that would be even worse than replacing UObject's Outer
-					// we need to replace the entire property instead. We need to find the new property on the object we want to replace the Owner with
-					UStruct* NewOwnerStruct = CastChecked<UStruct>(*ReplaceWith);
-					FField* ReplaceWithField = NewOwnerStruct->FindPropertyByName(InField->GetFName());
-					// Do we need to verify the existence of ReplaceWithField? Theoretically it could be missing on the new version
-					// of the owner struct and in this case we still don't want to keep the stale old property pointer around so it's safer to null it
-					InField = ReplaceWithField;
-					this->ReplacedReferences.FindOrAdd(OldOwnerStruct).AddUnique(this->GetSerializedProperty());
-					this->Count++;
-				}
-				// A->IsIn(A) returns false, but we don't want to NULL that reference out, so extra check here.
-				else if (OldOwnerStruct == this->SearchObject || OldOwnerStruct->IsIn(this->SearchObject))
-				{
-					bool bAlreadyAdded = false;
-					this->SerializedObjects.Add(OldOwnerStruct, &bAlreadyAdded);
-					if (!bAlreadyAdded)
-					{
-						// No recursion
-						this->PendingSerializationObjects.Add(OldOwnerStruct);
-					}
-				}
-				else if (this->bNullPrivateReferences && !OldOwnerStruct->HasAnyFlags(RF_Public))
-				{
-					checkf(false, TEXT("Can't null a reference to %s on property %s as it would be equivalent to nulling UObject's Outer."),
-						*OldOwnerStruct->GetPathName(), *InField->GetName());
-				}
-			}
-			else
-			{
-				// Just serialize the field to find any UObjects it may be referencing that we want to replace 
-				InField->Serialize(*this);
-			}
-		}
-		return *this;
-	}
-};
-
 struct FReplaceReferenceHelper
 {
 	static void IncludeCDO(UClass* OldClass, UClass* NewClass, TMap<UObject*, UObject*> &OldToNewInstanceMap, TArray<UObject*> &SourceObjects, UObject* OriginalCDO)
@@ -221,7 +131,7 @@ struct FReplaceReferenceHelper
 					{
 					public:
 						ReferenceReplace(UObject* InSearchObject, const TMap<UObject*, UObject*>& InReplacementMap, const TMap<FSoftObjectPath, UObject*>& InWeakReferencesMap)
-							: FArchiveReplaceObjectAndStructPropertyRef<UObject>(InSearchObject, InReplacementMap, false, false, false, true), WeakReferencesMap(InWeakReferencesMap)
+							: FArchiveReplaceObjectAndStructPropertyRef<UObject>(InSearchObject, InReplacementMap, EArchiveReplaceObjectFlags::DelayStart), WeakReferencesMap(InWeakReferencesMap)
 						{
 							SerializeSearchObject();
 						}
@@ -993,7 +903,10 @@ public:
 			if (ReplaceWith != nullptr)
 			{
 				Obj = ReplaceWith->ToUObject();
-				ReplacedReferences.FindOrAdd(Obj).AddUnique(GetSerializedProperty());
+				if (bTrackReplacedReferences)
+				{
+					ReplacedReferences.FindOrAdd(Obj).AddUnique(GetSerializedProperty());
+				}
 				Count++;
 			}
 			// A->IsIn(A) returns false, but we don't want to NULL that reference out, so extra check here.
@@ -1036,7 +949,6 @@ public:
 			if (ReplaceWith != nullptr)
 			{
 				Field = ReplaceWith->ToField();
-				//ReplacedReferences.FindOrAdd(Obj).AddUnique(GetSerializedProperty());
 				Count++;
 			}
 		}
diff --git a/Engine/Source/Editor/UnrealEd/Private/Kismet2/ReloadUtilities.cpp b/Engine/Source/Editor/UnrealEd/Private/Kismet2/ReloadUtilities.cpp
--- a/Engine/Source/Editor/UnrealEd/Private/Kismet2/ReloadUtilities.cpp
+++ b/Engine/Source/Editor/UnrealEd/Private/Kismet2/ReloadUtilities.cpp
@@ -446,7 +446,8 @@ FReloadClassReinstancer::FReloadClassReinstancer(UClass* InNewClass, UClass* InO
 
 		for (TObjectIterator<UBlueprint> BlueprintIt; BlueprintIt; ++BlueprintIt)
 		{
-			FArchiveReplaceObjectRef<UObject> ReplaceObjectArch(*BlueprintIt, ClassRedirects, false, true, true);
+			constexpr EArchiveReplaceObjectFlags ReplaceObjectArchFlags = (EArchiveReplaceObjectFlags::IgnoreOuterRef | EArchiveReplaceObjectFlags::IgnoreArchetypeRef);
+			FArchiveReplaceObjectRef<UObject> ReplaceObjectArch(*BlueprintIt, ClassRedirects, ReplaceObjectArchFlags);
 		}
 	}
 	else
diff --git a/Engine/Source/Editor/UnrealEd/Private/ObjectTools.cpp b/Engine/Source/Editor/UnrealEd/Private/ObjectTools.cpp
--- a/Engine/Source/Editor/UnrealEd/Private/ObjectTools.cpp
+++ b/Engine/Source/Editor/UnrealEd/Private/ObjectTools.cpp
@@ -843,7 +843,7 @@ namespace ObjectTools
 		GEditor->GetSelectedObjects()->Select( Object );
 
 		// Replace all references
-		FArchiveReplaceObjectRef<UObject> ReplaceAr( DupObject, ReplacementMap, false, true, true );
+		FArchiveReplaceObjectRef<UObject> ReplaceAr( DupObject, ReplacementMap, EArchiveReplaceObjectFlags::IgnoreOuterRef | EArchiveReplaceObjectFlags::IgnoreArchetypeRef );
 
 		return ReturnObject;
 	}
@@ -1119,19 +1119,9 @@ namespace ObjectTools
 					UBlueprint* BPObjectToUpdate = Cast<UBlueprint>(CurObject);
 					if (BPObjectToUpdate)
 					{
-						const bool bNullPrivateRefs = false;
-						const bool bIgnoreOuterRef = false;
-						const bool bIgnoreArchetypeRef = false;
-						const bool bDelayStart = false;
-						const bool bIgnoreClassGeneratedByRef = false;
-						FArchiveReplaceObjectRef<UObject> ReplaceAr(BPObjectToUpdate->GeneratedClass->ClassDefaultObject, ReplacementMap, bNullPrivateRefs, bIgnoreOuterRef, bIgnoreArchetypeRef, bDelayStart, bIgnoreClassGeneratedByRef);
+						FArchiveReplaceObjectRef<UObject> ReplaceAr(BPObjectToUpdate->GeneratedClass->ClassDefaultObject, ReplacementMap, EArchiveReplaceObjectFlags::IncludeClassGeneratedByRef);
 					}
-					const bool bNullPrivateRefs = false;
-					const bool bIgnoreOuterRef = false;
-					const bool bIgnoreArchetypeRef = false;
-					const bool bDelayStart = false;
-					const bool bIgnoreClassGeneratedByRef = false;
-					FArchiveReplaceObjectRef<UObject> ReplaceAr(CurObject, ReplacementMap, bNullPrivateRefs, bIgnoreOuterRef, bIgnoreArchetypeRef, bDelayStart, bIgnoreClassGeneratedByRef);
+					FArchiveReplaceObjectRef<UObject> ReplaceAr(CurObject, ReplacementMap, EArchiveReplaceObjectFlags::IncludeClassGeneratedByRef);
 				}
 			}
 		}
@@ -1145,12 +1135,7 @@ namespace ObjectTools
 				GWarn->StatusUpdate( NumObjsReplaced, ReferencingPropertiesMapKeys.Num(), NSLOCTEXT("UnrealEd", "ConsolidateAssetsUpdate_ReplacingReferences", "Replacing Asset References...") );
 
 				UObject* CurReplaceObj = ReferencingPropertiesMapKeys[Index];
-				const bool bNullPrivateRefs = false;
-				const bool bIgnoreOuterRef = false;
-				const bool bIgnoreArchetypeRef = false;
-				const bool bDelayStart = false;
-				const bool bIgnoreClassGeneratedByRef = false;
-				FArchiveReplaceObjectRef<UObject> ReplaceAr( CurReplaceObj, ReplacementMap, bNullPrivateRefs, bIgnoreOuterRef, bIgnoreArchetypeRef, bDelayStart, bIgnoreClassGeneratedByRef);
+				FArchiveReplaceObjectRef<UObject> ReplaceAr( CurReplaceObj, ReplacementMap, EArchiveReplaceObjectFlags::IncludeClassGeneratedByRef);
 			}
 		}
 		// Now alter the referencing objects the change has completed via PostEditChange,
diff --git a/Engine/Source/Editor/UnrealEd/Private/Tests/AutomationEditorCommon.cpp b/Engine/Source/Editor/UnrealEd/Private/Tests/AutomationEditorCommon.cpp
--- a/Engine/Source/Editor/UnrealEd/Private/Tests/AutomationEditorCommon.cpp
+++ b/Engine/Source/Editor/UnrealEd/Private/Tests/AutomationEditorCommon.cpp
@@ -177,7 +177,7 @@ void FAutomationEditorCommonUtils::NullReferencesToObject(UObject* InObject)
 		UObject* CurReplaceObj = MapIter.Key();
 		const TArray<FProperty*>& RefPropArray = MapIter.Value();
 
-		FArchiveReplaceObjectRef<UObject> ReplaceAr(CurReplaceObj, ReplacementMap, false, true, false);
+		FArchiveReplaceObjectRef<UObject> ReplaceAr(CurReplaceObj, ReplacementMap, EArchiveReplaceObjectFlags::IgnoreOuterRef);
 
 		for (TArray<FProperty*>::TConstIterator RefPropIter(RefPropArray); RefPropIter; ++RefPropIter)
 		{
diff --git a/Engine/Source/Runtime/CoreUObject/Private/Serialization/ArchiveUObject.cpp b/Engine/Source/Runtime/CoreUObject/Private/Serialization/ArchiveUObject.cpp
--- a/Engine/Source/Runtime/CoreUObject/Private/Serialization/ArchiveUObject.cpp
+++ b/Engine/Source/Runtime/CoreUObject/Private/Serialization/ArchiveUObject.cpp
@@ -223,6 +223,12 @@ void FSerializedPropertyScope::PopProperty()
 	}
 }
 
+const TMap<UObject*, TArray<FProperty*>>& FArchiveReplaceObjectRefBase::GetReplacedReferences() const
+{ 
+	ensure(bTrackReplacedReferences);
+	return ReplacedReferences; 
+}
+
 void FArchiveReplaceObjectRefBase::SerializeObject(UObject* ObjectToSerialize)
 {
 	// Simple FReferenceCollector proxy for FArchiveReplaceObjectRefBase
diff --git a/Engine/Source/Runtime/CoreUObject/Public/Serialization/ArchiveReplaceObjectAndStructPropertyRef.h b/Engine/Source/Runtime/CoreUObject/Public/Serialization/ArchiveReplaceObjectAndStructPropertyRef.h
--- a/Engine/Source/Runtime/CoreUObject/Public/Serialization/ArchiveReplaceObjectAndStructPropertyRef.h
+++ b/Engine/Source/Runtime/CoreUObject/Public/Serialization/ArchiveReplaceObjectAndStructPropertyRef.h
@@ -20,14 +20,17 @@ public:
 	/**
 	 * Initializes variables and starts the serialization search
 	 *
-	 * @param InSearchObject		The object to start the search on
-	 * @param InReplacementMap		Map of objects to find -> objects to replace them with (null zeros them)
-	 * @param bNullPrivateRefs		Whether references to non-public objects not contained within the SearchObject
-	 *								should be set to null
-	 * @param bIgnoreOuterRef		Whether we should replace Outer pointers on Objects.
-	 * @param bIgnoreArchetypeRef	Whether we should replace the ObjectArchetype reference on Objects.
-	 * @param bDelayStart			Specify true to prevent the constructor from starting the process.  Allows child classes' to do initialization stuff in their ctor
+	 * @param InSearchObject        The object to start the search on
+	 * @param InReplacementMap      Map of objects to find -> objects to replace them with (null zeros them)
+	 * @param Flags                 Enum specifying behavior of archive
 	 */
+	FArchiveReplaceObjectAndStructPropertyRef(UObject* InSearchObject, const TMap<T*, T*>& InReplacementMap, EArchiveReplaceObjectFlags Flags = EArchiveReplaceObjectFlags::None)
+		: FArchiveReplaceObjectRef<T>(InSearchObject, InReplacementMap, Flags)
+	{
+	}
+
+
+	UE_DEPRECATED(5.0, "Use version that supplies flags via enum.")
 	FArchiveReplaceObjectAndStructPropertyRef
 	(
 		UObject* InSearchObject,
@@ -65,7 +68,10 @@ public:
 					// Do we need to verify the existence of ReplaceWithField? Theoretically it could be missing on the new version
 					// of the owner struct and in this case we still don't want to keep the stale old property pointer around so it's safer to null it
 					InField = ReplaceWithField;
-					this->ReplacedReferences.FindOrAdd(OldOwnerStruct).AddUnique(this->GetSerializedProperty());
+					if (this->bTrackReplacedReferences)
+					{
+						this->ReplacedReferences.FindOrAdd(OldOwnerStruct).AddUnique(this->GetSerializedProperty());
+					}
 					this->Count++;
 				}
 				// A->IsIn(A) returns false, but we don't want to NULL that reference out, so extra check here.
diff --git a/Engine/Source/Runtime/CoreUObject/Public/Serialization/ArchiveReplaceObjectRef.h b/Engine/Source/Runtime/CoreUObject/Public/Serialization/ArchiveReplaceObjectRef.h
--- a/Engine/Source/Runtime/CoreUObject/Public/Serialization/ArchiveReplaceObjectRef.h
+++ b/Engine/Source/Runtime/CoreUObject/Public/Serialization/ArchiveReplaceObjectRef.h
@@ -24,7 +24,7 @@ public:
 	/**
 	* Returns a reference to the replaced references map
 	*/
-	const TMap<UObject*, TArray<FProperty*>>& GetReplacedReferences() const { return ReplacedReferences; }
+	const TMap<UObject*, TArray<FProperty*>>& GetReplacedReferences() const;
 
 	/**
 	* Returns the name of this archive.
@@ -96,10 +96,10 @@ protected:
 	void SerializeObject(UObject* ObjectToSerialize);
 
 	/** Initial object to start the reference search from */
-	UObject* SearchObject;
+	UObject* SearchObject = nullptr;
 
 	/** The number of times encountered */
-	int32 Count;
+	int32 Count = 0;
 
 	/** List of objects that have already been serialized */
 	TSet<UObject*> SerializedObjects;
@@ -110,11 +110,14 @@ protected:
 	/** Map of referencing objects to referencing properties */
 	TMap<UObject*, TArray<FProperty*>> ReplacedReferences;
 
+	/** Whether to populate the map of referencing objects to referencing properties */
+	bool bTrackReplacedReferences = false;
+
 	/**
 	* Whether references to non-public objects not contained within the SearchObject
 	* should be set to null
 	*/
-	bool bNullPrivateReferences;
+	bool bNullPrivateReferences = false;
 
 	/**
 	* Whether unresolved references to objects in other packages can be ignored when searching
@@ -122,6 +125,30 @@ protected:
 	TOptional<bool> CanIgnoreUnresolvedImports;
 };
 
+enum class EArchiveReplaceObjectFlags
+{
+	None                      = 0,
+
+	// References to non-public objects not contained within the SearchObject should be set to null
+	NullPrivateRefs            = 1 << 0,
+
+	// Do not replace Outer pointers on Objects.
+	IgnoreOuterRef             = 1 << 1,
+
+	// Do not replace the ObjectArchetype reference on Objects.
+	IgnoreArchetypeRef         = 1 << 2,
+
+	// Prevent the constructor from starting the process. Allows child classes to do initialization stuff in their constructor.
+	DelayStart                 = 1 << 3,
+
+	// Replace the ClassGeneratedBy reference in UClass
+	IncludeClassGeneratedByRef = 1 << 4,
+
+	// Populate the map of referencing objects to referencing properties
+	TrackReplacedReferences    = 1 << 5
+};
+ENUM_CLASS_FLAGS(EArchiveReplaceObjectFlags)
+
 /*----------------------------------------------------------------------------
 	FArchiveReplaceObjectRef.
 ----------------------------------------------------------------------------*/
@@ -140,14 +167,32 @@ public:
 	/**
 	 * Initializes variables and starts the serialization search
 	 *
-	 * @param InSearchObject		The object to start the search on
-	 * @param ReplacementMap		Map of objects to find -> objects to replace them with (null zeros them)
-	 * @param bNullPrivateRefs		Whether references to non-public objects not contained within the SearchObject
-	 *								should be set to null
-	 * @param bIgnoreOuterRef		Whether we should replace Outer pointers on Objects.
-	 * @param bIgnoreArchetypeRef	Whether we should replace the ObjectArchetype reference on Objects.
-	 * @param bDelayStart			Specify true to prevent the constructor from starting the process.  Allows child classes' to do initialization stuff in their ctor
+	 * @param InSearchObject        The object to start the search on
+	 * @param ReplacementMap        Map of objects to find -> objects to replace them with (null zeros them)
+	 * @param Flags                 Enum specifying behavior of archive 		
 	 */
+	FArchiveReplaceObjectRef(UObject* InSearchObject, const TMap<T*, T*>& InReplacementMap, EArchiveReplaceObjectFlags Flags = EArchiveReplaceObjectFlags::None)
+		: ReplacementMap(InReplacementMap)
+	{
+		bTrackReplacedReferences = !!(Flags & EArchiveReplaceObjectFlags::TrackReplacedReferences);
+
+		SearchObject = InSearchObject;
+		Count = 0;
+		bNullPrivateReferences = !!(Flags & EArchiveReplaceObjectFlags::NullPrivateRefs);
+
+		ArIsObjectReferenceCollector = true;
+		ArIsModifyingWeakAndStrongReferences = true;		// Also replace weak references too!
+		ArIgnoreArchetypeRef = !!(Flags & EArchiveReplaceObjectFlags::IgnoreArchetypeRef);
+		ArIgnoreOuterRef = !!(Flags & EArchiveReplaceObjectFlags::IgnoreOuterRef);
+		ArIgnoreClassGeneratedByRef = !(Flags & EArchiveReplaceObjectFlags::IncludeClassGeneratedByRef);
+
+		if (!(Flags & EArchiveReplaceObjectFlags::DelayStart))
+		{
+			SerializeSearchObject();
+		}
+	}
+
+	UE_DEPRECATED(5.0, "Use version that supplies flags via enum.")
 	FArchiveReplaceObjectRef
 	(
 		UObject* InSearchObject,
@@ -158,22 +203,13 @@ public:
 		bool bDelayStart = false,
 		bool bIgnoreClassGeneratedByRef = true
 	)
-	: ReplacementMap(inReplacementMap)
+	: FArchiveReplaceObjectRef(InSearchObject, inReplacementMap,
+		  (bNullPrivateRefs ? EArchiveReplaceObjectFlags::NullPrivateRefs : EArchiveReplaceObjectFlags::None)
+		| (bIgnoreOuterRef ? EArchiveReplaceObjectFlags::IgnoreOuterRef : EArchiveReplaceObjectFlags::None)
+		| (bIgnoreArchetypeRef ? EArchiveReplaceObjectFlags::IgnoreArchetypeRef : EArchiveReplaceObjectFlags::None)
+		| (bDelayStart ? EArchiveReplaceObjectFlags::DelayStart : EArchiveReplaceObjectFlags::None)
+		| (!bIgnoreClassGeneratedByRef ? EArchiveReplaceObjectFlags::IncludeClassGeneratedByRef : EArchiveReplaceObjectFlags::None))
 	{
-		SearchObject = InSearchObject;
-		Count = 0;
-		bNullPrivateReferences = bNullPrivateRefs;
-
-		ArIsObjectReferenceCollector = true;
-		ArIsModifyingWeakAndStrongReferences = true;		// Also replace weak references too!
-		ArIgnoreArchetypeRef = bIgnoreArchetypeRef;
-		ArIgnoreOuterRef = bIgnoreOuterRef;
-		ArIgnoreClassGeneratedByRef = bIgnoreClassGeneratedByRef;
-
-		if ( !bDelayStart )
-		{
-			SerializeSearchObject();
-		}
 	}
 
 	/**
@@ -205,11 +241,13 @@ public:
 		if (Obj != NULL)
 		{
 			// If these match, replace the reference
-			T* const* ReplaceWith = (T*const*)((const TMap<UObject*,UObject*>*)&ReplacementMap)->Find(Obj);
-			if ( ReplaceWith != NULL )
+			if (T* const* ReplaceWith = (T* const*)((const TMap<UObject*, UObject*>*)&ReplacementMap)->Find(Obj))
 			{
 				Obj = *ReplaceWith;
-				ReplacedReferences.FindOrAdd(Obj).AddUnique(GetSerializedProperty());
+				if (bTrackReplacedReferences)
+				{
+					ReplacedReferences.FindOrAdd(Obj).AddUnique(GetSerializedProperty());
+				}
 				Count++;
 			}
 			// A->IsIn(A) returns false, but we don't want to NULL that reference out, so extra check here.
diff --git a/Engine/Source/Runtime/CoreUObject/Public/Serialization/ArchiveReplaceOrClearExternalReferences.h b/Engine/Source/Runtime/CoreUObject/Public/Serialization/ArchiveReplaceOrClearExternalReferences.h
--- a/Engine/Source/Runtime/CoreUObject/Public/Serialization/ArchiveReplaceOrClearExternalReferences.h
+++ b/Engine/Source/Runtime/CoreUObject/Public/Serialization/ArchiveReplaceOrClearExternalReferences.h
@@ -25,7 +25,7 @@ public:
 		, const TMap<T*, T*>& InReplacementMap
 		, UPackage* InDestPackage
 		, bool bDelayStart = false )
-		: TSuper(InSearchObject, InReplacementMap, false, false, false, true)
+		: TSuper(InSearchObject, InReplacementMap, EArchiveReplaceObjectFlags::DelayStart)
 		, DestPackage(InDestPackage)
 	{
 		if (!bDelayStart)
diff --git a/Engine/Source/Runtime/Engine/Private/World.cpp b/Engine/Source/Runtime/Engine/Private/World.cpp
--- a/Engine/Source/Runtime/Engine/Private/World.cpp
+++ b/Engine/Source/Runtime/Engine/Private/World.cpp
@@ -873,12 +873,9 @@ void UWorld::PostDuplicate(bool bDuplicateForPIE)
 	// Now replace references from the old textures/classes to the new ones, if any were duplicated
 	if (ReplacementMap.Num() > 0)
 	{
-		const bool bNullPrivateRefs = false;
-		const bool bIgnoreOuterRef = true;
-		const bool bIgnoreArchetypeRef = false;
 		for (UObject* Obj : ObjectsToFixReferences)
 		{
-			FArchiveReplaceObjectRef<UObject> ReplaceAr(Obj, ReplacementMap, bNullPrivateRefs, bIgnoreOuterRef, bIgnoreArchetypeRef);
+			FArchiveReplaceObjectRef<UObject> ReplaceAr(Obj, ReplacementMap, EArchiveReplaceObjectFlags::IgnoreOuterRef);
 		}
 		// PostEditChange is required for some objects to react to the change, e.g. update render-thread proxies
 		for (UObject* Obj : ObjectsToFixReferences)
