EditorDomain: Add the projectname to the EditorDomain key to prevent loading an EditorDomain package in one project that was saved in another project and is therefore possibly incompatible.
#rnx

diff --git a/Engine/Source/Editor/UnrealEd/Private/EditorDomain/EditorDomainUtils.cpp b/Engine/Source/Editor/UnrealEd/Private/EditorDomain/EditorDomainUtils.cpp
--- a/Engine/Source/Editor/UnrealEd/Private/EditorDomain/EditorDomainUtils.cpp
+++ b/Engine/Source/Editor/UnrealEd/Private/EditorDomain/EditorDomainUtils.cpp
@@ -186,16 +186,16 @@ EPackageDigestResult AppendPackageDigest(FBlake3& Writer, EDomainUse& OutEditorD
 {
 	OutEditorDomainUse = EDomainUse::LoadEnabled | EDomainUse::SaveEnabled;
 	
-	FPackageFileVersion CurrentFileVersionUE = GPackageFileUEVersion;
-	int32 CurrentFileVersionLicenseeUE = GPackageFileLicenseeUEVersion;
-	Writer.Update(EditorDomainVersion, FCString::Strlen(EditorDomainVersion)*sizeof(TCHAR));
+	FStringView ProjectName(FApp::GetProjectName());
+	Writer.Update(ProjectName.GetData(), ProjectName.Len() * sizeof(ProjectName[0]));
+	Writer.Update(EditorDomainVersion, FCString::Strlen(EditorDomainVersion)*sizeof(EditorDomainVersion[0]));
 	uint8 EditorDomainSaveUnversioned = GetEditorDomainSaveUnversioned() ? 1 : 0;
 	Writer.Update(&EditorDomainSaveUnversioned, sizeof(EditorDomainSaveUnversioned));
 	PRAGMA_DISABLE_DEPRECATION_WARNINGS;
 	Writer.Update(&PackageData.PackageGuid, sizeof(PackageData.PackageGuid));
 	PRAGMA_ENABLE_DEPRECATION_WARNINGS;
-	Writer.Update(&CurrentFileVersionUE, sizeof(CurrentFileVersionUE));
-	Writer.Update(&CurrentFileVersionLicenseeUE, sizeof(CurrentFileVersionLicenseeUE));
+	Writer.Update(&GPackageFileUEVersion, sizeof(GPackageFileUEVersion));
+	Writer.Update(&GPackageFileLicenseeUEVersion, sizeof(GPackageFileLicenseeUEVersion));
 	TArray<int32> CustomVersionHandles;
 	// Reserve 10 custom versions per class times 100 classes per package times twice (once in package, once in class)
 	CustomVersionHandles.Reserve(10*100*2);
