Folders: Fix Folder rename comparison not checking root (causing folder to be removed when same folder path was found in different root)

diff --git a/Engine/Source/Editor/UnrealEd/Private/EditorActorFolders.cpp b/Engine/Source/Editor/UnrealEd/Private/EditorActorFolders.cpp
--- a/Engine/Source/Editor/UnrealEd/Private/EditorActorFolders.cpp
+++ b/Engine/Source/Editor/UnrealEd/Private/EditorActorFolders.cpp
@@ -589,7 +589,7 @@ bool FActorFolders::RenameFolderInWorld(UWorld& World, const FFolder& OldPath, c
 			bool ContainsFolder = false;
 			for (const auto& FolderPair : FoldersInWorld.Folders)
 			{
-				if (FolderPair.Key.GetPath().IsEqual(NewFolder.GetPath(), ENameCase::CaseSensitive))
+				if (FolderPair.Key.GetRootObject() == NewFolder.GetRootObject() && FolderPair.Key.GetPath().IsEqual(NewFolder.GetPath(), ENameCase::CaseSensitive))
 				{
 					ContainsFolder = true;
 					break;
