WorldPartitionResaveActorsBuilder
- Added -verbose commandline option to enable verbose logs.
- Added -reportonly commandline option to only report actions.
- Added -OnlyResaveOnActorDescDiff commandline option to only resave actors with an actor descriptor mismatch.

diff --git a/Engine/Source/Editor/UnrealEd/Private/WorldPartition/WorldPartitionResaveActorsBuilder.cpp b/Engine/Source/Editor/UnrealEd/Private/WorldPartition/WorldPartitionResaveActorsBuilder.cpp
--- a/Engine/Source/Editor/UnrealEd/Private/WorldPartition/WorldPartitionResaveActorsBuilder.cpp
+++ b/Engine/Source/Editor/UnrealEd/Private/WorldPartition/WorldPartitionResaveActorsBuilder.cpp
@@ -20,6 +20,19 @@ UWorldPartitionResaveActorsBuilder::UWorldPartitionResaveActorsBuilder(const FOb
 	: Super(ObjectInitializer)
 {}
 
+void UWorldPartitionResaveActorsBuilder::ParseCommandline()
+{
+	FParse::Value(FCommandLine::Get(), TEXT("ActorClass="), ActorClassName);
+
+	bReportOnly = FParse::Param(FCommandLine::Get(), TEXT("ReportOnly"));
+	bOnlyResaveOnActorDescDiff = FParse::Param(FCommandLine::Get(), TEXT("OnlyResaveOnActorDescDiff"));
+	
+	if(FParse::Param(FCommandLine::Get(), TEXT("Verbose")))
+	{
+		LogWorldPartitionResaveActorsBuilder.SetVerbosity(ELogVerbosity::Verbose);
+	}
+}
+
 bool UWorldPartitionResaveActorsBuilder::PreRun(UWorld* World, FPackageSourceControlHelper& PackageHelper)
 {
 	TArray<FString> Tokens, Switches;
@@ -82,10 +95,12 @@ bool UWorldPartitionResaveActorsBuilder::RunInternal(UWorld* World, const FCellI
 {
 	FPackageSourceControlHelper SCCHelper;
 
+	ParseCommandline();
 	UPackage* WorldPackage = World->GetPackage();
 
 	int32 LoadCount = 0;
 	int32 SaveCount = 0;
+	int32 SkipCount = 0;
 	int32 FailCount = 0;
 
 	// Actor Class Filter
@@ -133,17 +148,18 @@ bool UWorldPartitionResaveActorsBuilder::RunInternal(UWorld* World, const FCellI
 	{
 		World->PersistentLevel->ActorPackagingScheme = EActorPackagingScheme::Reduced;
 
-		FWorldPartitionHelpers::ForEachActorWithLoading(WorldPartition, ActorClass, [this, &LoadCount, &SaveCount, &FailCount, &SCCHelper, &PackagesToDelete, WorldPackage](const FWorldPartitionActorDesc* ActorDesc)
+		FWorldPartitionHelpers::ForEachActorWithLoading(WorldPartition, ActorClass, [this, &LoadCount, &SaveCount, &FailCount, &SkipCount, &SCCHelper, &PackagesToDelete, WorldPackage](const FWorldPartitionActorDesc* ActorDesc)
 		{
 			ON_SCOPE_EXIT
 			{
-				UE_LOG(LogWorldPartitionResaveActorsBuilder, Display, TEXT("Processed %d packages (%d Saved / %d Failed)"), LoadCount, SaveCount, FailCount);
+				UE_LOG(LogWorldPartitionResaveActorsBuilder, Display, TEXT("Processed %d packages (%d Saved / %d Skipped / %d Failed)"), SaveCount + SkipCount + FailCount, SaveCount, SkipCount, FailCount);
 			};
 
 			AActor* Actor = ActorDesc->GetActor();
 
 			if (!Actor)
 			{
+				UE_LOG(LogWorldPartitionResaveActorsBuilder, Error, TEXT("Error loading package %s."), *ActorDesc->GetActorPackage().ToString());
 				PackagesToDelete.Add(ActorDesc->GetActorPackage().ToString());
 				++FailCount;
 				return true;
@@ -277,33 +293,69 @@ bool UWorldPartitionResaveActorsBuilder::RunInternal(UWorld* World, const FCellI
 				UE_LOG(LogWorldPartitionResaveActorsBuilder, Log, TEXT("Package %s needs to be resaved."), *Package->GetName());
 			}
 
-			if (!bReportOnly)
+			if (bReportOnly || SCCHelper.Checkout(Package))
 			{
-				if (SCCHelper.Checkout(Package))
+				bool bShouldResaveActor = !bOnlyResaveOnActorDescDiff;
+
+				if (bOnlyResaveOnActorDescDiff)
+				{
+					TUniquePtr<FWorldPartitionActorDesc> CurActorDesc = Actor->CreateActorDesc();
+					const FWorldPartitionActorDesc& OldActorDesc = *ActorDesc;
+					const FWorldPartitionActorDesc& NewActorDesc = *CurActorDesc.Get();
+
+					// Bounds can change if the source asset changed
+					if (!OldActorDesc.GetBounds().Equals(NewActorDesc.GetBounds(), 1.0f))
+					{
+						bShouldResaveActor = true;
+					}
+					// Grid placement can have changed from defaults
+					else if (OldActorDesc.GetGridPlacement() != NewActorDesc.GetGridPlacement())
+					{
+						bShouldResaveActor = true;
+					}
+					// Editor-only flag can have changed from defaults
+					else if (OldActorDesc.GetActorIsEditorOnly() != NewActorDesc.GetActorIsEditorOnly())
+					{
+						bShouldResaveActor = true;
+					}
+
+					if (bShouldResaveActor)
+					{
+						UE_LOG(LogWorldPartitionResaveActorsBuilder, Verbose, TEXT("Actor descriptor mismatch: [%s]<->[%s]"), *OldActorDesc.ToString(), *NewActorDesc.ToString());
+					}
+				}
+
+				// Save package
+				if (!bReportOnly && bShouldResaveActor)
 				{
-					// Save package
 					FString PackageFileName = SourceControlHelpers::PackageFilename(Package);
-					FSavePackageArgs SaveArgs;
-					SaveArgs.TopLevelFlags = RF_Standalone;
-					if (!UPackage::SavePackage(Package, nullptr, *PackageFileName, SaveArgs))
+					if (!UPackage::SavePackage(Package, nullptr, RF_Standalone, *PackageFileName))
 					{
 						UE_LOG(LogWorldPartitionResaveActorsBuilder, Error, TEXT("Error saving package %s."), *Package->GetName());
 						++FailCount;
 						return true;
 					}
-
-					UE_LOG(LogWorldPartitionResaveActorsBuilder, Display, TEXT("Saved package %s."), *Package->GetName());
-					++SaveCount;
-					return true;
 				}
 				else
 				{
-					// It is possible the resave can't checkout everything. Continue processing.
-					UE_LOG(LogWorldPartitionResaveActorsBuilder, Warning, TEXT("Error checking out package %s."), *Package->GetName());
-					++FailCount;
+					UE_LOG(LogWorldPartitionResaveActorsBuilder, Display, TEXT("Skipped package %s."), *Package->GetName());
+					++SkipCount;
 					return true;
 				}
+
+				// It is possible the resave can't checkout everything. Continue processing.
+				UE_LOG(LogWorldPartitionResaveActorsBuilder, Display, TEXT("Saved package %s."), *Package->GetName());
+				++SaveCount;
+				return true;
 			}
+			else
+			{
+				// It is possible the resave can't checkout everything. Continue processing.
+				UE_LOG(LogWorldPartitionResaveActorsBuilder, Warning, TEXT("Error checking out package %s."), *Package->GetName());
+				++FailCount;
+				return true;
+			}
+			
 		}
 
 		return true;
@@ -311,28 +363,31 @@ bool UWorldPartitionResaveActorsBuilder::RunInternal(UWorld* World, const FCellI
 
 	if (!bReportOnly)
 	{
+		UPackage::WaitForAsyncFileWrites();
 		SCCHelper.Delete(PackagesToDelete);
-	}
-
-	if (bSwitchActorPackagingSchemeToReduced && !bReportOnly)
-	{
-		// checkout world package
-		if (!SCCHelper.Checkout(WorldPackage))
+		
+		if (bSwitchActorPackagingSchemeToReduced)
 		{
-			UE_LOG(LogWorldPartitionResaveActorsBuilder, Error, TEXT("Error checking out world package %s."), *WorldPackage->GetName());
-			return false;
-		}
+			// checkout world package
+			if (!SCCHelper.Checkout(WorldPackage))
+			{
+				UE_LOG(LogWorldPartitionResaveActorsBuilder, Error, TEXT("Error checking out world package %s."), *WorldPackage->GetName());
+				return false;
+			}
 
-		// Save world package
-		FString PackageFileName = SourceControlHelpers::PackageFilename(WorldPackage);
-		FSavePackageArgs SaveArgs;
-		SaveArgs.TopLevelFlags = RF_Standalone;
-		if (!UPackage::SavePackage(WorldPackage, nullptr, *PackageFileName, SaveArgs))
-		{
-			UE_LOG(LogWorldPartitionResaveActorsBuilder, Error, TEXT("Error saving world %s."), *WorldPackage->GetName());
-			return false;
+			// Save world package
+			FString PackageFileName = SourceControlHelpers::PackageFilename(WorldPackage);
+			FSavePackageArgs SaveArgs;
+			SaveArgs.TopLevelFlags = RF_Standalone;
+			if (!UPackage::SavePackage(WorldPackage, nullptr, *PackageFileName, SaveArgs))
+			{
+				UE_LOG(LogWorldPartitionResaveActorsBuilder, Error, TEXT("Error saving world %s."), *WorldPackage->GetName());
+				return false;
+			}
 		}
 	}
 
+	
+
 	return true;
 }
\ No newline at end of file
diff --git a/Engine/Source/Runtime/Engine/Classes/GameFramework/Actor.h b/Engine/Source/Runtime/Engine/Classes/GameFramework/Actor.h
--- a/Engine/Source/Runtime/Engine/Classes/GameFramework/Actor.h
+++ b/Engine/Source/Runtime/Engine/Classes/GameFramework/Actor.h
@@ -72,6 +72,18 @@ enum class EActorGridPlacement : uint8
 	AlwaysLoaded,
 	None UMETA(Hidden)
 };
+
+inline const TCHAR* GetActorGridPlacementName(EActorGridPlacement ActorGridPlacement)
+{
+	switch(ActorGridPlacement)
+	{
+	case EActorGridPlacement::Bounds: return TEXT("Bounds");
+	case EActorGridPlacement::Location: return TEXT("Location");
+	case EActorGridPlacement::AlwaysLoaded: return TEXT("AlwaysLoaded");
+	default: check(0);
+	}
+	return TEXT("Invalid");
+}
 #endif
 
 ENGINE_API DECLARE_LOG_CATEGORY_EXTERN(LogActor, Log, Warning);
@@ -887,19 +899,20 @@ public:
 	bool CanPlayFromHere() const { return bCanPlayFromHere; }
 
 	/**
-	 * Creates an uninitialized actor descriptor. Can be overriden.
+	 * Creates an uninitialized actor descriptor from this actor. Meant to be called on the class CDO.
 	 */
 	virtual TUniquePtr<class FWorldPartitionActorDesc> CreateClassActorDesc() const;
 
+public:
 	/**
-	 * Creates an initialized actor descriptor.
+	 * Creates an initialized actor descriptor from this actor.
 	 */
 	TUniquePtr<class FWorldPartitionActorDesc> CreateActorDesc() const;
 
 	/**
-	 * Creates an uninitialized actor descriptor.
+	 * Creates an uninitialized actor descriptor from a specific class.
 	 */
-	static TUniquePtr<class FWorldPartitionActorDesc> CreateClassActorDesc(const TSubclassOf<AActor>& ActorClass);
+	static TUniquePtr<class FWorldPartitionActorDesc> StaticCreateClassActorDesc(const TSubclassOf<AActor>& ActorClass);
 #endif // WITH_EDITOR
 
 public:
diff --git a/Engine/Source/Runtime/Engine/Private/ActorEditor.cpp b/Engine/Source/Runtime/Engine/Private/ActorEditor.cpp
--- a/Engine/Source/Runtime/Engine/Private/ActorEditor.cpp
+++ b/Engine/Source/Runtime/Engine/Private/ActorEditor.cpp
@@ -853,7 +853,7 @@ TUniquePtr<FWorldPartitionActorDesc> AActor::CreateActorDesc() const
 	return ActorDesc;
 }
 
-TUniquePtr<class FWorldPartitionActorDesc> AActor::CreateClassActorDesc(const TSubclassOf<AActor>& ActorClass)
+TUniquePtr<class FWorldPartitionActorDesc> AActor::StaticCreateClassActorDesc(const TSubclassOf<AActor>& ActorClass)
 {
 	return CastChecked<AActor>(ActorClass->GetDefaultObject())->CreateClassActorDesc();
 }
diff --git a/Engine/Source/Runtime/Engine/Private/WorldPartition/ActorDescContainer.cpp b/Engine/Source/Runtime/Engine/Private/WorldPartition/ActorDescContainer.cpp
--- a/Engine/Source/Runtime/Engine/Private/WorldPartition/ActorDescContainer.cpp
+++ b/Engine/Source/Runtime/Engine/Private/WorldPartition/ActorDescContainer.cpp
@@ -81,7 +81,7 @@ void UActorDescContainer::Initialize(UWorld* InWorld, FName InPackageName, TFunc
 				ActorDescInitData.ActorPath = InAssetData.ObjectPath;
 				FBase64::Decode(ActorMetaDataStr, ActorDescInitData.SerializedData);
 
-				TUniquePtr<FWorldPartitionActorDesc> NewActorDesc(AActor::CreateClassActorDesc(ActorDescInitData.NativeClass));
+				TUniquePtr<FWorldPartitionActorDesc> NewActorDesc(AActor::StaticCreateClassActorDesc(ActorDescInitData.NativeClass));
 
 				NewActorDesc->Init(this, ActorDescInitData);
 			
diff --git a/Engine/Source/Runtime/Engine/Private/WorldPartition/WorldPartitionActorDesc.cpp b/Engine/Source/Runtime/Engine/Private/WorldPartition/WorldPartitionActorDesc.cpp
--- a/Engine/Source/Runtime/Engine/Private/WorldPartition/WorldPartitionActorDesc.cpp
+++ b/Engine/Source/Runtime/Engine/Private/WorldPartition/WorldPartitionActorDesc.cpp
@@ -190,7 +190,20 @@ void FWorldPartitionActorDesc::TransformInstance(const FString& From, const FStr
 
 FString FWorldPartitionActorDesc::ToString() const
 {
-	return FString::Printf(TEXT("Guid:%s Class:%s Name:%s SpatiallyLoaded:%s"), *Guid.ToString(), *Class.ToString(), *FPaths::GetExtension(ActorPath.ToString()), (GridPlacement != EActorGridPlacement::AlwaysLoaded) ? TEXT("True") : TEXT("False"));
+	return FString::Printf(
+		TEXT("Guid:%s Class:%s Name:%s SpatiallyLoaded:%s Bounds:%s GridPlacement:%s RuntimeGrid:%s EditorOnly:%s LevelBoundsRelevant:%s HLODRelevant:%s FolderPath:%s"), 
+		*Guid.ToString(), 
+		*Class.ToString(), 
+		*FPaths::GetExtension(ActorPath.ToString()), 
+		(GridPlacement != EActorGridPlacement::AlwaysLoaded) ? TEXT("True") : TEXT("False"),
+		*GetBounds().ToString(),
+		GetActorGridPlacementName(GridPlacement),
+		*RuntimeGrid.ToString(),
+		bActorIsEditorOnly ? TEXT("true") : TEXT("false"),
+		bLevelBoundsRelevant ? TEXT("true") : TEXT("false"),
+		bActorIsHLODRelevant ? TEXT("true") : TEXT("false"),
+		*FolderPath.ToString()
+	);
 }
 
 void FWorldPartitionActorDesc::Serialize(FArchive& Ar)
