Added PreCreatePIEInstances as a PIE lifecycle event on the editor engine (similar to PreCreatePIEServerInstance but called always, not just when a dedicated server is requested)

diff --git a/Engine/Source/Editor/UnrealEd/Classes/Editor/EditorEngine.h b/Engine/Source/Editor/UnrealEd/Classes/Editor/EditorEngine.h
--- a/Engine/Source/Editor/UnrealEd/Classes/Editor/EditorEngine.h
+++ b/Engine/Source/Editor/UnrealEd/Classes/Editor/EditorEngine.h
@@ -2646,6 +2646,9 @@ public:
 	UE_DEPRECATED(4.25, "This is now handled as part of OnLoginPIEComplete_Deferred.")
 	bool CreatePIEWorldFromLogin(FWorldContext& PieWorldContext, EPlayNetMode PlayNetMode, FPieLoginStruct& DataStruct);
 
+	/** Called before creating PIE instance(s). */
+	virtual FGameInstancePIEResult PreCreatePIEInstances(const bool bAnyBlueprintErrors, const bool bStartInSpectatorMode, const float PIEStartTime, const bool bSupportsOnlinePIE, int32& InNumOnlinePIEInstances);
+
 	/** Called before creating a PIE server instance. */
 	virtual FGameInstancePIEResult PreCreatePIEServerInstance(const bool bAnyBlueprintErrors, const bool bStartInSpectatorMode, const float PIEStartTime, const bool bSupportsOnlinePIE, int32& InNumOnlinePIEInstances);
 
diff --git a/Engine/Source/Editor/UnrealEd/Private/PlayLevel.cpp b/Engine/Source/Editor/UnrealEd/Private/PlayLevel.cpp
--- a/Engine/Source/Editor/UnrealEd/Private/PlayLevel.cpp
+++ b/Engine/Source/Editor/UnrealEd/Private/PlayLevel.cpp
@@ -1497,6 +1497,11 @@ static bool ShowBlueprintErrorDialog( TArray<UBlueprint*> ErroredBlueprints )
 	return ButtonPressed == 0;
 }
 
+FGameInstancePIEResult UEditorEngine::PreCreatePIEInstances(const bool bAnyBlueprintErrors, const bool bStartInSpectatorMode, const float PIEStartTime, const bool bSupportsOnlinePIE, int32& InNumOnlinePIEInstances)
+{
+	return FGameInstancePIEResult::Success();
+}
+
 FGameInstancePIEResult UEditorEngine::PreCreatePIEServerInstance(const bool bAnyBlueprintErrors, const bool bStartInSpectatorMode, const float PIEStartTime, const bool bSupportsOnlinePIE, int32& InNumOnlinePIEInstances)
 {
 	return FGameInstancePIEResult::Success();
@@ -2677,15 +2682,23 @@ void UEditorEngine::StartPlayInEditorSession(FRequestPlaySessionParams& InReques
 		const bool bNetModeRequiresSeparateServer = NetMode == EPlayNetMode::PIE_Client;
 		const bool bLaunchExtraServerAnyways = InRequestParams.EditorPlaySettings->bLaunchSeparateServer;
 		const bool bNeedsServer = bNetModeRequiresSeparateServer || bLaunchExtraServerAnyways;
-		
+
+		// Allow the engine to cancel the PIE request if needed.
+		FGameInstancePIEResult PreCreateResult = PreCreatePIEInstances(
+			ErroredBlueprints.Num() > 0, false /*bStartInSpectorMode*/, PIEStartTime, SupportsOnlinePIE(), PlayInEditorSessionInfo->NumOutstandingPIELogins);
+		if (!PreCreateResult.IsSuccess())
+		{
+			UE_LOG(LogPlayLevel, Warning, TEXT("PlayInEditor Session failed (%s::PreCreatePIEInstances) and will not be started."), *GetClass()->GetName());
+			return;
+		}
 
 		// If they require a separate server we'll give the EditorEngine a chance to handle any additional prep-work.
 		if (bNeedsServer)
 		{
 			// Allow the engine to cancel the server request if needed.
-			FGameInstancePIEResult PreCreateResult = PreCreatePIEServerInstance(
-			ErroredBlueprints.Num() > 0, false /*bStartInSpectorMode*/, PIEStartTime, true, PlayInEditorSessionInfo->NumOutstandingPIELogins);
-			if (!PreCreateResult.IsSuccess())
+			FGameInstancePIEResult ServerPreCreateResult = PreCreatePIEServerInstance(
+				ErroredBlueprints.Num() > 0, false /*bStartInSpectorMode*/, PIEStartTime, true, PlayInEditorSessionInfo->NumOutstandingPIELogins);
+			if (!ServerPreCreateResult.IsSuccess())
 			{
 				// ToDo: This will skip client creation as well right now. Probably OK though.
 				UE_LOG(LogPlayLevel, Warning, TEXT("PlayInEditor Session Server failed Pre-Create and will not be started."));
