Implements CookGlobalShaders'  -live that recompile/deploy/reload global shader automatically when a shader source file has changed

diff --git a/Engine/Source/Editor/UnrealEd/Classes/Commandlets/CookGlobalShadersCommandlet.h b/Engine/Source/Editor/UnrealEd/Classes/Commandlets/CookGlobalShadersCommandlet.h
--- a/Engine/Source/Editor/UnrealEd/Classes/Commandlets/CookGlobalShadersCommandlet.h
+++ b/Engine/Source/Editor/UnrealEd/Classes/Commandlets/CookGlobalShadersCommandlet.h
@@ -5,8 +5,12 @@
 #include "CoreMinimal.h"
 #include "UObject/ObjectMacros.h"
 #include "Commandlets/Commandlet.h"
+#include "Interfaces/ITargetPlatform.h"
 #include "CookGlobalShadersCommandlet.generated.h"
 
+struct FFileChangeData;
+class UCookGlobalShadersDeviceHelperBase;
+
 UCLASS(config=Editor)
 class UCookGlobalShadersCommandlet : public UCommandlet
 {
@@ -15,6 +19,28 @@ class UCookGlobalShadersCommandlet : public UCommandlet
 	//~ Begin UCommandlet Interface
 	virtual int32 Main(const FString& Params) override;
 	//~ End UCommandlet Interface
+
+private:
+	// All the derived info from command line parameters.
+	bool bDeployToDevice = false;
+	bool bCopyToStaged = false;
+	bool bExecuteReload = false;
+	FString DeployFolder;
+	FString PlatformName;
+	ITargetPlatform* TargetPlatform = nullptr;
+	ITargetDevicePtr TargetDevice = nullptr;
+	class UCookGlobalShadersDeviceHelperBase* DeviceHelper = nullptr;
+	TArray<FName> ShaderFormats;
+
+	// Cook, deploy and reload global shaders
+	void CookGlobalShaders() const;
+
+	// Callback when a shader source directory changes.
+	void HandleDirectoryChanged(const TArray<FFileChangeData>& InFileChangeDatas);
+
+	// Issues CookGlobalShaders() automatically when a change is detected to shader source directories.
+	void CookGlobalShadersOnDirectoriesChanges();
+
 };
 
 UCLASS(abstract, transient, MinimalAPI)
diff --git a/Engine/Source/Editor/UnrealEd/Private/Commandlets/CookGlobalShaderCommandlet.cpp b/Engine/Source/Editor/UnrealEd/Private/Commandlets/CookGlobalShaderCommandlet.cpp
--- a/Engine/Source/Editor/UnrealEd/Private/Commandlets/CookGlobalShaderCommandlet.cpp
+++ b/Engine/Source/Editor/UnrealEd/Private/Commandlets/CookGlobalShaderCommandlet.cpp
@@ -13,6 +13,8 @@
 #include "Interfaces/ITargetPlatformManagerModule.h"
 #include "Interfaces/ITargetPlatform.h"
 #include "DerivedDataCacheInterface.h"
+#include "DirectoryWatcherModule.h"
+#include "IDirectoryWatcher.h"
 
 DEFINE_LOG_CATEGORY_STATIC(LogCookGlobalShaders, Log, All);
 
@@ -26,10 +28,192 @@ bool UCookGlobalShadersDeviceHelperStaged::CopyFilesToDevice(class ITargetDevice
 		FString LocalFile = FileToCopy.Key;
 		FString RemoteFile = FPaths::Combine(StagedBuildPath, FileToCopy.Value);
 		bSuccess &= PlatformFile.CopyFile(*RemoteFile, *LocalFile);
+
+		UE_LOG(LogCookGlobalShaders, Warning, TEXT("%s -> %s"), *LocalFile, *RemoteFile);
 	}
 	return bSuccess;
 }
 
+void UCookGlobalShadersCommandlet::CookGlobalShaders() const
+{
+	// Cook shaders
+	UE_LOG(LogCookGlobalShaders, Log, TEXT("Cooking Global Shaders..."));
+	FString OutputDir = FPaths::ProjectSavedDir() / TEXT("CookGlobalShaders") / PlatformName;
+	TArray<uint8> OutGlobalShaderMap;
+	FShaderRecompileData Arguments(PlatformName, SP_NumPlatforms, ODSCRecompileCommand::Global, nullptr, nullptr, &OutGlobalShaderMap);
+	RecompileShadersForRemote(Arguments, OutputDir);
+
+	// Build list of files to copy
+	TArray<TPair<FString, FString>> FilesToCopy;
+	for (FName ShaderFormat : ShaderFormats)
+	{
+		const FString GlobalShaderCacheName = FPaths::Combine(OutputDir, TEXT("Engine"), TEXT("GlobalShaderCache-") + ShaderFormat.ToString() + TEXT(".bin"));
+		const FString OverrideGlobalShaderCacheName = FPaths::Combine(TEXT("Engine"), TEXT("OverrideGlobalShaderCache-") + ShaderFormat.ToString() + TEXT(".bin"));
+		FilesToCopy.Emplace(GlobalShaderCacheName, OverrideGlobalShaderCacheName);
+	}
+
+	// Are we copying the built files somewhere?
+	bool bCopySucceeded = false;
+	if (DeviceHelper != nullptr)
+	{
+		// Execute Copy
+		UE_LOG(LogCookGlobalShaders, Display, TEXT("Copying Cooked Files..."));
+		bCopySucceeded = DeviceHelper->CopyFilesToDevice(TargetDevice.Get(), FilesToCopy);
+	}
+	// if no helper, but we want to deploy, use the TargetPlatform
+	else if (bDeployToDevice && TargetDevice != nullptr)
+	{
+		bCopySucceeded = true;
+
+		TMap<FString,FString> CustomPlatformData;
+		CustomPlatformData.Add(TEXT("DeployFolder"), DeployFolder);
+
+		for (auto It : FilesToCopy)
+		{
+			bool bThisCopySucceeded = TargetPlatform->CopyFileToTarget(TargetDevice->GetId().GetDeviceName(), It.Key, It.Value, CustomPlatformData);
+
+			if (bThisCopySucceeded)
+			{
+				UE_LOG(LogCookGlobalShaders, Display, TEXT("Deployement of %s to devkit %s succeeded"), *It.Key, *TargetDevice->GetName());
+			}
+			else
+			{
+				UE_LOG(LogCookGlobalShaders, Error, TEXT("Deployement of %s to devkit %s failed"), *It.Key, *TargetDevice->GetName());
+				bCopySucceeded = false;
+			}
+		}
+	}
+
+	// Execute Reload
+	if (bExecuteReload && bCopySucceeded && TargetDevice.IsValid())
+	{
+		UE_LOG(LogCookGlobalShaders, Display, TEXT("Sending ReloadGlobalShaders command to devkit %s"), *TargetDevice->GetName());
+		TargetDevice->ReloadGlobalShadersMap(OutputDir / TEXT("Engine"));
+	}
+	
+	// Wait for any DDC writes to complete
+	GetDerivedDataCacheRef().WaitForQuiescence(true);
+}
+
+void UCookGlobalShadersCommandlet::HandleDirectoryChanged(const TArray<FFileChangeData>& InFileChangeDatas)
+{
+	bool bHasAnyShaderFileChanges = false;
+
+	for (const FFileChangeData& It : InFileChangeDatas)
+	{
+		if (It.Filename.EndsWith(TEXT(".usf")) || It.Filename.EndsWith(TEXT(".ush")) || It.Filename.EndsWith(TEXT(".h")))
+		{
+			bHasAnyShaderFileChanges = true;
+			UE_LOG(LogCookGlobalShaders, Display, TEXT("Detected change on %s"), *It.Filename);
+		}
+	}
+
+	if (bHasAnyShaderFileChanges)
+	{
+		CookGlobalShaders();
+
+		UE_LOG(LogCookGlobalShaders, Display, TEXT("Ready for new shader file changes"));
+	}
+}
+
+void UCookGlobalShadersCommandlet::CookGlobalShadersOnDirectoriesChanges()
+{
+	FDirectoryWatcherModule& DirectoryWatcherModule = FModuleManager::LoadModuleChecked<FDirectoryWatcherModule>(TEXT("DirectoryWatcher"));
+	IDirectoryWatcher* DirectoryWatcher = DirectoryWatcherModule.Get();
+	
+	if (!DirectoryWatcher)
+	{
+		return;
+	}
+	
+	// Handle if we are watching a directory for changes.
+	TMap<FString, FDelegateHandle> DirectoryWatcherHandles;
+	{
+		UE_LOG(LogCookGlobalShaders, Display, TEXT("Register directory watchers"));
+
+		const TMap<FString, FString>& ShaderSourceDirectoryMappings = AllShaderSourceDirectoryMappings();
+
+		DirectoryWatcherHandles.Reserve(ShaderSourceDirectoryMappings.Num());
+
+		for (const auto& It : ShaderSourceDirectoryMappings)
+		{
+			FString DirectoryToWatch = It.Value;
+			if (FPaths::IsRelative(DirectoryToWatch))
+			{
+				DirectoryToWatch = FPaths::ConvertRelativePathToFull(DirectoryToWatch);
+			}
+			
+			DirectoryWatcherHandles.Add(DirectoryToWatch, FDelegateHandle());
+			FDelegateHandle& DirectoryWatcherHandle = DirectoryWatcherHandles[DirectoryToWatch];
+
+			DirectoryWatcher->RegisterDirectoryChangedCallback_Handle(
+				DirectoryToWatch,
+				IDirectoryWatcher::FDirectoryChanged::CreateUObject(this, &UCookGlobalShadersCommandlet::HandleDirectoryChanged),
+				DirectoryWatcherHandle);
+
+			if (DirectoryWatcherHandle.IsValid())
+			{
+				UE_LOG(LogCookGlobalShaders, Display, TEXT("Watching %s -> %s"), *It.Key, *DirectoryToWatch);
+			}
+			else
+			{
+				UE_LOG(LogCookGlobalShaders, Error, TEXT("Failed to set up directory watcher %s -> %s"), *It.Key, *DirectoryToWatch);
+			}
+		}
+	}
+
+	GIsRunning = true;
+
+	//@todo abstract properly or delete
+#if PLATFORM_WINDOWS// Windows only
+	// Used by the .com wrapper to notify that the Ctrl-C handler was triggered.
+	// This shared event is checked each tick so that the log file can be cleanly flushed.
+	FEvent* ComWrapperShutdownEvent = FPlatformProcess::GetSynchEventFromPool(true);
+#endif
+
+	UE_LOG(LogCookGlobalShaders, Display, TEXT("Ready for new shader file changes"));
+
+	while (GIsRunning && !IsEngineExitRequested())
+	{
+		GEngine->UpdateTimeAndHandleMaxTickRate();
+		GEngine->Tick(FApp::GetDeltaTime(), false);
+		
+		// tick the directory watcher
+		DirectoryWatcherModule.Get()->Tick(FApp::GetDeltaTime());
+		
+		// flush log
+		GLog->FlushThreadedLogs();
+
+#if PLATFORM_WINDOWS
+		if (ComWrapperShutdownEvent->Wait(0))
+		{
+			RequestEngineExit(TEXT("CookGlobalShadersCommandlet ComWrapperShutdownEvent"));
+		}
+#endif
+	}
+
+	//@todo abstract properly or delete 
+#if PLATFORM_WINDOWS
+	FPlatformProcess::ReturnSynchEventToPool(ComWrapperShutdownEvent);
+	ComWrapperShutdownEvent = nullptr;
+#endif
+
+	GIsRunning = false;
+	
+	// Unregisters all directories.
+	{
+		UE_LOG(LogCookGlobalShaders, Display, TEXT("Unregister directory watchers"));
+
+		for (const auto& It : DirectoryWatcherHandles)
+		{
+			if (It.Value.IsValid() && DirectoryWatcher)
+			{
+				DirectoryWatcher->UnregisterDirectoryChangedCallback_Handle(It.Key, It.Value);
+			}
+		}
+	}
+}
+
 int32 UCookGlobalShadersCommandlet::Main(const FString& Params)
 {
 	TArray<FString> Tokens;
@@ -40,31 +224,29 @@ int32 UCookGlobalShadersCommandlet::Main(const FString& Params)
 	// Display help
 	if (Switches.Contains("help"))
 	{
-		UE_LOG(LogCookGlobalShaders, Log, TEXT("CookGlobalShaders"));
-		UE_LOG(LogCookGlobalShaders, Log, TEXT("This commandlet will allow you to generate the global shaders file which can be used to override what is used in a cooked build by deploying the loose file."));
-		UE_LOG(LogCookGlobalShaders, Log, TEXT("Options:"));
-		UE_LOG(LogCookGlobalShaders, Log, TEXT(" Required: -platform=<platform>             (Which platform you want to cook for, i.e. windows)"));
-		UE_LOG(LogCookGlobalShaders, Log, TEXT(" Optional: -device=<name>                   (Set which device to use, when enabled the reload command will be sent to the device once the shaders are cooked)"));
-		UE_LOG(LogCookGlobalShaders, Log, TEXT(" Optional: -deploy=<optional deploy folder> (Must be used with -device and will deploy the shader file onto the device rather than in the staged builds folder)"));
-		UE_LOG(LogCookGlobalShaders, Log, TEXT(" Optional: -stage=<optional path>           (Moved the shader file into the staged builds folder, destination can be overriden)"));
-		UE_LOG(LogCookGlobalShaders, Log, TEXT(" Optional: -reload                          (Execute a shader reload on the device, only works if the device is valid or a default one was found"));
-		UE_LOG(LogCookGlobalShaders, Log, TEXT(" Optional: -shaderpdb=<path>                (Sets the shader pdb root)"));
+		UE_LOG(LogCookGlobalShaders, Display, TEXT("CookGlobalShaders"));
+		UE_LOG(LogCookGlobalShaders, Display, TEXT("This commandlet will allow you to generate the global shaders file which can be used to override what is used in a cooked build by deploying the loose file."));
+		UE_LOG(LogCookGlobalShaders, Display, TEXT("Options:"));
+		UE_LOG(LogCookGlobalShaders, Display, TEXT(" Required: -platform=<platform>             (Which platform you want to cook for, i.e. windows)"));
+		UE_LOG(LogCookGlobalShaders, Display, TEXT(" Optional: -device=<name>                   (Set which device to use, when enabled the reload command will be sent to the device once the shaders are cooked)"));
+		UE_LOG(LogCookGlobalShaders, Display, TEXT(" Optional: -deploy=<optional deploy folder> (Must be used with -device and will deploy the shader file onto the device rather than in the staged builds folder)"));
+		UE_LOG(LogCookGlobalShaders, Display, TEXT(" Optional: -stage=<optional path>           (Moved the shader file into the staged builds folder, destination can be overriden)"));
+		UE_LOG(LogCookGlobalShaders, Display, TEXT(" Optional: -reload                          (Execute a shader reload on the device, only works if the device is valid or a default one was found)"));
+		UE_LOG(LogCookGlobalShaders, Display, TEXT(" Optional: -live                            (Keep commandlet open and automatically recompile shader map when shader file changes happen)"));
+		UE_LOG(LogCookGlobalShaders, Display, TEXT(" Optional: -shaderpdb=<path>                (Sets the shader pdb root)"));
 		return 0;
 	}
 
-	const bool bDeployToDevice = Switches.Contains(TEXT("deploy")) || ParamVals.Contains(TEXT("deploy"));
-	const bool bCopyToStaged = Switches.Contains(TEXT("stage"));
-	const bool bExecuteReload = Switches.Contains(TEXT("reload"));
+	bDeployToDevice = Switches.Contains(TEXT("deploy")) || ParamVals.Contains(TEXT("deploy"));
+	bCopyToStaged = Switches.Contains(TEXT("stage"));
+	bExecuteReload = Switches.Contains(TEXT("reload"));
 
-	FString DeployFolder;
 	if (ParamVals.Contains(TEXT("deploy")))
 	{
 		DeployFolder = ParamVals[TEXT("deploy")];
 	}
 
 	// Parse platform
-	FString PlatformName;
-	ITargetPlatform* TargetPlatform = nullptr;
 	{
 		ITargetPlatformManagerModule& TPM = GetTargetPlatformManagerRef();
 
@@ -94,7 +276,6 @@ int32 UCookGlobalShadersCommandlet::Main(const FString& Params)
 	}
 
 	// Get target device
-	ITargetDevicePtr TargetDevice;
 	FString TargetDeviceName;
 	if ( FParse::Value(*Params, TEXT("device="), TargetDeviceName, true) )
 	{
@@ -131,7 +312,6 @@ int32 UCookGlobalShadersCommandlet::Main(const FString& Params)
 	}
 
 	// Find DeviceHelper class to use
-	UCookGlobalShadersDeviceHelperBase* DeviceHelper = nullptr;
 	if (TargetDevice.IsValid() && bDeployToDevice)
 	{
 		for (TObjectIterator<UClass> ClassIt; ClassIt; ++ClassIt)
@@ -165,65 +345,26 @@ int32 UCookGlobalShadersCommandlet::Main(const FString& Params)
 	}
 
 	// Cook shaders
-	TArray<FName> ShaderFormats;
 	TargetPlatform->GetAllTargetedShaderFormats(ShaderFormats);
 
-	//const bool bContinous = Switches.Contains(TEXT("continuous"));
-	//while (bContinous)
+	const bool bLive = Switches.Contains(TEXT("live"));
+	if (bLive)
 	{
-		// Cook shaders
-		UE_LOG(LogCookGlobalShaders, Log, TEXT("Cooking Global Shaders..."));
-		FString OutputDir = FPaths::ProjectSavedDir() / TEXT("CookGlobalShaders") / PlatformName;
-		TArray<uint8> OutGlobalShaderMap;
-		FShaderRecompileData Arguments(PlatformName, SP_NumPlatforms, ODSCRecompileCommand::Global, nullptr, nullptr, &OutGlobalShaderMap);
-		RecompileShadersForRemote(Arguments, OutputDir);
-
-		bool bCopySucceeded = false;
-
-		// Build list of files to copy
-		TArray<TPair<FString, FString>> FilesToCopy;
-		for (FName ShaderFormat : ShaderFormats)
-		{
-			const FString GlobalShaderCacheName = FPaths::Combine(OutputDir, TEXT("Engine"), TEXT("GlobalShaderCache-") + ShaderFormat.ToString() + TEXT(".bin"));
-			const FString OverrideGlobalShaderCacheName = FPaths::Combine(TEXT("Engine"), TEXT("OverrideGlobalShaderCache-") + ShaderFormat.ToString() + TEXT(".bin"));
-			FilesToCopy.Emplace(GlobalShaderCacheName, OverrideGlobalShaderCacheName);
-		}
-
-		// Are we copying the built files somewhere?
-		if (DeviceHelper != nullptr)
-		{
-			// Execute Copy
-			UE_LOG(LogCookGlobalShaders, Log, TEXT("Copying Cooked Files..."));
-			bCopySucceeded = DeviceHelper->CopyFilesToDevice(TargetDevice.Get(), FilesToCopy);
-		}
-		// if no helper, but we want to deploy, use the TargetPlatform
-		else if (bDeployToDevice && TargetDevice != nullptr)
-		{
-			bCopySucceeded = true;
-
-			TMap<FString,FString> CustomPlatformData;
-			CustomPlatformData.Add(TEXT("DeployFolder"), DeployFolder);
-
-			for (auto It : FilesToCopy)
-			{
-				bCopySucceeded = bCopySucceeded && TargetPlatform->CopyFileToTarget(TargetDevice->GetId().GetDeviceName(), It.Key, It.Value, CustomPlatformData);
-			}
-		}
-
+		// Cook the global shaders once in case they have been modified before starting the commandlet.
+		CookGlobalShaders();
 
-		// Execute Reload
-		if (bCopySucceeded && bExecuteReload && TargetDevice.IsValid())
-		{
-			UE_LOG(LogCookGlobalShaders, Log, TEXT("Sending Reload Command..."));
-			TargetDevice->ReloadGlobalShadersMap(OutputDir / TEXT("Engine"));
-		}
+		// Keep compiling shaders when directories changes.
+		CookGlobalShadersOnDirectoriesChanges();
+	}
+	else
+	{
+		// Cook the global shaders once and return.
+		CookGlobalShaders();
 	}
+
 	UE_LOG(LogCookGlobalShaders, Log, TEXT("Complete"));
 
 	DeviceHelper = nullptr;
 
-	// Wait for any DDC writes to complete
-	GetDerivedDataCacheRef().WaitForQuiescence(true);
-
 	return 0;
 }
