Merge //UE5/Release-5.0 into //UE5/Private-Frosty-To-5.0-Staging
CL-18242543

diff --git a/Engine/Source/Editor/UnrealEd/Private/WorldPartition/WorldPartitionResaveActorsBuilder.cpp b/Engine/Source/Editor/UnrealEd/Private/WorldPartition/WorldPartitionResaveActorsBuilder.cpp
--- a/Engine/Source/Editor/UnrealEd/Private/WorldPartition/WorldPartitionResaveActorsBuilder.cpp
+++ b/Engine/Source/Editor/UnrealEd/Private/WorldPartition/WorldPartitionResaveActorsBuilder.cpp
@@ -33,6 +33,18 @@ void UWorldPartitionResaveActorsBuilder::ParseCommandline()
 	}
 }
 
+bool UWorldPartitionResaveActorsBuilder::OnPartitionBuildStarted(UWorld* World, FPackageSourceControlHelper& PackageHelper)
+{
+	FParse::Value(FCommandLine::Get(), TEXT("ActorClass="), ActorClassName);
+	
+	TArray<FString> Tokens, Switches;
+	UCommandlet::ParseCommandLine(FCommandLine::Get(), Tokens, Switches);
+
+	bSwitchActorPackagingSchemeToReduced = Switches.Contains(TEXT("SwitchActorPackagingSchemeToReduced"));
+
+	return true;
+}
+
 bool UWorldPartitionResaveActorsBuilder::PreRun(UWorld* World, FPackageSourceControlHelper& PackageHelper)
 {
 	TArray<FString> Tokens, Switches;
@@ -343,7 +355,6 @@ bool UWorldPartitionResaveActorsBuilder::RunInternal(UWorld* World, const FCellI
 					return true;
 				}
 
-				// It is possible the resave can't checkout everything. Continue processing.
 				UE_LOG(LogWorldPartitionResaveActorsBuilder, Display, TEXT("Saved package %s."), *Package->GetName());
 				++SaveCount;
 				return true;
@@ -390,4 +401,4 @@ bool UWorldPartitionResaveActorsBuilder::RunInternal(UWorld* World, const FCellI
 	
 
 	return true;
-}
\ No newline at end of file
+}
diff --git a/Engine/Source/Runtime/GeometryCore/Private/Spatial/SparseDynamicOctree3.cpp b/Engine/Source/Runtime/GeometryCore/Private/Spatial/SparseDynamicOctree3.cpp
--- a/Engine/Source/Runtime/GeometryCore/Private/Spatial/SparseDynamicOctree3.cpp
+++ b/Engine/Source/Runtime/GeometryCore/Private/Spatial/SparseDynamicOctree3.cpp
@@ -225,7 +225,7 @@ bool FSparseDynamicOctree3::ReinsertObject(int32 ObjectID, const FAxisAlignedBox
 
 
 
-double FSparseDynamicOctree3::FindNearestRayCellIntersection(const FSparseOctreeCell& Cell, const UE::Geometry::FRay3d& Ray) const
+double FSparseDynamicOctree3::FindNearestRayCellIntersection(const FSparseOctreeCell& Cell, const FRay3d& Ray) const
 {
 	FAxisAlignedBox3d Box = GetCellBox(Cell, MaxExpandFactor);
 	double ray_t = TNumericLimits<double>::Max();
@@ -241,9 +241,9 @@ double FSparseDynamicOctree3::FindNearestRayCellIntersection(const FSparseOctree
 
 
 
-int32 FSparseDynamicOctree3::FindNearestHitObject(const UE::Geometry::FRay3d& Ray,
+int32 FSparseDynamicOctree3::FindNearestHitObject(const FRay3d& Ray,
 	TFunctionRef<FAxisAlignedBox3d(int)> GetObjectBoundsFunc,
-	TFunctionRef<double(int, const UE::Geometry::FRay3d&)> HitObjectDistFunc,
+	TFunctionRef<double(int, const FRay3d&)> HitObjectDistFunc,
 	double MaxDistance) const
 {
 	// this should take advantage of raster!
diff --git a/Engine/Source/Runtime/Renderer/Private/SceneVisibility.cpp b/Engine/Source/Runtime/Renderer/Private/SceneVisibility.cpp
--- a/Engine/Source/Runtime/Renderer/Private/SceneVisibility.cpp
+++ b/Engine/Source/Runtime/Renderer/Private/SceneVisibility.cpp
@@ -275,6 +275,14 @@ static TAutoConsoleVariable<int32> CVarUseVisibilityOctree(
 	TEXT("Use the octree for visibility calculations."), 
 	ECVF_RenderThreadSafe);
 
+static bool GOcclusionSingleRHIThreadStall = false;
+static FAutoConsoleVariableRef CVarOcclusionSingleRHIThreadStall(
+	TEXT("r.Occlusion.SingleRHIThreadStall"),
+	GOcclusionSingleRHIThreadStall,
+	TEXT("Enable a single RHI thread stall before polling occlusion queries. This will only happen if the RHI's occlusion queries would normally stall the RHI thread themselves."),
+	ECVF_RenderThreadSafe
+);
+
 #if !UE_BUILD_SHIPPING
 
 static TAutoConsoleVariable<int32> CVarTAADebugOverrideTemporalIndex(
