World Partition Builders: renamed OnPartitionBuildStarted/OnPartitionBuildEnded to PreRun/PostRun.

diff --git a/Engine/Source/Editor/UnrealEd/Private/WorldPartition/WorldPartitionResaveActorsBuilder.cpp b/Engine/Source/Editor/UnrealEd/Private/WorldPartition/WorldPartitionResaveActorsBuilder.cpp
--- a/Engine/Source/Editor/UnrealEd/Private/WorldPartition/WorldPartitionResaveActorsBuilder.cpp
+++ b/Engine/Source/Editor/UnrealEd/Private/WorldPartition/WorldPartitionResaveActorsBuilder.cpp
@@ -20,20 +20,7 @@ UWorldPartitionResaveActorsBuilder::UWorldPartitionResaveActorsBuilder(const FOb
 	: Super(ObjectInitializer)
 {}
 
-void UWorldPartitionResaveActorsBuilder::ParseCommandline()
-{
-	FParse::Value(FCommandLine::Get(), TEXT("ActorClass="), ActorClassName);
-
-	bReportOnly = FParse::Param(FCommandLine::Get(), TEXT("ReportOnly"));
-	bOnlyResaveOnActorDescDiff = FParse::Param(FCommandLine::Get(), TEXT("OnlyResaveOnActorDescDiff"));
-	
-	if(FParse::Param(FCommandLine::Get(), TEXT("Verbose")))
-	{
-		LogWorldPartitionResaveActorsBuilder.SetVerbosity(ELogVerbosity::Verbose);
-	}
-}
-
-bool UWorldPartitionResaveActorsBuilder::OnPartitionBuildStarted(UWorld* World, FPackageSourceControlHelper& PackageHelper)
+bool UWorldPartitionResaveActorsBuilder::PreRun(UWorld* World, FPackageSourceControlHelper& PackageHelper)
 {
 	FParse::Value(FCommandLine::Get(), TEXT("ActorClass="), ActorClassName);
 	
