[ODSC] - Break out recompile shaders functionality into its own file.

- Break out common functionality for both ShaderCompiler.cpp and RecompileShaders.cpp into it's own shared file.

diff --git a/Engine/Source/Editor/UnrealEd/Private/Commandlets/CookGlobalShaderCommandlet.cpp b/Engine/Source/Editor/UnrealEd/Private/Commandlets/CookGlobalShaderCommandlet.cpp
--- a/Engine/Source/Editor/UnrealEd/Private/Commandlets/CookGlobalShaderCommandlet.cpp
+++ b/Engine/Source/Editor/UnrealEd/Private/Commandlets/CookGlobalShaderCommandlet.cpp
@@ -8,7 +8,7 @@
 #include "GenericPlatform/GenericPlatformFile.h"
 #include "UObject/UObjectIterator.h"
 #include "UObject/Class.h"
-#include "ShaderCompiler.h"
+#include "RecompileShaders.h"
 #include "CookOnTheSide/CookOnTheFlyServer.h"
 #include "Interfaces/ITargetPlatformManagerModule.h"
 #include "Interfaces/ITargetPlatform.h"
