SavePackage: Deprecate Conform and DiffMap arguments, and change the interface to use a FSavePackageArgs struct instead of a large number of separate arguments.
[RN] Core, Minor

diff --git a/Engine/Source/Editor/UnrealEd/Private/WorldPartition/WorldPartitionResaveActorsBuilder.cpp b/Engine/Source/Editor/UnrealEd/Private/WorldPartition/WorldPartitionResaveActorsBuilder.cpp
--- a/Engine/Source/Editor/UnrealEd/Private/WorldPartition/WorldPartitionResaveActorsBuilder.cpp
+++ b/Engine/Source/Editor/UnrealEd/Private/WorldPartition/WorldPartitionResaveActorsBuilder.cpp
@@ -328,7 +328,9 @@ bool UWorldPartitionResaveActorsBuilder::RunInternal(UWorld* World, const FCellI
 				if (!bReportOnly && bShouldResaveActor)
 				{
 					FString PackageFileName = SourceControlHelpers::PackageFilename(Package);
-					if (!UPackage::SavePackage(Package, nullptr, RF_Standalone, *PackageFileName))
+					FSavePackageArgs SaveArgs;
+					SaveArgs.TopLevelFlags = RF_Standalone;
+					if (!UPackage::SavePackage(Package, nullptr, *PackageFileName, SaveArgs))
 					{
 						UE_LOG(LogWorldPartitionResaveActorsBuilder, Error, TEXT("Error saving package %s."), *Package->GetName());
 						++FailCount;
