UWorldPartitionResaveActorsBuilder
- Automatically match commandline arguments to properties when they exists (this will be ported to all builders in the future).
- Added `ResaveDirtyActorDescsOnly` to only resave actors that have notable differences in their actor descriptors.
- Added `ReportOnly` option to only report what would be done.

diff --git a/Engine/Source/Editor/UnrealEd/Private/WorldPartition/WorldPartitionResaveActorsBuilder.cpp b/Engine/Source/Editor/UnrealEd/Private/WorldPartition/WorldPartitionResaveActorsBuilder.cpp
--- a/Engine/Source/Editor/UnrealEd/Private/WorldPartition/WorldPartitionResaveActorsBuilder.cpp
+++ b/Engine/Source/Editor/UnrealEd/Private/WorldPartition/WorldPartitionResaveActorsBuilder.cpp
@@ -22,12 +22,58 @@ UWorldPartitionResaveActorsBuilder::UWorldPartitionResaveActorsBuilder(const FOb
 
 bool UWorldPartitionResaveActorsBuilder::PreRun(UWorld* World, FPackageSourceControlHelper& PackageHelper)
 {
-	FParse::Value(FCommandLine::Get(), TEXT("ActorClass="), ActorClassName);
-	
 	TArray<FString> Tokens, Switches;
 	UCommandlet::ParseCommandLine(FCommandLine::Get(), Tokens, Switches);
 
-	bSwitchActorPackagingSchemeToReduced = Switches.Contains(TEXT("SwitchActorPackagingSchemeToReduced"));
+	//@todo_ow: generalize to all builders
+	for (const FString& Switch : Switches)
+	{
+		FString Key, Value;
+		if (!Switch.Split(TEXT("="), &Key, &Value))
+		{
+			Key = Switch;
+		}
+
+		// Lookup property
+		const FProperty* Property = GetClass()->FindPropertyByName(*Key);
+
+		// If we can't find the property, try for properties with the 'b' prefix
+		if (!Property)
+		{
+			Key = TEXT("b") + Key;
+			Property = GetClass()->FindPropertyByName(*Key);
+		}
+
+		if (Property)
+		{
+			// If the property is a bool, treat no values as true
+			if (Property->IsA(FBoolProperty::StaticClass()) && Value.IsEmpty())
+			{
+				Value = TEXT("True");
+			}
+
+			uint8* Container = (uint8*)this;
+			if (!FBlueprintEditorUtils::PropertyValueFromString(Property, Value, Container, nullptr))
+			{
+				UE_LOG(LogWorldPartitionResaveActorsBuilder, Error, TEXT("Cannot set value for '%s': '%s'"), *Key, *Value);
+				return false;
+			}
+		}
+	}
+
+	if (bSwitchActorPackagingSchemeToReduced)
+	{
+		if (!ActorClassName.IsEmpty())
+		{
+			UE_LOG(LogWorldPartitionResaveActorsBuilder, Error, TEXT("SwitchActorPackagingSchemeToReduced is not compatible with ActorClassName"));
+			return false;
+		}
+		else if (bResaveDirtyActorDescsOnly)
+		{
+			UE_LOG(LogWorldPartitionResaveActorsBuilder, Error, TEXT("SwitchActorPackagingSchemeToReduced is not compatible with ResaveDirtyActorDescsOnly"));
+			return false;
+		}
+	}
 
 	return true;
 }
@@ -99,7 +145,6 @@ bool UWorldPartitionResaveActorsBuilder::RunInternal(UWorld* World, const FCellI
 
 	int32 LoadCount = 0;
 	int32 SaveCount = 0;
-	int32 SkipCount = 0;
 	int32 FailCount = 0;
 
 	// Actor Class Filter
@@ -147,11 +192,11 @@ bool UWorldPartitionResaveActorsBuilder::RunInternal(UWorld* World, const FCellI
 	{
 		World->PersistentLevel->ActorPackagingScheme = EActorPackagingScheme::Reduced;
 
-		FWorldPartitionHelpers::ForEachActorWithLoading(WorldPartition, ActorClass, [this, &LoadCount, &SaveCount, &FailCount, &SkipCount, &SCCHelper, &PackagesToDelete, WorldPackage](const FWorldPartitionActorDesc* ActorDesc)
+		FWorldPartitionHelpers::ForEachActorWithLoading(WorldPartition, ActorClass, [this, &LoadCount, &SaveCount, &FailCount, &SCCHelper, &PackagesToDelete, WorldPackage](const FWorldPartitionActorDesc* ActorDesc)
 		{
 			ON_SCOPE_EXIT
 			{
-				UE_LOG(LogWorldPartitionResaveActorsBuilder, Display, TEXT("Processed %d packages (%d Saved / %d Skipped / %d Failed)"), SaveCount + SkipCount + FailCount, SaveCount, SkipCount, FailCount);
+				UE_LOG(LogWorldPartitionResaveActorsBuilder, Display, TEXT("Processed %d packages (%d Saved / %d Failed)"), LoadCount, SaveCount, FailCount);
 			};
 
 			AActor* Actor = ActorDesc->GetActor();
@@ -336,24 +381,18 @@ bool UWorldPartitionResaveActorsBuilder::RunInternal(UWorld* World, const FCellI
 						++FailCount;
 						return true;
 					}
+
+					UE_LOG(LogWorldPartitionResaveActorsBuilder, Display, TEXT("Saved package %s."), *Package->GetName());
+					++SaveCount;
+					return true;
 				}
 				else
 				{
-					UE_LOG(LogWorldPartitionResaveActorsBuilder, Display, TEXT("Skipped package %s."), *Package->GetName());
-					++SkipCount;
+					// It is possible the resave can't checkout everything. Continue processing.
+					UE_LOG(LogWorldPartitionResaveActorsBuilder, Warning, TEXT("Error checking out package %s."), *Package->GetName());
+					++FailCount;
 					return true;
 				}
-
-				UE_LOG(LogWorldPartitionResaveActorsBuilder, Display, TEXT("Saved package %s."), *Package->GetName());
-				++SaveCount;
-				return true;
-			}
-			else
-			{
-				// It is possible the resave can't checkout everything. Continue processing.
-				UE_LOG(LogWorldPartitionResaveActorsBuilder, Warning, TEXT("Error checking out package %s."), *Package->GetName());
-				++FailCount;
-				return true;
 			}
 			
 		}
@@ -390,4 +429,4 @@ bool UWorldPartitionResaveActorsBuilder::RunInternal(UWorld* World, const FCellI
 	
 
 	return true;
-}
+}
\ No newline at end of file
