Fixing bad merges
#rnx

diff --git a/Engine/Source/Editor/UnrealEd/Private/WorldPartition/WorldPartitionResaveActorsBuilder.cpp b/Engine/Source/Editor/UnrealEd/Private/WorldPartition/WorldPartitionResaveActorsBuilder.cpp
--- a/Engine/Source/Editor/UnrealEd/Private/WorldPartition/WorldPartitionResaveActorsBuilder.cpp
+++ b/Engine/Source/Editor/UnrealEd/Private/WorldPartition/WorldPartitionResaveActorsBuilder.cpp
@@ -78,63 +78,6 @@ bool UWorldPartitionResaveActorsBuilder::PreRun(UWorld* World, FPackageSourceCon
 	return true;
 }
 
-bool UWorldPartitionResaveActorsBuilder::PreRun(UWorld* World, FPackageSourceControlHelper& PackageHelper)
-{
-	TArray<FString> Tokens, Switches;
-	UCommandlet::ParseCommandLine(FCommandLine::Get(), Tokens, Switches);
-
-	//@todo_ow: generalize to all builders
-	for (const FString& Switch : Switches)
-	{
-		FString Key, Value;
-		if (!Switch.Split(TEXT("="), &Key, &Value))
-		{
-			Key = Switch;
-		}
-
-		// Lookup property
-		const FProperty* Property = GetClass()->FindPropertyByName(*Key);
-
-		// If we can't find the property, try for properties with the 'b' prefix
-		if (!Property)
-		{
-			Key = TEXT("b") + Key;
-			Property = GetClass()->FindPropertyByName(*Key);
-		}
-
-		if (Property)
-		{
-			// If the property is a bool, treat no values as true
-			if (Property->IsA(FBoolProperty::StaticClass()) && Value.IsEmpty())
-			{
-				Value = TEXT("True");
-			}
-
-			uint8* Container = (uint8*)this;
-			if (!FBlueprintEditorUtils::PropertyValueFromString(Property, Value, Container, nullptr))
-			{
-				UE_LOG(LogWorldPartitionResaveActorsBuilder, Error, TEXT("Cannot set value for '%s': '%s'"), *Key, *Value);
-				return false;
-			}
-		}
-	}
-
-	if (bSwitchActorPackagingSchemeToReduced)
-	{
-		if (!ActorClassName.IsEmpty())
-		{
-			UE_LOG(LogWorldPartitionResaveActorsBuilder, Error, TEXT("SwitchActorPackagingSchemeToReduced is not compatible with ActorClassName"));
-			return false;
-		}
-		else if (bResaveDirtyActorDescsOnly)
-		{
-			UE_LOG(LogWorldPartitionResaveActorsBuilder, Error, TEXT("SwitchActorPackagingSchemeToReduced is not compatible with ResaveDirtyActorDescsOnly"));
-			return false;
-		}
-	}
-
-	return true;
-}
 
 bool UWorldPartitionResaveActorsBuilder::RunInternal(UWorld* World, const FCellInfo& InCellInfo, FPackageSourceControlHelper& PackageHelper)
 {
diff --git a/Engine/Source/Runtime/Engine/Public/WorldPartition/WorldPartitionEditorPerProjectUserSettings.h b/Engine/Source/Runtime/Engine/Public/WorldPartition/WorldPartitionEditorPerProjectUserSettings.h
--- a/Engine/Source/Runtime/Engine/Public/WorldPartition/WorldPartitionEditorPerProjectUserSettings.h
+++ b/Engine/Source/Runtime/Engine/Public/WorldPartition/WorldPartitionEditorPerProjectUserSettings.h
@@ -104,8 +104,8 @@ public:
 		}
 	}
 
-	const TArray<FName>& GetWorldDataLayersNotLoadedInEditor(UWorld* InWorld) const;
-	const TArray<FName>& GetWorldDataLayersLoadedInEditor(UWorld* InWorld) const;
+	TArray<FName> GetWorldDataLayersNotLoadedInEditor(UWorld* InWorld) const;
+	TArray<FName> GetWorldDataLayersLoadedInEditor(UWorld* InWorld) const;
 	
 	void SetWorldDataLayersNonDefaultEditorLoadStates(UWorld* InWorld, const TArray<FName>& InDataLayersLoadedInEditor, const TArray<FName>& InDataLayersNotLoadedInEditor);
 
diff --git a/Engine/Source/Runtime/Renderer/Private/ScenePrivate.h b/Engine/Source/Runtime/Renderer/Private/ScenePrivate.h
--- a/Engine/Source/Runtime/Renderer/Private/ScenePrivate.h
+++ b/Engine/Source/Runtime/Renderer/Private/ScenePrivate.h
@@ -3038,6 +3038,7 @@ public:
 	virtual void ReleasePrimitive(UPrimitiveComponent* Primitive) override;
 	virtual void UpdateAllPrimitiveSceneInfos(FRDGBuilder& GraphBuilder, bool bAsyncCreateLPIs = false) override;
 	virtual void UpdatePrimitiveTransform(UPrimitiveComponent* Primitive) override;
+	virtual void UpdatePrimitiveInstances(UPrimitiveComponent* Primitive) override;
 	virtual void UpdatePrimitiveOcclusionBoundsSlack(UPrimitiveComponent* Primitive, float NewSlack) override;
 	virtual void UpdatePrimitiveAttachment(UPrimitiveComponent* Primitive) override;
 	virtual void UpdateCustomPrimitiveData(UPrimitiveComponent* Primitive) override;
