Fix bad merges.

diff --git a/Engine/Source/Editor/UnrealEd/Private/WorldPartition/WorldPartitionResaveActorsBuilder.cpp b/Engine/Source/Editor/UnrealEd/Private/WorldPartition/WorldPartitionResaveActorsBuilder.cpp
--- a/Engine/Source/Editor/UnrealEd/Private/WorldPartition/WorldPartitionResaveActorsBuilder.cpp
+++ b/Engine/Source/Editor/UnrealEd/Private/WorldPartition/WorldPartitionResaveActorsBuilder.cpp
@@ -78,12 +78,10 @@ bool UWorldPartitionResaveActorsBuilder::PreRun(UWorld* World, FPackageSourceCon
 	return true;
 }
 
-
 bool UWorldPartitionResaveActorsBuilder::RunInternal(UWorld* World, const FCellInfo& InCellInfo, FPackageSourceControlHelper& PackageHelper)
 {
 	FPackageSourceControlHelper SCCHelper;
 
-	ParseCommandline();
 	UPackage* WorldPackage = World->GetPackage();
 
 	int32 LoadCount = 0;
@@ -146,7 +144,6 @@ bool UWorldPartitionResaveActorsBuilder::RunInternal(UWorld* World, const FCellI
 
 			if (!Actor)
 			{
-				UE_LOG(LogWorldPartitionResaveActorsBuilder, Error, TEXT("Error loading package %s."), *ActorDesc->GetActorPackage().ToString());
 				PackagesToDelete.Add(ActorDesc->GetActorPackage().ToString());
 				++FailCount;
 				return true;
@@ -280,41 +277,11 @@ bool UWorldPartitionResaveActorsBuilder::RunInternal(UWorld* World, const FCellI
 				UE_LOG(LogWorldPartitionResaveActorsBuilder, Log, TEXT("Package %s needs to be resaved."), *Package->GetName());
 			}
 
-			if (bReportOnly || SCCHelper.Checkout(Package))
+			if (!bReportOnly)
 			{
-				bool bShouldResaveActor = !bOnlyResaveOnActorDescDiff;
-
-				if (bOnlyResaveOnActorDescDiff)
-				{
-					TUniquePtr<FWorldPartitionActorDesc> CurActorDesc = Actor->CreateActorDesc();
-					const FWorldPartitionActorDesc& OldActorDesc = *ActorDesc;
-					const FWorldPartitionActorDesc& NewActorDesc = *CurActorDesc.Get();
-
-					// Bounds can change if the source asset changed
-					if (!OldActorDesc.GetBounds().Equals(NewActorDesc.GetBounds(), 1.0f))
-					{
-						bShouldResaveActor = true;
-					}
-					// Grid placement can have changed from defaults
-					else if (OldActorDesc.GetGridPlacement() != NewActorDesc.GetGridPlacement())
-					{
-						bShouldResaveActor = true;
-					}
-					// Editor-only flag can have changed from defaults
-					else if (OldActorDesc.GetActorIsEditorOnly() != NewActorDesc.GetActorIsEditorOnly())
-					{
-						bShouldResaveActor = true;
-					}
-
-					if (bShouldResaveActor)
-					{
-						UE_LOG(LogWorldPartitionResaveActorsBuilder, Verbose, TEXT("Actor descriptor mismatch: [%s]<->[%s]"), *OldActorDesc.ToString(), *NewActorDesc.ToString());
-					}
-				}
-
-				// Save package
-				if (!bReportOnly && bShouldResaveActor)
+				if (SCCHelper.Checkout(Package))
 				{
+					// Save package
 					FString PackageFileName = SourceControlHelpers::PackageFilename(Package);
 					FSavePackageArgs SaveArgs;
 					SaveArgs.TopLevelFlags = RF_Standalone;
@@ -337,7 +304,6 @@ bool UWorldPartitionResaveActorsBuilder::RunInternal(UWorld* World, const FCellI
 					return true;
 				}
 			}
-			
 		}
 
 		return true;
@@ -345,31 +311,28 @@ bool UWorldPartitionResaveActorsBuilder::RunInternal(UWorld* World, const FCellI
 
 	if (!bReportOnly)
 	{
-		UPackage::WaitForAsyncFileWrites();
 		SCCHelper.Delete(PackagesToDelete);
-		
-		if (bSwitchActorPackagingSchemeToReduced)
+	}
+
+	if (bSwitchActorPackagingSchemeToReduced && !bReportOnly)
+	{
+		// checkout world package
+		if (!SCCHelper.Checkout(WorldPackage))
 		{
-			// checkout world package
-			if (!SCCHelper.Checkout(WorldPackage))
-			{
-				UE_LOG(LogWorldPartitionResaveActorsBuilder, Error, TEXT("Error checking out world package %s."), *WorldPackage->GetName());
-				return false;
-			}
+			UE_LOG(LogWorldPartitionResaveActorsBuilder, Error, TEXT("Error checking out world package %s."), *WorldPackage->GetName());
+			return false;
+		}
 
-			// Save world package
-			FString PackageFileName = SourceControlHelpers::PackageFilename(WorldPackage);
-			FSavePackageArgs SaveArgs;
-			SaveArgs.TopLevelFlags = RF_Standalone;
-			if (!UPackage::SavePackage(WorldPackage, nullptr, *PackageFileName, SaveArgs))
-			{
-				UE_LOG(LogWorldPartitionResaveActorsBuilder, Error, TEXT("Error saving world %s."), *WorldPackage->GetName());
-				return false;
-			}
+		// Save world package
+		FString PackageFileName = SourceControlHelpers::PackageFilename(WorldPackage);
+		FSavePackageArgs SaveArgs;
+		SaveArgs.TopLevelFlags = RF_Standalone;
+		if (!UPackage::SavePackage(WorldPackage, nullptr, *PackageFileName, SaveArgs))
+		{
+			UE_LOG(LogWorldPartitionResaveActorsBuilder, Error, TEXT("Error saving world %s."), *WorldPackage->GetName());
+			return false;
 		}
 	}
 
-	
-
 	return true;
 }
\ No newline at end of file
diff --git a/Engine/Source/Runtime/Engine/Private/WorldPartition/WorldPartitionRuntimeSpatialHash.cpp b/Engine/Source/Runtime/Engine/Private/WorldPartition/WorldPartitionRuntimeSpatialHash.cpp
--- a/Engine/Source/Runtime/Engine/Private/WorldPartition/WorldPartitionRuntimeSpatialHash.cpp
+++ b/Engine/Source/Runtime/Engine/Private/WorldPartition/WorldPartitionRuntimeSpatialHash.cpp
@@ -667,7 +667,7 @@ void UWorldPartitionRuntimeSpatialHash::SetDefaultValues()
 	MainGrid.DebugColor = FLinearColor::Gray;
 }
 
-bool UWorldPartitionRuntimeSpatialHash::GenerateStreaming(class UWorldPartitionStreamingPolicy* StreamingPolicy, const FActorClusterContext& ActorClusterContext, TArray<FString>* OutPackagesToGenerate = nullptr)
+bool UWorldPartitionRuntimeSpatialHash::GenerateStreaming(UWorldPartitionStreamingPolicy* StreamingPolicy, const FActorClusterContext& ActorClusterContext, TArray<FString>* OutPackagesToGenerate)
 {
 	TRACE_CPUPROFILER_EVENT_SCOPE(UWorldPartitionRuntimeSpatialHash::GenerateStreaming);
 	UWorldPartition* WorldPartition = GetOuterUWorldPartition();
@@ -848,14 +848,6 @@ bool UWorldPartitionRuntimeSpatialHash::CreateStreamingGrid(const FSpatialHashRu
 				{
 					for (const FActorInstance& ActorInstance : GridCellDataChunk.GetActors())
 					{
-						if (ActorInstance.ShouldStripFromStreaming())
-						{
-							const FWorldPartitionActorDescView& ActorDescView = ActorInstance.GetActorDescView();
-							UE_LOG(LogWorldPartition, Verbose, TEXT("Stripping Actor %s (%s) from streaming grid (Container %08x)"),
-								*(ActorDescView.GetActorPath().ToString()), *ActorInstance.Actor.ToString(EGuidFormats::UniqueObjectGuid), ActorInstance.ContainerInstance->ID);
-							continue;
-						}
-
 						if (bIsMainWorldPartition && !IsRunningCookCommandlet())
 						{
 							const FWorldPartitionActorDescView& ActorDescView = ActorInstance.GetActorDescView();
