Fix non-unity

diff --git a/Engine/Source/Editor/UnrealEd/Private/Elements/Component/ComponentElementEditorViewportInteractionCustomization.cpp b/Engine/Source/Editor/UnrealEd/Private/Elements/Component/ComponentElementEditorViewportInteractionCustomization.cpp
--- a/Engine/Source/Editor/UnrealEd/Private/Elements/Component/ComponentElementEditorViewportInteractionCustomization.cpp
+++ b/Engine/Source/Editor/UnrealEd/Private/Elements/Component/ComponentElementEditorViewportInteractionCustomization.cpp
@@ -5,6 +5,7 @@
 #include "Components/SceneComponent.h"
 
 #include "Editor.h"
+#include "EditorModeManager.h"
 #include "EditorSupportDelegates.h"
 
 bool FComponentElementEditorViewportInteractionCustomization::GetGizmoPivotLocation(const TTypedElement<ITypedElementWorldInterface>& InElementWorldHandle, const UE::Widget::EWidgetMode InWidgetMode, FVector& OutPivotLocation)
