Cooker: Delete PackageNameCache; FPackageDatas now covers its functionality and removing it saves cputime and memory from the duplicate maps of packagename to filename.
Make construction of new FPackageData threadsafe.
#rnx

diff --git a/Engine/Source/Editor/UnrealEd/Classes/CookOnTheSide/CookOnTheFlyServer.h b/Engine/Source/Editor/UnrealEd/Classes/CookOnTheSide/CookOnTheFlyServer.h
--- a/Engine/Source/Editor/UnrealEd/Classes/CookOnTheSide/CookOnTheFlyServer.h
+++ b/Engine/Source/Editor/UnrealEd/Classes/CookOnTheSide/CookOnTheFlyServer.h
@@ -30,7 +30,6 @@ class IAssetRegistry;
 class ICookedPackageWriter;
 class IPlugin;
 class ITargetPlatform;
-struct FPackageNameCache;
 struct FPropertyChangedEvent;
 
 enum class ECookInitializationFlags
@@ -117,6 +116,7 @@ namespace UE::Cook
 	class FRequestCluster;
 	class FSaveCookedPackageContext;
 	class ICookOnTheFlyRequestManager;
+	struct FConstructPackageData;
 	struct FCookerTimer;
 	struct FCookSavePackageContext;
 	struct FGeneratorPackage;
@@ -748,10 +748,11 @@ private:
 	*
 	* @param AssetRegistryPath path of the assetregistry.bin file to read
 	* @param bVerifyPackagesExist whether or not we should verify the packages exist on disk.
-	* @param OutPackageNames out list of uncooked package filenames which were contained in the asset registry file
+	* @param OutPackageDatas out list of packagename/filenames for packages contained in the asset registry file
 	* @return true if successfully read false otherwise
 	*/
-	bool GetAllPackageFilenamesFromAssetRegistry( const FString& AssetRegistryPath, bool bVerifyPackagesExist, TArray<FName>& OutPackageFilenames ) const;
+	bool GetAllPackageFilenamesFromAssetRegistry( const FString& AssetRegistryPath, bool bVerifyPackagesExist,
+		TArray<UE::Cook::FConstructPackageData>& OutPackageDatas) const;
 
 	/**
 	* BuildMapDependencyGraph
@@ -1047,9 +1048,6 @@ private:
 	/** Generates long package names for all files to be cooked */
 	void GenerateLongPackageNames(TArray<FName>& FilesInPath, TMap<FName, UE::Cook::FInstigator>& Instigators);
 
-	/** Return the PackageNameCache that is caching FileNames on disk for this CookOnTheFlyServer. */
-	const FPackageNameCache& GetPackageNameCache() const;
-
 	/**
 	* Try to find a registered CookPackageSplitter to split the given package.
 	*
diff --git a/Engine/Source/Editor/UnrealEd/Private/CookOnTheFlyServer.cpp b/Engine/Source/Editor/UnrealEd/Private/CookOnTheFlyServer.cpp
--- a/Engine/Source/Editor/UnrealEd/Private/CookOnTheFlyServer.cpp
+++ b/Engine/Source/Editor/UnrealEd/Private/CookOnTheFlyServer.cpp
@@ -29,7 +29,6 @@
 #include "Cooker/IoStoreCookOnTheFlyRequestManager.h"
 #include "Cooker/LooseCookedPackageWriter.h"
 #include "Cooker/NetworkFileCookOnTheFlyRequestManager.h"
-#include "Cooker/PackageNameCache.h"
 #include "Cooker/PackageTracker.h"
 #include "CookerSettings.h"
 #include "CookPackageSplitter.h"
@@ -1017,11 +1016,10 @@ bool UCookOnTheFlyServer::ContainsRedirector(const FName& PackageName, TMap<FNam
 
 				if (!bDestinationValid)
 				{
-					// we can't call GetCachedStandardFileName with None
 					if (RedirectedPath != NAME_None)
 					{
-						FName StandardPackageName = GetPackageNameCache().GetCachedStandardFileName(FName(*FPackageName::ObjectPathToPackageName(RedirectedPathString)));
-						if (StandardPackageName != NAME_None)
+						FName StandardPackageName = PackageDatas->GetFileNameByPackageName(FName(*FPackageName::ObjectPathToPackageName(RedirectedPathString)));
+						if (!StandardPackageName.IsNone())
 						{
 							bDestinationValid = true;
 						}
@@ -1993,18 +1991,15 @@ void UCookOnTheFlyServer::LoadPackageInQueue(UE::Cook::FPackageData& PackageData
 		}
 		check(LoadedPackage != nullptr && LoadedPackage->IsFullyLoaded());
 
-		FName NewPackageFileName(GetPackageNameCache().GetCachedStandardFileName(LoadedPackage));
 		if (LoadedPackage->GetFName() != PackageData.GetPackageName())
 		{
 			// The PackageName is not the name that we loaded. This can happen due to CoreRedirects.
 			// We refuse to cook requests for packages that no longer exist in PumpExternalRequests, but it is possible
 			// that a CoreRedirect exists from a (externally requested or requested as a reference) package that still exists.
 			// Mark the original PackageName as cooked for all platforms and send a request to cook the new FileName
-			check(NewPackageFileName != PackageFileName);
-
-			UE_LOG(LogCook, Verbose, TEXT("Request for %s received going to save %s"), *PackageFileName.ToString(),
-				*NewPackageFileName.ToString());
 			FPackageData& OtherPackageData = PackageDatas->AddPackageDataByPackageNameChecked(LoadedPackage->GetFName());
+			UE_LOG(LogCook, Verbose, TEXT("Request for %s received going to save %s"), *PackageFileName.ToString(),
+				*OtherPackageData.GetFileName().ToString());
 			TArray<const ITargetPlatform*, TInlineAllocator<ExpectedMaxNumPlatforms>> RequestedPlatforms;
 			PackageData.GetRequestedPlatforms(RequestedPlatforms);
 			OtherPackageData.UpdateRequestData(RequestedPlatforms, PackageData.GetIsUrgent(), FCompletionCallback(),
@@ -2014,20 +2009,6 @@ void UCookOnTheFlyServer::LoadPackageInQueue(UE::Cook::FPackageData& PackageData
 			RejectPackageToLoad(PackageData, TEXT("is redirected to another filename"));
 			return;
 		}
-
-		if (NewPackageFileName != PackageFileName)
-		{
-			// This case should never happen since we are checking for the existence of the file in PumpExternalRequests
-			UE_LOG(LogCook, Warning, TEXT("Unexpected change in FileName when loading a requested package. \"%s\" changed to \"%s\"."),
-				*PackageFileName.ToString(), *NewPackageFileName.ToString());
-
-			UE_LOG(LogCook, Verbose, TEXT("Request for %s received going to save %s"), *PackageFileName.ToString(), *NewPackageFileName.ToString());
-			// The package FileName is not the FileName that we loaded
-			//  sounds unpossible.... but it is due to searching for files and such
-			PackageDatas->UpdateFileName(LoadedPackage->GetFName());
-			// Register the FileName that we loaded as pointing to this PackageData
-			PackageDatas->RegisterFileNameAlias(PackageData, PackageFileName);
-		}
 	}
 	else
 	{
@@ -2136,18 +2117,18 @@ UE::Cook::FPackageData* UCookOnTheFlyServer::QueueDiscoveredPackage(UPackage* Pa
 {
 	check(Package != nullptr);
 
-	const FName FileName = GetPackageNameCache().GetCachedStandardFileName(Package);
-	if (FileName.IsNone())
+	UE::Cook::FPackageData* PackageData = PackageDatas->TryAddPackageDataByPackageName(Package->GetFName());
+	if (!PackageData)
 	{
-		return nullptr;	// if we have name none that means we are in core packages or something...
+		return nullptr;	// Getting the PackageData will fail if e.g. it is a script package
 	}
-	UE::Cook::FPackageData& PackageData = PackageDatas->FindOrAddPackageData(Package->GetFName(), FileName);
+
 	if (bOutWasInProgress)
 	{
-		*bOutWasInProgress = PackageData.IsInProgress();
+		*bOutWasInProgress = PackageData->IsInProgress();
 	}
-	QueueDiscoveredPackageData(PackageData, MoveTemp(Instigator), true /* bIsLoadReady */);
-	return &PackageData;
+	QueueDiscoveredPackageData(*PackageData, MoveTemp(Instigator), true /* bIsLoadReady */);
+	return PackageData;
 }
 
 void UCookOnTheFlyServer::QueueDiscoveredPackageData(UE::Cook::FPackageData& PackageData,
@@ -2342,7 +2323,7 @@ void UCookOnTheFlyServer::SplitPackage(UE::Cook::FGeneratorPackage* GeneratorStr
 	if (!GeneratorStruct->HasGeneratedList())
 	{
 		// Call the splitter to generate the list
-		if (!GeneratorStruct->TryGenerateList(SplitObject, GetPackageNameCache(), *PackageDatas))
+		if (!GeneratorStruct->TryGenerateList(SplitObject, *PackageDatas))
 		{
 			bOutError = true;
 			return;
@@ -4679,8 +4660,11 @@ bool UCookOnTheFlyServer::Exec(class UWorld* InWorld, const TCHAR* Cmd, FOutputD
 		StartupOptions.TargetPlatforms.Add(TargetPlatform);
 		for (const FName& StandardPackageName : PackageNames)
 		{
-			FName PackageFileFName = GetPackageNameCache().GetCachedStandardFileName(StandardPackageName);
-			StartupOptions.CookMaps.Add(StandardPackageName.ToString());
+			FName PackageFileFName = PackageDatas->GetFileNameByPackageName(StandardPackageName);
+			if (!PackageFileFName.IsNone())
+			{
+				StartupOptions.CookMaps.Add(StandardPackageName.ToString());
+			}
 		}
 		StartupOptions.CookOptions = ECookByTheBookOptions::NoAlwaysCookMaps | ECookByTheBookOptions::NoDefaultMaps | ECookByTheBookOptions::NoGameAlwaysCookPackages | ECookByTheBookOptions::NoInputPackages | ECookByTheBookOptions::NoSlatePackages | ECookByTheBookOptions::SkipSoftReferences | ECookByTheBookOptions::ForceDisableSaveGlobalShaders;
 		
@@ -5759,7 +5743,7 @@ void UCookOnTheFlyServer::GenerateAssetRegistry()
 		UE_LOG(LogCook, Display, TEXT("Creating asset registry"));
 
 		ModifiedAssetFilenames.Reset();
-		GetPackageNameCache().SetAssetRegistry(AssetRegistry);
+		UE::Cook::FPackageDatas::OnAssetRegistryGenerated(*AssetRegistry);
 	}
 }
 
@@ -6122,9 +6106,9 @@ void UCookOnTheFlyServer::CollectFilesToCook(TArray<FName>& FilesInPath, TMap<FN
 
 			for (FName NeverCookPackage : PackagesToNeverCook)
 			{
-				const FName StandardPackageFilename = GetPackageNameCache().GetCachedStandardFileName(NeverCookPackage);
+				const FName StandardPackageFilename = PackageDatas->GetFileNameByFlexName(NeverCookPackage);
 
-				if (StandardPackageFilename != NAME_None)
+				if (!StandardPackageFilename.IsNone())
 				{
 					PackageTracker->NeverCookPackageList.Add(StandardPackageFilename);
 				}
@@ -6160,11 +6144,11 @@ void UCookOnTheFlyServer::CollectFilesToCook(TArray<FName>& FilesInPath, TMap<FN
 
 		for (FName NeverCookPackage : PackagesToNeverCook)
 		{
-			const FName* StandardPackageFilename = GetPackageNameCache().GetCachedPackageNameFromStandardFileName(NeverCookPackage);
+			FName StandardPackageFilename = PackageDatas->GetFileNameByFlexName(NeverCookPackage);
 
-			if (StandardPackageFilename && *StandardPackageFilename != NAME_None)
+			if (!StandardPackageFilename.IsNone())
 			{
-				PackageTracker->NeverCookPackageList.Add(*StandardPackageFilename);
+				PackageTracker->NeverCookPackageList.Add(StandardPackageFilename);
 			}
 		}
 	}
@@ -6778,6 +6762,8 @@ void UCookOnTheFlyServer::CleanShaderCodeLibraries()
 
 void UCookOnTheFlyServer::CookByTheBookFinished()
 {
+	using namespace UE::Cook;
+
 	check( IsInGameThread() );
 	check( IsCookByTheBookMode() );
 	check( CookByTheBookOptions->bRunning == true );
@@ -6825,55 +6811,66 @@ void UCookOnTheFlyServer::CookByTheBookFinished()
 			UE_SCOPED_HIERARCHICAL_COOKTIMER(SavingAssetRegistry);
 			for (const ITargetPlatform* TargetPlatform : PlatformManager->GetSessionPlatforms())
 			{
-				UE::Cook::FPlatformData* PlatformData = PlatformManager->GetPlatformData(TargetPlatform);
+				FPlatformData* PlatformData = PlatformManager->GetPlatformData(TargetPlatform);
 				FAssetRegistryGenerator& Generator = *PlatformData->RegistryGenerator;
-				TArray<FName> CookedPackagesFilenames;
-				TArray<FName> IgnorePackageFilenames;
+				TArray<FPackageData*> CookedPackageDatas;
+				TArray<FPackageData*> IgnorePackageDatas;
 
 				const FName& PlatformName = FName(*TargetPlatform->PlatformName());
 				FString PlatformNameString = PlatformName.ToString();
 
-				PackageDatas->GetCookedFileNamesForPlatform(TargetPlatform, CookedPackagesFilenames, false, /* include successful */ true);
+				PackageDatas->GetCookedPackagesForPlatform(TargetPlatform, CookedPackageDatas, false, /* include successful */ true);
 
 				// ignore any packages which failed to cook
-				PackageDatas->GetCookedFileNamesForPlatform(TargetPlatform, IgnorePackageFilenames, /* include failed */ true, false);
+				PackageDatas->GetCookedPackagesForPlatform(TargetPlatform, IgnorePackageDatas, /* include failed */ true, false);
 
 				bool bForceNoFilterAssetsFromAssetRegistry = false;
 
 				if (IsCookingDLC())
 				{
-					TSet<FName> CookedPackagesSet(CookedPackagesFilenames);
+					TMap<FName, FPackageData*> CookedPackagesMap;
+					CookedPackagesMap.Reserve(CookedPackageDatas.Num());
+					for (FPackageData* PackageData : CookedPackageDatas)
+					{
+						CookedPackagesMap.Add(PackageData->GetFileName(), PackageData);
+					}
 					bForceNoFilterAssetsFromAssetRegistry = true;
 					// remove the previous release cooked packages from the new asset registry, add to ignore list
 					UE_SCOPED_HIERARCHICAL_COOKTIMER(RemovingOldManifestEntries);
-					
+
 					const TArray<FName>* PreviousReleaseCookedPackages = CookByTheBookOptions->BasedOnReleaseCookedPackages.Find(PlatformName);
 					if (PreviousReleaseCookedPackages)
 					{
 						for (FName PreviousReleaseCookedPackage : *PreviousReleaseCookedPackages)
 						{
-							CookedPackagesSet.Remove(PreviousReleaseCookedPackage);
-							IgnorePackageFilenames.Add(PreviousReleaseCookedPackage);
+							FPackageData* PackageData;
+							if (!CookedPackagesMap.RemoveAndCopyValue(PreviousReleaseCookedPackage, PackageData))
+							{
+								PackageData = PackageDatas->FindPackageDataByFileName(PreviousReleaseCookedPackage);
+							}
+							if (PackageData)
+							{
+								IgnorePackageDatas.Add(PackageData);
+							}
 						}
 					}
-					CookedPackagesFilenames = CookedPackagesSet.Array();
+					CookedPackageDatas.Reset();
+					for (TPair<FName, FPackageData*>& Pair : CookedPackagesMap)
+					{
+						CookedPackageDatas.Add(Pair.Value);
+					}
 				}
 
-				// convert from filenames to package names
 				TSet<FName> CookedPackageNames;
-				for (FName PackageFilename : CookedPackagesFilenames)
+				for (FPackageData* PackageData : CookedPackageDatas)
 				{
-					const FName *FoundLongPackageFName = GetPackageNameCache().GetCachedPackageNameFromStandardFileName(PackageFilename);
-					check(FoundLongPackageFName);
-					CookedPackageNames.Add(*FoundLongPackageFName);
+					CookedPackageNames.Add(PackageData->GetPackageName());
 				}
 
 				TSet<FName> IgnorePackageNames;
-				for (FName PackageFilename : IgnorePackageFilenames)
+				for (FPackageData* PackageData : IgnorePackageDatas)
 				{
-					const FName *FoundLongPackageFName = GetPackageNameCache().GetCachedPackageNameFromStandardFileName(PackageFilename);
-					check(FoundLongPackageFName);
-					IgnorePackageNames.Add(*FoundLongPackageFName);
+					IgnorePackageNames.Add(PackageData->GetPackageName());
 				}
 
 				// ignore packages that weren't cooked because they were only referenced by editor-only properties
@@ -7010,14 +7007,14 @@ void UCookOnTheFlyServer::BuildMapDependencyGraph(const ITargetPlatform* TargetP
 {
 	auto& MapDependencyGraph = CookByTheBookOptions->MapDependencyGraphs.FindChecked(TargetPlatform);
 
-	TArray<FName> PlatformCookedPackages;
-	PackageDatas->GetCookedFileNamesForPlatform(TargetPlatform, PlatformCookedPackages, /* include failed */ true, /* include successful */ true);
+	TArray<UE::Cook::FPackageData*> PlatformCookedPackages;
+	PackageDatas->GetCookedPackagesForPlatform(TargetPlatform, PlatformCookedPackages, /* include failed */ true, /* include successful */ true);
 
 	// assign chunks for all the map packages
-	for (const FName& CookedPackage : PlatformCookedPackages)
+	for (const UE::Cook::FPackageData* const CookedPackage : PlatformCookedPackages)
 	{
 		TArray<FAssetData> PackageAssets;
-		FName Name = FName(*FPackageName::FilenameToLongPackageName(CookedPackage.ToString()));
+		FName Name = CookedPackage->GetPackageName();
 
 		if (!ContainsMap(Name))
 		{
@@ -7457,7 +7454,7 @@ UE::Cook::FCookSavePackageContext* UCookOnTheFlyServer::CreateSaveContext(const
 	else
 	{
 		PackageWriter = new FLooseCookedPackageWriter(ResolvedRootPath, ResolvedMetadataPath, TargetPlatform,
-			GetAsyncIODelete(), GetPackageNameCache(), PluginsToRemap);
+			GetAsyncIODelete(), *PackageDatas, PluginsToRemap);
 		WriterDebugName = TEXT("DirectoryWriter");
 	}
 
@@ -7567,15 +7564,9 @@ void UCookOnTheFlyServer::DiscoverPlatformSpecificNeverCookPackages(
 	}
 }
 
-const FPackageNameCache& UCookOnTheFlyServer::GetPackageNameCache() const
-{
-	return PackageDatas->GetPackageNameCache();
-}
-
 void UCookOnTheFlyServer::TermSandbox()
 {
 	ClearAllCookedData();
-	GetPackageNameCache().ClearPackageFileNameCache(nullptr);
 	SandboxFile = nullptr;
 }
 
@@ -7779,21 +7770,14 @@ void UCookOnTheFlyServer::StartCookByTheBook( const FCookByTheBookStartupOptions
 	
 	if ( IsCookingDLC() )
 	{
-		const FPackageNameCache& PackageNameCache = GetPackageNameCache();
-		IAssetRegistry* CacheAssetRegistry = PackageNameCache.GetAssetRegistry();
-		if (CacheAssetRegistry == nullptr)
-		{
-			UE_LOG(LogCook, Log, TEXT("Temporarily Replacing PackageNameCache Asset Registry with the CookOnTheFlyServer's AssetRegistry to initialise Cache"));
-			PackageNameCache.SetAssetRegistry(AssetRegistry);
-		}
-
 		// If we're cooking against a fixed base, we don't need to verify the packages exist on disk, we simply want to use the Release Data 
 		const bool bVerifyPackagesExist = !IsCookingAgainstFixedBase();
 
 		// if we are cooking dlc we must be based on a release version cook
 		check( !BasedOnReleaseVersion.IsEmpty() );
 
-		auto ReadDevelopmentAssetRegistry = [this, &BasedOnReleaseVersion, bVerifyPackagesExist](TArray<FName>& OutPackageList, const FString& InPlatformName)
+		auto ReadDevelopmentAssetRegistry = [this, &BasedOnReleaseVersion, bVerifyPackagesExist]
+			(TArray<UE::Cook::FConstructPackageData>& OutPackageList, const FString& InPlatformName)
 		{
 			TArray<FString> AttemptedNames;
 			FString OriginalSandboxRegistryFilename = GetBasedOnReleaseVersionAssetRegistryPath(BasedOnReleaseVersion, InPlatformName ) / TEXT("Metadata") / GetDevelopmentAssetRegistryFilename();
@@ -7837,7 +7821,7 @@ void UCookOnTheFlyServer::StartCookByTheBook( const FCookByTheBookStartupOptions
 			return bSucceeded;
 		};
 
-		TArray<FName> OverridePackageList;
+		TArray<UE::Cook::FConstructPackageData> OverridePackageList;
 		FString DevelopmentAssetRegistryPlatformOverride;
 		const bool bUsingDevRegistryOverride = FParse::Value(FCommandLine::Get(), TEXT("DevelopmentAssetRegistryPlatformOverride="), DevelopmentAssetRegistryPlatformOverride);
 		if (bUsingDevRegistryOverride)
@@ -7857,7 +7841,7 @@ void UCookOnTheFlyServer::StartCookByTheBook( const FCookByTheBookStartupOptions
 		for ( const ITargetPlatform* TargetPlatform: TargetPlatforms )
 		{
 			SCOPED_BOOT_TIMING("AddCookedPlatforms");
-			TArray<FName> PackageList;
+			TArray<UE::Cook::FConstructPackageData> PackageList;
 			FString PlatformNameString = TargetPlatform->PlatformName();
 			FName PlatformName(*PlatformNameString);
 
@@ -7873,25 +7857,20 @@ void UCookOnTheFlyServer::StartCookByTheBook( const FCookByTheBookStartupOptions
 				}
 			}
 
-			TArray<FName>& ActivePackageList = OverridePackageList.Num() > 0 ? OverridePackageList : PackageList;
+			TArray<UE::Cook::FConstructPackageData>& ActivePackageList = OverridePackageList.Num() > 0 ? OverridePackageList : PackageList;
 			if (ActivePackageList.Num() > 0)
 			{
 				SCOPED_BOOT_TIMING("AddPackageDataByFileNamesForPlatform");
 				PackageDatas->AddExistingPackageDatasForPlatform(ActivePackageList, TargetPlatform);
 			}
 
-			if (OverridePackageList.Num() > 0)
+			TArray<FName>& PlatformBasedPackages = CookByTheBookOptions->BasedOnReleaseCookedPackages.FindOrAdd(PlatformName);
+			PlatformBasedPackages.Reset(ActivePackageList.Num());
+			for (UE::Cook::FConstructPackageData& PackageData : ActivePackageList)
 			{
-				// This is the override list, so we can't give the memory away because we will need it for the other platforms
-				CookByTheBookOptions->BasedOnReleaseCookedPackages.Add(PlatformName, OverridePackageList);
-			}
-			else
-			{
-				CookByTheBookOptions->BasedOnReleaseCookedPackages.Add(PlatformName, MoveTemp(PackageList));
+				PlatformBasedPackages.Add(PackageData.NormalizedFileName);
 			}
 		}
-
-		PackageNameCache.SetAssetRegistry(CacheAssetRegistry);
 	}
 	
 	// add shader library chunkers
@@ -7998,25 +7977,25 @@ void UCookOnTheFlyServer::StartCookByTheBook( const FCookByTheBookStartupOptions
 		GenerateLongPackageNames(FilesInPath, FilesInPathInstigators);
 	}
 	// add all the files to the cook list for the requested platforms
-	for ( const FName& FileFName : FilesInPath )
+	for (FName PackageName : FilesInPath)
 	{
-		if (FileFName == NAME_None)
+		if (PackageName.IsNone())
 		{
 			continue;
 		}
 
-		const FName PackageFileFName = GetPackageNameCache().GetCachedStandardFileName(FileFName);
+		const FName PackageFileFName = PackageDatas->GetFileNameByPackageName(PackageName);
 		
-		if (PackageFileFName != NAME_None)
+		if (!PackageFileFName.IsNone())
 		{
-			UE::Cook::FInstigator& Instigator = FilesInPathInstigators.FindChecked(FileFName);
-			ExternalRequests->EnqueueUnique( UE::Cook::FFilePlatformRequest( PackageFileFName, MoveTemp(Instigator), TargetPlatforms ) );
+			UE::Cook::FInstigator& Instigator = FilesInPathInstigators.FindChecked(PackageName);
+			ExternalRequests->EnqueueUnique(UE::Cook::FFilePlatformRequest(PackageFileFName, MoveTemp(Instigator), TargetPlatforms));
 		}
-		else if (!FLinkerLoad::IsKnownMissingPackage(FileFName))
+		else if (!FLinkerLoad::IsKnownMissingPackage(PackageName))
 		{
-			const FString FileName = FileFName.ToString();
-			LogCookerMessage( FString::Printf(TEXT("Unable to find package for cooking %s"), *FileName), EMessageSeverity::Warning );
-		}	
+			LogCookerMessage( FString::Printf(TEXT("Unable to find package for cooking %s"), *PackageName.ToString()),
+				EMessageSeverity::Warning );
+		}
 	}
 
 	if (!IsCookingDLC())
@@ -8033,15 +8012,16 @@ void UCookOnTheFlyServer::StartCookByTheBook( const FCookByTheBookStartupOptions
 				// if we are based of a cook and we are creating a new one we need to make sure that at least all the old packages are cooked as well as the new ones
 				FString OriginalAssetRegistryPath = GetBasedOnReleaseVersionAssetRegistryPath( BasedOnReleaseVersion, TargetPlatform->PlatformName() ) / GetAssetRegistryFilename();
 
-				TArray<FName> PackageFiles;
-				verify( GetAllPackageFilenamesFromAssetRegistry(OriginalAssetRegistryPath, true, PackageFiles) );
+				TArray<UE::Cook::FConstructPackageData> BasedOnReleaseDatas;
+				verify( GetAllPackageFilenamesFromAssetRegistry(OriginalAssetRegistryPath, true, BasedOnReleaseDatas) );
 
 				TArray<const ITargetPlatform*, TInlineAllocator<1>> RequestPlatforms;
 				RequestPlatforms.Add(TargetPlatform);
-				for ( const FName& PackageFilename : PackageFiles )
+				for ( const UE::Cook::FConstructPackageData& PackageData : BasedOnReleaseDatas)
 				{
 					ExternalRequests->EnqueueUnique(
-						UE::Cook::FFilePlatformRequest( PackageFilename, UE::Cook::EInstigator::PreviousAssetRegistry, RequestPlatforms) );
+						UE::Cook::FFilePlatformRequest(PackageData.NormalizedFileName,
+							UE::Cook::EInstigator::PreviousAssetRegistry, RequestPlatforms));
 				}
 			}
 		}
@@ -8114,7 +8094,7 @@ TArray<FName> UCookOnTheFlyServer::GetNeverCookPackageFileNames(TArrayView<const
 	TArray<FName> NeverCookNormalizedFileNames;
 	for (const FString& NeverCookPackagePath : NeverCookPackagesPaths)
 	{
-		NeverCookNormalizedFileNames.Add(FPackageNameCache::GetStandardFileName(NeverCookPackagePath));
+		NeverCookNormalizedFileNames.Add(UE::Cook::FPackageDatas::GetStandardFileName(NeverCookPackagePath));
 	}
 	return NeverCookNormalizedFileNames;
 }
@@ -8154,15 +8134,16 @@ void UCookOnTheFlyServer::MaybeMarkPackageAsAlreadyLoaded(UPackage *Package)
 		return;
 	}
 
-	FName StandardName = GetPackageNameCache().GetCachedStandardFileName(Package);
-
-	// UE_LOG(LogCook, Display, TEXT("Loading package %s"), *StandardName.ToString());
-
 	bool bShouldMarkAsAlreadyProcessed = false;
 
 	TArray<const ITargetPlatform*> CookedPlatforms;
 	UE::Cook::FPackageData* PackageData = PackageDatas->FindPackageDataByPackageName(Package->GetFName());
-	if (PackageData && PackageData->HasAnyCookedPlatform())
+	if (!PackageData)
+	{
+		return;
+	}
+	FName StandardName = PackageData->GetFileName();
+	if (PackageData->HasAnyCookedPlatform())
 	{
 		bShouldMarkAsAlreadyProcessed = PackageData->HasAllCookedPlatforms(PlatformManager->GetSessionPlatforms(), true /* bIncludeFailed */);
 
@@ -8257,8 +8238,11 @@ void UCookOnTheFlyServer::GetCookOnTheFlyUnsolicitedFiles(const ITargetPlatform*
 	}
 }
 
-bool UCookOnTheFlyServer::GetAllPackageFilenamesFromAssetRegistry( const FString& AssetRegistryPath, bool bVerifyPackagesExist, TArray<FName>& OutPackageFilenames ) const
+bool UCookOnTheFlyServer::GetAllPackageFilenamesFromAssetRegistry(const FString& AssetRegistryPath,
+	bool bVerifyPackagesExist, TArray<UE::Cook::FConstructPackageData>& OutPackageDatas) const
 {
+	using namespace UE::Cook;
+
 	UE_SCOPED_COOKTIMER(GetAllPackageFilenamesFromAssetRegistry);
 	SCOPED_BOOT_TIMING("GetAllPackageFilenamesFromAssetRegistry");
 	TUniquePtr<FArchive> Reader(IFileManager::Get().CreateFileReader(*AssetRegistryPath));
@@ -8289,73 +8273,86 @@ bool UCookOnTheFlyServer::GetAllPackageFilenamesFromAssetRegistry( const FString
 		// possibly use the preloaded state
 		const TMap<FName, const FAssetData*>& RegistryDataMap = GPreloadedARState.GetObjectPathToAssetDataMap();
 
-		check(OutPackageFilenames.Num() == 0);
-		OutPackageFilenames.SetNum(RegistryDataMap.Num());
+		check(OutPackageDatas.Num() == 0);
 
-		TArray<const FAssetData*> Packages;
-		Packages.Reserve(RegistryDataMap.Num());
+		int32 NumPackages = RegistryDataMap.Num();
+		TArray<const FAssetData*> AssetDatas;
+		TSet<FName> PackageNames;
+		AssetDatas.Reserve(NumPackages);
+		OutPackageDatas.Reserve(NumPackages);
 
+		// Convert the Map of RegistryData into an Array of FAssetData and populate PackageNames in the output array
 		for (const TPair<FName, const FAssetData*>& RegistryData : RegistryDataMap)
 		{
-			if (!FPackageName::GetPackageMountPoint(RegistryData.Value->PackageName.ToString()).IsNone())
+			FName PackageName = RegistryData.Value->PackageName;
+			bool bPackageAlreadyAdded;
+			PackageNames.Add(PackageName, &bPackageAlreadyAdded);
+			if (bPackageAlreadyAdded)
 			{
-				int32 AddedIndex = Packages.Add(RegistryData.Value);
-				if (GetPackageNameCache().ContainsPackageName(Packages.Last()->PackageName))
-				{
-					OutPackageFilenames[AddedIndex] = GetPackageNameCache().GetCachedStandardFileName(Packages.Last()->PackageName);
-				}
+				continue;
+			}
+			if (FPackageName::GetPackageMountPoint(PackageName.ToString()).IsNone())
+			{
+				// Skip any packages that are not currently mounted; if we tried to find their FileNames below
+				// we would get log spam
+				continue;
+			}
+
+			AssetDatas.Add(RegistryData.Value);
+			FConstructPackageData& PackageData = OutPackageDatas.Emplace_GetRef();
+			PackageData.PackageName = PackageName;
+
+			// For any PackageNames that already have PackageDatas, mark them ahead of the loop to
+			// skip the effort of checking whether they exist on disk inside the loop
+			FPackageData* ExistingPackageData = PackageDatas->FindPackageDataByPackageName(PackageName);
+			if (ExistingPackageData)
+			{
+				PackageData.NormalizedFileName = ExistingPackageData->GetFileName();
 			}
 		}
+		NumPackages = AssetDatas.Num();
 
-		TArray<TTuple<FName, FPackageNameCache::FCachedPackageFilename>> PackageToStandardFileNames;
-		PackageToStandardFileNames.SetNum(RegistryDataMap.Num());
-		ParallelFor(Packages.Num(), [&AssetRegistryPath, &OutPackageFilenames, &PackageToStandardFileNames, &Packages, this, bVerifyPackagesExist](int32 AssetIndex)
+		ParallelFor(NumPackages,
+			[&AssetRegistryPath, &OutPackageDatas, &AssetDatas, this, bVerifyPackagesExist](int32 AssetIndex)
 			{
-				if (!OutPackageFilenames[AssetIndex].IsNone())
+				FConstructPackageData& PackageData = OutPackageDatas[AssetIndex];
+				if (!PackageData.NormalizedFileName.IsNone())
 				{
 					return;
 				}
+				const FName PackageName = PackageData.PackageName;
 
-				const FName PackageName = Packages[AssetIndex]->PackageName;
-
-				FString StandardFilename;
-				// TODO ICookPackageSplitter: Need to handle GeneratedPackages that exist in the cooked AssetRegistry we are reading, but do not exist
-				// in the workspace domain and so are not found when we try to calculate their cache data here.
-				if (!GetPackageNameCache().TryCalculateCacheData(PackageName, StandardFilename, OutPackageFilenames[AssetIndex]))
+				// TODO ICookPackageSplitter: Need to handle GeneratedPackages that exist in the cooked AssetRegistry we are 
+				// reading, but do not exist in WorkspaceDomain and so are not found when we look them up here.
+				FName PackageFileName = FPackageDatas::LookupFileNameOnDisk(PackageName, true /* bRequireExists */);
+				if (!PackageFileName.IsNone())
+				{
+					PackageData.NormalizedFileName = PackageFileName;
+				}
+				else
 				{
 					if (bVerifyPackagesExist)
 					{
-						UE_LOG(LogCook, Warning, TEXT("Could not resolve package %s from %s"), *PackageName.ToString(), *AssetRegistryPath);
+						UE_LOG(LogCook, Warning, TEXT("Could not resolve package %s from %s"),
+							*PackageName.ToString(), *AssetRegistryPath);
 					}
 					else
 					{
-						const bool bContainsMap = Packages[AssetIndex]->PackageFlags & PKG_ContainsMap;
-						FString PackageNameStr = PackageName.ToString();
-
-						if (FPackageName::TryConvertLongPackageNameToFilename(PackageNameStr, StandardFilename, bContainsMap ? FPackageName::GetMapPackageExtension() : FPackageName::GetAssetPackageExtension()))
+						const bool bContainsMap = AssetDatas[AssetIndex]->PackageFlags & PKG_ContainsMap;
+						PackageFileName = FPackageDatas::LookupFileNameOnDisk(PackageName,
+							false /* bRequireExists */, bContainsMap);
+						if (!PackageFileName.IsNone())
 						{
-							OutPackageFilenames[AssetIndex] = FPackageNameCache::GetStandardFileName(StandardFilename);
-							StandardFilename = OutPackageFilenames[AssetIndex].ToString();
+							PackageData.NormalizedFileName = PackageFileName;
 						}
 					}
 				}
-				
-				FName StandardFileFName(*StandardFilename);
-				PackageToStandardFileNames[AssetIndex] = TTuple<FName, FPackageNameCache::FCachedPackageFilename>(PackageName, FPackageNameCache::FCachedPackageFilename(MoveTemp(StandardFilename), StandardFileFName));
 			});
 
-		for (int32 Idx = OutPackageFilenames.Num() - 1; Idx >= 0; --Idx)
-		{
-			if (OutPackageFilenames[Idx] == NAME_None)
+		OutPackageDatas.RemoveAllSwap([](FConstructPackageData& PackageData)
 			{
-				OutPackageFilenames.RemoveAtSwap(Idx);
-			}
-		}
-
-		{
-			SCOPED_BOOT_TIMING("AppendCacheResults");
-			GetPackageNameCache().AppendCacheResults(MoveTemp(PackageToStandardFileNames));
-		}
+				return PackageData.NormalizedFileName.IsNone();
+			}, false /* bAllowShrinking */);
 		return true;
 	}
 
@@ -8462,7 +8459,7 @@ uint32 UCookOnTheFlyServer::FullLoadAndSave(uint32& CookedPackageCount)
 				}
 
 				FName PackageName = Package->GetFName();
-				if (PackageTracker->NeverCookPackageList.Contains(GetPackageNameCache().GetCachedStandardFileName(PackageName)))
+				if (PackageTracker->NeverCookPackageList.Contains(PackageDatas->GetFileNameByPackageName(PackageName)))
 				{
 					// refuse to save this package
 					continue;
@@ -8626,7 +8623,7 @@ uint32 UCookOnTheFlyServer::FullLoadAndSave(uint32& CookedPackageCount)
 			GIsCookerLoadingPackage = true;
 			for (const FString& ToLoad : PackagesToLoad)
 			{
-				FName BuildFilenameFName = GetPackageNameCache().GetCachedStandardFileName(FName(*ToLoad));
+				FName BuildFilenameFName = PackageDatas->GetFileNameByPackageName(FName(*ToLoad));
 				if (!PackageTracker->NeverCookPackageList.Contains(BuildFilenameFName))
 				{
 					LoadPackage(nullptr, *ToLoad, LOAD_None);
@@ -8951,4 +8948,4 @@ UE::Cook::FCookSavePackageContext& UCookOnTheFlyServer::FindOrCreateSaveContext(
 }
 
 
-#undef LOCTEXT_NAMESPACE
+#undef LOCTEXT_NAMESPACE
\ No newline at end of file
diff --git a/Engine/Source/Editor/UnrealEd/Private/Cooker/CookPackageData.cpp b/Engine/Source/Editor/UnrealEd/Private/Cooker/CookPackageData.cpp
--- a/Engine/Source/Editor/UnrealEd/Private/Cooker/CookPackageData.cpp
+++ b/Engine/Source/Editor/UnrealEd/Private/Cooker/CookPackageData.cpp
@@ -4,7 +4,9 @@
 
 #include "Algo/AnyOf.h"
 #include "Algo/Count.h"
+#include "Algo/Find.h"
 #include "AssetCompilingManager.h"
+#include "AssetRegistry/IAssetRegistry.h"
 #include "Async/ParallelFor.h"
 #include "Cooker/CookPlatformManager.h"
 #include "Cooker/CookRequestCluster.h"
@@ -21,12 +23,14 @@
 #include "Misc/Parse.h"
 #include "Misc/Paths.h"
 #include "Misc/PreloadableFile.h"
+#include "Misc/ScopeLock.h"
 #include "ShaderCompiler.h"
 #include "UObject/Object.h"
 #include "UObject/Package.h"
 #include "UObject/ReferenceChainSearch.h"
 #include "UObject/UObjectGlobals.h"
 #include "UObject/UObjectHash.h"
+#include "UObject/ObjectRedirector.h"
 
 namespace UE::Cook
 {
@@ -1136,8 +1140,7 @@ void FGeneratorPackage::ClearGeneratedPackages()
 	}
 }
 
-bool FGeneratorPackage::TryGenerateList(UObject* OwnerObject, const FPackageNameCache& PackageNameCache,
-	FPackageDatas& PackageDatas)
+bool FGeneratorPackage::TryGenerateList(UObject* OwnerObject, FPackageDatas& PackageDatas)
 {
 	UPackage* OwnerPackage = Owner.GetPackage();
 	check(OwnerPackage);
@@ -1161,28 +1164,27 @@ bool FGeneratorPackage::TryGenerateList(UObject* OwnerObject, const FPackageName
 		GeneratedStruct.bCreateAsMap = *SplitterData.GetCreateAsMap();
 
 		const FName PackageFName(*PackageName);
-		const FName FileName = PackageNameCache.GetCachedStandardFileName(PackageFName, false /* bRequireExists */,
-			GeneratedStruct.bCreateAsMap);
-		if (FileName.IsNone())
+		UE::Cook::FPackageData* PackageData = PackageDatas.TryAddPackageDataByPackageName(PackageFName,
+			false /* bRequireExists */, GeneratedStruct.bCreateAsMap);
+		if (!PackageData)
 		{
 			UE_LOG(LogCook, Error, TEXT("PackageSplitter could not find mounted filename for generated packagepath. Splitter=%s, Generated=%s."),
 				*this->GetSplitDataObjectName().ToString(), *PackageName);
 			return false;
 		}
-		UE::Cook::FPackageData& PackageData = PackageDatas.FindOrAddPackageData(PackageFName, FileName);
-		if (IFileManager::Get().FileExists(*FileName.ToString()))
+		if (IFileManager::Get().FileExists(*PackageData->GetFileName().ToString()))
 		{
 			UE_LOG(LogCook, Warning, TEXT("PackageSplitter specified a generated package that already exists in the workspace domain. Splitter=%s, Generated=%s."),
 				*this->GetSplitDataObjectName().ToString(), *PackageName);
 			return false;
 		}
-		GeneratedStruct.PackageData = &PackageData;
-		PackageData.SetGenerated(true);
+		GeneratedStruct.PackageData = PackageData;
+		PackageData->SetGenerated(true);
 		// No package should be generated by two different splitters. If an earlier run of this splitter generated
 		// the package, the package's owner should have been reset to null when we called ClearGeneratedPackages
 		// between then and now
-		check(PackageData.GetGeneratedOwner() == nullptr); 
-		PackageData.SetGeneratedOwner(this);
+		check(PackageData->GetGeneratedOwner() == nullptr); 
+		PackageData->SetGeneratedOwner(this);
 	}
 	RemainingToPopulate = GeneratorDatas.Num();
 	return true;
@@ -1543,6 +1545,7 @@ void FPackageDataMonitor::TrackUrgentRequests(EPackageState State, int32 Delta)
 //////////////////////////////////////////////////////////////////////////
 // FPackageDatas
 
+IAssetRegistry* FPackageDatas::AssetRegistry = nullptr;
 
 FPackageDatas::FPackageDatas(UCookOnTheFlyServer& InCookOnTheFlyServer)
 	: CookOnTheFlyServer(InCookOnTheFlyServer)
@@ -1566,16 +1569,12 @@ void FPackageDatas::BeginCook()
 		else
 		{
 			FName PackageFName(*PackageName);
-			FName FileName = PackageNameCache.GetCachedStandardFileName(PackageFName);
-			if (FileName.IsNone())
+			ShowInstigatorPackageData = TryAddPackageDataByPackageName(PackageFName);
+			if (!ShowInstigatorPackageData)
 			{
 				UE_LOG(LogCook, Fatal, TEXT("-CookShowInstigator argument %s could not be found on disk"),
 					*FileOrPackageName);
 			}
-			else
-			{
-				ShowInstigatorPackageData = &FindOrAddPackageData(PackageFName, FileName);
-			}
 		}
 	}
 }
@@ -1585,6 +1584,11 @@ FPackageDatas::~FPackageDatas()
 	Clear();
 }
 
+void FPackageDatas::OnAssetRegistryGenerated(IAssetRegistry& InAssetRegistry)
+{
+	AssetRegistry = &InAssetRegistry;
+}
+
 FString FPackageDatas::GetReferencerName() const
 {
 	return TEXT("FPackageDatas");
@@ -1595,11 +1599,6 @@ void FPackageDatas::AddReferencedObjects(FReferenceCollector& Collector)
 	return CookOnTheFlyServer.CookerAddReferencedObjects(Collector);
 }
 
-const FPackageNameCache& FPackageDatas::GetPackageNameCache() const
-{
-	return PackageNameCache;
-}
-
 FPackageDataMonitor& FPackageDatas::GetMonitor()
 {
 	return Monitor;
@@ -1622,18 +1621,21 @@ FPackageDataQueue& FPackageDatas::GetSaveQueue()
 
 FPackageData& FPackageDatas::FindOrAddPackageData(const FName& PackageName, const FName& NormalizedFileName)
 {
-	FPackageData** PackageDataMapAddr = PackageNameToPackageData.Find(PackageName);
-	if (PackageDataMapAddr != nullptr)
 	{
-		FPackageData** FileNameMapAddr = FileNameToPackageData.Find(NormalizedFileName);
-		check(FileNameMapAddr && *FileNameMapAddr == *PackageDataMapAddr);
-		return **PackageDataMapAddr;
-	}
+		FReadScopeLock ExistenceReadLock(ExistenceLock);
+		FPackageData** PackageDataMapAddr = PackageNameToPackageData.Find(PackageName);
+		if (PackageDataMapAddr != nullptr)
+		{
+			FPackageData** FileNameMapAddr = FileNameToPackageData.Find(NormalizedFileName);
+			check(FileNameMapAddr && *FileNameMapAddr == *PackageDataMapAddr);
+			return **PackageDataMapAddr;
+		}
 
-	checkf(FileNameToPackageData.Find(NormalizedFileName) == nullptr,
-		TEXT("Package \"%s\" and package \"%s\" share the same filename \"%s\"."),
-		*PackageName.ToString(), *(*FileNameToPackageData.Find(NormalizedFileName))->GetPackageName().ToString(),
-		*NormalizedFileName.ToString());
+		checkf(FileNameToPackageData.Find(NormalizedFileName) == nullptr,
+			TEXT("Package \"%s\" and package \"%s\" share the same filename \"%s\"."),
+			*PackageName.ToString(), *(*FileNameToPackageData.Find(NormalizedFileName))->GetPackageName().ToString(),
+			*NormalizedFileName.ToString());
+	}
 	return CreatePackageData(PackageName, NormalizedFileName);
 }
 
@@ -1644,97 +1646,116 @@ FPackageData* FPackageDatas::FindPackageDataByPackageName(const FName& PackageNa
 		return nullptr;
 	}
 
+	FReadScopeLock ExistenceReadLock(ExistenceLock);
 	FPackageData** PackageDataMapAddr = PackageNameToPackageData.Find(PackageName);
 	return PackageDataMapAddr ? *PackageDataMapAddr : nullptr;
 }
 
-FPackageData* FPackageDatas::TryAddPackageDataByPackageName(const FName& PackageName)
+FPackageData* FPackageDatas::TryAddPackageDataByPackageName(const FName& PackageName, bool bRequireExists,
+	bool bCreateAsMap)
 {
 	if (PackageName.IsNone())
 	{
 		return nullptr;
 	}
 
-	FPackageData** PackageDataMapAddr = PackageNameToPackageData.Find(PackageName);
-	if (PackageDataMapAddr != nullptr)
 	{
-		return *PackageDataMapAddr;
+		FReadScopeLock ExistenceReadLock(ExistenceLock);
+		FPackageData** PackageDataMapAddr = PackageNameToPackageData.Find(PackageName);
+		if (PackageDataMapAddr != nullptr)
+		{
+			return *PackageDataMapAddr;
+		}
 	}
 
-	FName FileName = PackageNameCache.GetCachedStandardFileName(PackageName);
+	FName FileName = LookupFileNameOnDisk(PackageName, bRequireExists, bCreateAsMap);
 	if (FileName.IsNone())
 	{
-		// This can happen if PackageName is a script package
+		// This will happen if PackageName does not exist on disk
 		return nullptr;
 	}
-	checkf(FileNameToPackageData.Find(FileName) == nullptr,
-		TEXT("Package \"%s\" and package \"%s\" share the same filename \"%s\"."),
-		*PackageName.ToString(), *(*FileNameToPackageData.Find(FileName))->GetPackageName().ToString(),
-		*FileName.ToString()); 
+	{
+		FReadScopeLock ExistenceReadLock(ExistenceLock);
+		checkf(FileNameToPackageData.Find(FileName) == nullptr,
+			TEXT("Package \"%s\" and package \"%s\" share the same filename \"%s\"."),
+			*PackageName.ToString(), *(*FileNameToPackageData.Find(FileName))->GetPackageName().ToString(),
+			*FileName.ToString());
+	}
 	return &CreatePackageData(PackageName, FileName);
 }
 
-FPackageData& FPackageDatas::AddPackageDataByPackageNameChecked(const FName& PackageName)
+FPackageData& FPackageDatas::AddPackageDataByPackageNameChecked(const FName& PackageName, bool bRequireExists,
+	bool bCreateAsMap)
 {
-	FPackageData* PackageData = TryAddPackageDataByPackageName(PackageName);
+	FPackageData* PackageData = TryAddPackageDataByPackageName(PackageName, bRequireExists, bCreateAsMap);
 	check(PackageData);
 	return *PackageData;
 }
 
 FPackageData* FPackageDatas::FindPackageDataByFileName(const FName& InFileName)
 {
-	FName FileName(FPackageNameCache::GetStandardFileName(InFileName));
+	FName FileName(GetStandardFileName(InFileName));
 	if (FileName.IsNone())
 	{
 		return nullptr;
 	}
 
+	FReadScopeLock ExistenceReadLock(ExistenceLock);
 	FPackageData** PackageDataMapAddr = FileNameToPackageData.Find(FileName);
 	return PackageDataMapAddr ? *PackageDataMapAddr : nullptr;
 }
 
 FPackageData* FPackageDatas::TryAddPackageDataByFileName(const FName& InFileName)
 {
-	FName FileName(FPackageNameCache::GetStandardFileName(InFileName));
+	return TryAddPackageDataByStandardFileName(GetStandardFileName(InFileName));
+}
+
+FPackageData* FPackageDatas::TryAddPackageDataByStandardFileName(const FName& FileName)
+{
 	if (FileName.IsNone())
 	{
 		return nullptr;
 	}
 
-	FPackageData** PackageDataMapAddr = FileNameToPackageData.Find(FileName);
-	if (PackageDataMapAddr != nullptr)
 	{
-		return *PackageDataMapAddr;
+		FReadScopeLock ExistenceReadLock(ExistenceLock);
+		FPackageData** PackageDataMapAddr = FileNameToPackageData.Find(FileName);
+		if (PackageDataMapAddr != nullptr)
+		{
+			return *PackageDataMapAddr;
+		}
 	}
 
-	const FName* PackageName = PackageNameCache.GetCachedPackageNameFromStandardFileName(FileName);
-	if (!PackageName)
+	FName PackageName = LookupPackageNameOnDisk(FileName);
+	if (PackageName.IsNone())
 	{
 		return nullptr;
 	}
-	return &CreatePackageData(*PackageName, FileName);
+	return &CreatePackageData(PackageName, FileName);
 }
 
 FPackageData& FPackageDatas::CreatePackageData(FName PackageName, FName FileName)
 {
-	if (PackageName.IsNone())
-	{
-		check(!FileName.IsNone());
-		const FName* FoundPackageName = PackageNameCache.GetCachedPackageNameFromStandardFileName(FileName);
-		check(FoundPackageName);
-		PackageName = *FoundPackageName;
-		check(!PackageName.IsNone());
-	}
-	else if (FileName.IsNone())
+	check(!PackageName.IsNone());
+	check(!FileName.IsNone());
+	FPackageData* PackageData = new FPackageData(*this, PackageName, FileName);
+
+	FWriteScopeLock ExistenceWriteLock(ExistenceLock);
+	FPackageData*& ExistingByPackageName = PackageNameToPackageData.FindOrAdd(PackageName);
+	FPackageData*& ExistingByFileName = FileNameToPackageData.FindOrAdd(FileName);
+	if (ExistingByPackageName)
 	{
-		FileName = PackageNameCache.GetCachedStandardFileName(PackageName);
-		check(!FileName.IsNone());
+		// The other CreatePackageData call should have added the FileName as well
+		check(ExistingByFileName == ExistingByPackageName);
+		delete PackageData;
+		return *ExistingByPackageName;
 	}
-
-	FPackageData* PackageData = new FPackageData(*this, PackageName, FileName);
+	// If no other CreatePackageData added the PackageName, then they should not have added
+	// the FileName either
+	check(!ExistingByFileName);
+	ExistingByPackageName = PackageData;
+	ExistingByFileName = PackageData;
 	PackageDatas.Add(PackageData);
-	PackageNameToPackageData.Add(PackageName, PackageData);
-	FileNameToPackageData.Add(FileName, PackageData);
 	return *PackageData;
 }
 
@@ -1745,25 +1766,178 @@ FPackageData& FPackageDatas::AddPackageDataByFileNameChecked(const FName& FileNa
 	return *PackageData;
 }
 
-void FPackageDatas::AddExistingPackageDatasForPlatform(const TArray<FName>& ExistingPackages,
-	const ITargetPlatform* TargetPlatform)
+FName FPackageDatas::GetFileNameByPackageName(FName PackageName, bool bRequireExists, bool bCreateAsMap)
 {
-	TArray<FPackageData*> NewPackageDataObjects;
-	NewPackageDataObjects.AddZeroed(ExistingPackages.Num());
+	FPackageData* PackageData = TryAddPackageDataByPackageName(PackageName, bRequireExists, bCreateAsMap);
+	return PackageData ? PackageData->GetFileName() : NAME_None;
+}
 
-	// parallelize the read-only operations (and write NewPackageDataObjects by index which has no threading issues)
-	ParallelFor(ExistingPackages.Num(),
-		[&ExistingPackages, TargetPlatform, &NewPackageDataObjects, this](int Index)
+FName FPackageDatas::GetFileNameByFlexName(FName PackageOrFileName, bool bRequireExists, bool bCreateAsMap)
+{
+	FString Buffer = PackageOrFileName.ToString();
+	if (!FPackageName::TryConvertFilenameToLongPackageName(Buffer, Buffer))
+	{
+		return NAME_None;
+	}
+	return GetFileNameByPackageName(FName(Buffer), bRequireExists, bCreateAsMap);
+}
+
+FName FPackageDatas::LookupFileNameOnDisk(FName PackageName, bool bRequireExists, bool bCreateAsMap)
+{
+	FString FilenameOnDisk;
+	if (TryLookupFileNameOnDisk(PackageName, FilenameOnDisk))
 	{
-		const FName& PackageFilename = ExistingPackages[Index];
+	}
+	else if (!bRequireExists)
+	{
+		FString Extension = bCreateAsMap ? FPackageName::GetMapPackageExtension() : FPackageName::GetAssetPackageExtension();
+		if (!FPackageName::TryConvertLongPackageNameToFilename(PackageName.ToString(), FilenameOnDisk, Extension))
+		{
+			return NAME_None;
+		}
+	}
+	else
+	{
+		return NAME_None;
+	}
+	FilenameOnDisk = FPaths::ConvertRelativePathToFull(FilenameOnDisk);
+	FPaths::MakeStandardFilename(FilenameOnDisk);
+	return FName(FilenameOnDisk);
+}
 
-		FName FileName(FPackageNameCache::GetStandardFileName(PackageFilename));
-		if (FileName.IsNone())
+bool FPackageDatas::TryLookupFileNameOnDisk(FName PackageName, FString& OutFileName)
+{
+	FString PackageNameStr = PackageName.ToString();
+
+	// Verse packages are editor-generated in-memory packages which don't have a corresponding 
+	// asset file (yet). However, we still want to cook these packages out, producing cooked 
+	// asset files for packaged projects.
+	if (FPackageName::IsVersePackage(PackageNameStr))
+	{
+		if (FindPackage(/*Outer =*/nullptr, *PackageNameStr))
 		{
-			return;
+			return FPackageName::TryConvertLongPackageNameToFilename(PackageNameStr, OutFileName,
+				FPackageName::GetAssetPackageExtension());
 		}
+		// else, the cooker could be responding to a NotifyUObjectCreated() event, and the object hasn't
+		// been fully constructed yet (missing from the FindObject() list) -- in this case, we've found 
+		// that the linker loader is creating a dummy object to fill a referencing import slot, not loading
+		// the proper object (which means we want to ignore it).
+	}
 
-		FPackageData** PackageDataMapAddr = FileNameToPackageData.Find(FileName);
+	if (!AssetRegistry)
+	{
+		return FPackageName::DoesPackageExist(PackageNameStr, &OutFileName, false /* InAllowTextFormats */);
+	}
+	else
+	{
+		TArray<FAssetData> Assets;
+		AssetRegistry->GetAssetsByPackageName(PackageName, Assets, /*bIncludeOnlyDiskAssets =*/true);
+
+		if (Assets.Num() <= 0)
+		{
+			return false;
+		}
+
+		FName ClassRedirector = UObjectRedirector::StaticClass()->GetFName();
+		bool bContainsMap = false;
+		bool bContainsRedirector = false;
+		for (const FAssetData& Asset : Assets)
+		{
+			bContainsMap = bContainsMap | ((Asset.PackageFlags & PKG_ContainsMap) != 0);
+			bContainsRedirector = bContainsRedirector | (Asset.AssetClass == ClassRedirector);
+		}
+		if (!bContainsMap && bContainsRedirector)
+		{
+			// presence of map -> .umap
+			// But we can only assume lack of map -> .uasset if we know the type of every object in the package.
+			// If we don't, because there was a redirector, we have to check the package on disk
+			// TODO: Have the AssetRegistry store the extension of the package so that we don't have to look it up
+			// Guessing the extension based on map vs non-map also does not support text assets and maps which have a different extension
+			return FPackageName::DoesPackageExist(PackageNameStr, &OutFileName, false /* InAllowTextFormats */);
+		}
+		const FString& PackageExtension = bContainsMap ? FPackageName::GetMapPackageExtension() : FPackageName::GetAssetPackageExtension();
+		return FPackageName::TryConvertLongPackageNameToFilename(PackageNameStr, OutFileName, PackageExtension);
+	}
+}
+
+FName FPackageDatas::LookupPackageNameOnDisk(FName NormalizedFileName)
+{
+	if (NormalizedFileName.IsNone())
+	{
+		return NAME_None;
+	}
+	FString Buffer = NormalizedFileName.ToString();
+	if (!FPackageName::TryConvertFilenameToLongPackageName(Buffer, Buffer))
+	{
+		return NAME_None;
+	}
+	FName PackageName = FName(*Buffer);
+
+	FName DiscoveredFileName = LookupFileNameOnDisk(PackageName, true /* bRequireExists */, false /* bCreateAsMap */);
+	if (DiscoveredFileName == NormalizedFileName)
+	{
+		return PackageName;
+	}
+	else
+	{
+		// Either the file does not exist on disk or the input FileName was not in normalized format
+		return NAME_None;
+	}
+}
+
+FName FPackageDatas::GetStandardFileName(FName FileName)
+{
+	FString FileNameString(FileName.ToString());
+	FPaths::MakeStandardFilename(FileNameString);
+	return FName(FileNameString);
+}
+
+FName FPackageDatas::GetStandardFileName(FStringView InFileName)
+{
+	FString FileName(InFileName);
+	FPaths::MakeStandardFilename(FileName);
+	return FName(FileName);
+}
+
+void FPackageDatas::AddExistingPackageDatasForPlatform(TConstArrayView<FConstructPackageData> ExistingPackages,
+	const ITargetPlatform* TargetPlatform)
+{
+	int32 NumPackages = ExistingPackages.Num();
+
+	// Make the list unique
+	TMap<FName, FName> UniquePackages;
+	UniquePackages.Reserve(NumPackages);
+	for (const FConstructPackageData& PackageToAdd : ExistingPackages)
+	{
+		FName& AddedFileName = UniquePackages.FindOrAdd(PackageToAdd.PackageName, PackageToAdd.NormalizedFileName);
+		check(AddedFileName == PackageToAdd.NormalizedFileName);
+	}
+	TArray<FConstructPackageData> UniqueArray;
+	if (UniquePackages.Num() != NumPackages)
+	{
+		NumPackages = UniquePackages.Num();
+		UniqueArray.Reserve(NumPackages);
+		for (TPair<FName, FName>& Pair : UniquePackages)
+		{
+			UniqueArray.Add(FConstructPackageData{ Pair.Key, Pair.Value });
+		}
+		ExistingPackages = UniqueArray;
+	}
+
+	// parallelize the read-only operations (and write NewPackageDataObjects by index which has no threading issues)
+	TArray<FPackageData*> NewPackageDataObjects;
+	NewPackageDataObjects.AddZeroed(NumPackages);
+	FWriteScopeLock ExistenceWriteLock(ExistenceLock);
+	ParallelFor(NumPackages,
+		[&ExistingPackages, TargetPlatform, &NewPackageDataObjects, this](int Index)
+	{
+		FName PackageName = ExistingPackages[Index].PackageName;
+		FName NormalizedFileName = ExistingPackages[Index].NormalizedFileName;
+		check(!PackageName.IsNone());
+		check(!NormalizedFileName.IsNone());
+
+		FPackageData** PackageDataMapAddr = FileNameToPackageData.Find(NormalizedFileName);
 		FPackageData* PackageData = nullptr;
 		if (PackageDataMapAddr != nullptr)
 		{
@@ -1771,80 +1945,66 @@ void FPackageDatas::AddExistingPackageDatasForPlatform(const TArray<FName>& Exis
 		}
 		else
 		{
-			// we expect all of these to be in PackageNameCache from earlier in the AR load process
-			const FName* PackageName = PackageNameCache.FindExistingCachedPackageNameFromStandardFileName(FileName);
-			checkf(PackageName != nullptr, TEXT("AddExistingPackageDatasForPlatform should only be used when the PackageCache has already been filled out with the packages"));
-
 			// create the package data and remember it for updating caches after the the ParallelFor
-			PackageData = new FPackageData(*this, *PackageName, FileName);
+			PackageData = new FPackageData(*this, PackageName, NormalizedFileName);
 			NewPackageDataObjects[Index] = PackageData;
 		}
 		PackageData->SetPlatformCooked(TargetPlatform, true /* Succeeded */);
 	});
 
-	// update cache for all newly created objects (taken from CreatePackageData)
+	// update cache for all newly created objects (copied from CreatePackageData)
 	for (FPackageData* PackageData : NewPackageDataObjects)
 	{
 		if (PackageData)
 		{
+			FPackageData* ExistingByFileName = FileNameToPackageData.Add(PackageData->FileName, PackageData);
+			// We looked up by FileName in the loop; it should still have been unset before the write we just did
+			check(ExistingByFileName == PackageData);
+			FPackageData* ExistingByPackageName = PackageNameToPackageData.FindOrAdd(PackageData->PackageName, PackageData);
+			// If no other CreatePackageData added the FileName, then they should not have added
+			// the PackageName either
+			check(ExistingByPackageName == PackageData);
 			PackageDatas.Add(PackageData);
-			PackageNameToPackageData.Add(PackageData->PackageName, PackageData);
-			FileNameToPackageData.Add(PackageData->FileName, PackageData);
 		}
 	}
 }
 
-FPackageData* FPackageDatas::UpdateFileName(const FName& PackageName)
+FPackageData* FPackageDatas::UpdateFileName(FName PackageName)
 {
-	if (!PackageNameCache.HasCacheForPackageName(PackageName))
-	{
-		return nullptr;
-	}
-
-	FName OldFileName = PackageNameCache.GetCachedStandardFileName(PackageName);
-	PackageNameCache.ClearPackageFileNameCacheForPackage(PackageName);
-	FName NewFileName = PackageNameCache.GetCachedStandardFileName(PackageName);
+	FWriteScopeLock ExistenceWriteLock(ExistenceLock);
 
 	FPackageData** PackageDataAddr = PackageNameToPackageData.Find(PackageName);
 	if (!PackageDataAddr)
 	{
-		check(OldFileName.IsNone() || !FileNameToPackageData.Find(OldFileName));
+		FName NewFileName = LookupFileNameOnDisk(PackageName);
+		check(NewFileName.IsNone() || !FileNameToPackageData.Find(NewFileName));
 		return nullptr;
 	}
 	FPackageData* PackageData = *PackageDataAddr;
-
+	FName OldFileName = PackageData->GetFileName();
+	bool bIsMap = FPackageName::IsMapPackageExtension(*FPaths::GetExtension(OldFileName.ToString()));
+	FName NewFileName = LookupFileNameOnDisk(PackageName, false /* bRequireExists */, bIsMap);
 	if (OldFileName == NewFileName)
 	{
 		return PackageData;
 	}
-
-	if (!OldFileName.IsNone())
+	if (NewFileName.IsNone())
 	{
-		PackageDataAddr = FileNameToPackageData.Find(OldFileName);
-		check(PackageDataAddr == nullptr || *PackageDataAddr == PackageData);
-		FileNameToPackageData.Remove(OldFileName);
-	}
-	PackageData->SetFileName(NewFileName);
-	if (!NewFileName.IsNone())
-	{
-		check(FileNameToPackageData.Find(NewFileName) == nullptr);
-		FileNameToPackageData.Add(NewFileName, PackageData);
+		UE_LOG(LogCook, Error, TEXT("Cannot update FileName for package %s because the package is no longer mounted."),
+			*PackageName.ToString())
+			return PackageData;
 	}
 
-	return PackageData;
-}
+	check(!OldFileName.IsNone());
+	FPackageData* ExistingByFileName;
+	ensure(FileNameToPackageData.RemoveAndCopyValue(OldFileName, ExistingByFileName));
+	check(ExistingByFileName == PackageData);
 
-void FPackageDatas::RegisterFileNameAlias(FPackageData& PackageData, FName FileName)
-{
-	FileName = FPackageNameCache::GetStandardFileName(FileName);
-	if (FileName.IsNone())
-	{
-		return;
-	}
+	PackageData->SetFileName(NewFileName);
+	FPackageData* AddedByFileName = FileNameToPackageData.FindOrAdd(NewFileName, PackageData);
+	check(AddedByFileName == PackageData);
 
-	FPackageData*& PackageDataMapAddr = FileNameToPackageData.FindOrAdd(FileName);
-	check(PackageDataMapAddr == nullptr || PackageDataMapAddr == &PackageData);
-	PackageDataMapAddr = &PackageData;
+	return PackageData;
 }
 
 int32 FPackageDatas::GetNumCooked()
@@ -1852,22 +2012,23 @@ int32 FPackageDatas::GetNumCooked()
 	return Monitor.GetNumCooked();
 }
 
-void FPackageDatas::GetCookedFileNamesForPlatform(const ITargetPlatform* Platform, TArray<FName>& CookedFiles,
+void FPackageDatas::GetCookedPackagesForPlatform(const ITargetPlatform* Platform, TArray<FPackageData*>& CookedPackages,
 	bool bGetFailedCookedPackages, bool bGetSuccessfulCookedPackages)
 {
-	for (const FPackageData* PackageData : PackageDatas)
+	for (FPackageData* PackageData : PackageDatas)
 	{
 		ECookResult CookResults = PackageData->GetCookResults(Platform);
 		if (((CookResults == ECookResult::Succeeded) & (bGetSuccessfulCookedPackages != 0)) |
 			((CookResults == ECookResult::Failed) & (bGetFailedCookedPackages != 0)))
 		{
-			CookedFiles.Add(PackageData->GetFileName());
+			CookedPackages.Add(PackageData);
 		}
 	}
 }
 
 void FPackageDatas::Clear()
 {
+	FWriteScopeLock ExistenceWriteLock(ExistenceLock);
 	PendingCookedPlatformDatas.Empty(); // These destructors will dereference PackageDatas
 	RequestQueue.Empty();
 	SaveQueue.Empty();
diff --git a/Engine/Source/Editor/UnrealEd/Private/Cooker/CookPackageData.h b/Engine/Source/Editor/UnrealEd/Private/Cooker/CookPackageData.h
--- a/Engine/Source/Editor/UnrealEd/Private/Cooker/CookPackageData.h
+++ b/Engine/Source/Editor/UnrealEd/Private/Cooker/CookPackageData.h
@@ -13,9 +13,9 @@
 #include "Containers/StringFwd.h"
 #include "Containers/UnrealString.h"
 #include "Engine/ICookInfo.h"
+#include "HAL/CriticalSection.h"
 #include "HAL/Platform.h"
 #include "Misc/EnumClassFlags.h"
-#include "PackageNameCache.h"
 #include "Templates/SharedPointer.h"
 #include "UObject/GCObject.h"
 #include "UObject/NameTypes.h"
@@ -588,8 +588,7 @@ public:
 	void ClearGeneratedPackages();
 
 	/** Call the Splitter's GetGenerateList and create the PackageDatas */
-	bool TryGenerateList(UObject* OwnerObject, const FPackageNameCache& PackageNameCache,
-		FPackageDatas& PackageDatas);
+	bool TryGenerateList(UObject* OwnerObject, FPackageDatas& PackageDatas);
 
 	/** Accessor for the packages to generate */
 	TArrayView<FGeneratedStruct> GetPackagesToGenerate() { return PackagesToGenerate; }
@@ -839,6 +838,13 @@ public:
 	FPackageDataQueue EntryQueue;
 };
 
+/** Data necessary to create an FPackageData without checking the disk, for e.g. AddExistingPackageDatasForPlatform */
+struct FConstructPackageData
+{
+	FName PackageName;
+	FName NormalizedFileName;
+};
+
 /*
  * Class that manages the list of all PackageDatas for a CookOnTheFlyServer. PackageDatas is an associative
  * array for extra data about a package (e.g. the requested platforms) that is needed by the CookOnTheFlyServer.
@@ -852,6 +858,8 @@ struct FPackageDatas : public FGCObject
 public:
 	FPackageDatas(UCookOnTheFlyServer& InCookOnTheFlyServer);
 	~FPackageDatas();
+	/** Called when the initial AssetRegistry search is done and it can be used to determine package existence. */
+	static void OnAssetRegistryGenerated(IAssetRegistry& InAssetRegistry);
 
 	/** Called each time BeginCook is called, for deferred intialization steps. */
 	void BeginCook();
@@ -864,8 +872,6 @@ public:
 	 */
 	virtual void AddReferencedObjects(FReferenceCollector& Collector) override;
 
-	/** Return the PackageNameCache that is caching FileNames on disk for the CookOnTheFlyServer. */
-	const FPackageNameCache& GetPackageNameCache() const;
 	/** Return the Monitor used to report aggregated information about FPackageDatas. */
 	FPackageDataMonitor& GetMonitor();
 	/** Return the backpointer to the CookOnTheFlyServer */
@@ -901,39 +907,102 @@ public:
 	FPackageData* FindPackageDataByPackageName(const FName& PackageName);
 	/**
 	 * Return a pointer to the PackageData for the given PackageName. If one does not already exist,
-	 * find its FileName on disk and create the PackageData. If no filename exists, return nullptr.
+	 * find its FileName on disk and create the PackageData.  Will fail if the path is not mounted,
+	 * or if the file does not exist and bRequireExists is true. If it fails, returns nullptr.
+	 * 
+	 * @param bRequireExists If true, returns nullptr if the PackageData does not already exist and the
+	 *                       package does not exist on disk in the Workspace Domain.
+	 *                       If false, creates the PackageData so long as the package path is mounted.
+	 * @param bCreateAsMap Only used if the PackageData does not already exist, the package does not exist
+	 *                     on disk and bRequireExists is false. If true the extension is set to .umap,
+	 *                     if false, it is set to .uasset.
 	 */
-	FPackageData* TryAddPackageDataByPackageName(const FName& PackageName);
+	FPackageData* TryAddPackageDataByPackageName(const FName& PackageName, bool bRequireExists = true,
+		bool bCreateAsMap = false);
 	/**
 	 * Return a reference to the PackageData for the given PackageName. If one does not already exist,
-	 * find its FileName on disk and create the PackageData. Asserts if FileName does not exist;
-	 * caller is claiming it does.
+	 * find its FileName on disk and create the PackageData. Will fail if the path is not mounted,
+	 * or if the file does not exist and bRequireExists is true. If it fails, this function will assert.
+	 *
+	 * @param bRequireExists If true, asserts if the PackageData does not already exist and the
+	 *                       package does not exist on disk in the Workspace Domain.
+	 *                       If false, creates the PackageData so long as the package path is mounted.
+	 * @param bCreateAsMap Only used if the PackageData does not already exist, the package does not exist
+	 *                     on disk and bRequireExists is false. If true the extension is set to .umap,
+	 *                     if false, it is set to .uasset.
 	 */
-	FPackageData& AddPackageDataByPackageNameChecked(const FName& PackageName);
+	FPackageData& AddPackageDataByPackageNameChecked(const FName& PackageName, bool bRequireExists = true,
+		bool bCreateAsMap = false);
 
 	/**
-	 * Return the PackageData with the given FileName (or that has been marked as having the given FileName as an
-	 * alias) if one exists, otherwise return nullptr.
+	 * Return the PackageData with the given FileName if one exists, otherwise return nullptr.
 	 */
 	FPackageData* FindPackageDataByFileName(const FName& InFileName);
 	/**
-	 * Return a pointer to the PackageData for the given FileName (whether the actual name or an alias).
+	 * Return a pointer to the PackageData for the given FileName.
 	 * If one does not already exist, verify the FileName on disk and create the PackageData.
 	 * If no filename exists for the package on disk, return nullptr.
 	 */
 	FPackageData* TryAddPackageDataByFileName(const FName& InFileName);
 	/**
-	 * Return a reference to the PackageData for the given FileName (whether the actual name or an alias).
+	 * Return a pointer to the PackageData for the given FileName.
+	 * If one does not already exist, verify the FileName on disk and create the PackageData.
+	 * If no filename exists for the package on disk, return nullptr.
+	 * FileName must have been normalized by FPackageDatas::GetStandardFileName.
+	 */
+	FPackageData* TryAddPackageDataByStandardFileName(const FName& InFileName);
+	/**
+	 * Return a reference to the PackageData for the given FileName.
 	 * If one does not already exist, verify the FileName on disk and create the PackageData. Asserts if FileName
 	 * does not exist; caller is claiming it does.
 	 */
 	FPackageData& AddPackageDataByFileNameChecked(const FName& FileName);
 
 	/**
-	 * Optimized version of TryAddPackageDataByFileName for use by DLC when loading up the shipped packages.
-	 * This expects the PackageNameCache to have already been updated (for instance with AppendCacheResults)
+	 * Return the local path for the given packagename.
+	 * 
+	 * @param PackageName The LongPackageName of the Package
+	 * @param bRequireExists If true, fails if the package does not exist on disk in the Workspace Domain.
+	 *                       If false, returns the filename it would have so long as the package path is mounted.
+	 * @param bCreateAsMap Only used if bRequireExists is false and the path does not exist. If true
+	 *                     the extension is set to .umap, if false, it is set to .uasset.
+	 * @return Local WorkspaceDomain path in FPaths::MakeStandardFilename form, or NAME_None if it does not exist.
+	 */
+	FName GetFileNameByPackageName(FName PackageName, bool bRequireExists = true, bool bCreateAsMap = false);
+
+	/**
+	 * Return the local path for the given LongPackageName or unnormalized localpath.
+	 *
+	 * @param PackageName The name or local path for the package.
+	 * @param bRequireExists If true, fails if the package does not exist on disk in the Workspace Domain.
+	 *                       If false, returns the filename it would have so long as the package path is mounted.
+	 * @param bCreateAsMap Only used if bRequireExists is false and the path does not exist. If true
+	 *                     the extension is set to .umap, if false, it is set to .uasset.
+	 * @return Local WorkspaceDomain path in FPaths::MakeStandardFilename form, or NAME_None if it does not exist.
+	 */
+	FName GetFileNameByFlexName(FName PackageOrFileName, bool bRequireExists = true, bool bCreateAsMap = false);
+
+	/**
+	 * Uncached; reads the AssetRegistry and disk to find the filename for the given PackageName.
+	 * This is the same function that the other caching lookup functions use internally.
+	 * It can be called from multiple threads on a batch of files to accelerate the lookup, and
+	 * the results passed to FindOrAddPackageData(PackageName, FileName) which will then skip the lookup.
+	 * 
+	 * @param bRequireExists If true, fails if the package does not exist on disk in the Workspace Domain.
+	 *                       If false, returns the filename it would have so long as the package path is mounted.
+	 * @param bCreateAsMap Only used if bRequireExists is false and the path does not exist. If true
+	 *                     the extension is set to .umap, if false, it is set to .uasset.
+	 * @return The FPaths::MakeStandardFilename format of the localpath for the packagename, or NAME_None if
+	 *         it does not exist.
 	 */
-	void AddExistingPackageDatasForPlatform(const TArray<FName>& ExistingPackages,
+	static FName LookupFileNameOnDisk(FName PackageName, bool bRequireExists = true, bool bCreateAsMap = false);
+	/** Normalize the given FileName for use in looking up the cached data associated with the FileName. */
+	static FName GetStandardFileName(FName FileName);
+	/** Normalize the given FileName for use in looking up the cached data associated with the FileName. */
+	static FName GetStandardFileName(FStringView FileName);
+
+	/** Create and mark-cooked a batch of PackageDatas, used by DLC for cooked-in-earlier-release packages. */
+	void AddExistingPackageDatasForPlatform(TConstArrayView<FConstructPackageData> ExistingPackages,
 		const ITargetPlatform* TargetPlatform);
 
 	/**
@@ -943,21 +1012,15 @@ public:
 	 * and got redirected to another FileName.
 	 * Returns the PackageData if it exists.
 	 */
-	FPackageData* UpdateFileName(const FName& PackageName);
-	/**
-	 * Mark that the give PackageData should also be returned for searches for the given FileName in addition to
-	 * searches for its own FileName.
-	 * This is used for e.g. mapping AssetName.umap to AssetName.uasset.
-	 */
-	void RegisterFileNameAlias(FPackageData& PackageData, FName FileName);
+	FPackageData* UpdateFileName(FName PackageName);
 
 	/** Report the number of packages that have cooked any platform. Used by cook commandlet progress reporting. */
 	int32 GetNumCooked();
 	/**
-	 * Add onto CookedFiles the list of all packages that have cooked any package, either successfully if
+	 * Append to CookedPackages all packages that have cooked any platform, either successfully if
 	 * bGetSuccessfulCookedPackages is true and/or unsuccessfully if bGetFailedCookedPackages is true.
 	 */
-	void GetCookedFileNamesForPlatform(const ITargetPlatform* Platform, TArray<FName>& CookedFiles,
+	void GetCookedPackagesForPlatform(const ITargetPlatform* Platform, TArray<FPackageData*>& CookedPackages,
 		bool bGetFailedCookedPackages, bool bGetSuccessfulCookedPackages);
 
 	/**
@@ -997,11 +1060,14 @@ private:
 	 * New FPackageData are always created in the Idle state.
 	 */
 	FPackageData& CreatePackageData(FName PackageName, FName FileName);
+	/** Return whether a filename exists for the package on disk, and return it, unnormalized. */
+	static bool TryLookupFileNameOnDisk(FName PackageName, FString& OutFileName);
+	/** Return the corresponding PackageName if the normalized filename exists on disk. */
+	static FName LookupPackageNameOnDisk(FName NormalizedFileName);
 
 	/** Collection of pointers to all FPackageData that have been constructed by *this. */
 	TArray<FPackageData*> PackageDatas;
 	FPackageDataMonitor Monitor;
-	FPackageNameCache PackageNameCache;
 	TMap<FName, FPackageData*> PackageNameToPackageData;
 	TMap<FName, FPackageData*> FileNameToPackageData;
 	TArray<FPendingCookedPlatformData> PendingCookedPlatformDatas;
@@ -1010,9 +1076,12 @@ private:
 	FPackageDataQueue LoadReadyQueue;
 	FPackageDataQueue SaveQueue;
 	UCookOnTheFlyServer& CookOnTheFlyServer;
+	FRWLock ExistenceLock;
 	FPackageData* ShowInstigatorPackageData = nullptr;
 	double LastPollAsyncTime;
 	bool bLogDiscoveredPackages = false;
+
+	static IAssetRegistry* AssetRegistry;
 };
 
 /**
diff --git a/Engine/Source/Editor/UnrealEd/Private/Cooker/CookRequestCluster.cpp b/Engine/Source/Editor/UnrealEd/Private/Cooker/CookRequestCluster.cpp
--- a/Engine/Source/Editor/UnrealEd/Private/Cooker/CookRequestCluster.cpp
+++ b/Engine/Source/Editor/UnrealEd/Private/Cooker/CookRequestCluster.cpp
@@ -12,7 +12,6 @@
 #include "Cooker/CookProfiling.h"
 #include "Cooker/CookRequests.h"
 #include "Cooker/CookTypes.h"
-#include "Cooker/PackageNameCache.h"
 #include "Cooker/PackageTracker.h"
 #include "CookOnTheSide/CookOnTheFlyServer.h"
 #include "DerivedDataCache.h"
@@ -32,7 +31,6 @@ FRequestCluster::FRequestCluster(UCookOnTheFlyServer& COTFS, TArray<const ITarge
 	: Platforms(MoveTemp(InPlatforms))
 	, PackageDatas(*COTFS.PackageDatas)
 	, AssetRegistry(*IAssetRegistry::Get())
-	, PackageNameCache(COTFS.GetPackageNameCache())
 	, PackageTracker(*COTFS.PackageTracker)
 	, BuildDefinitions(*COTFS.BuildDefinitions)
 {
@@ -42,7 +40,6 @@ FRequestCluster::FRequestCluster(UCookOnTheFlyServer& COTFS, TConstArrayView<con
 	: Platforms(InPlatforms)
 	, PackageDatas(*COTFS.PackageDatas)
 	, AssetRegistry(*IAssetRegistry::Get())
-	, PackageNameCache(COTFS.GetPackageNameCache())
 	, PackageTracker(*COTFS.PackageTracker)
 	, BuildDefinitions(*COTFS.BuildDefinitions)
 {
@@ -156,10 +153,8 @@ void FRequestCluster::AddClusters(UCookOnTheFlyServer& COTFS, FPackageDataSet& U
 		// they are hard dependencies, avoid the work of creating a cluster just for them,
 		// by checking non-cookable and sending the package to idle instead of adding to
 		// a cluster
-		FName FileName = PackageData->GetFileName();
 		PackageData->GetRequestedPlatforms(RequestedPlatforms);
-		if (!IsRequestCookable(PackageData->GetPackageName(), FileName, PackageData,
-			COTFS.GetPackageNameCache(), *COTFS.PackageDatas, *COTFS.PackageTracker,
+		if (!IsRequestCookable(PackageData->GetPackageName(), PackageData, *COTFS.PackageDatas, *COTFS.PackageTracker,
 			DLCPath, bErrorOnEngineContentUse, RequestedPlatforms))
 		{
 			PackageData->SendToState(EPackageState::Idle, ESendFlags::QueueAdd);
@@ -298,9 +293,8 @@ void FRequestCluster::FetchPackageNames(const FCookerTimer& CookerTimer, bool& b
 #if DEBUG_COOKONTHEFLY
 		UE_LOG(LogCook, Display, TEXT("Processing request for package %s"), *OriginalName.ToString());
 #endif
-		const FName* PackageName = PackageNameCache.GetCachedPackageNameFromStandardFileName(OriginalName,
-			/* bExactMatchRequired */ false, &Request.FileName);
-		if (!PackageName)
+		FPackageData* PackageData = PackageDatas.TryAddPackageDataByStandardFileName(OriginalName);
+		if (!PackageData)
 		{
 			LogCookerMessage(FString::Printf(TEXT("Could not find package at file %s!"),
 				*OriginalName.ToString()), EMessageSeverity::Error);
@@ -313,14 +307,13 @@ void FRequestCluster::FetchPackageNames(const FCookerTimer& CookerTimer, bool& b
 		}
 		else
 		{
-			FPackageData& PackageData = PackageDatas.FindOrAddPackageData(*PackageName, Request.FileName);
-			if (TryTakeOwnership(PackageData, Request.bUrgent, MoveTemp(Request.CompletionCallback), Request.Instigator))
+			if (TryTakeOwnership(*PackageData, Request.bUrgent, MoveTemp(Request.CompletionCallback), Request.Instigator))
 			{
 				bool bAlreadyExists;
-				OwnedPackageDatas.Add(&PackageData, &bAlreadyExists);
+				OwnedPackageDatas.Add(PackageData, &bAlreadyExists);
 				if (!bAlreadyExists)
 				{
-					Requests.Add(&PackageData);
+					Requests.Add(PackageData);
 				}
 			}
 		}
@@ -846,11 +839,10 @@ FRequestCluster::FVertexData::FVertexData(EAsyncType, int32 NumPlatforms)
 
 FRequestCluster::FVertexData::FVertexData(ESkipDependenciesType, FPackageData& InPackageData, FRequestCluster& Cluster)
 {
-	FName FileName = InPackageData.GetFileName();
 	PackageName = InPackageData.GetPackageName();
 	PackageData = &InPackageData;
 	bInitialRequest = true;
-	bCookable = Cluster.IsRequestCookable(PackageData->GetPackageName(), FileName, PackageData);
+	bCookable = Cluster.IsRequestCookable(PackageData->GetPackageName(), PackageData);
 	bExploreDependencies = false;
 }
 
@@ -897,12 +889,7 @@ FPackageData* FRequestCluster::FGraphSearch::FindOrAddVertex(FName PackageName,
 	}
 	VisitStatus.bVisited = true;
 
-	FName FileName = NAME_None;
-	if (PackageData)
-	{
-		FileName = PackageData->GetFileName();
-	}
-	bool bCookable = Cluster.IsRequestCookable(PackageName, FileName, PackageData);
+	bool bCookable = Cluster.IsRequestCookable(PackageName, PackageData);
 	if (!bInitialRequest)
 	{
 		if (!bCookable)
@@ -1126,14 +1113,14 @@ TMap<FPackageData*, TArray<FPackageData*>>& FRequestCluster::FGraphSearch::GetGr
 	return GraphEdges;
 }
 
-bool FRequestCluster::IsRequestCookable(FName PackageName, FName& InOutFileName, FPackageData*& PackageData)
+bool FRequestCluster::IsRequestCookable(FName PackageName, FPackageData*& InOutPackageData)
 {
-	return IsRequestCookable(PackageName, InOutFileName, PackageData, PackageNameCache, PackageDatas, PackageTracker,
+	return IsRequestCookable(PackageName, InOutPackageData, PackageDatas, PackageTracker,
 		DLCPath, bErrorOnEngineContentUse, GetPlatforms());
 }
 
-bool FRequestCluster::IsRequestCookable(FName PackageName, FName& InOutFileName, FPackageData*& PackageData,
-	const FPackageNameCache& InPackageNameCache, FPackageDatas& InPackageDatas, FPackageTracker& InPackageTracker,
+bool FRequestCluster::IsRequestCookable(FName PackageName, FPackageData*& InOutPackageData,
+	FPackageDatas& InPackageDatas, FPackageTracker& InPackageTracker,
 	FStringView InDLCPath, bool bInErrorOnEngineContentUse, TConstArrayView<const ITargetPlatform*> RequestPlatforms)
 {
 	TStringBuilder<256> NameBuffer;
@@ -1145,33 +1132,30 @@ bool FRequestCluster::IsRequestCookable(FName PackageName, FName& InOutFileName,
 		return false;
 	}
 
-	if (InOutFileName.IsNone())
+	if (!InOutPackageData)
 	{
-		InOutFileName = InPackageNameCache.GetCachedStandardFileName(PackageName);
-		if (InOutFileName.IsNone())
+		InOutPackageData = InPackageDatas.TryAddPackageDataByPackageName(PackageName);
+		if (!InOutPackageData)
 		{
 			// Package does not exist on disk
 			return false;
 		}
 	}
 
-	if (InPackageTracker.NeverCookPackageList.Contains(InOutFileName))
+	FName FileName = InOutPackageData->GetFileName();
+	if (InPackageTracker.NeverCookPackageList.Contains(FileName))
 	{
 		UE_LOG(LogCook, Verbose, TEXT("Package %s is referenced but is in the never cook package list, discarding request"), *NameBuffer);
 		return false;
 	}
 
-	if (!PackageData)
-	{
-		PackageData = &InPackageDatas.FindOrAddPackageData(PackageName, InOutFileName);
-	}
 
 	if (bInErrorOnEngineContentUse && !InDLCPath.IsEmpty())
 	{
-		InOutFileName.ToString(NameBuffer);
+		FileName.ToString(NameBuffer);
 		if (!FStringView(NameBuffer).StartsWith(InDLCPath))
 		{
-			if (!PackageData->HasAllCookedPlatforms(RequestPlatforms, true /* bIncludeFailed */))
+			if (!InOutPackageData->HasAllCookedPlatforms(RequestPlatforms, true /* bIncludeFailed */))
 			{
 				UE_LOG(LogCook, Error, TEXT("Uncooked Engine or Game content %s is being referenced by DLC!"), *NameBuffer);
 			}
@@ -1182,4 +1166,4 @@ bool FRequestCluster::IsRequestCookable(FName PackageName, FName& InOutFileName,
 	return true;
 }
 
-} // namespace UE::Cook
+} // namespace UE::Cook
\ No newline at end of file
diff --git a/Engine/Source/Editor/UnrealEd/Private/Cooker/CookRequestCluster.h b/Engine/Source/Editor/UnrealEd/Private/Cooker/CookRequestCluster.h
--- a/Engine/Source/Editor/UnrealEd/Private/Cooker/CookRequestCluster.h
+++ b/Engine/Source/Editor/UnrealEd/Private/Cooker/CookRequestCluster.h
@@ -18,9 +18,9 @@ class IAssetRegistry;
 class ITargetPlatform;
 class FEvent;
 class UCookOnTheFlyServer;
-struct FPackageNameCache;
 namespace UE::Cook { class FRequestQueue; }
 namespace UE::Cook { struct FFilePlatformRequest; }
+namespace UE::Cook { struct FPackageDatas; }
 namespace UE::Cook { struct FPackageTracker; }
 
 namespace UE::Cook
@@ -295,9 +295,9 @@ private:
 	void StartAsync(const FCookerTimer& CookerTimer, bool& bOutComplete);
 	bool TryTakeOwnership(FPackageData& PackageData, bool bUrgent, FCompletionCallback&& CompletionCallback,
 		const FInstigator& InInstigator);
-	bool IsRequestCookable(FName PackageName, FName& InOutFileName, FPackageData*& PackageData);
-	static bool IsRequestCookable(FName PackageName, FName& InOutFileName, FPackageData*& PackageData,
-		const FPackageNameCache& InPackageNameCache, FPackageDatas& InPackageDatas, FPackageTracker& InPackageTracker,
+	bool IsRequestCookable(FName PackageName, FPackageData*& InOutPackageData);
+	static bool IsRequestCookable(FName PackageName, FPackageData*& InOutPackageData,
+		FPackageDatas& InPackageDatas, FPackageTracker& InPackageTracker,
 		FStringView InDLCPath, bool bInErrorOnEngineContentUse, TConstArrayView<const ITargetPlatform*> RequestPlatforms);
 
 	TArray<FFileNameRequest> InRequests;
@@ -310,7 +310,6 @@ private:
 	TUniquePtr<FGraphSearch> GraphSearch; // Needs to be dynamic-allocated because of large alignment
 	FPackageDatas& PackageDatas;
 	IAssetRegistry& AssetRegistry;
-	const FPackageNameCache& PackageNameCache;
 	FPackageTracker& PackageTracker;
 	FBuildDefinitions& BuildDefinitions;
 	int32 NextRequest = 0;
diff --git a/Engine/Source/Editor/UnrealEd/Private/Cooker/CookRequests.cpp b/Engine/Source/Editor/UnrealEd/Private/Cooker/CookRequests.cpp
--- a/Engine/Source/Editor/UnrealEd/Private/Cooker/CookRequests.cpp
+++ b/Engine/Source/Editor/UnrealEd/Private/Cooker/CookRequests.cpp
@@ -3,12 +3,12 @@
 #include "CookRequests.h"
 
 #include "Algo/Find.h"
+#include "CookPackageData.h"
 #include "CookPlatformManager.h"
 #include "HAL/Event.h"
 #include "Interfaces/ITargetPlatform.h"
 #include "Misc/CoreMiscDefines.h"
 #include "Misc/ScopeLock.h"
-#include "PackageNameCache.h"
 #include "Templates/UnrealTemplate.h"
 
 namespace UE
@@ -76,7 +76,7 @@ namespace Cook
 
 	void FFilePlatformRequest::SetFilename(FString InFilename)
 	{
-		Filename = FPackageNameCache::GetStandardFileName(InFilename);
+		Filename = FPackageDatas::GetStandardFileName(InFilename);
 	}
 
 	const FName& FFilePlatformRequest::GetFilename() const
@@ -326,4 +326,4 @@ namespace Cook
 	}
 
 }
-}
+}
\ No newline at end of file
diff --git a/Engine/Source/Editor/UnrealEd/Private/Cooker/LooseCookedPackageWriter.cpp b/Engine/Source/Editor/UnrealEd/Private/Cooker/LooseCookedPackageWriter.cpp
--- a/Engine/Source/Editor/UnrealEd/Private/Cooker/LooseCookedPackageWriter.cpp
+++ b/Engine/Source/Editor/UnrealEd/Private/Cooker/LooseCookedPackageWriter.cpp
@@ -6,8 +6,8 @@
 #include "Async/Async.h"
 #include "Async/ParallelFor.h"
 #include "Cooker/AsyncIODelete.h"
+#include "Cooker/CookPackageData.h"
 #include "Cooker/CookTypes.h"
-#include "Cooker/PackageNameCache.h"
 #include "HAL/FileManager.h"
 #include "HAL/PlatformFile.h"
 #include "HAL/PlatformFileManager.h"
@@ -28,12 +28,13 @@
 #include "UObject/Package.h"
 #include "PackageStoreOptimizer.h"
 
-FLooseCookedPackageWriter::FLooseCookedPackageWriter(const FString& InOutputPath, const FString& InMetadataDirectoryPath, const ITargetPlatform* InTargetPlatform,
-	FAsyncIODelete& InAsyncIODelete, const FPackageNameCache& InPackageNameCache, const TArray<TSharedRef<IPlugin>>& InPluginsToRemap)
+FLooseCookedPackageWriter::FLooseCookedPackageWriter(const FString& InOutputPath,
+	const FString& InMetadataDirectoryPath, const ITargetPlatform* InTargetPlatform, FAsyncIODelete& InAsyncIODelete,
+	UE::Cook::FPackageDatas& InPackageDatas, const TArray<TSharedRef<IPlugin>>& InPluginsToRemap)
 	: OutputPath(InOutputPath)
 	, MetadataDirectoryPath(InMetadataDirectoryPath)
 	, TargetPlatform(*InTargetPlatform)
-	, PackageNameCache(InPackageNameCache)
+	, PackageDatas(InPackageDatas)
 	, PackageStoreManifest(InOutputPath)
 	, PluginsToRemap(InPluginsToRemap)
 	, AsyncIODelete(InAsyncIODelete)
@@ -423,7 +424,7 @@ TUniquePtr<FAssetRegistryState> FLooseCookedPackageWriter::LoadPreviousAssetRegi
 		{
 			FName PackageName = Pair.Key;
 			const FAssetPackageData* PackageData = Pair.Value;
-			const FName UncookedFilename = PackageNameCache.GetCachedStandardFileName(PackageName);
+			const FName UncookedFilename = PackageDatas.GetFileNameByPackageName(PackageName);
 			bool bExistsOnDisk = false;
 			if (!UncookedFilename.IsNone())
 			{
@@ -467,18 +468,23 @@ void FLooseCookedPackageWriter::RemoveCookedPackages(TArrayView<const FName> Pac
 		// that means it's also conceivable that we will recook the same package which currently has an outstanding async write request
 		UPackage::WaitForAsyncFileWrites();
 
-		// PackageNameCache is read-game-thread-only, so we have to read it before calling parallel-for
-		TArray<FName> UncookedFileNamesToRemove;
-		UncookedFileNamesToRemove.Reserve(PackageNamesToRemove.Num());
-		for (FName PackageName : PackageNamesToRemove)
+		auto DeletePackageLambda = [&PackageNamesToRemove, this](int32 PackageIndex)
 		{
-			const FName UncookedFileName = PackageNameCache.GetCachedStandardFileName(PackageName);
-			if (!UncookedFileName.IsNone())
+			const FName PackageName = PackageNamesToRemove[PackageIndex];
+			const FName UncookedFileName = PackageDatas.GetFileNameByPackageName(PackageName);
+			if (UncookedFileName.IsNone())
 			{
-				UncookedFileNamesToRemove.Add(UncookedFileName);
+				return;
 			}
-		}
-		RemoveCookedPackagesByUncookedFilename(UncookedFileNamesToRemove);
+			FName* CookedFileName = UncookedPathToCookedPath.Find(UncookedFileName);
+			if (CookedFileName)
+			{
+				TStringBuilder<256> FilePath;
+				CookedFileName->ToString(FilePath);
+				IFileManager::Get().Delete(*FilePath, true, true, true);
+			}
+		};
+		ParallelFor(PackageNamesToRemove.Num(), DeletePackageLambda);
 	}
 
 	// We no longer have a use for UncookedPathToCookedPath, after the RemoveCookedPackages call at the beginning of the cook.
@@ -645,7 +651,7 @@ FName FLooseCookedPackageWriter::ConvertCookedPathToUncookedPath(
 		BuildUncookedPath(FullCookedFilename, SandboxRootDir, RelativeRootDir);
 	}
 
-	// Convert to a standard filename as required by FPackageNameCache where this path is used.
+	// Convert to a standard filename as required by PackageDatas where this path is used.
 	FPaths::MakeStandardFilename(OutUncookedPath);
 
 	return FName(*OutUncookedPath);
diff --git a/Engine/Source/Editor/UnrealEd/Private/Cooker/LooseCookedPackageWriter.h b/Engine/Source/Editor/UnrealEd/Private/Cooker/LooseCookedPackageWriter.h
--- a/Engine/Source/Editor/UnrealEd/Private/Cooker/LooseCookedPackageWriter.h
+++ b/Engine/Source/Editor/UnrealEd/Private/Cooker/LooseCookedPackageWriter.h
@@ -13,7 +13,7 @@
 class FAsyncIODelete;
 class IPlugin;
 class ITargetPlatform;
-struct FPackageNameCache;
+namespace UE::Cook { struct FPackageDatas; }
 
 /** A CookedPackageWriter that saves cooked packages in separate .uasset,.uexp,.ubulk files in the Saved\Cooked\[Platform] directory. */
 class FLooseCookedPackageWriter : public TPackageWriterToSharedBuffer<ICookedPackageWriter>
@@ -21,8 +21,9 @@ class FLooseCookedPackageWriter : public TPackageWriterToSharedBuffer<ICookedPac
 public:
 	using Super = TPackageWriterToSharedBuffer<ICookedPackageWriter>;
 
-	FLooseCookedPackageWriter(const FString& OutputPath, const FString& MetadataDirectoryPath, const ITargetPlatform* TargetPlatform,
-		FAsyncIODelete& InAsyncIODelete, const FPackageNameCache& InPackageNameCache, const TArray<TSharedRef<IPlugin>>& InPluginsToRemap);
+	FLooseCookedPackageWriter(const FString& OutputPath, const FString& MetadataDirectoryPath,
+		const ITargetPlatform* TargetPlatform, FAsyncIODelete& InAsyncIODelete,
+		UE::Cook::FPackageDatas& InPackageDatas, const TArray<TSharedRef<IPlugin>>& InPluginsToRemap);
 	~FLooseCookedPackageWriter();
 
 	virtual FCookCapabilities GetCookCapabilities() const override
@@ -120,7 +121,7 @@ private:
 	FString OutputPath;
 	FString MetadataDirectoryPath;
 	const ITargetPlatform& TargetPlatform;
-	const FPackageNameCache& PackageNameCache;
+	UE::Cook::FPackageDatas& PackageDatas;
 	FPackageStoreManifest PackageStoreManifest;
 	const TArray<TSharedRef<IPlugin>>& PluginsToRemap;
 	FAsyncIODelete& AsyncIODelete;
diff --git a/Engine/Source/Editor/UnrealEd/Private/Cooker/PackageNameCache.h b/Engine/Source/Editor/UnrealEd/Private/Cooker/PackageNameCache.h
deleted file mode 100644
diff --git a/Engine/Source/Editor/UnrealEd/Private/Cooker/PackageTracker.h b/Engine/Source/Editor/UnrealEd/Private/Cooker/PackageTracker.h
--- a/Engine/Source/Editor/UnrealEd/Private/Cooker/PackageTracker.h
+++ b/Engine/Source/Editor/UnrealEd/Private/Cooker/PackageTracker.h
@@ -201,6 +201,7 @@ namespace Cook
 		FThreadSafeUnsolicitedPackagesList UnsolicitedCookedPackages;
 		FThreadSafeQueue<FRecompileShaderRequest> RecompileRequests;
 
+		/** Packages to never cook - entries are localpaths in FPaths::MakeStandardFilename format. */
 		FThreadSafeSet<FName> NeverCookPackageList;
 		FThreadSafeSet<FName> UncookedEditorOnlyPackages; // set of packages that have been rejected due to being referenced by editor-only properties
 		TFastPointerMap<const ITargetPlatform*, TSet<FName>> PlatformSpecificNeverCookPackages;
