Non-Unity fixes

diff --git a/Engine/Plugins/Animation/IKRig/Source/IKRigEditor/Private/RigEditor/IKRigDetailCustomizations.cpp b/Engine/Plugins/Animation/IKRig/Source/IKRigEditor/Private/RigEditor/IKRigDetailCustomizations.cpp
--- a/Engine/Plugins/Animation/IKRig/Source/IKRigEditor/Private/RigEditor/IKRigDetailCustomizations.cpp
+++ b/Engine/Plugins/Animation/IKRig/Source/IKRigEditor/Private/RigEditor/IKRigDetailCustomizations.cpp
@@ -9,6 +9,7 @@
 #include "IDetailPropertyRow.h"
 #include "Widgets/SWidget.h"
 #include "Widgets/Input/SSegmentedControl.h"
+#include "Framework/MultiBox/MultiBoxBuilder.h"
 
 #define LOCTEXT_NAMESPACE "IKRigDetailCustomizations"
 
diff --git a/Engine/Source/Editor/UnrealEd/Private/Commandlets/CookGlobalShaderCommandlet.cpp b/Engine/Source/Editor/UnrealEd/Private/Commandlets/CookGlobalShaderCommandlet.cpp
--- a/Engine/Source/Editor/UnrealEd/Private/Commandlets/CookGlobalShaderCommandlet.cpp
+++ b/Engine/Source/Editor/UnrealEd/Private/Commandlets/CookGlobalShaderCommandlet.cpp
@@ -15,6 +15,7 @@
 #include "DerivedDataCacheInterface.h"
 #include "DirectoryWatcherModule.h"
 #include "IDirectoryWatcher.h"
+#include "Modules/ModuleManager.h"
 
 DEFINE_LOG_CATEGORY_STATIC(LogCookGlobalShaders, Log, All);
 
diff --git a/Engine/Source/Editor/UnrealEd/Private/Elements/Actor/ActorElementEditorViewportInteractionCustomization.cpp b/Engine/Source/Editor/UnrealEd/Private/Elements/Actor/ActorElementEditorViewportInteractionCustomization.cpp
--- a/Engine/Source/Editor/UnrealEd/Private/Elements/Actor/ActorElementEditorViewportInteractionCustomization.cpp
+++ b/Engine/Source/Editor/UnrealEd/Private/Elements/Actor/ActorElementEditorViewportInteractionCustomization.cpp
@@ -5,6 +5,7 @@
 #include "GameFramework/Actor.h"
 
 #include "Editor.h"
+#include "EditorModeManager.h"
 #include "AI/NavigationSystemBase.h"
 #include "Components/BrushComponent.h"
 #include "Engine/DocumentationActor.h"
