[World Partition Navmesh] Make the WorldPartitionNavigationDataBuilder use it's own grid instead of using the one from the runtime spatial hash

diff --git a/Engine/Source/Editor/UnrealEd/Private/WorldPartition/WorldPartitionNavigationDataBuilder.cpp b/Engine/Source/Editor/UnrealEd/Private/WorldPartition/WorldPartitionNavigationDataBuilder.cpp
--- a/Engine/Source/Editor/UnrealEd/Private/WorldPartition/WorldPartitionNavigationDataBuilder.cpp
+++ b/Engine/Source/Editor/UnrealEd/Private/WorldPartition/WorldPartitionNavigationDataBuilder.cpp
@@ -6,9 +6,11 @@
 #include "EngineUtils.h"
 #include "FileHelpers.h"
 #include "StaticMeshCompiler.h"
+#include "ActorPartition/ActorPartitionSubsystem.h"
 #include "Logging/LogMacros.h"
 #include "UObject/SavePackage.h"
 #include "Commandlets/Commandlet.h"
+#include "AI/NavigationSystemBase.h"
 
 #include "WorldPartition/WorldPartition.h"
 #include "WorldPartition/WorldPartitionSubsystem.h"
@@ -19,17 +21,21 @@ DEFINE_LOG_CATEGORY_STATIC(LogWorldPartitionNavigationDataBuilder, Log, All);
 UWorldPartitionNavigationDataBuilder::UWorldPartitionNavigationDataBuilder(const FObjectInitializer& ObjectInitializer)
 	: Super(ObjectInitializer)
 {
+}
+
+bool UWorldPartitionNavigationDataBuilder::PreRun(UWorld* World, FPackageSourceControlHelper& PackageHelper)
+{
+	const TSubclassOf<APartitionActor>& NavigationDataActorClass = ANavigationDataChunkActor::StaticClass();
+	const uint32 GridSize = NavigationDataActorClass->GetDefaultObject<APartitionActor>()->GetDefaultGridSize(World);
+	
 	// Size of loaded cell. Set as big as your hardware can afford.
 	// @todo: move to a config file.
-	IterativeCellSize = 204800;
+	IterativeCellSize = 2*GridSize;
 	
 	// Extra padding around loaded cell.
 	// @todo: set value programatically.
-	IterativeCellOverlapSize = 2000 + 51200;	// tile size + data chunk actor half size (because chunks are currently centered)
-}
+	IterativeCellOverlapSize = 2000; // navmesh tile size
 
-bool UWorldPartitionNavigationDataBuilder::PreRun(UWorld* World, FPackageSourceControlHelper& PackageHelper)
-{
 	TArray<FString> Tokens, Switches;
 	UCommandlet::ParseCommandLine(FCommandLine::Get(), Tokens, Switches);
 
@@ -62,7 +68,7 @@ bool UWorldPartitionNavigationDataBuilder::RunInternal(UWorld* World, const FCel
 
 	// Destroy any existing navigation data chunk actors within bounds we are generating, we will make new ones.
 	int32 Count = 0;
-	FBox GeneratingBounds = InCellInfo.Bounds.ExpandBy(-IterativeCellOverlapSize);
+	const FBox GeneratingBounds = InCellInfo.Bounds.ExpandBy(-IterativeCellOverlapSize);
 
 	UE_LOG(LogWorldPartitionNavigationDataBuilder, Verbose, TEXT("   GeneratingBounds %s"), *GeneratingBounds.ToString());
 	
@@ -118,7 +124,7 @@ bool UWorldPartitionNavigationDataBuilder::RunInternal(UWorld* World, const FCel
 	FStaticMeshCompilingManager::Get().FinishAllCompilation();
 	
 	// Rebuild ANavigationDataChunkActor in loaded bounds
-	WorldPartition->GenerateNavigationData(InCellInfo.Bounds);
+	GenerateNavigationData(WorldPartition, InCellInfo.Bounds, GeneratingBounds);
 
 	// Gather all packages again to include newly created ANavigationDataChunkActor actors
 	for (TActorIterator<ANavigationDataChunkActor> ItActor(World); ItActor; ++ItActor)
@@ -226,6 +232,122 @@ bool UWorldPartitionNavigationDataBuilder::RunInternal(UWorld* World, const FCel
 	return true;
 }
 
+FName GetCellName(UWorldPartition* WorldPartition, const UActorPartitionSubsystem::FCellCoord& InCellCoord)
+{
+	const FString PackageName = FPackageName::GetShortName(WorldPartition->GetPackage());
+	const FString PackageNameNoPIEPrefix = UWorld::RemovePIEPrefix(PackageName);
+	return FName(*FString::Printf(TEXT("%s_%i_%i"), *PackageNameNoPIEPrefix, InCellCoord.X, InCellCoord.Y));
+}
+
+bool UWorldPartitionNavigationDataBuilder::GenerateNavigationData(UWorldPartition* WorldPartition, const FBox& LoadedBounds, const FBox& GeneratingBounds)
+{
+	TRACE_CPUPROFILER_EVENT_SCOPE(UWorldPartitionNavigationDataBuilder::GenerateNavigationData);
+
+	static int32 CallCount = 0;
+	CallCount++;
+	UE_LOG(LogWorldPartitionNavigationDataBuilder, Log, TEXT("%i. GenerateNavigationData for LoadedBounds %s"), CallCount, *LoadedBounds.ToString());
+
+	UWorld* World = WorldPartition->World;
+
+	// Generate navmesh
+	// Make sure navigation is added and initialized in EditorMode
+	FNavigationSystem::AddNavigationSystemToWorld(*World, FNavigationSystemRunMode::EditorMode);
+
+	UNavigationSystemBase* NavSystem = World->GetNavigationSystem();
+	if (NavSystem == nullptr)
+	{
+		UE_LOG(LogWorldPartitionNavigationDataBuilder, Verbose, TEXT("No navigation system to generate navigation data."));
+		return false;
+	}
+
+	// First check if there is any intersection with the nav bounds.
+	const FBox NavBounds = NavSystem->GetNavigableWorldBounds();
+	if (!LoadedBounds.Intersect(NavBounds))
+	{
+		// No intersections, early out
+		UE_LOG(LogWorldPartitionNavigationDataBuilder, Log, TEXT("GenerateNavigationData finished (no intersection with nav bounds)."));
+		return true;
+	}
+
+	// Invoke navigation data generator
+	NavSystem->SetBuildBounds(LoadedBounds);
+	FNavigationSystem::Build(*World);
+
+	// Compute navdata bounds from tiles
+	const FBox NavDataBounds = NavSystem->ComputeNavDataBounds();
+	if (!LoadedBounds.Intersect(NavDataBounds))
+	{
+		// No intersections, early out
+		UE_LOG(LogWorldPartitionNavigationDataBuilder, Log, TEXT("GenerateNavigationData finished (no intersection with generated navigation data)."));
+		return true;
+	}
+
+	// For each cell, gather navmesh and generate a datachunk actor
+	const FBox WorldBounds = WorldPartition->GetWorldBounds();
+
+	int32 ActorCount = 0;
+
+	// Keep track of all valid navigation data chunk actors
+	TSet<ANavigationDataChunkActor*> ValidNavigationDataChunkActors;
+
+	// A DataChunkActor will be generated for each tile touching the generating bounds.
+	const TSubclassOf<APartitionActor>& NavigationDataActorClass = ANavigationDataChunkActor::StaticClass();
+	FIntRect GeneratingBounds2D(GeneratingBounds.Min.X, GeneratingBounds.Min.Y, GeneratingBounds.Max.X, GeneratingBounds.Max.Y);
+	FActorPartitionGridHelper::ForEachIntersectingCell(NavigationDataActorClass, GeneratingBounds2D, World->PersistentLevel,
+		[&WorldPartition, &ActorCount, World, &ValidNavigationDataChunkActors, &NavDataBounds, this](const UActorPartitionSubsystem::FCellCoord& InCellCoord, const FIntRect& InCellBounds)->bool
+		{
+			TRACE_CPUPROFILER_EVENT_SCOPE(MakeNavigationDataChunkActorForGridCell);
+		
+			const FBox2D CellBounds(FVector2D(InCellBounds.Min), FVector2D(InCellBounds.Max));
+			
+			if (CellBounds.GetExtent().X < 1.f || CellBounds.GetExtent().Y < 1.f)
+			{
+				// Safety, since we reduce by 1.f below.
+				UE_LOG(LogWorldPartitionNavigationDataBuilder, Warning, TEXT("%s: grid cell too small."), ANSI_TO_TCHAR(__FUNCTION__));
+				return false;
+			}
+
+			constexpr float HalfHeight = HALF_WORLD_MAX;
+			const FBox QueryBounds(FVector(CellBounds.Min.X, CellBounds.Min.Y, -HalfHeight), FVector(CellBounds.Max.X, CellBounds.Max.Y, HalfHeight));
+
+			if (!NavDataBounds.IsValid || !NavDataBounds.Intersect(QueryBounds))
+			{
+				// Skip if there is no navdata for this cell
+				return true;
+			}
+
+			//@todo_ow: Properly handle data layers
+			FActorSpawnParameters SpawnParams;
+			SpawnParams.bDeferConstruction = true;
+			SpawnParams.bCreateActorPackage = true;
+			ANavigationDataChunkActor* DataChunkActor = World->SpawnActor<ANavigationDataChunkActor>(SpawnParams);
+			ActorCount++;
+
+			const FVector2D CellCenter = CellBounds.GetCenter();
+			DataChunkActor->SetActorLocation(FVector(CellCenter.X, CellCenter.Y, 0.f));
+
+			FBox TilesBounds(EForceInit::ForceInit);
+			DataChunkActor->CollectNavData(QueryBounds, TilesBounds);
+
+			FBox ChunkActorBounds(FVector(QueryBounds.Min.X, QueryBounds.Min.Y, TilesBounds.Min.Z), FVector(QueryBounds.Max.X, QueryBounds.Max.Y, TilesBounds.Max.Z));
+			ChunkActorBounds = ChunkActorBounds.ExpandBy(FVector(-1.f, -1.f, 1.f)); //reduce XY by 1cm to avoid precision issues causing potential overflow on neighboring cell, add 1cm in Z to have a minimum of volume.
+			UE_LOG(LogWorldPartitionNavigationDataBuilder, VeryVerbose, TEXT("Setting ChunkActorBounds to %s"), *ChunkActorBounds.ToString());
+			DataChunkActor->SetDataChunkActorBounds(ChunkActorBounds);
+
+			const FName CellName = GetCellName(WorldPartition, InCellCoord);
+			DataChunkActor->SetActorLabel(FString::Printf(TEXT("NavDataChunkActor_%s"), *CellName.ToString()));
+
+			ValidNavigationDataChunkActors.Add(DataChunkActor);
+
+			UE_LOG(LogWorldPartitionNavigationDataBuilder, Verbose, TEXT("%i) %s added."), ActorCount, *DataChunkActor->GetName());
+
+			return true;
+		});
+
+	UE_LOG(LogWorldPartitionNavigationDataBuilder, Log, TEXT("GenerateNavigationData finished (%i actors added)."), ActorCount);
+	return true;
+}
+
 bool UWorldPartitionNavigationDataBuilder::SavePackages(const TArray<UPackage*>& PackagesToSave)
 {
 	// Save packages
diff --git a/Engine/Source/Editor/UnrealEd/Public/WorldPartition/WorldPartitionNavigationDataBuilder.h b/Engine/Source/Editor/UnrealEd/Public/WorldPartition/WorldPartitionNavigationDataBuilder.h
--- a/Engine/Source/Editor/UnrealEd/Public/WorldPartition/WorldPartitionNavigationDataBuilder.h
+++ b/Engine/Source/Editor/UnrealEd/Public/WorldPartition/WorldPartitionNavigationDataBuilder.h
@@ -5,6 +5,8 @@
 #include "WorldPartition/WorldPartitionBuilder.h"
 #include "WorldPartitionNavigationDataBuilder.generated.h"
 
+class UWorldPartition;
+
 UCLASS()
 class UWorldPartitionNavigationDataBuilder : public UWorldPartitionBuilder
 {
@@ -20,6 +22,8 @@ protected:
 	virtual bool RunInternal(UWorld* World, const FCellInfo& InCellInfo, FPackageSourceControlHelper& PackageHelper) override;
 	// UWorldPartitionBuilder interface end
 
+	bool GenerateNavigationData(UWorldPartition* WorldPartition, const FBox& LoadedBounds, const FBox& GeneratingBounds);
+
 	bool SavePackages(const TArray<UPackage*>& PackagesToSave);
 
 	bool bCleanBuilderPackages = false;
diff --git a/Engine/Source/Runtime/Engine/Classes/GameFramework/WorldSettings.h b/Engine/Source/Runtime/Engine/Classes/GameFramework/WorldSettings.h
--- a/Engine/Source/Runtime/Engine/Classes/GameFramework/WorldSettings.h
+++ b/Engine/Source/Runtime/Engine/Classes/GameFramework/WorldSettings.h
@@ -563,6 +563,10 @@ public:
 	UPROPERTY(EditAnywhere, Category = Foliage)
 	uint32 InstancedFoliageGridSize;
 
+	/** Size of the grid for navigation data chunk actors*/
+	UPROPERTY(EditAnywhere, Category = Navigation)
+	uint32 NavigationDataChunkGridSize;
+	
 	/** Default size of the grid for placed elements from the editor */
 	UPROPERTY(EditAnywhere, Category=Editor)
 	uint32 DefaultPlacementGridSize;
diff --git a/Engine/Source/Runtime/Engine/Private/WorldPartition/NavigationData/NavigationDataChunkActor.cpp b/Engine/Source/Runtime/Engine/Private/WorldPartition/NavigationData/NavigationDataChunkActor.cpp
--- a/Engine/Source/Runtime/Engine/Private/WorldPartition/NavigationData/NavigationDataChunkActor.cpp
+++ b/Engine/Source/Runtime/Engine/Private/WorldPartition/NavigationData/NavigationDataChunkActor.cpp
@@ -2,6 +2,7 @@
 
 #include "WorldPartition/NavigationData/NavigationDataChunkActor.h"
 #include "AI/NavigationSystemBase.h"
+#include "GameFramework/WorldSettings.h"
 
 #if WITH_EDITOR
 #include "Editor.h"
@@ -59,6 +60,11 @@ void ANavigationDataChunkActor::BeginDestroy()
 	Super::BeginDestroy();
 }
 
+uint32 ANavigationDataChunkActor::GetDefaultGridSize(UWorld* InWorld) const
+{
+	return InWorld->GetWorldSettings()->NavigationDataChunkGridSize;
+}
+
 void ANavigationDataChunkActor::AddNavigationDataChunkInEditor(const UNavigationSystemBase& NavSys)
 {
 	if (GEditor)
diff --git a/Engine/Source/Runtime/Engine/Private/WorldPartition/RuntimeSpatialHash/RuntimeSpatialHashNavigation.cpp b/Engine/Source/Runtime/Engine/Private/WorldPartition/RuntimeSpatialHash/RuntimeSpatialHashNavigation.cpp
deleted file mode 100644
diff --git a/Engine/Source/Runtime/Engine/Private/WorldPartition/WorldPartition.cpp b/Engine/Source/Runtime/Engine/Private/WorldPartition/WorldPartition.cpp
--- a/Engine/Source/Runtime/Engine/Private/WorldPartition/WorldPartition.cpp
+++ b/Engine/Source/Runtime/Engine/Private/WorldPartition/WorldPartition.cpp
@@ -1301,11 +1301,6 @@ void UWorldPartition::SavePerUserSettings()
 	}
 }
 
-void UWorldPartition::GenerateNavigationData(const FBox& LoadedBounds)
-{
-	RuntimeHash->GenerateNavigationData(LoadedBounds);
-}
-
 void UWorldPartition::DumpActorDescs(const FString& Path)
 {
 	if (FArchive* LogFile = IFileManager::Get().CreateFileWriter(*Path))
diff --git a/Engine/Source/Runtime/Engine/Private/WorldSettings.cpp b/Engine/Source/Runtime/Engine/Private/WorldSettings.cpp
--- a/Engine/Source/Runtime/Engine/Private/WorldSettings.cpp
+++ b/Engine/Source/Runtime/Engine/Private/WorldSettings.cpp
@@ -92,6 +92,7 @@ PRAGMA_ENABLE_DEPRECATION_WARNINGS
 
 #if WITH_EDITORONLY_DATA
 	InstancedFoliageGridSize = DefaultPlacementGridSize = 25600;
+	NavigationDataChunkGridSize = 102400;
 	bIncludeGridSizeInNameForFoliageActors = false;
 	bIncludeGridSizeInNameForPartitionedActors = false;
 #endif
diff --git a/Engine/Source/Runtime/Engine/Public/WorldPartition/NavigationData/NavigationDataChunkActor.h b/Engine/Source/Runtime/Engine/Public/WorldPartition/NavigationData/NavigationDataChunkActor.h
--- a/Engine/Source/Runtime/Engine/Public/WorldPartition/NavigationData/NavigationDataChunkActor.h
+++ b/Engine/Source/Runtime/Engine/Public/WorldPartition/NavigationData/NavigationDataChunkActor.h
@@ -2,7 +2,7 @@
 
 #pragma once
 
-#include "GameFramework/Actor.h"
+#include "ActorPartition/PartitionActor.h"
 #include "Containers/Array.h"
 #include "NavigationDataChunkActor.generated.h"
 
@@ -10,7 +10,7 @@ class UNavigationDataChunk;
 class UNavigationSystemBase;
 
 UCLASS(NotPlaceable)
-class ENGINE_API ANavigationDataChunkActor : public AActor
+class ENGINE_API ANavigationDataChunkActor : public APartitionActor
 {
 	GENERATED_UCLASS_BODY()
 
@@ -21,6 +21,10 @@ public:
 	virtual void BeginDestroy() override;
 	//~ End UObject Interface
 
+	//~ Begin APartitionActor Interface
+	virtual uint32 GetDefaultGridSize(UWorld* InWorld) const override;
+	//~ End APartitionActor Interface
+	
 	void AddNavigationDataChunkInEditor(const UNavigationSystemBase& NavSys);
 #endif // WITH_EDITOR
 
diff --git a/Engine/Source/Runtime/Engine/Public/WorldPartition/WorldPartition.h b/Engine/Source/Runtime/Engine/Public/WorldPartition/WorldPartition.h
--- a/Engine/Source/Runtime/Engine/Public/WorldPartition/WorldPartition.h
+++ b/Engine/Source/Runtime/Engine/Public/WorldPartition/WorldPartition.h
@@ -136,7 +136,6 @@ public:
 	FBox GetWorldBounds() const;
 	FBox GetEditorWorldBounds() const;
 	void GenerateHLOD(ISourceControlHelper* SourceControlHelper, bool bCreateActorsOnly);
-	void GenerateNavigationData(const FBox& LoadedBounds);
 
 	// Debugging Methods
 	void DrawRuntimeHashPreview();
diff --git a/Engine/Source/Runtime/Engine/Public/WorldPartition/WorldPartitionRuntimeHash.h b/Engine/Source/Runtime/Engine/Public/WorldPartition/WorldPartitionRuntimeHash.h
--- a/Engine/Source/Runtime/Engine/Public/WorldPartition/WorldPartitionRuntimeHash.h
+++ b/Engine/Source/Runtime/Engine/Public/WorldPartition/WorldPartitionRuntimeHash.h
@@ -37,7 +37,6 @@ class ENGINE_API UWorldPartitionRuntimeHash : public UObject
 	virtual bool GenerateStreaming(class UWorldPartitionStreamingPolicy* StreamingPolicy, const FActorClusterContext& ActorClusterContext, TArray<FString>* OutPackagesToGenerate) { return false; }
 	virtual void FlushStreaming() {}
 	virtual bool GenerateHLOD(ISourceControlHelper* SourceControlHelper, FActorClusterContext& ActorClusterContext, bool bCreateActorsOnly) { return false; }
-	virtual bool GenerateNavigationData(const FBox& LoadedBounds) { return false; }
 	virtual void DrawPreview() const {}
 
 	// PIE/Game methods
diff --git a/Engine/Source/Runtime/Engine/Public/WorldPartition/WorldPartitionRuntimeSpatialHash.h b/Engine/Source/Runtime/Engine/Public/WorldPartition/WorldPartitionRuntimeSpatialHash.h
--- a/Engine/Source/Runtime/Engine/Public/WorldPartition/WorldPartitionRuntimeSpatialHash.h
+++ b/Engine/Source/Runtime/Engine/Public/WorldPartition/WorldPartitionRuntimeSpatialHash.h
@@ -232,7 +232,6 @@ public:
 	virtual bool FinalizeGeneratorPackageForCook(const TArray<ICookPackageSplitter::FGeneratedPackageForPreSave>& InGeneratedPackages) override;
 	virtual void FlushStreaming() override;
 	virtual bool GenerateHLOD(ISourceControlHelper* SourceControlHelper, FActorClusterContext& ActorClusterContext, bool bCreateActorsOnly) override;
-	virtual bool GenerateNavigationData(const FBox& LoadedBounds) override;
 	virtual void DrawPreview() const override;
 
 	FName GetCellName(FName InGridName, const FIntVector& InCellGlobalCoord, const FDataLayersID& InDataLayerID) const;
