Fix condition for preview shader format in UEditorEngine::SetPreviewPlatform.

#rnx

diff --git a/Engine/Source/Editor/UnrealEd/Private/EditorEngine.cpp b/Engine/Source/Editor/UnrealEd/Private/EditorEngine.cpp
--- a/Engine/Source/Editor/UnrealEd/Private/EditorEngine.cpp
+++ b/Engine/Source/Editor/UnrealEd/Private/EditorEngine.cpp
@@ -7428,7 +7428,7 @@ void UEditorEngine::SetPreviewPlatform(const FPreviewPlatformInfo& NewPreviewPla
 		NewPreviewPlatform.bPreviewFeatureLevelActive != PreviewPlatform.bPreviewFeatureLevelActive;
 	const ERHIFeatureLevel::Type EffectiveFeatureLevel = NewPreviewPlatform.GetEffectivePreviewFeatureLevel();
 
-	if (NewPreviewPlatform.PreviewPlatformName != NAME_None)
+	if (NewPreviewPlatform.PreviewShaderFormatName != NAME_None)
 	{
 		// Force generation of the autogen files if they don't already exist. We'll ignore the returned data.
 		FGBufferInfo GBufferInfo = FShaderCompileUtilities::FetchGBufferInfoAndWriteAutogen(ShaderPlatform, MaxFeatureLevel);
