LayerSubsystem: Prevent adding Level Instance actors into layers, fix hide behavior by calling SetIsHiddenEdLayer instead of setting field directly

diff --git a/Engine/Source/Editor/UnrealEd/Private/Layers/LayersSubsystem.cpp b/Engine/Source/Editor/UnrealEd/Private/Layers/LayersSubsystem.cpp
--- a/Engine/Source/Editor/UnrealEd/Private/Layers/LayersSubsystem.cpp
+++ b/Engine/Source/Editor/UnrealEd/Private/Layers/LayersSubsystem.cpp
@@ -176,18 +176,7 @@ void ULayersSubsystem::RemoveLevelLayerInformation(ULevel* Level)
 
 bool ULayersSubsystem::IsActorValidForLayer(AActor* Actor)
 {
-	if (!Actor || Actor->GetClass() == nullptr || Actor->GetWorld() == nullptr)
-	{
-		return false;
-	}
-
-	const bool bIsBuilderBrush = FActorEditorUtils::IsABuilderBrush(Actor);
-	const bool bIsHidden = (Actor->GetClass()->GetDefaultObject<AActor>()->bHiddenEd == true);
-	const bool bIsInEditorWorld = (Actor->GetWorld()->WorldType == EWorldType::Editor);
-	const bool bIsPartitionedActor = Actor->GetLevel()->bIsPartitioned;
-	const bool bIsValid = !bIsHidden && !bIsBuilderBrush && bIsInEditorWorld && !bIsPartitionedActor;
-
-	return bIsValid;
+	return Actor && Actor->GetClass() && Actor->GetWorld() && Actor->SupportsLayers();
 }
 
 bool ULayersSubsystem::InitializeNewActorLayers(AActor* Actor)
@@ -851,15 +840,7 @@ bool ULayersSubsystem::UpdateActorVisibility(AActor* Actor, bool& bOutSelectionC
 	// If the actor doesn't belong to any layers
 	if( Actor->Layers.Num() == 0)
 	{
-		// If the actor is also hidden
-		if( Actor->bHiddenEdLayer )
-		{
-			// Actors that don't belong to any layer shouldn't be hidden
-			Actor->bHiddenEdLayer = false;
-			Actor->MarkComponentsRenderStateDirty();
-			bOutActorModified = true;
-		}
-
+		bOutActorModified = Actor->SetIsHiddenEdLayer(false);
 		return bOutActorModified;
 	}
 
@@ -875,10 +856,8 @@ bool ULayersSubsystem::UpdateActorVisibility(AActor* Actor, bool& bOutSelectionC
 
 		if( Actor->Layers.Contains( Layer->GetLayerName() ) )
 		{
-			if ( Actor->bHiddenEdLayer )
+			if (Actor->SetIsHiddenEdLayer(false))
 			{
-				Actor->bHiddenEdLayer = false;
-				Actor->MarkComponentsRenderStateDirty();
 				bOutActorModified = true;
 
 				if (ABrush* Brush = Cast<ABrush>(Actor))
@@ -897,10 +876,8 @@ bool ULayersSubsystem::UpdateActorVisibility(AActor* Actor, bool& bOutSelectionC
 	// If the actor isn't part of a visible layer, hide and de-select it.
 	if( !bActorBelongsToVisibleLayer )
 	{
-		if ( !Actor->bHiddenEdLayer )
+		if (Actor->SetIsHiddenEdLayer(true))
 		{
-			Actor->bHiddenEdLayer = true;
-			Actor->MarkComponentsRenderStateDirty();
 			bOutActorModified = true;
 
 			if (ABrush* Brush = Cast<ABrush>(Actor))
diff --git a/Engine/Source/Runtime/Engine/Classes/Engine/Brush.h b/Engine/Source/Runtime/Engine/Classes/Engine/Brush.h
--- a/Engine/Source/Runtime/Engine/Classes/Engine/Brush.h
+++ b/Engine/Source/Runtime/Engine/Classes/Engine/Brush.h
@@ -178,6 +178,7 @@ public:
 	virtual void CheckForErrors() override;
 	virtual void SetIsTemporarilyHiddenInEditor( bool bIsHidden ) override;
 	virtual bool SetIsHiddenEdLayer(bool bIsHiddenEdLayer) override;
+	virtual bool SupportsLayers() const override;
 	virtual bool SupportsExternalPackaging() const override;
 public:
 
diff --git a/Engine/Source/Runtime/Engine/Classes/GameFramework/Actor.h b/Engine/Source/Runtime/Engine/Classes/GameFramework/Actor.h
--- a/Engine/Source/Runtime/Engine/Classes/GameFramework/Actor.h
+++ b/Engine/Source/Runtime/Engine/Classes/GameFramework/Actor.h
@@ -2209,6 +2209,9 @@ public:
 	/** Changes bHiddenEdLayer flag and returns true if flag changed. */
 	virtual bool SetIsHiddenEdLayer(bool bIsHiddenEdLayer);
 
+	/** Returns true if the actor supports modifications to its Layers property */
+	virtual bool SupportsLayers() const;
+
 	/** Returns if level should keep a reference to the external actor for PIE (used for always loaded actors). */
 	bool IsForceExternalActorLevelReferenceForPIE() const
 	{
diff --git a/Engine/Source/Runtime/Engine/Private/ActorEditor.cpp b/Engine/Source/Runtime/Engine/Private/ActorEditor.cpp
--- a/Engine/Source/Runtime/Engine/Private/ActorEditor.cpp
+++ b/Engine/Source/Runtime/Engine/Private/ActorEditor.cpp
@@ -31,6 +31,7 @@
 #include "Folder.h"
 #include "Engine/LevelStreaming.h"
 #include "WorldPartition/WorldPartitionActorDesc.h"
+#include "LevelInstance/LevelInstanceSubsystem.h"
 
 #define LOCTEXT_NAMESPACE "ErrorChecking"
 
@@ -749,6 +750,28 @@ bool AActor::SetIsHiddenEdLayer(bool bIsHiddenEdLayer)
 	return false;
 }
 
+bool AActor::SupportsLayers() const
+{
+	const bool bIsHidden = (GetClass()->GetDefaultObject<AActor>()->bHiddenEd == true);
+	const bool bIsInEditorWorld = (GetWorld()->WorldType == EWorldType::Editor);
+	const bool bIsPartitionedActor = GetLevel()->bIsPartitioned;
+	const bool bIsValid = !bIsHidden && bIsInEditorWorld && !bIsPartitionedActor;
+
+	if (bIsValid)
+	{
+		// Actors part of Level Instance are not valid for layers
+		if (ULevelInstanceSubsystem* LevelInstanceSubsystem = GetWorld()->GetSubsystem<ULevelInstanceSubsystem>())
+		{
+			if (ALevelInstance* LevelInstance = LevelInstanceSubsystem->GetParentLevelInstance(this))
+			{
+				return false;
+			}
+		}
+	}
+
+	return bIsValid;
+}
+
 bool AActor::IsEditable() const
 {
 	return bEditable;
diff --git a/Engine/Source/Runtime/Engine/Private/Brush.cpp b/Engine/Source/Runtime/Engine/Private/Brush.cpp
--- a/Engine/Source/Runtime/Engine/Private/Brush.cpp
+++ b/Engine/Source/Runtime/Engine/Private/Brush.cpp
@@ -213,6 +213,11 @@ bool ABrush::SetIsHiddenEdLayer(bool bIsHiddenEdLayer)
 	return false;
 }
 
+bool ABrush::SupportsLayers() const
+{
+	return !FActorEditorUtils::IsABuilderBrush(this) && Super::SupportsLayers();
+}
+
 bool ABrush::SupportsExternalPackaging() const
 {
 	// Base class ABrush actors do not support OFPA
