CookRequestCluster: Fix TryAddPackageDataByStandardFileName not handling missing extension; handling the missing extension is needed in CookOnTheFly.
#rnx

diff --git a/Engine/Source/Editor/UnrealEd/Private/Cooker/CookPackageData.cpp b/Engine/Source/Editor/UnrealEd/Private/Cooker/CookPackageData.cpp
--- a/Engine/Source/Editor/UnrealEd/Private/Cooker/CookPackageData.cpp
+++ b/Engine/Source/Editor/UnrealEd/Private/Cooker/CookPackageData.cpp
@@ -23,6 +23,7 @@
 #include "Misc/Parse.h"
 #include "Misc/Paths.h"
 #include "Misc/PreloadableFile.h"
+#include "Misc/ScopeExit.h"
 #include "Misc/ScopeLock.h"
 #include "ShaderCompiler.h"
 #include "UObject/Object.h"
@@ -1710,8 +1711,11 @@ FPackageData* FPackageDatas::TryAddPackageDataByFileName(const FName& InFileName
 	return TryAddPackageDataByStandardFileName(GetStandardFileName(InFileName));
 }
 
-FPackageData* FPackageDatas::TryAddPackageDataByStandardFileName(const FName& FileName)
+FPackageData* FPackageDatas::TryAddPackageDataByStandardFileName(const FName& FileName, bool bExactMatchRequired,
+	FName* OutFoundFileName)
 {
+	FName FoundFileName = FileName;
+	ON_SCOPE_EXIT{ if (OutFoundFileName) { *OutFoundFileName = FoundFileName; } };
 	if (FileName.IsNone())
 	{
 		return nullptr;
@@ -1726,12 +1730,14 @@ FPackageData* FPackageDatas::TryAddPackageDataByStandardFileName(const FName& Fi
 		}
 	}
 
-	FName PackageName = LookupPackageNameOnDisk(FileName);
+	FName ExistingFileName;
+	FName PackageName = LookupPackageNameOnDisk(FileName, bExactMatchRequired, ExistingFileName);
 	if (PackageName.IsNone())
 	{
 		return nullptr;
 	}
-	return &CreatePackageData(PackageName, FileName);
+	FoundFileName = ExistingFileName;
+	return &CreatePackageData(PackageName, ExistingFileName);
 }
 
 FPackageData& FPackageDatas::CreatePackageData(FName PackageName, FName FileName)
@@ -1861,8 +1867,9 @@ bool FPackageDatas::TryLookupFileNameOnDisk(FName PackageName, FString& OutFileN
 	}
 }
 
-FName FPackageDatas::LookupPackageNameOnDisk(FName NormalizedFileName)
+FName FPackageDatas::LookupPackageNameOnDisk(FName NormalizedFileName, bool bExactMatchRequired, FName& FoundFileName)
 {
+	FoundFileName = NormalizedFileName;
 	if (NormalizedFileName.IsNone())
 	{
 		return NAME_None;
@@ -1875,13 +1882,14 @@ FName FPackageDatas::LookupPackageNameOnDisk(FName NormalizedFileName)
 	FName PackageName = FName(*Buffer);
 
 	FName DiscoveredFileName = LookupFileNameOnDisk(PackageName, true /* bRequireExists */, false /* bCreateAsMap */);
-	if (DiscoveredFileName == NormalizedFileName)
+	if (DiscoveredFileName == NormalizedFileName || !bExactMatchRequired)
 	{
+		FoundFileName = DiscoveredFileName;
 		return PackageName;
 	}
 	else
 	{
-		// Either the file does not exist on disk or the input FileName was not in normalized format
+		// Either the file does not exist on disk or NormalizedFileName did not match its format or extension
 		return NAME_None;
 	}
 }
diff --git a/Engine/Source/Editor/UnrealEd/Private/Cooker/CookPackageData.h b/Engine/Source/Editor/UnrealEd/Private/Cooker/CookPackageData.h
--- a/Engine/Source/Editor/UnrealEd/Private/Cooker/CookPackageData.h
+++ b/Engine/Source/Editor/UnrealEd/Private/Cooker/CookPackageData.h
@@ -949,8 +949,13 @@ public:
 	 * If one does not already exist, verify the FileName on disk and create the PackageData.
 	 * If no filename exists for the package on disk, return nullptr.
 	 * FileName must have been normalized by FPackageDatas::GetStandardFileName.
+	 * 
+	 * @param bExactMatchRequired If true, returns true even if the FileName is not an exact match, e.g.
+	 *                            because it is missing the extension.
+	 * @param FoundFileName If non-null, will be set to the discovered FileName, only useful if !bExactMatchRequired
 	 */
-	FPackageData* TryAddPackageDataByStandardFileName(const FName& InFileName);
+	FPackageData* TryAddPackageDataByStandardFileName(const FName& InFileName, bool bExactMatchRequired=true,
+		FName* OutFoundFileName=nullptr);
 	/**
 	 * Return a reference to the PackageData for the given FileName.
 	 * If one does not already exist, verify the FileName on disk and create the PackageData. Asserts if FileName
@@ -1063,7 +1068,7 @@ private:
 	/** Return whether a filename exists for the package on disk, and return it, unnormalized. */
 	static bool TryLookupFileNameOnDisk(FName PackageName, FString& OutFileName);
 	/** Return the corresponding PackageName if the normalized filename exists on disk. */
-	static FName LookupPackageNameOnDisk(FName NormalizedFileName);
+	static FName LookupPackageNameOnDisk(FName NormalizedFileName, bool bExactMatchRequired, FName& FoundFileName);
 
 	/** Collection of pointers to all FPackageData that have been constructed by *this. */
 	TArray<FPackageData*> PackageDatas;
diff --git a/Engine/Source/Editor/UnrealEd/Private/Cooker/CookRequestCluster.cpp b/Engine/Source/Editor/UnrealEd/Private/Cooker/CookRequestCluster.cpp
--- a/Engine/Source/Editor/UnrealEd/Private/Cooker/CookRequestCluster.cpp
+++ b/Engine/Source/Editor/UnrealEd/Private/Cooker/CookRequestCluster.cpp
@@ -293,7 +293,11 @@ void FRequestCluster::FetchPackageNames(const FCookerTimer& CookerTimer, bool& b
 #if DEBUG_COOKONTHEFLY
 		UE_LOG(LogCook, Display, TEXT("Processing request for package %s"), *OriginalName.ToString());
 #endif
-		FPackageData* PackageData = PackageDatas.TryAddPackageDataByStandardFileName(OriginalName);
+		// The input filenames are normalized, but might be missing their extension, so allow PackageDatas
+		// to correct the filename if the package is found with a different filename
+		bool bExactMatchRequired = false;
+		FPackageData* PackageData = PackageDatas.TryAddPackageDataByStandardFileName(OriginalName, bExactMatchRequired,
+			&Request.FileName);
 		if (!PackageData)
 		{
 			LogCookerMessage(FString::Printf(TEXT("Could not find package at file %s!"),
