Add InterfaceProperty support to Debugger, better logging for future unsupported property types

diff --git a/Engine/Source/Editor/Kismet/Private/Debugging/SKismetDebugTreeView.cpp b/Engine/Source/Editor/Kismet/Private/Debugging/SKismetDebugTreeView.cpp
--- a/Engine/Source/Editor/Kismet/Private/Debugging/SKismetDebugTreeView.cpp
+++ b/Engine/Source/Editor/Kismet/Private/Debugging/SKismetDebugTreeView.cpp
@@ -76,7 +76,7 @@ namespace
 			TSharedPtr<FPropertyInstanceInfo> Data = PropertyInfo.Get();
 			if (Data.IsValid())
 			{
-				if (Data->Property->IsA<FObjectProperty>())
+				if (Data->Property->IsA<FObjectProperty>() || Data->Property->IsA<FInterfaceProperty>())
 				{
 					ChildSlot
 						[
diff --git a/Engine/Source/Editor/UnrealEd/Private/Kismet2/KismetDebugUtilities.cpp b/Engine/Source/Editor/UnrealEd/Private/Kismet2/KismetDebugUtilities.cpp
--- a/Engine/Source/Editor/UnrealEd/Private/Kismet2/KismetDebugUtilities.cpp
+++ b/Engine/Source/Editor/UnrealEd/Private/Kismet2/KismetDebugUtilities.cpp
@@ -1899,13 +1899,14 @@ FPropertyInstanceInfo::FPropertyInstanceInfo(FPropertyInstance PropertyInstance)
 	Type(UEdGraphSchema_K2::TypeToText(PropertyInstance.Property)),
 	Property(PropertyInstance.Property)
 {
-	check(PropertyInstance.Property);
+	FProperty* ResolvedProperty = Property.Get();
+	check(ResolvedProperty);
 	if (PropertyInstance.Value == nullptr)
 	{
 		return;
 	}
 
-	if (FByteProperty* ByteProperty = CastField<FByteProperty>(Property.Get()))
+	if (FByteProperty* ByteProperty = CastField<FByteProperty>(ResolvedProperty))
 	{
 		UEnum* Enum = ByteProperty->GetIntPropertyEnum();
 		if (Enum)
@@ -1925,34 +1926,34 @@ FPropertyInstanceInfo::FPropertyInstanceInfo(FPropertyInstance PropertyInstance)
 		// if there is no Enum we need to fall through and treat this as a FNumericProperty
 	}
 
-	if (FNumericProperty* NumericProperty = CastField<FNumericProperty>(Property.Get()))
+	if (FNumericProperty* NumericProperty = CastField<FNumericProperty>(ResolvedProperty))
 	{
 		Value = FText::FromString(NumericProperty->GetNumericPropertyValueToString(PropertyInstance.Value));
 		return;
 	}
-	else if (FBoolProperty* BoolProperty = CastField<FBoolProperty>(Property.Get()))
+	else if (FBoolProperty* BoolProperty = CastField<FBoolProperty>(ResolvedProperty))
 	{
 		const FCoreTexts& CoreTexts = FCoreTexts::Get();
 
 		Value = BoolProperty->GetPropertyValue(PropertyInstance.Value) ? CoreTexts.True : CoreTexts.False;
 		return;
 	}
-	else if (FNameProperty* NameProperty = CastField<FNameProperty>(Property.Get()))
+	else if (FNameProperty* NameProperty = CastField<FNameProperty>(ResolvedProperty))
 	{
 		Value = FText::FromName(*(FName*)PropertyInstance.Value);
 		return;
 	}
-	else if (FTextProperty* TextProperty = CastField<FTextProperty>(Property.Get()))
+	else if (FTextProperty* TextProperty = CastField<FTextProperty>(ResolvedProperty))
 	{
 		Value = TextProperty->GetPropertyValue(PropertyInstance.Value);
 		return;
 	}
-	else if (FStrProperty* StringProperty = CastField<FStrProperty>(Property.Get()))
+	else if (FStrProperty* StringProperty = CastField<FStrProperty>(ResolvedProperty))
 	{
 		Value = FText::FromString(StringProperty->GetPropertyValue(PropertyInstance.Value));
 		return;
 	}
-	else if (FArrayProperty* ArrayProperty = CastField<FArrayProperty>(Property.Get()))
+	else if (FArrayProperty* ArrayProperty = CastField<FArrayProperty>(ResolvedProperty))
 	{
 		checkSlow(ArrayProperty->Inner);
 
@@ -1961,14 +1962,14 @@ FPropertyInstanceInfo::FPropertyInstanceInfo(FPropertyInstance PropertyInstance)
 		Value = FText::Format(LOCTEXT("ArraySize", "Num={0}"), FText::AsNumber(ArrayHelper.Num()));
 		return;
 	}
-	else if (FStructProperty* StructProperty = CastField<FStructProperty>(Property.Get()))
+	else if (FStructProperty* StructProperty = CastField<FStructProperty>(ResolvedProperty))
 	{
 		FString WatchText;
 		StructProperty->ExportTextItem(WatchText, PropertyInstance.Value, PropertyInstance.Value, nullptr, PPF_PropertyWindow | PPF_BlueprintDebugView, nullptr);
 		Value = FText::FromString(WatchText);
 		return;
 	}
-	else if (FEnumProperty* EnumProperty = CastField<FEnumProperty>(Property.Get()))
+	else if (FEnumProperty* EnumProperty = CastField<FEnumProperty>(ResolvedProperty))
 	{
 		FNumericProperty* LocalUnderlyingProp = EnumProperty->GetUnderlyingProperty();
 		UEnum* Enum = EnumProperty->GetEnum();
@@ -1995,19 +1996,19 @@ FPropertyInstanceInfo::FPropertyInstanceInfo(FPropertyInstance PropertyInstance)
 
 		return;
 	}
-	else if (FMapProperty* MapProperty = CastField<FMapProperty>(Property.Get()))
+	else if (FMapProperty* MapProperty = CastField<FMapProperty>(ResolvedProperty))
 	{
 		FScriptMapHelper MapHelper(MapProperty, PropertyInstance.Value);
 		Value = FText::Format(LOCTEXT("MapSize", "Num={0}"), FText::AsNumber(MapHelper.Num()));
 		return;
 	}
-	else if (FSetProperty* SetProperty = CastField<FSetProperty>(Property.Get()))
+	else if (FSetProperty* SetProperty = CastField<FSetProperty>(ResolvedProperty))
 	{
 		FScriptSetHelper SetHelper(SetProperty, PropertyInstance.Value);
 		Value = FText::Format(LOCTEXT("SetSize", "Num={0}"), FText::AsNumber(SetHelper.Num()));
 		return;
 	}
-	else if (FObjectPropertyBase* ObjectPropertyBase = CastField<FObjectPropertyBase>(Property.Get()))
+	else if (FObjectPropertyBase* ObjectPropertyBase = CastField<FObjectPropertyBase>(ResolvedProperty))
 	{
 		Object = ObjectPropertyBase->GetObjectPropertyValue(PropertyInstance.Value);
 		if (Object.IsValid())
@@ -2021,7 +2022,7 @@ FPropertyInstanceInfo::FPropertyInstanceInfo(FPropertyInstance PropertyInstance)
 
 		return;
 	}
-	else if (FDelegateProperty* DelegateProperty = CastField<FDelegateProperty>(Property.Get()))
+	else if (FDelegateProperty* DelegateProperty = CastField<FDelegateProperty>(ResolvedProperty))
 	{
 		if (DelegateProperty->SignatureFunction)
 		{
@@ -2034,7 +2035,7 @@ FPropertyInstanceInfo::FPropertyInstanceInfo(FPropertyInstance PropertyInstance)
 
 		return;
 	}
-	else if (FMulticastDelegateProperty* MulticastDelegateProperty = CastField<FMulticastDelegateProperty>(Property.Get()))
+	else if (FMulticastDelegateProperty* MulticastDelegateProperty = CastField<FMulticastDelegateProperty>(ResolvedProperty))
 	{
 		if (MulticastDelegateProperty->SignatureFunction)
 		{
@@ -2047,8 +2048,24 @@ FPropertyInstanceInfo::FPropertyInstanceInfo(FPropertyInstance PropertyInstance)
 
 		return;
 	}
+	else if (FInterfaceProperty* InterfaceProperty = CastField<FInterfaceProperty>(ResolvedProperty))
+	{
+		const FScriptInterface* InterfaceData = StaticCast<const FScriptInterface*>(PropertyInstance.Value);
+		Object = InterfaceData->GetObject();
+		
+		if (Object.IsValid())
+		{
+			Value = FText::FromString(Object->GetFullName());
+		}
+		else
+		{
+			Value = FText::FromString(TEXT("None"));
+		}
+
+		return;
+	}
 
-	ensureMsgf(false, TEXT("Failed to identify property type. This function may need to be exanded to include it"));
+	ensureMsgf(false, TEXT("Failed to identify property type. This function may need to be exanded to include it: %s"), *ResolvedProperty->GetClass()->GetName());
 }
 
 TSharedPtr<FPropertyInstanceInfo> FPropertyInstanceInfo::FindOrMake(FPropertyInstance PropertyInstance,
