Adds OnlyMergeSubLevels option when converting Levels and Sublevels to convert them to OFPA without WorldPartition.
Allows the converted level to be used as a LevelInstance in WP.

#changelist validated

diff --git a/Engine/Source/Editor/UnrealEd/Classes/Commandlets/WorldPartitionConvertCommandlet.h b/Engine/Source/Editor/UnrealEd/Classes/Commandlets/WorldPartitionConvertCommandlet.h
--- a/Engine/Source/Editor/UnrealEd/Classes/Commandlets/WorldPartitionConvertCommandlet.h
+++ b/Engine/Source/Editor/UnrealEd/Classes/Commandlets/WorldPartitionConvertCommandlet.h
@@ -40,15 +40,15 @@ public:
 	virtual int32 Main(const FString& Params) override;
 	//~ End UCommandlet Interface
 
-	const FString& GetConversionSuffix() const { return ConversionSuffix; }
+	static const FString GetConversionSuffix(const bool bInOnlyMergeSubLevels);
 
 private:
-	void GatherAndPrepareSubLevelsToConvert(const UWorldPartition* WorldPartition, ULevel* Level, TArray<ULevel*>& SubLevels);
+	void GatherAndPrepareSubLevelsToConvert(ULevel* Level, TArray<ULevel*>& SubLevels);
 	EActorGridPlacement GetLevelGridPlacement(ULevel* Level, EActorGridPlacement DefaultGridPlacement);
 
 protected:
 	virtual bool GetAdditionalLevelsToConvert(ULevel* Level, TArray<ULevel*>& SubLevels);
-	virtual bool PrepareStreamingLevelForConversion(const UWorldPartition* WorldPartition, ULevelStreaming* StreamingLevel);
+	virtual bool PrepareStreamingLevelForConversion(ULevelStreaming* StreamingLevel);
 	virtual bool ShouldConvertStreamingLevel(ULevelStreaming* StreamingLevel);
 	virtual bool ShouldDeleteActor(AActor* Actor, bool bMainLevel) const;
 	virtual void PerformAdditionalWorldCleanup(UWorld* World) const;
@@ -89,6 +89,7 @@ protected:
 	TArray<UPackage*> PackagesToDelete;
 	FPackageSourceControlHelper PackageHelper;
 
+	bool bOnlyMergeSubLevels;
 	bool bDeleteSourceLevels;
 	bool bGenerateIni;
 	bool bReportOnly;
diff --git a/Engine/Source/Editor/UnrealEd/Private/Commandlets/WorldPartitionConvertCommandlet.cpp b/Engine/Source/Editor/UnrealEd/Private/Commandlets/WorldPartitionConvertCommandlet.cpp
--- a/Engine/Source/Editor/UnrealEd/Private/Commandlets/WorldPartitionConvertCommandlet.cpp
+++ b/Engine/Source/Editor/UnrealEd/Private/Commandlets/WorldPartitionConvertCommandlet.cpp
@@ -232,7 +232,7 @@ UWorldPartition* UWorldPartitionConvertCommandlet::CreateWorldPartition(AWorldSe
 	return WorldPartition;
 }
 
-void UWorldPartitionConvertCommandlet::GatherAndPrepareSubLevelsToConvert(const UWorldPartition* WorldPartition, ULevel* Level, TArray<ULevel*>& SubLevels)
+void UWorldPartitionConvertCommandlet::GatherAndPrepareSubLevelsToConvert(ULevel* Level, TArray<ULevel*>& SubLevels)
 {
 	TRACE_CPUPROFILER_EVENT_SCOPE(UWorldPartitionConvertCommandlet::GatherAndPrepareSubLevelsToConvert);
 
@@ -259,7 +259,7 @@ void UWorldPartitionConvertCommandlet::GatherAndPrepareSubLevelsToConvert(const
 	
 	for(ULevelStreaming* StreamingLevel: StreamingLevels)
 	{
-		if (PrepareStreamingLevelForConversion(WorldPartition, StreamingLevel))
+		if (PrepareStreamingLevelForConversion(StreamingLevel))
 		{
 			ULevel* SubLevel = StreamingLevel->GetLoadedLevel();
 			check(SubLevel);
@@ -267,7 +267,7 @@ void UWorldPartitionConvertCommandlet::GatherAndPrepareSubLevelsToConvert(const
 			SubLevels.Add(SubLevel);
 
 			// Recursively obtain sub levels to convert
-			GatherAndPrepareSubLevelsToConvert(WorldPartition, SubLevel, SubLevels);
+			GatherAndPrepareSubLevelsToConvert(SubLevel, SubLevels);
 		}
 	}
 }
@@ -282,14 +282,14 @@ EActorGridPlacement UWorldPartitionConvertCommandlet::GetLevelGridPlacement(ULev
 	return DefaultGridPlacement;
 }
 
-bool UWorldPartitionConvertCommandlet::PrepareStreamingLevelForConversion(const UWorldPartition* WorldPartition, ULevelStreaming* StreamingLevel)
+bool UWorldPartitionConvertCommandlet::PrepareStreamingLevelForConversion(ULevelStreaming* StreamingLevel)
 {
 	TRACE_CPUPROFILER_EVENT_SCOPE(UWorldPartitionConvertCommandlet::PrepareStreamingLevelForConversion);
 
 	ULevel* SubLevel = StreamingLevel->GetLoadedLevel();
 	check(SubLevel);
 
-	if (StreamingLevel->ShouldBeAlwaysLoaded() || StreamingLevel->bDisableDistanceStreaming)
+	if (bOnlyMergeSubLevels || StreamingLevel->ShouldBeAlwaysLoaded() || StreamingLevel->bDisableDistanceStreaming)
 	{
 		FString WorldPath = SubLevel->GetPackage()->GetName();
 		if (!LevelsGridPlacement.Contains(*WorldPath))
@@ -713,11 +713,14 @@ int32 UWorldPartitionConvertCommandlet::Main(const FString& Params)
 		return 1;
 	}
 
+	bOnlyMergeSubLevels = Switches.Contains(TEXT("OnlyMergeSubLevels"));
 	bDeleteSourceLevels = Switches.Contains(TEXT("DeleteSourceLevels"));
 	bGenerateIni = Switches.Contains(TEXT("GenerateIni"));
 	bReportOnly = bGenerateIni || Switches.Contains(TEXT("ReportOnly"));
 	bVerbose = Switches.Contains(TEXT("Verbose"));
 
+	ConversionSuffix = GetConversionSuffix(bOnlyMergeSubLevels);
+
 	if (!Switches.Contains(TEXT("AllowCommandletRendering")))
 	{
 		UE_LOG(LogWorldPartitionConvertCommandlet, Error, TEXT("The option \"-AllowCommandletRendering\" is required."));
@@ -821,9 +824,10 @@ int32 UWorldPartitionConvertCommandlet::Main(const FString& Params)
 		return 1;
 	}
 
-	// Setup the world partition object
-	UWorldPartition* WorldPartition = CreateWorldPartition(MainWorldSettings);
-	if (!WorldPartition)
+	// Setup the world partition object, do not create world partition object if only merging sublevels
+	UWorldPartition* WorldPartition = bOnlyMergeSubLevels ? nullptr : CreateWorldPartition(MainWorldSettings);
+	
+	if (!bOnlyMergeSubLevels && !WorldPartition)
 	{
 		return 1;
 	}
@@ -844,7 +848,8 @@ int32 UWorldPartitionConvertCommandlet::Main(const FString& Params)
 
 	UPackage* MainPackage = MainLevel->GetPackage();
 	AWorldDataLayers* MainWorldDataLayers = MainWorld->GetWorldDataLayers();
-	check(MainWorldDataLayers);
+	// DataLayers are only needed if converting to WorldPartition
+	check(bOnlyMergeSubLevels || MainWorldDataLayers);
 
 	OnWorldLoaded(MainWorld);
 
@@ -1015,8 +1020,11 @@ int32 UWorldPartitionConvertCommandlet::Main(const FString& Params)
 							Actor->AddDataLayer(DataLayer);
 						}
 					}
-					// Clear actor layers as they are not supported yet in world partition
-					Actor->Layers.Empty();
+					// Clear actor layers as they are not supported yet in world partition, keep them if only merging
+					if (!bOnlyMergeSubLevels)
+					{
+						Actor->Layers.Empty();
+					}
 				}
 			}
 		}
@@ -1045,7 +1053,7 @@ int32 UWorldPartitionConvertCommandlet::Main(const FString& Params)
 
 	// Gather and load sublevels
 	TArray<ULevel*> SubLevelsToConvert;
-	GatherAndPrepareSubLevelsToConvert(WorldPartition, MainLevel, SubLevelsToConvert);
+	GatherAndPrepareSubLevelsToConvert(MainLevel, SubLevelsToConvert);
 
 	if (!GetAdditionalLevelsToConvert(MainLevel, SubLevelsToConvert))
 	{
@@ -1309,8 +1317,8 @@ int32 UWorldPartitionConvertCommandlet::Main(const FString& Params)
 	bool bForceInitializeWorld = false;
 	bool bInitializedPhysicsSceneForSave = GEditor->InitializePhysicsSceneForSaveIfNecessary(MainWorld, bForceInitializeWorld);
 
-	// Create the world's minimap
-	if (!Switches.Contains(TEXT("SkipMiniMapGeneration")))
+	// Create the world's minimap, do not create the minimap if only merging
+	if (!bOnlyMergeSubLevels && !Switches.Contains(TEXT("SkipMiniMapGeneration")))
 	{
 		CreateWorldMiniMapTexture(MainWorld);
 	}
@@ -1344,7 +1352,10 @@ int32 UWorldPartitionConvertCommandlet::Main(const FString& Params)
 	ActorList.Append(ChildActorList);
 	ChildActorList.Empty();
 
-	WorldPartition->AddToRoot();
+	if (!bOnlyMergeSubLevels)
+	{
+		WorldPartition->AddToRoot();
+	}
 
 	if (!bReportOnly)
 	{
@@ -1394,7 +1405,7 @@ int32 UWorldPartitionConvertCommandlet::Main(const FString& Params)
 
 		MainLevel->bUseExternalActors = true;
 		MainWorld->WorldComposition = nullptr;
-		MainLevel->bIsPartitioned = true;
+		MainLevel->bIsPartitioned = !bOnlyMergeSubLevels;
 
 		if (bDeleteSourceLevels)
 		{
@@ -1483,8 +1494,12 @@ int32 UWorldPartitionConvertCommandlet::Main(const FString& Params)
 		if (bGenerateIni || !FPlatformFileManager::Get().GetPlatformFile().FileExists(*LevelConfigFilename))
 		{
 			SaveConfig(CPF_Config, *LevelConfigFilename);
-			WorldPartition->EditorHash->SaveConfig(CPF_Config, *LevelConfigFilename);
-			WorldPartition->RuntimeHash->SaveConfig(CPF_Config, *LevelConfigFilename);
+
+			if (!bOnlyMergeSubLevels)
+			{
+				WorldPartition->EditorHash->SaveConfig(CPF_Config, *LevelConfigFilename);
+				WorldPartition->RuntimeHash->SaveConfig(CPF_Config, *LevelConfigFilename);
+			}
 			
 			for(const auto& Pair : HLODLayers)
 			{
@@ -1500,6 +1515,11 @@ int32 UWorldPartitionConvertCommandlet::Main(const FString& Params)
 	return 0;
 }
 
+const FString UWorldPartitionConvertCommandlet::GetConversionSuffix(const bool bInOnlyMergeSubLevels)
+{
+	return bInOnlyMergeSubLevels ? TEXT("_OFPA") : TEXT("_WP");
+}
+
 bool UWorldPartitionConvertCommandlet::ShouldConvertStreamingLevel(ULevelStreaming* StreamingLevel)
 {
 	return StreamingLevel && !ExcludedLevels.Contains(StreamingLevel->GetWorldAssetPackageName());
diff --git a/Engine/Source/Editor/UnrealEd/Private/WorldPartition/WorldPartitionBuilder.cpp b/Engine/Source/Editor/UnrealEd/Private/WorldPartition/WorldPartitionBuilder.cpp
--- a/Engine/Source/Editor/UnrealEd/Private/WorldPartition/WorldPartitionBuilder.cpp
+++ b/Engine/Source/Editor/UnrealEd/Private/WorldPartition/WorldPartitionBuilder.cpp
@@ -183,7 +183,7 @@ bool UWorldPartitionBuilder::Run(UWorld* World, FPackageSourceControlHelper& Pac
 		UE_LOG(LogWorldPartitionBuilder, Display, TEXT("Cell Size %d"), IterativeCellSize);
 		UE_LOG(LogWorldPartitionBuilder, Display, TEXT("Cell Overlap %d"), IterativeCellOverlapSize);
 		UE_LOG(LogWorldPartitionBuilder, Display, TEXT("WorldBounds: Min %s, Max %s"), *CellInfo.EditorBounds.Min.ToString(), *CellInfo.EditorBounds.Max.ToString());
-		UE_LOG(LogWorldPartitionBuilder, Display, TEXT("Iteration Count: %d"), (MaxCellCoords.Z-MinCellCoords.Z) * (MaxCellCoords.Y-MinCellCoords.Y) * (MaxCellCoords.X - MinCellCoords.X));
+		UE_LOG(LogWorldPartitionBuilder, Display, TEXT("Iteration Count: %d"), ((LoadingMode == ELoadingMode::IterativeCells2D) ? 1 : (MaxCellCoords.Z-MinCellCoords.Z + 1)) * (MaxCellCoords.Y-MinCellCoords.Y + 1) * (MaxCellCoords.X - MinCellCoords.X + 1));
 		
 		FBox LoadedBounds(ForceInit);
 
diff --git a/Engine/Source/Editor/WorldPartitionEditor/Private/WorldPartitionEditorModule.cpp b/Engine/Source/Editor/WorldPartitionEditor/Private/WorldPartitionEditorModule.cpp
--- a/Engine/Source/Editor/WorldPartitionEditor/Private/WorldPartitionEditorModule.cpp
+++ b/Engine/Source/Editor/WorldPartitionEditor/Private/WorldPartitionEditorModule.cpp
@@ -296,7 +296,7 @@ bool FWorldPartitionEditorModule::ConvertMap(const FString& InLongPackageName)
 				FString MapToLoad = InLongPackageName;
 				if (!DefaultConvertOptions->bInPlace)
 				{
-					MapToLoad += UWorldPartitionConvertCommandlet::StaticClass()->GetDefaultObject<UWorldPartitionConvertCommandlet>()->GetConversionSuffix();
+					MapToLoad += UWorldPartitionConvertCommandlet::GetConversionSuffix(DefaultConvertOptions->bOnlyMergeSubLevels);
 				}
 				
 				AssetRegistry.ScanModifiedAssetFiles({ MapToLoad });
@@ -446,6 +446,11 @@ FString UWorldPartitionConvertOptions::ToCommandletArgs() const
 	{
 		CommandletArgsBuilder.Append(TEXT(" -Verbose"));
 	}
+
+	if (bOnlyMergeSubLevels)
+	{
+		CommandletArgsBuilder.Append(TEXT(" -OnlyMergeSubLevels"));
+	}
 	
 	return CommandletArgsBuilder.ToString();
 }
diff --git a/Engine/Source/Editor/WorldPartitionEditor/Public/WorldPartition/WorldPartitionConvertOptions.h b/Engine/Source/Editor/WorldPartitionEditor/Public/WorldPartition/WorldPartitionConvertOptions.h
--- a/Engine/Source/Editor/WorldPartitionEditor/Public/WorldPartition/WorldPartitionConvertOptions.h
+++ b/Engine/Source/Editor/WorldPartitionEditor/Public/WorldPartition/WorldPartitionConvertOptions.h
@@ -38,5 +38,8 @@ public:
 	UPROPERTY(EditAnywhere, AdvancedDisplay, Category = Convert)
 	bool bSkipMiniMapGeneration;
 
+	UPROPERTY(EditAnywhere, AdvancedDisplay, Category = Convert)
+	bool bOnlyMergeSubLevels;
+
 	FString LongPackageName;
 };
\ No newline at end of file
