Prevent crash due to MRU requiring a valid writable LongPackageName when saving a package in /Temp.
#rnx

diff --git a/Engine/Source/Editor/UnrealEd/Private/EditorEngine.cpp b/Engine/Source/Editor/UnrealEd/Private/EditorEngine.cpp
--- a/Engine/Source/Editor/UnrealEd/Private/EditorEngine.cpp
+++ b/Engine/Source/Editor/UnrealEd/Private/EditorEngine.cpp
@@ -4537,7 +4537,7 @@ void UEditorEngine::OnPostSaveWorld(UWorld* World, FObjectPostSaveContext Object
 			{
 				// Put the map into the MRU and mark it as not dirty.
 
-				if ( bIsPersistentLevel )
+				if ( bIsPersistentLevel && FPackageName::IsValidLongPackageName(WorldPackage->GetName()))
 				{
 					// Set the map filename.
 					const FString Filename = FPackageName::LongPackageNameToFilename(WorldPackage->GetName(), FPackageName::GetMapPackageExtension());
