Fix for PSO commandlet filtering out valid PSOs.

- Also a fix for garbage being printed instead of the readable PSO desc in the dump.

diff --git a/Engine/Source/Editor/UnrealEd/Private/Commandlets/ShaderPipelineCacheToolsCommandlet.cpp b/Engine/Source/Editor/UnrealEd/Private/Commandlets/ShaderPipelineCacheToolsCommandlet.cpp
--- a/Engine/Source/Editor/UnrealEd/Private/Commandlets/ShaderPipelineCacheToolsCommandlet.cpp
+++ b/Engine/Source/Editor/UnrealEd/Private/Commandlets/ShaderPipelineCacheToolsCommandlet.cpp
@@ -413,8 +413,8 @@ int32 DumpPSOSC(FString& Token, const FString& StableKeyFileDir)
 	for (FPipelineCacheFileFormatPSO& Item : PSOs)
 	{
 		UE_LOG(LogShaderPipelineCacheTools, Display, TEXT("--- Entry %d --------------------------------"), Count);
-		FStringView ReadablePSODesc = Item.ToStringReadable();
-		UE_LOG(LogShaderPipelineCacheTools, Display, TEXT("%.*s"), ReadablePSODesc.Len(), ReadablePSODesc.GetData());
+		FString ReadablePSODesc = Item.ToStringReadable();
+		UE_LOG(LogShaderPipelineCacheTools, Display, TEXT("%s"), *ReadablePSODesc);
 
 		if (Item.Type == FPipelineCacheFileFormatPSO::DescriptorType::Compute)
 		{
@@ -1550,7 +1550,8 @@ void FilterInvalidPSOs(TSet<FPipelineCacheFileFormatPSO>& InOutPSOs, const TMult
 		TEXT("FWaterTileVS"),
 		TEXT("FRenderSkyAtmosphereVS"),
 		TEXT("TPageTableUpdateVS<true>"),
-		TEXT("TPageTableUpdateVS<false>")
+		TEXT("TPageTableUpdateVS<false>"),
+		TEXT("FShaderDrawDebugVS")
 	};
 
 	TSet<FName> VShadersUsableWithEmptyVertexDecl;
