EditorDomain: Use the EditorDomain's cache of EditorDomainHash instead of recalculating it every time it is needed for BulkData and TargetDomain keys.
#rnx

diff --git a/Engine/Source/Editor/UnrealEd/Private/Cooker/CookRequestCluster.cpp b/Engine/Source/Editor/UnrealEd/Private/Cooker/CookRequestCluster.cpp
--- a/Engine/Source/Editor/UnrealEd/Private/Cooker/CookRequestCluster.cpp
+++ b/Engine/Source/Editor/UnrealEd/Private/Cooker/CookRequestCluster.cpp
@@ -389,7 +389,8 @@ void FRequestCluster::StartAsync(const FCookerTimer& CookerTimer, bool& bOutComp
 		return;
 	}
 
-	if (FEditorDomain::Get())
+	FEditorDomain* EditorDomain = FEditorDomain::Get();
+	if (EditorDomain)
 	{
 		// Disable BatchDownload until cache storage implements fetch-head requests in response to a Get with SkipData
 		bool bBatchDownloadEnabled = true;
@@ -398,17 +399,13 @@ void FRequestCluster::StartAsync(const FCookerTimer& CookerTimer, bool& bOutComp
 		{
 			// If the EditorDomain is active, then batch-download all packages to cook from remote cache into local
 			TArray<UE::DerivedData::FCacheKey> CacheKeys;
-			FString ErrorMessage;
 			CacheKeys.Reserve(Requests.Num());
 			for (FPackageData* PackageData : Requests)
 			{
-				FPackageDigest PackageDigest;
-				EDomainUse EditorDomainUse;
-				EPackageDigestResult Result = GetPackageDigest(AssetRegistry, PackageData->GetPackageName(),
-					PackageDigest, EditorDomainUse, ErrorMessage);
-				if (Result == EPackageDigestResult::Success && EnumHasAnyFlags(EditorDomainUse, EDomainUse::LoadEnabled))
+				FPackageDigest PackageDigest = EditorDomain->GetPackageDigest(PackageData->GetPackageName());
+				if (PackageDigest.IsSuccessful() && EnumHasAnyFlags(PackageDigest.DomainUse, EDomainUse::LoadEnabled))
 				{
-					CacheKeys.Add(GetEditorDomainPackageKey(PackageDigest));
+					CacheKeys.Add(GetEditorDomainPackageKey(PackageDigest.Hash));
 				}
 			}
 
diff --git a/Engine/Source/Editor/UnrealEd/Private/EditorDomain/EditorDomain.cpp b/Engine/Source/Editor/UnrealEd/Private/EditorDomain/EditorDomain.cpp
--- a/Engine/Source/Editor/UnrealEd/Private/EditorDomain/EditorDomain.cpp
+++ b/Engine/Source/Editor/UnrealEd/Private/EditorDomain/EditorDomain.cpp
@@ -145,20 +145,14 @@ FEditorDomain::FLocks::FLocks(FEditorDomain& InOwner)
 {
 }
 
-bool FEditorDomain::TryFindOrAddPackageSource(const FPackagePath& PackagePath, TRefCountPtr<FPackageSource>& OutSource)
+bool FEditorDomain::TryFindOrAddPackageSource(FName PackageName,
+	TRefCountPtr<FPackageSource>& OutSource, UE::EditorDomain::FPackageDigest* OutErrorDigest)
 {
 	// Called within Locks.Lock
 	using namespace UE::EditorDomain;
 
 	// EDITOR_DOMAIN_TODO: Need to delete entries from PackageSources when the assetregistry reports the package is
 	// resaved on disk.
-	FName PackageName = PackagePath.GetPackageFName();
-	if (PackageName.IsNone())
-	{
-		OutSource.SafeRelease();
-		return false;
-	}
-
 	TRefCountPtr<FPackageSource>& PackageSource = PackageSources.FindOrAdd(PackageName);
 	if (PackageSource)
 	{
@@ -167,25 +161,27 @@ bool FEditorDomain::TryFindOrAddPackageSource(const FPackagePath& PackagePath, T
 	}
 
 	FString ErrorMessage;
-	FPackageDigest PackageDigest;
-	EDomainUse EditorDomainUseForPackage;
-	EPackageDigestResult Result = GetPackageDigest(*AssetRegistry, PackageName, PackageDigest, EditorDomainUseForPackage, ErrorMessage);
-	switch (Result)
+	FPackageDigest PackageDigest = CalculatePackageDigest(*AssetRegistry, PackageName);
+	switch (PackageDigest.Status)
 	{
-	case EPackageDigestResult::Success:
+	case FPackageDigest::EStatus::Successful:
 		PackageSource = new FPackageSource();
-		PackageSource->Digest = PackageDigest;
+		PackageSource->Digest = MoveTemp(PackageDigest);
 		OutSource = PackageSource;
-		if (!bEditorDomainReadEnabled || !EnumHasAnyFlags(EditorDomainUseForPackage, EDomainUse::LoadEnabled))
+		if (!bEditorDomainReadEnabled || !EnumHasAnyFlags(PackageDigest.DomainUse, EDomainUse::LoadEnabled))
 		{
 			PackageSource->Source = EPackageSource::Workspace;
 		}
 		return true;
-	case EPackageDigestResult::FileDoesNotExist:
+	case FPackageDigest::EStatus::DoesNotExistInAssetRegistry:
 		OutSource.SafeRelease();
 		// Remove the entry in PackageSources that we added; we added it to avoid a double lookup for new packages,
 		// but for non-existent packages we want it not to be there to avoid wasting memory on it
 		PackageSources.Remove(PackageName);
+		if (OutErrorDigest)
+		{
+			*OutErrorDigest = MoveTemp(PackageDigest);
+		}
 		return false;
 	default:
 		UE_LOG(LogEditorDomain, Warning,
@@ -198,6 +194,20 @@ bool FEditorDomain::TryFindOrAddPackageSource(const FPackagePath& PackagePath, T
 	}
 }
 
+UE::EditorDomain::FPackageDigest FEditorDomain::GetPackageDigest(FName PackageName)
+{
+	using namespace UE::EditorDomain;
+
+	TRefCountPtr<FPackageSource> PackageSource;
+	FPackageDigest ErrorDigest;
+	FScopeLock ScopeLock(&Locks->Lock);
+	if (!TryFindOrAddPackageSource(PackageName, PackageSource, &ErrorDigest))
+	{
+		return ErrorDigest;
+	}
+	return PackageSource->Digest;
+}
+
 void FEditorDomain::PrecachePackageDigest(FName PackageName)
 {
 	AssetRegistry->WaitForPackage(PackageName.ToString());
@@ -253,7 +263,13 @@ int64 FEditorDomain::FileSize(const FPackagePath& PackagePath, EPackageSegment P
 	{
 		FScopeLock ScopeLock(&Locks->Lock);
 		TRefCountPtr<FPackageSource> PackageSource;
-		if (!TryFindOrAddPackageSource(PackagePath, PackageSource) || PackageSource->Source == EPackageSource::Workspace)
+		FName PackageName = PackagePath.GetPackageFName();
+		if (PackageName.IsNone())
+		{
+			return Workspace->FileSize(PackagePath, PackageSegment, OutUpdatedPath);
+		}
+
+		if (!TryFindOrAddPackageSource(PackageName, PackageSource) || PackageSource->Source == EPackageSource::Workspace)
 		{
 			return Workspace->FileSize(PackagePath, PackageSegment, OutUpdatedPath);
 		}
@@ -292,7 +308,8 @@ int64 FEditorDomain::FileSize(const FPackagePath& PackagePath, EPackageSegment P
 		// Fetch meta-data only
 		ECachePolicy SkipFlags = ECachePolicy::SkipData & ~ECachePolicy::SkipMeta;
 		Owner.Emplace(EPriority::Highest);
-		RequestEditorDomainPackage(PackagePath, PackageSource->Digest, SkipFlags, *Owner, MoveTemp(MetaDataGetComplete));
+		RequestEditorDomainPackage(PackagePath, PackageSource->Digest.Hash, SkipFlags,
+			*Owner, MoveTemp(MetaDataGetComplete));
 	}
 	Owner->Wait();
 	return FileSize;
@@ -309,8 +326,13 @@ FOpenPackageResult FEditorDomain::OpenReadPackage(const FPackagePath& PackagePat
 	{
 		return Workspace->OpenReadPackage(PackagePath, PackageSegment, OutUpdatedPath);
 	}
+	FName PackageName = PackagePath.GetPackageFName();
+	if (PackageName.IsNone())
+	{
+		return Workspace->OpenReadPackage(PackagePath, PackageSegment, OutUpdatedPath);
+	}
 	TRefCountPtr<FPackageSource> PackageSource;
-	if (!TryFindOrAddPackageSource(PackagePath, PackageSource) || (PackageSource->Source == EPackageSource::Workspace))
+	if (!TryFindOrAddPackageSource(PackageName, PackageSource) || (PackageSource->Source == EPackageSource::Workspace))
 	{
 		return Workspace->OpenReadPackage(PackagePath, PackageSegment, OutUpdatedPath);
 	}
@@ -319,7 +341,7 @@ FOpenPackageResult FEditorDomain::OpenReadPackage(const FPackagePath& PackagePat
 	// and don't need to block on the result before exiting this function
 	EPriority Priority = EPriority::Blocking;
 	FEditorDomainReadArchive* Result = new FEditorDomainReadArchive(Locks, PackagePath, PackageSource, Priority);
-	const FPackageDigest PackageSourceDigest = PackageSource->Digest;
+	const FIoHash PackageEditorHash = PackageSource->Digest.Hash;
 	const bool bHasEditorSource = (PackageSource->Source == EPackageSource::Editor);
 
 	// Unlock before requesting the package because the completion callback takes the lock.
@@ -327,7 +349,7 @@ FOpenPackageResult FEditorDomain::OpenReadPackage(const FPackagePath& PackagePat
 
 	// Fetch only meta-data in the initial request
 	ECachePolicy SkipFlags = ECachePolicy::SkipData & ~ECachePolicy::SkipMeta;
-	RequestEditorDomainPackage(PackagePath, PackageSourceDigest,
+	RequestEditorDomainPackage(PackagePath, PackageEditorHash,
 		SkipFlags, Result->GetRequestOwner(),
 		[Result](FCacheGetCompleteParams&& Params)
 		{
@@ -360,8 +382,13 @@ FOpenAsyncPackageResult FEditorDomain::OpenAsyncReadPackage(const FPackagePath&
 		return Workspace->OpenAsyncReadPackage(PackagePath, PackageSegment);
 	}
 
+	FName PackageName = PackagePath.GetPackageFName();
+	if (PackageName.IsNone())
+	{
+		return Workspace->OpenAsyncReadPackage(PackagePath, PackageSegment);
+	}
 	TRefCountPtr<FPackageSource> PackageSource;
-	if (!TryFindOrAddPackageSource(PackagePath, PackageSource) ||
+	if (!TryFindOrAddPackageSource(PackageName, PackageSource) ||
 		(PackageSource->Source == EPackageSource::Workspace))
 	{
 		return Workspace->OpenAsyncReadPackage(PackagePath, PackageSegment);
@@ -373,14 +400,14 @@ FOpenAsyncPackageResult FEditorDomain::OpenAsyncReadPackage(const FPackagePath&
 	FEditorDomainAsyncReadFileHandle* Result =
 		new FEditorDomainAsyncReadFileHandle(Locks, PackagePath, PackageSource, Priority);
 	const bool bHasEditorSource = (PackageSource->Source == EPackageSource::Editor);
+	FIoHash EditorDomainHash = PackageSource->Digest.Hash;
 
 	// Unlock before requesting the package because the completion callback takes the lock.
 	ScopeLock.Unlock();
 
 	// Fetch meta-data only in the initial request
 	ECachePolicy SkipFlags = ECachePolicy::SkipData & ~ECachePolicy::SkipMeta;
-	RequestEditorDomainPackage(PackagePath, PackageSource->Digest,
-		SkipFlags, Result->GetRequestOwner(),
+	RequestEditorDomainPackage(PackagePath, EditorDomainHash, SkipFlags, Result->GetRequestOwner(),
 		[Result](FCacheGetCompleteParams&& Params)
 		{
 			// Note that ~FEditorDomainAsyncReadFileHandle waits for this callback to be called, so Result cannot dangle
@@ -566,3 +593,49 @@ void FEditorDomain::OnAssetUpdatedOnDisk(const FAssetData& AssetData)
 	FScopeLock ScopeLock(&Locks->Lock);
 	PackageSources.Remove(PackageName);
 }
+
+namespace UE::EditorDomain
+{
+
+FPackageDigest::FPackageDigest(EStatus InStatus, FName InStatusArg)
+ : Status(InStatus)
+ , StatusArg(InStatusArg)
+{
+}
+
+bool FPackageDigest::IsSuccessful() const
+{
+	return Status == EStatus::Successful;
+}
+
+FString FPackageDigest::GetStatusString() const
+{
+	return LexToString(Status, StatusArg);
+}
+
+}
+
+FString LexToString(UE::EditorDomain::FPackageDigest::EStatus Status, FName StatusArg)
+{
+	using namespace UE::EditorDomain;
+
+	switch (Status)
+	{
+	case FPackageDigest::EStatus::NotYetRequested:
+		return TEXT("Has not been requested.");
+	case FPackageDigest::EStatus::Successful:
+		return TEXT("Successful.");
+	case FPackageDigest::EStatus::InvalidPackageName:
+		return TEXT("PackageName is not a valid LongPackageName.");
+	case FPackageDigest::EStatus::DoesNotExistInAssetRegistry:
+		return TEXT("Does not exist in AssetRegistry.");
+	case FPackageDigest::EStatus::MissingClass:
+		return FString::Printf(TEXT("Uses class %s that is not loaded."), *StatusArg.ToString());
+	case FPackageDigest::EStatus::MissingCustomVersion:
+		return FString::Printf(
+			TEXT("Uses CustomVersion guid %s but that guid is not available in FCurrentCustomVersions."),
+			*StatusArg.ToString());
+	default:
+		return TEXT("Unknown result code.");
+	}
+}
\ No newline at end of file
diff --git a/Engine/Source/Editor/UnrealEd/Private/EditorDomain/EditorDomain.h b/Engine/Source/Editor/UnrealEd/Private/EditorDomain/EditorDomain.h
--- a/Engine/Source/Editor/UnrealEd/Private/EditorDomain/EditorDomain.h
+++ b/Engine/Source/Editor/UnrealEd/Private/EditorDomain/EditorDomain.h
@@ -21,11 +21,58 @@ class IAssetRegistry;
 class UObject;
 class UPackage;
 struct FAssetData;
+
 namespace UE::EditorDomain
 {
-	typedef FIoHash FPackageDigest;
+
+/** Flags for whether a package is allowed to be loaded/saved in the EditorDomain. */
+enum class EDomainUse : uint8
+{
+	None = 0x0,
+	/** The package can be loaded from the EditorDomain. */
+	LoadEnabled = 0x1,
+	/** The package can be saved to the EditorDomain. */
+	SaveEnabled = 0x2,
+};
+ENUM_CLASS_FLAGS(EDomainUse);
+
+/** Information about a package's loading from the EditorDomain. */
+struct FPackageDigest
+{
+	enum class EStatus : uint8
+	{
+		NotYetRequested,
+		Successful,
+		InvalidPackageName,
+		DoesNotExistInAssetRegistry,
+		MissingClass,
+		MissingCustomVersion,
+	};
+
+	FPackageDigest() = default;
+	FPackageDigest(EStatus InStatus, FName InStatusArg = NAME_None);
+	bool IsSuccessful() const;
+	FString GetStatusString() const;
+
+	/**
+	 * Hash used to lookup the resaved package in the cache;
+	 * created from properties that change if the saved version would change (package bytes, serialization versions).
+	 */
+	FIoHash Hash;
+	/** Allow flags for whether the package can be saved/loaded from EditorDomain. */
+	EDomainUse DomainUse = EDomainUse::None;
+	/** Status for creation of this digest. Either Success or an error code for why it couldn't be created. */
+	EStatus Status = EStatus::NotYetRequested;
+	/** Extended information for the status description (e.g. missing class name). */
+	FName StatusArg;
+	/** List of CustomVersions used to save the package. */
+	UE::AssetRegistry::FPackageCustomVersionsHandle CustomVersions;
+};
+
 }
 
+extern FString LexToString(UE::EditorDomain::FPackageDigest::EStatus Status, FName StatusArg);
+
 DECLARE_LOG_CATEGORY_EXTERN(LogEditorDomain, Log, All);
 
 /**
@@ -93,7 +140,13 @@ public:
 	virtual TStatId GetStatId() const override { return TStatId(); }
 
 	// EditorDomain interface
+	/** Fetch data from game-thread sources that is required to calculate the PackageDigest of the given PackageName. */
 	void PrecachePackageDigest(FName PackageName);
+	/**
+	 * Calculate the PackageDigest for the given PackageName.
+	 * Callable from any thread, but will fail if game-thread data is required and not cached.
+	 */
+	UE::EditorDomain::FPackageDigest GetPackageDigest(FName PackageDigest);
 
 private:
 	/**
@@ -127,7 +180,8 @@ private:
 	FEditorDomain(FEditorDomain&& Other) = delete;
 
 	/** Read the PackageSource data from PackageSources, or from the asset registry if not in PackageSources. */
-	bool TryFindOrAddPackageSource(const FPackagePath& PackagePath, TRefCountPtr<FPackageSource>& OutSource);
+	bool TryFindOrAddPackageSource(FName PackageName, TRefCountPtr<FPackageSource>& OutSource,
+		UE::EditorDomain::FPackageDigest* OutErrorDigest=nullptr);
 	/** Return the PackageSource data in PackageSources, if it exists */
 	TRefCountPtr<FPackageSource> FindPackageSource(const FPackagePath& PackagePath);
 	/** Mark that we had to load the Package from the workspace domain, and schedule its save into the EditorDomain. */
diff --git a/Engine/Source/Editor/UnrealEd/Private/EditorDomain/EditorDomainArchive.cpp b/Engine/Source/Editor/UnrealEd/Private/EditorDomain/EditorDomainArchive.cpp
--- a/Engine/Source/Editor/UnrealEd/Private/EditorDomain/EditorDomainArchive.cpp
+++ b/Engine/Source/Editor/UnrealEd/Private/EditorDomain/EditorDomainArchive.cpp
@@ -23,7 +23,7 @@ FEditorDomainPackageSegments::FEditorDomainPackageSegments(const TRefCountPtr<FE
 	, RequestOwner(Priority)
 	, PackagePath(InPackagePath)
 	, PackageSource(InPackageSource)
-	, PackageDigest(InPackageSource->Digest)
+	, EditorDomainHash(InPackageSource->Digest.Hash)
 {
 }
 
@@ -407,7 +407,7 @@ void FEditorDomainPackageSegments::SendSegmentRequest(FSegment& Segment)
 
 	// Note that Segment.RequestOwner is Interface-only and so we can write it outside the lock
 	Segment.RequestOwner.Emplace(EPriority::Normal);
-	FCacheChunkRequest SegmentChunk{UE::EditorDomain::GetEditorDomainPackageKey(PackageDigest), Segment.PayloadId};
+	FCacheChunkRequest SegmentChunk{UE::EditorDomain::GetEditorDomainPackageKey(EditorDomainHash), Segment.PayloadId};
 	SegmentChunk.Policy = ECachePolicy::Local;
 	Cache.GetChunks({SegmentChunk}, PackagePath.GetDebugName(), *Segment.RequestOwner,
 		[this, &Segment](FCacheGetChunkCompleteParams&& Params)
diff --git a/Engine/Source/Editor/UnrealEd/Private/EditorDomain/EditorDomainArchive.h b/Engine/Source/Editor/UnrealEd/Private/EditorDomain/EditorDomainArchive.h
--- a/Engine/Source/Editor/UnrealEd/Private/EditorDomain/EditorDomainArchive.h
+++ b/Engine/Source/Editor/UnrealEd/Private/EditorDomain/EditorDomainArchive.h
@@ -162,8 +162,8 @@ private:
 	 * Read-only, *pointer requires EditorDomainLocks.Lock.
 	 */
 	TRefCountPtr<FEditorDomain::FPackageSource> PackageSource;
-	/** PackageDigest for requesting payloads, copied from PackageSource. Read-only. */
-	UE::EditorDomain::FPackageDigest PackageDigest;
+	/** EditorDomainHash for requesting payloads, copied from PackageSource. Read-only. */
+	FIoHash EditorDomainHash;
 
 	/**
 	 * Data read from the EditorDomain; each Segment is requested individually.
diff --git a/Engine/Source/Editor/UnrealEd/Private/EditorDomain/EditorDomainUtils.cpp b/Engine/Source/Editor/UnrealEd/Private/EditorDomain/EditorDomainUtils.cpp
--- a/Engine/Source/Editor/UnrealEd/Private/EditorDomain/EditorDomainUtils.cpp
+++ b/Engine/Source/Editor/UnrealEd/Private/EditorDomain/EditorDomainUtils.cpp
@@ -181,11 +181,21 @@ void FKnownCustomVersions::FindGuidsChecked(TArray<FGuid>& OutGuids, TConstArray
 	}
 }
 
-EPackageDigestResult AppendPackageDigest(FBlake3& Writer, EDomainUse& OutEditorDomainUse, FString& OutErrorMessage,
-	const FAssetPackageData& PackageData, FName PackageName, TArray<FGuid>* OutCustomVersions)
+FPackageDigest CalculatePackageDigest(IAssetRegistry& AssetRegistry, FName PackageName)
 {
-	OutEditorDomainUse = EDomainUse::LoadEnabled | EDomainUse::SaveEnabled;
-	
+	AssetRegistry.WaitForPackage(PackageName.ToString());
+	TOptional<FAssetPackageData> PackageDataOptional = AssetRegistry.GetAssetPackageDataCopy(PackageName);
+	if (!PackageDataOptional)
+	{
+		return FPackageDigest(FPackageDigest::EStatus::DoesNotExistInAssetRegistry);
+	}
+	FAssetPackageData& PackageData = *PackageDataOptional;
+	FPackageDigest Result;
+	Result.DomainUse = EDomainUse::LoadEnabled | EDomainUse::SaveEnabled;;
+	EnumSetFlagsAnd(Result.DomainUse, EDomainUse::LoadEnabled | EDomainUse::SaveEnabled,
+		Result.DomainUse, ~MapFindRef(GPackageBlockedUses, PackageName, EDomainUse::None));
+
+	FBlake3 Writer;
 	FStringView ProjectName(FApp::GetProjectName());
 	Writer.Update(ProjectName.GetData(), ProjectName.Len() * sizeof(ProjectName[0]));
 	Writer.Update(EditorDomainVersion, FCString::Strlen(EditorDomainVersion)*sizeof(EditorDomainVersion[0]));
@@ -212,9 +222,7 @@ EPackageDigestResult AppendPackageDigest(FBlake3& Writer, EDomainUse& OutEditorD
 			// EDITORDOMAIN_TODO: Remove this !IsInGameThread check once FindObject no longer asserts if GIsSavingPackage
 			if (Attempt > 1 || !IsInGameThread())
 			{
-				OutErrorMessage = FString::Printf(TEXT("Package %s uses Class %s but that class is not loaded"),
-					*PackageName.ToString(), *PackageData.ImportedClasses[NextClass].ToString());
-				return EPackageDigestResult::MissingClass;
+				return FPackageDigest(FPackageDigest::EStatus::MissingClass, PackageData.ImportedClasses[NextClass]);
 			}
 			TConstArrayView<FName> RemainingClasses(PackageData.ImportedClasses);
 			RemainingClasses = RemainingClasses.Slice(NextClass, PackageData.ImportedClasses.Num() - NextClass);
@@ -234,8 +242,8 @@ EPackageDigestResult AppendPackageDigest(FBlake3& Writer, EDomainUse& OutEditorD
 				Writer.Update(&ExistingData->SchemaHash, sizeof(ExistingData->SchemaHash));
 			}
 			CustomVersionHandles.Append(ExistingData->CustomVersionHandles);
-			EnumSetFlagsAnd(OutEditorDomainUse, EDomainUse::LoadEnabled | EDomainUse::SaveEnabled,
-				OutEditorDomainUse, ExistingData->EditorDomainUse);
+			EnumSetFlagsAnd(Result.DomainUse, EDomainUse::LoadEnabled | EDomainUse::SaveEnabled,
+				Result.DomainUse, ExistingData->EditorDomainUse);
 		}
 	}
 
@@ -244,28 +252,28 @@ EPackageDigestResult AppendPackageDigest(FBlake3& Writer, EDomainUse& OutEditorD
 	CustomVersionHandles.Sort();
 	CustomVersionHandles.SetNum(Algo::Unique(CustomVersionHandles));
 
-	TArray<FGuid> CustomVersionGuidBuffer;
-	TArray<FGuid>& CustomVersionGuids(OutCustomVersions ? *OutCustomVersions : CustomVersionGuidBuffer);
+	TArray<UE::AssetRegistry::FPackageCustomVersion> CustomVersions;
+	TArray<FGuid> CustomVersionGuids;
 	FKnownCustomVersions::FindGuidsChecked(CustomVersionGuids, CustomVersionHandles);
 	CustomVersionGuids.Sort();
+	CustomVersions.Reserve(CustomVersionGuids.Num());
 
 	for (const FGuid& CustomVersionGuid : CustomVersionGuids)
 	{
 		Writer.Update(&CustomVersionGuid, sizeof(CustomVersionGuid));
 		TOptional<FCustomVersion> CurrentVersion = FCurrentCustomVersions::Get(CustomVersionGuid);
-		if (CurrentVersion.IsSet())
-		{
-			Writer.Update(&CurrentVersion->Version, sizeof(CurrentVersion->Version));
-		}
-		else
+		if (!CurrentVersion.IsSet())
 		{
-			OutErrorMessage = FString::Printf(TEXT("Package %s uses CustomVersion guid %s but that guid is not available in FCurrentCustomVersions"),
-				*PackageName.ToString(), *CustomVersionGuid.ToString());
-			return EPackageDigestResult::MissingCustomVersion;
+			return FPackageDigest(FPackageDigest::EStatus::MissingCustomVersion, FName(CustomVersionGuid.ToString()));
 		}
+		Writer.Update(&CurrentVersion->Version, sizeof(CurrentVersion->Version));
+		CustomVersions.Emplace(CustomVersionGuid, CurrentVersion->Version);
 	}
 
-	return EPackageDigestResult::Success;
+	Result.Hash = Writer.Finalize();
+	Result.Status = FPackageDigest::EStatus::Successful;
+	Result.CustomVersions = UE::AssetRegistry::FPackageCustomVersionsHandle::FindOrAdd(CustomVersions);
+	return Result;
 }
 
 /**
@@ -880,63 +888,26 @@ void UtilsPostEngineInit()
 	}
 }
 
-EPackageDigestResult GetPackageDigest(IAssetRegistry& AssetRegistry, FName PackageName,
-	FPackageDigest& OutPackageDigest, EDomainUse& OutEditorDomainUse, FString& OutErrorMessage,
-	TArray<FGuid>* OutCustomVersions)
-{
-	FBlake3 Builder;
-	EPackageDigestResult Result = AppendPackageDigest(AssetRegistry, PackageName, Builder, OutEditorDomainUse, 
-		OutErrorMessage, OutCustomVersions);
-	if (Result == EPackageDigestResult::Success)
-	{
-		OutPackageDigest = Builder.Finalize();
-	}
-	return Result;
-}
-
-EPackageDigestResult AppendPackageDigest(IAssetRegistry& AssetRegistry, FName PackageName,
-	FBlake3& Builder, EDomainUse& OutEditorDomainUse, FString& OutErrorMessage, TArray<FGuid>* OutCustomVersions)
-{
-	AssetRegistry.WaitForPackage(PackageName.ToString());
-	TOptional<FAssetPackageData> PackageData = AssetRegistry.GetAssetPackageDataCopy(PackageName);
-	if (!PackageData)
-	{
-		OutErrorMessage = FString::Printf(TEXT("Package %s does not exist in the AssetRegistry"),
-			*PackageName.ToString());
-		OutEditorDomainUse = EDomainUse::LoadEnabled | EDomainUse::SaveEnabled;
-		if (OutCustomVersions)
-		{
-			OutCustomVersions->Reset();
-		}
-		return EPackageDigestResult::FileDoesNotExist;
-	}
-	EPackageDigestResult Result = AppendPackageDigest(Builder, OutEditorDomainUse, OutErrorMessage, *PackageData,
-		PackageName, OutCustomVersions);
-	EnumSetFlagsAnd(OutEditorDomainUse, EDomainUse::LoadEnabled | EDomainUse::SaveEnabled,
-		OutEditorDomainUse, ~MapFindRef(GPackageBlockedUses, PackageName, EDomainUse::None));
-	return Result;
-}
-
-UE::DerivedData::FCacheKey GetEditorDomainPackageKey(const FPackageDigest& PackageDigest)
+UE::DerivedData::FCacheKey GetEditorDomainPackageKey(const FIoHash& EditorDomainHash)
 {
 	static UE::DerivedData::FCacheBucket EditorDomainPackageCacheBucket(EditorDomainPackageBucketName);
-	return UE::DerivedData::FCacheKey{EditorDomainPackageCacheBucket, PackageDigest};
+	return UE::DerivedData::FCacheKey{EditorDomainPackageCacheBucket, EditorDomainHash };
 }
 
-UE::DerivedData::FCacheKey GetBulkDataListKey(const FPackageDigest& PackageDigest)
+UE::DerivedData::FCacheKey GetBulkDataListKey(const FIoHash& EditorDomainHash)
 {
 	static UE::DerivedData::FCacheBucket EditorDomainBulkDataListBucket(EditorDomainBulkDataListBucketName);
-	return UE::DerivedData::FCacheKey{ EditorDomainBulkDataListBucket, PackageDigest };
+	return UE::DerivedData::FCacheKey{ EditorDomainBulkDataListBucket, EditorDomainHash };
 }
 
-UE::DerivedData::FCacheKey GetBulkDataPayloadIdKey(const FIoHash& PackageAndGuidDigest)
+UE::DerivedData::FCacheKey GetBulkDataPayloadIdKey(const FIoHash& PackageAndGuidHash)
 {
 	static UE::DerivedData::FCacheBucket EditorDomainBulkDataPayloadIdBucket(EditorDomainBulkDataPayloadIdBucketName);
-	return UE::DerivedData::FCacheKey{ EditorDomainBulkDataPayloadIdBucket, PackageAndGuidDigest };
+	return UE::DerivedData::FCacheKey{ EditorDomainBulkDataPayloadIdBucket, PackageAndGuidHash };
 }
 
 void RequestEditorDomainPackage(const FPackagePath& PackagePath,
-	const FPackageDigest& PackageDigest, UE::DerivedData::ECachePolicy SkipFlags, UE::DerivedData::IRequestOwner& Owner,
+	const FIoHash& EditorDomainHash, UE::DerivedData::ECachePolicy SkipFlags, UE::DerivedData::IRequestOwner& Owner,
 	UE::DerivedData::FOnCacheGetComplete&& Callback)
 {
 	using namespace UE::DerivedData;
@@ -950,7 +921,7 @@ void RequestEditorDomainPackage(const FPackagePath& PackagePath,
 	// But set the CachePolicy to store into remote. This will cause the CacheStore to push
 	// any existing local value into upstream storage and refresh the last-used time in the upstream.
 	ECachePolicy CachePolicy = SkipFlags | ECachePolicy::Local | ECachePolicy::StoreRemote;
-	Cache.Get({ GetEditorDomainPackageKey(PackageDigest) }, PackagePath.GetDebugName(),
+	Cache.Get({ GetEditorDomainPackageKey(EditorDomainHash) }, PackagePath.GetDebugName(),
 		CachePolicy, Owner, MoveTemp(Callback));
 }
 
@@ -1109,24 +1080,24 @@ private:
 bool TrySavePackage(UPackage* Package)
 {
 	using namespace UE::DerivedData;
+	FEditorDomain* EditorDomain = FEditorDomain::Get();
+	if (!EditorDomain)
+	{
+		return false;
+	}
 
 	FString ErrorMessage;
-	FPackageDigest PackageDigest;
-	EDomainUse EditorDomainUse;
-	TArray<FGuid> CustomVersionGuids;
-	EPackageDigestResult FindHashResult = GetPackageDigest(*IAssetRegistry::Get(), Package->GetFName(), PackageDigest,
-		EditorDomainUse, ErrorMessage, &CustomVersionGuids);
-	switch (FindHashResult)
-	{
-	case EPackageDigestResult::Success:
-		break;
-	default:
-		UE_LOG(LogEditorDomain, Warning, TEXT("Could not save package to EditorDomain: %s."), *ErrorMessage)
+	FPackageDigest PackageDigest = EditorDomain->GetPackageDigest(Package->GetFName());
+	if (!PackageDigest.IsSuccessful())
+	{
+		UE_LOG(LogEditorDomain, Warning, TEXT("Could not save package to EditorDomain: %s."),
+			*PackageDigest.GetStatusString());
 		return false;
 	}
-	if (!EnumHasAnyFlags(EditorDomainUse, EDomainUse::SaveEnabled))
+	if (!EnumHasAnyFlags(PackageDigest.DomainUse, EDomainUse::SaveEnabled))
 	{
-		UE_LOG(LogEditorDomain, Verbose, TEXT("Skipping save of blocked package to EditorDomain: %s."), *Package->GetName())
+		UE_LOG(LogEditorDomain, Verbose, TEXT("Skipping save of blocked package to EditorDomain: %s."),
+			*Package->GetName());
 		return false;
 	}
 	UE_LOG(LogEditorDomain, Verbose, TEXT("Saving to EditorDomain: %s."), *Package->GetName())
@@ -1192,11 +1163,12 @@ bool TrySavePackage(UPackage* Package)
 			*Package->GetName());
 		return false;
 	}
+	TConstArrayView<UE::AssetRegistry::FPackageCustomVersion> KnownCustomVersions = PackageDigest.CustomVersions.Get();
 	TSet<FGuid> KnownGuids;
-	KnownGuids.Reserve(CustomVersionGuids.Num());
-	for (const FGuid& CustomVersionGuid : CustomVersionGuids)
+	KnownGuids.Reserve(KnownCustomVersions.Num());
+	for (const UE::AssetRegistry::FPackageCustomVersion& CustomVersion : KnownCustomVersions)
 	{
-		KnownGuids.Add(CustomVersionGuid);
+		KnownGuids.Add(CustomVersion.Key);
 	}
 	TArray<FGuid> UnknownGuids;
 	for (const FCustomVersion& CustomVersion : SavedCustomVersions.GetAllVersions())
@@ -1287,7 +1259,7 @@ bool TrySavePackage(UPackage* Package)
 	MetaData << "FileSize" << FileSize;
 	MetaData.EndObject();
 
-	FCacheRecordBuilder RecordBuilder(GetEditorDomainPackageKey(PackageDigest));
+	FCacheRecordBuilder RecordBuilder(GetEditorDomainPackageKey(PackageDigest.Hash));
 	for (const FEditorDomainPackageWriter::FAttachment& Attachment : PackageWriter->GetAttachments())
 	{
 		RecordBuilder.AddAttachment(Attachment.Buffer, Attachment.PayloadId);
@@ -1314,22 +1286,21 @@ bool TrySavePackage(UPackage* Package)
 
 void GetBulkDataList(FName PackageName, UE::DerivedData::IRequestOwner& Owner, TUniqueFunction<void(FSharedBuffer Buffer)>&& Callback)
 {
-	FString ErrorMessage;
-	FPackageDigest PackageDigest;
-	EDomainUse EditorDomainUse;
-	EPackageDigestResult FindHashResult = GetPackageDigest(*IAssetRegistry::Get(), PackageName, PackageDigest,
-		EditorDomainUse, ErrorMessage);
-	switch (FindHashResult)
-	{
-	case EPackageDigestResult::Success:
-		break;
-	default:
+	FEditorDomain* EditorDomain = FEditorDomain::Get();
+	if (!EditorDomain)
 	{
 		Callback(FSharedBuffer());
 		return;
 	}
+
+	FPackageDigest PackageDigest = EditorDomain->GetPackageDigest(PackageName);
+	if (!PackageDigest.IsSuccessful())
+	{
+		Callback(FSharedBuffer());
+		return;
 	}
-	if (!EnumHasAnyFlags(EditorDomainUse, EDomainUse::LoadEnabled))
+
+	if (!EnumHasAnyFlags(PackageDigest.DomainUse, EDomainUse::LoadEnabled))
 	{
 		Callback(FSharedBuffer());
 		return;
@@ -1337,7 +1308,7 @@ void GetBulkDataList(FName PackageName, UE::DerivedData::IRequestOwner& Owner, T
 
 	using namespace UE::DerivedData;
 	ICache& Cache = GetCache();
-	Cache.Get({ GetBulkDataListKey(PackageDigest) },
+	Cache.Get({ GetBulkDataListKey(PackageDigest.Hash) },
 		WriteToString<128>(PackageName), ECachePolicy::Default, Owner,
 		[InnerCallback = MoveTemp(Callback)](FCacheGetCompleteParams&& Params)
 		{
@@ -1348,21 +1319,19 @@ void GetBulkDataList(FName PackageName, UE::DerivedData::IRequestOwner& Owner, T
 
 void PutBulkDataList(FName PackageName, FSharedBuffer Buffer)
 {
-	FString ErrorMessage;
-	FPackageDigest PackageDigest;
-	EDomainUse EditorDomainUse;
-	EPackageDigestResult FindHashResult = GetPackageDigest(*IAssetRegistry::Get(), PackageName, PackageDigest,
-		EditorDomainUse, ErrorMessage);
-	switch (FindHashResult)
-	{
-	case EPackageDigestResult::Success:
-		break;
-	default:
+	FEditorDomain* EditorDomain = FEditorDomain::Get();
+	if (!EditorDomain)
 	{
 		return;
 	}
+
+	FPackageDigest PackageDigest = EditorDomain->GetPackageDigest(PackageName);
+	if (!PackageDigest.IsSuccessful())
+	{
+		return;
 	}
-	if (!EnumHasAnyFlags(EditorDomainUse, EDomainUse::SaveEnabled))
+
+	if (!EnumHasAnyFlags(PackageDigest.DomainUse, EDomainUse::SaveEnabled))
 	{
 		return;
 	}
@@ -1370,14 +1339,16 @@ void PutBulkDataList(FName PackageName, FSharedBuffer Buffer)
 	using namespace UE::DerivedData;
 	ICache& Cache = GetCache();
 	FRequestOwner Owner(EPriority::Normal);
-	FCacheRecordBuilder RecordBuilder(GetBulkDataListKey(PackageDigest));
+	FCacheRecordBuilder RecordBuilder(GetBulkDataListKey(PackageDigest.Hash));
 	RecordBuilder.SetValue(Buffer);
 	Cache.Put({RecordBuilder.Build()}, WriteToString<128>(PackageName), ECachePolicy::Default, Owner);
 	Owner.KeepAlive();
 }
 
-FIoHash GetPackageAndGuidDigest(FBlake3& Builder, const FGuid& BulkDataId)
+FIoHash GetPackageAndGuidHash(FIoHash& EditorDomainHash, const FGuid& BulkDataId)
 {
+	FBlake3 Builder;
+	Builder.Update(&EditorDomainHash, sizeof(EditorDomainHash));
 	Builder.Update(&BulkDataId, sizeof(BulkDataId));
 	return Builder.Finalize();
 }
@@ -1385,30 +1356,30 @@ FIoHash GetPackageAndGuidDigest(FBlake3& Builder, const FGuid& BulkDataId)
 void GetBulkDataPayloadId(FName PackageName, const FGuid& BulkDataId, UE::DerivedData::IRequestOwner& Owner,
 	TUniqueFunction<void(FSharedBuffer Buffer)>&& Callback)
 {
-	FString ErrorMessage;
-	FBlake3 Builder;
-	EDomainUse EditorDomainUse;
-	EPackageDigestResult FindHashResult = AppendPackageDigest(*IAssetRegistry::Get(), PackageName, Builder, EditorDomainUse, ErrorMessage);
-	switch (FindHashResult)
-	{
-	case EPackageDigestResult::Success:
-		break;
-	default:
+	FEditorDomain* EditorDomain = FEditorDomain::Get();
+	if (!EditorDomain)
 	{
 		Callback(FSharedBuffer());
 		return;
 	}
+
+	FPackageDigest PackageDigest = EditorDomain->GetPackageDigest(PackageName);
+	if (!PackageDigest.IsSuccessful())
+	{
+		Callback(FSharedBuffer());
+		return;
 	}
-	if (!EnumHasAnyFlags(EditorDomainUse, EDomainUse::LoadEnabled))
+
+	if (!EnumHasAnyFlags(PackageDigest.DomainUse, EDomainUse::LoadEnabled))
 	{
 		Callback(FSharedBuffer());
 		return;
 	}
-	FIoHash PackageAndGuidDigest = GetPackageAndGuidDigest(Builder, BulkDataId);
+	FIoHash PackageAndGuidHash = GetPackageAndGuidHash(PackageDigest.Hash, BulkDataId);
 
 	using namespace UE::DerivedData;
 	ICache& Cache = GetCache();
-	Cache.Get({ GetBulkDataPayloadIdKey(PackageAndGuidDigest) },
+	Cache.Get({ GetBulkDataPayloadIdKey(PackageAndGuidHash) },
 		WriteToString<192>(PackageName, TEXT("/"), BulkDataId), ECachePolicy::Default, Owner,
 		[InnerCallback = MoveTemp(Callback)](FCacheGetCompleteParams&& Params)
 	{
@@ -1419,32 +1390,31 @@ void GetBulkDataPayloadId(FName PackageName, const FGuid& BulkDataId, UE::Derive
 
 void PutBulkDataPayloadId(FName PackageName, const FGuid& BulkDataId, FSharedBuffer Buffer)
 {
-	FString ErrorMessage;
-	FBlake3 Builder;
-	EDomainUse EditorDomainUse;
-	EPackageDigestResult FindHashResult = AppendPackageDigest(*IAssetRegistry::Get(), PackageName, Builder, EditorDomainUse, ErrorMessage);
-	switch (FindHashResult)
-	{
-	case EPackageDigestResult::Success:
-		break;
-	default:
+	FEditorDomain* EditorDomain = FEditorDomain::Get();
+	if (!EditorDomain)
 	{
 		return;
 	}
+
+	FPackageDigest PackageDigest = EditorDomain->GetPackageDigest(PackageName);
+	if (!PackageDigest.IsSuccessful())
+	{
+		return;
 	}
-	if (!EnumHasAnyFlags(EditorDomainUse, EDomainUse::SaveEnabled))
+
+	if (!EnumHasAnyFlags(PackageDigest.DomainUse, EDomainUse::SaveEnabled))
 	{
 		return;
 	}
-	FIoHash PackageAndGuidDigest = GetPackageAndGuidDigest(Builder, BulkDataId);
+	FIoHash PackageAndGuidHash = GetPackageAndGuidHash(PackageDigest.Hash, BulkDataId);
 
 	using namespace UE::DerivedData;
 	ICache& Cache = GetCache();
 	FRequestOwner Owner(EPriority::Normal);
-	FCacheRecordBuilder RecordBuilder(GetBulkDataPayloadIdKey(PackageAndGuidDigest));
+	FCacheRecordBuilder RecordBuilder(GetBulkDataPayloadIdKey(PackageAndGuidHash));
 	RecordBuilder.SetValue(Buffer);
 	Cache.Put({RecordBuilder.Build()}, WriteToString<128>(PackageName), ECachePolicy::Default, Owner);
 	Owner.KeepAlive();
 }
 
-}
+}
\ No newline at end of file
diff --git a/Engine/Source/Editor/UnrealEd/Private/EditorDomain/EditorDomainUtils.h b/Engine/Source/Editor/UnrealEd/Private/EditorDomain/EditorDomainUtils.h
--- a/Engine/Source/Editor/UnrealEd/Private/EditorDomain/EditorDomainUtils.h
+++ b/Engine/Source/Editor/UnrealEd/Private/EditorDomain/EditorDomainUtils.h
@@ -30,22 +30,6 @@ namespace UE::DerivedData { struct FCacheKey; }
 namespace UE::EditorDomain
 {
 
-enum class EPackageDigestResult
-{
-	Success,
-	FileDoesNotExist,
-	MissingCustomVersion,
-	MissingClass,
-};
-
-enum class EDomainUse : uint8
-{
-	None = 0x0,
-	LoadEnabled = 0x1,
-	SaveEnabled = 0x2,
-};
-ENUM_CLASS_FLAGS(EDomainUse);
-
 /** A UClass's data that is used in the EditorDomain Digest, and holds other information about Classes the EditorDomain needs. */
 struct FClassDigestData
 {
@@ -74,22 +58,16 @@ struct FClassDigestMap
  * Calculate the PackageDigest for the given packagePath.
  * Reads information from the AssetRegistry to compute the digest.
  */
-EPackageDigestResult GetPackageDigest(IAssetRegistry& AssetRegistry, FName PackageName,
-	FPackageDigest& OutPackageDigest, EDomainUse& OutEditorDomainUse, FString& OutErrorMessage,
-	TArray<FGuid>* OutCustomVersions = nullptr);
-/** Appends the fields to calculate the packagedigest; call Builder.Save().GetRangeHash() to get digest. */
-EPackageDigestResult AppendPackageDigest(IAssetRegistry& AssetRegistry, FName PackageName,
-	FBlake3& Builder, EDomainUse& OutEditorDomainUse, FString& OutErrorMessage,
-	TArray<FGuid>* OutCustomVersions = nullptr);
+FPackageDigest CalculatePackageDigest(IAssetRegistry& AssetRegistry, FName PackageName);
 
 /** For any ClassNames not already in ClassDigests, look up their UStruct and add them. */
 void PrecacheClassDigests(TConstArrayView<FName> ClassNames);
 
 /** Get the CacheRequest for the given package from the EditorDomain cache bucket. */
 void RequestEditorDomainPackage(const FPackagePath& PackagePath,
-	const FPackageDigest& PackageDigest, UE::DerivedData::ECachePolicy SkipFlags,
+	const FIoHash& EditorDomainHash, UE::DerivedData::ECachePolicy SkipFlags,
 	UE::DerivedData::IRequestOwner& Owner, UE::DerivedData::FOnCacheGetComplete&& Callback);
-UE::DerivedData::FCacheKey GetEditorDomainPackageKey(const FPackageDigest& PackageDigest);
+UE::DerivedData::FCacheKey GetEditorDomainPackageKey(const FIoHash& EditorDomainHash);
 
 /** Save the given package into the EditorDomain. */
 bool TrySavePackage(UPackage* Package);
diff --git a/Engine/Source/Editor/UnrealEd/Private/TargetDomain/TargetDomainUtils.cpp b/Engine/Source/Editor/UnrealEd/Private/TargetDomain/TargetDomainUtils.cpp
--- a/Engine/Source/Editor/UnrealEd/Private/TargetDomain/TargetDomainUtils.cpp
+++ b/Engine/Source/Editor/UnrealEd/Private/TargetDomain/TargetDomainUtils.cpp
@@ -82,29 +82,34 @@ bool TryCreateKey(FName PackageName, TConstArrayView<FName> SortedBuildDependenc
 		if (OutErrorMessage) *OutErrorMessage = TEXT("AssetRegistry is unavailable.");
 		return false;
 	}
+	FEditorDomain* EditorDomain = FEditorDomain::Get();
+	if (!EditorDomain)
+	{
+		if (OutErrorMessage) *OutErrorMessage = TEXT("EditorDomain is unavailable.");
+		return false;
+	}
 	FBlake3 KeyBuilder;
-	UE::EditorDomain::EPackageDigestResult Result;
-	FString ErrorMessage;
-	UE::EditorDomain::EDomainUse EditorDomainUse;
-	Result = UE::EditorDomain::AppendPackageDigest(*AssetRegistry, PackageName, KeyBuilder, EditorDomainUse, ErrorMessage);
-	if (Result != UE::EditorDomain::EPackageDigestResult::Success)
+	UE::EditorDomain::FPackageDigest PackageDigest = EditorDomain->GetPackageDigest(PackageName);
+	if (!PackageDigest.IsSuccessful())
 	{
-		if (OutErrorMessage) *OutErrorMessage = MoveTemp(ErrorMessage);
+		if (OutErrorMessage) *OutErrorMessage = PackageDigest.GetStatusString();
 		return false;
 	}
+	KeyBuilder.Update(&PackageDigest.Hash, sizeof(PackageDigest.Hash));
 
 	for (FName DependencyName : SortedBuildDependencies)
 	{
-		Result = UE::EditorDomain::AppendPackageDigest(*AssetRegistry, DependencyName, KeyBuilder, EditorDomainUse, ErrorMessage);
-		if (Result != UE::EditorDomain::EPackageDigestResult::Success)
+		PackageDigest = EditorDomain->GetPackageDigest(PackageName);
+		if (!PackageDigest.IsSuccessful())
 		{
 			if (OutErrorMessage)
 			{
 				*OutErrorMessage = FString::Printf(TEXT("Could not create PackageDigest for %s: %s"),
-					*DependencyName.ToString(), *ErrorMessage);
+					*DependencyName.ToString(), *PackageDigest.GetStatusString());
 			}
 			return false;
 		}
+		KeyBuilder.Update(&PackageDigest.Hash, sizeof(PackageDigest.Hash));
 	}
 
 	if (OutHash)
@@ -122,17 +127,22 @@ bool TryCollectKeyAndDependencies(UPackage* Package, const ITargetPlatform* Targ
 		if (OutErrorMessage) *OutErrorMessage = TEXT("Invalid null package.");
 		return false;
 	}
-
-	FName PackageName = Package->GetFName();
-	TSet<FName> BuildDependencies;
-	TSet<FName> RuntimeOnlyDependencies;
-
 	IAssetRegistry* AssetRegistry = IAssetRegistry::Get();
 	if (!AssetRegistry)
 	{
 		if (OutErrorMessage) *OutErrorMessage = TEXT("AssetRegistry is unavailable.");
 		return false;
 	}
+	FEditorDomain* EditorDomain = FEditorDomain::Get();
+	if (!EditorDomain)
+	{
+		if (OutErrorMessage) *OutErrorMessage = TEXT("EditorDomain is unavailable.");
+		return false;
+	}
+
+	FName PackageName = Package->GetFName();
+	TSet<FName> BuildDependencies;
+	TSet<FName> RuntimeOnlyDependencies;
 
 	TArray<FName> AssetDependencies;
 	AssetRegistry->GetDependencies(PackageName, AssetDependencies, UE::AssetRegistry::EDependencyCategory::Package,
