Remove bCheckReferencesOnDelete.  It has no application to UE5

diff --git a/Engine/Source/Editor/UnrealEd/Classes/Settings/LevelEditorMiscSettings.h b/Engine/Source/Editor/UnrealEd/Classes/Settings/LevelEditorMiscSettings.h
--- a/Engine/Source/Editor/UnrealEd/Classes/Settings/LevelEditorMiscSettings.h
+++ b/Engine/Source/Editor/UnrealEd/Classes/Settings/LevelEditorMiscSettings.h
@@ -55,10 +55,6 @@ public:
 	UPROPERTY(EditAnywhere, config, Category = Editing, meta = (DisplayName = "Avoid Actor Relabel on Paste Selected"))
 	uint32 bAvoidRelabelOnPasteSelected:1;
 
-	/** If enabled, will force checking references during a delete. Note: disabling this is for advanced usage only and not recommended.*/
-	UPROPERTY(EditAnywhere, config, AdvancedDisplay, Category = Editing, meta = (DisplayName = "Check References on Delete"))
-	uint32 bCheckReferencesOnDelete:1;
-
 public:
 	/** If checked audio playing in the editor will continue to play even if the editor is in the background */
 	UPROPERTY(EditAnywhere, config, Category=Sound)
