Add MaxInputSizeForRemoteBuilds setting, to limit large requests from being made to Horde.

This will be unnecessary when jupiter is used as storage for remote execution.

#rnx

diff --git a/Engine/Source/Editor/UnrealEd/Private/DerivedDataBuildRemoteExecutor.cpp b/Engine/Source/Editor/UnrealEd/Private/DerivedDataBuildRemoteExecutor.cpp
--- a/Engine/Source/Editor/UnrealEd/Private/DerivedDataBuildRemoteExecutor.cpp
+++ b/Engine/Source/Editor/UnrealEd/Private/DerivedDataBuildRemoteExecutor.cpp
@@ -379,6 +379,13 @@ private:
 				MinInputSizeForRemoteBuilds = (uint64)SignedMinInputSizeForRemoteBuilds;
 			}
 
+			int32 SignedMaxInputSizeForRemoteBuilds{MAX_int32};
+			GConfig->GetInt(Section, TEXT("MaxInputSizeForRemoteBuilds"), SignedMaxInputSizeForRemoteBuilds, GEngineIni);
+			if ((SignedMaxInputSizeForRemoteBuilds >= 0) && (SignedMaxInputSizeForRemoteBuilds < MAX_int32))
+			{
+				MaxInputSizeForRemoteBuilds = (uint64)SignedMaxInputSizeForRemoteBuilds;
+			}
+
 			int32 SignedMaxMissingInputSizeForRemoteBuilds{MAX_int32};
 			GConfig->GetInt(Section, TEXT("MaxMissingInputSizeForRemoteBuilds"), SignedMaxMissingInputSizeForRemoteBuilds, GEngineIni);
 			if ((SignedMaxMissingInputSizeForRemoteBuilds >= 0) && (SignedMaxMissingInputSizeForRemoteBuilds < MAX_int32))
@@ -399,6 +406,11 @@ private:
 				return false;
 			}
 
+			if (InputSize > MaxInputSizeForRemoteBuilds)
+			{
+				return false;
+			}
+
 			if (MissingInputSize > MaxMissingInputSizeForRemoteBuilds)
 			{
 				return false;
@@ -434,6 +446,7 @@ private:
 		uint64 MaxTotalRemoteBuilds{MAX_uint64};
 		uint32 MaxInFlightRemoteBuilds{MAX_uint32};
 		uint64 MinInputSizeForRemoteBuilds{0};
+		uint64 MaxInputSizeForRemoteBuilds{MAX_uint64};
 		uint64 MaxMissingInputSizeForRemoteBuilds{MAX_uint64};
 		bool bEnableLimits{false};
 	};
