Cooker performance: cap the number of loads made before saving to prevent  a large save queue and a large amount of loads redone in the event of a garbage collect.
#rnx

diff --git a/Engine/Source/Editor/UnrealEd/Classes/CookOnTheSide/CookOnTheFlyServer.h b/Engine/Source/Editor/UnrealEd/Classes/CookOnTheSide/CookOnTheFlyServer.h
--- a/Engine/Source/Editor/UnrealEd/Classes/CookOnTheSide/CookOnTheFlyServer.h
+++ b/Engine/Source/Editor/UnrealEd/Classes/CookOnTheSide/CookOnTheFlyServer.h
@@ -212,6 +212,8 @@ private:
 	uint32 DesiredLoadQueueLength;
 	/** A knob to tune performance - how many packages to pull off in each call to PumpRequests. */
 	uint32 RequestBatchSize;
+	/** A knob to tune performance - how many packages to load in each call to PumpLoads. */
+	int32 LoadBatchSize;
 
 	ECookInitializationFlags CookFlags = ECookInitializationFlags::None;
 	TUniquePtr<class FSandboxPlatformFile> SandboxFile;
diff --git a/Engine/Source/Editor/UnrealEd/Private/CookOnTheFlyServer.cpp b/Engine/Source/Editor/UnrealEd/Private/CookOnTheFlyServer.cpp
--- a/Engine/Source/Editor/UnrealEd/Private/CookOnTheFlyServer.cpp
+++ b/Engine/Source/Editor/UnrealEd/Private/CookOnTheFlyServer.cpp
@@ -1857,7 +1857,11 @@ void UCookOnTheFlyServer::PumpLoads(UE::Cook::FTickStackData& StackData, uint32
 	OutNumPushed = 0;
 	bOutBusy = false;
 
-	while (LoadReadyQueue.Num() + LoadPrepareQueue.Num() > static_cast<int32>(DesiredQueueLength))
+	// Process loads until we reduce the queue size down to the desired size or we hit the max number of loads per batch
+	// We do not want to load too many packages without saving because if we hit the memory limit and GC every package
+	// we load will have to be loaded again
+	while (LoadReadyQueue.Num() + LoadPrepareQueue.Num() > static_cast<int32>(DesiredQueueLength) &&
+		OutNumPushed < LoadBatchSize)
 	{
 		if (StackData.Timer.IsTimeUp())
 		{
@@ -4507,6 +4511,7 @@ void UCookOnTheFlyServer::Initialize( ECookMode::Type DesiredCookMode, ECookInit
 	MaxPreloadAllocated = 16;
 	DesiredSaveQueueLength = 8;
 	DesiredLoadQueueLength = 8;
+	LoadBatchSize = 16;
 	RequestBatchSize = 16;
 
 	MinFreeUObjectIndicesBeforeGC = 100000;
