Remote Execute: Suppress zen api errors

#rnx

diff --git a/Engine/Source/Editor/UnrealEd/Private/DerivedDataBuildRemoteExecutor.cpp b/Engine/Source/Editor/UnrealEd/Private/DerivedDataBuildRemoteExecutor.cpp
--- a/Engine/Source/Editor/UnrealEd/Private/DerivedDataBuildRemoteExecutor.cpp
+++ b/Engine/Source/Editor/UnrealEd/Private/DerivedDataBuildRemoteExecutor.cpp
@@ -623,7 +623,7 @@ void FRemoteBuildExecutionRequest::DetermineIfWorkerExists_Async()
 			WorkerUri.AppendAnsi("/apply/workers/");
 			WorkerUri << State.WorkerDescriptionId;
 
-			UE::Zen::FZenScopedRequestPtr Request(Executor.RequestPool.Get());
+			UE::Zen::FZenScopedRequestPtr Request(Executor.RequestPool.Get(), false);
 			State.Result = Request->PerformBlockingDownload(WorkerUri, nullptr, Zen::EContentType::CbObject);
 			State.ResponseCode = Request->GetResponseCode();
 			
@@ -644,7 +644,7 @@ void FRemoteBuildExecutionRequest::PostWorkerObject_Async()
 
 			State.NeedHashes.Empty();
 
-			UE::Zen::FZenScopedRequestPtr Request(Executor.RequestPool.Get());
+			UE::Zen::FZenScopedRequestPtr Request(Executor.RequestPool.Get(), false);
 			State.Result = Request->PerformBlockingPost(WorkerUri, State.WorkerDescriptor);
 			State.ResponseCode = Request->GetResponseCode();
 
@@ -718,7 +718,7 @@ void FRemoteBuildExecutionRequest::PostWorkerPackage_Async()
 
 			Executor.Stats.TotalActionPackagesUploaded.AddBlob(AttachmentBytes + State.WorkerDescriptor.GetSize());
 
-			UE::Zen::FZenScopedRequestPtr Request(Executor.RequestPool.Get());
+			UE::Zen::FZenScopedRequestPtr Request(Executor.RequestPool.Get(), false);
 			State.Result = Request->PerformBlockingPostPackage(WorkerUri, Package);
 			State.ResponseCode = Request->GetResponseCode();
 			
@@ -746,7 +746,7 @@ void FRemoteBuildExecutionRequest::PostActionObject_Async()
 
 			Executor.Stats.TotalActionObjectsUploaded.AddBlob(State.Action.GetSize());
 
-			UE::Zen::FZenScopedRequestPtr Request(Executor.RequestPool.Get());
+			UE::Zen::FZenScopedRequestPtr Request(Executor.RequestPool.Get(), false);
 			State.Result = Request->PerformBlockingPost(ActionUri, State.Action);
 			State.ResponseCode = Request->GetResponseCode();
 
@@ -799,7 +799,7 @@ void FRemoteBuildExecutionRequest::PostActionPackage_Async()
 
 			Executor.Stats.TotalActionPackagesUploaded.AddBlob(AttachmentBytes + State.Action.GetSize());
 
-			UE::Zen::FZenScopedRequestPtr Request(Executor.RequestPool.Get());
+			UE::Zen::FZenScopedRequestPtr Request(Executor.RequestPool.Get(), false);
 			State.Result = Request->PerformBlockingPostPackage(ActionUri, ActionPackage);
 			State.ResponseCode = Request->GetResponseCode();
 			
@@ -818,7 +818,7 @@ void FRemoteBuildExecutionRequest::GetResultPackage_Async()
 			JobGetUri << "/";
 			JobGetUri << State.ActionId;
 
-			UE::Zen::FZenScopedRequestPtr Request(Executor.RequestPool.Get());
+			UE::Zen::FZenScopedRequestPtr Request(Executor.RequestPool.Get(), false);
 			State.Result = Request->PerformBlockingDownload(JobGetUri.ToString(), nullptr, UE::Zen::EContentType::CbPackage);
 			State.ResponseCode = Request->GetResponseCode();
 			if (State.ResponseCode == 200)
