Refactor code paths to use SaveLevelAs always so that we can get the proper save name to reflect level editor.

diff --git a/Engine/Source/Editor/UnrealEd/Private/EditorLevelUtils.cpp b/Engine/Source/Editor/UnrealEd/Private/EditorLevelUtils.cpp
--- a/Engine/Source/Editor/UnrealEd/Private/EditorLevelUtils.cpp
+++ b/Engine/Source/Editor/UnrealEd/Private/EditorLevelUtils.cpp
@@ -648,6 +648,37 @@ ULevelStreaming* UEditorLevelUtils::CreateNewStreamingLevel(TSubclassOf<ULevelSt
 	return nullptr;
 }
 
+namespace UE::EditorLevelUtils::Private
+{
+
+ULevel* GetPersistentLevelForNewStreamingLevel(UWorld* NewLevelWorld, UWorld* InTemplateWorld)
+{
+	check(NewLevelWorld || InTemplateWorld);
+	return NewLevelWorld ? NewLevelWorld->PersistentLevel : InTemplateWorld->PersistentLevel;
+}
+
+UWorld* GetWorldForNewStreamingLevel(UWorld* InTemplateWorld)
+{
+	if (!InTemplateWorld)
+	{
+		// Create a new world
+		UWorldFactory* Factory = NewObject<UWorldFactory>();
+		// streaming levels shouldn't be WP
+		Factory->bCreateWorldPartition = false;
+		Factory->WorldType = EWorldType::Inactive;
+		UPackage* Pkg = CreatePackage( NULL);
+		FName WorldName(TEXT("Untitled"));
+		EObjectFlags Flags = RF_Public | RF_Standalone;
+		UWorld* NewLevelWorld = CastChecked<UWorld>(Factory->FactoryCreateNew(UWorld::StaticClass(), Pkg, WorldName, Flags, NULL, GWarn));
+		if (NewLevelWorld)
+		{
+			FAssetRegistryModule::AssetCreated(NewLevelWorld);
+		}
+		return NewLevelWorld;
+	}
+	return nullptr;
+}
+};
 
 ULevelStreaming* UEditorLevelUtils::CreateNewStreamingLevelForWorld(UWorld& InWorld, TSubclassOf<ULevelStreaming> LevelStreamingClass, const FString& DefaultFilename /* = TEXT( "" ) */, bool bMoveSelectedActorsIntoNewLevel /* = false */, UWorld* InTemplateWorld /* = nullptr */)
 {
@@ -656,22 +687,24 @@ ULevelStreaming* UEditorLevelUtils::CreateNewStreamingLevelForWorld(UWorld& InWo
 	{
 		GLevelEditorModeTools().DeactivateAllModes();
 	}
+	using namespace UE::EditorLevelUtils::Private;
 
 	// This is the world we are adding the new level to
 	UWorld* WorldToAddLevelTo = &InWorld;
 
 	// This is the new streaming level's world not the persistent level world
-	UWorld* NewLevelWorld = nullptr;
+	UWorld* NewLevelWorld = GetWorldForNewStreamingLevel(InTemplateWorld);
+	ULevel* PersistentLevel = GetPersistentLevelForNewStreamingLevel(NewLevelWorld, InTemplateWorld);
+
 	bool bNewWorldSaved = false;
 	FString NewPackageName = DefaultFilename;
 
-	if (InTemplateWorld)
+	bNewWorldSaved = FEditorFileUtils::SaveLevelAs(PersistentLevel, &NewPackageName);
+	if (bNewWorldSaved && !NewPackageName.IsEmpty())
 	{
-		// Copy and save the new world to disk.
-		bNewWorldSaved = FEditorFileUtils::SaveLevelAs(InTemplateWorld->PersistentLevel, &NewPackageName);
-		if (bNewWorldSaved && !NewPackageName.IsEmpty())
+		NewPackageName = FPackageName::FilenameToLongPackageName(NewPackageName);
+		if (!NewLevelWorld)
 		{
-			NewPackageName = FPackageName::FilenameToLongPackageName(NewPackageName);
 			UPackage* NewPackage = LoadPackage(nullptr, *NewPackageName, LOAD_None);
 			if (NewPackage)
 			{
@@ -679,29 +712,6 @@ ULevelStreaming* UEditorLevelUtils::CreateNewStreamingLevelForWorld(UWorld& InWo
 			}
 		}
 	}
-	else
-	{
-		// Create a new world
-		UWorldFactory* Factory = NewObject<UWorldFactory>();
-		// streaming levels shouldn't be WP
-		Factory->bCreateWorldPartition = false;
-		Factory->WorldType = EWorldType::Inactive;
-		UPackage* Pkg = CreatePackage( NULL);
-		FName WorldName(TEXT("Untitled"));
-		EObjectFlags Flags = RF_Public | RF_Standalone;
-		NewLevelWorld = CastChecked<UWorld>(Factory->FactoryCreateNew(UWorld::StaticClass(), Pkg, WorldName, Flags, NULL, GWarn));
-		if (NewLevelWorld)
-		{
-			FAssetRegistryModule::AssetCreated(NewLevelWorld);
-		}
-
-		// Save the new world to disk.
-		bNewWorldSaved = FEditorFileUtils::SaveLevel(NewLevelWorld->PersistentLevel, DefaultFilename);
-		if (bNewWorldSaved)
-		{
-			NewPackageName = NewLevelWorld->GetOutermost()->GetName();
-		}
-	}
 
 	// If the new world was saved successfully, import it as a streaming level.
 	ULevelStreaming* NewStreamingLevel = nullptr;
