Fixing a bug where creating temporary actors via typed elements would not be able to delete them again in single place tool

diff --git a/Engine/Source/Editor/AssetPlacementEdMode/Private/Tools/PlacementPlaceSingleTool.h b/Engine/Source/Editor/AssetPlacementEdMode/Private/Tools/PlacementPlaceSingleTool.h
--- a/Engine/Source/Editor/AssetPlacementEdMode/Private/Tools/PlacementPlaceSingleTool.h
+++ b/Engine/Source/Editor/AssetPlacementEdMode/Private/Tools/PlacementPlaceSingleTool.h
@@ -72,7 +72,6 @@ protected:
 	void DestroyPreviewElements();
 	void EnterTweakState(TArrayView<const FTypedElementHandle> InElementHandles);
 	void ExitTweakState(bool bClearSelectionSet);
-	FAssetPlacementInfo UpdatePlacementInfo();
 	void UpdateElementTransforms(TArrayView<const FTypedElementHandle> InElements, const FTransform& InTransform, bool bLocalTransform = false);
 	void NotifyMovementStarted(TArrayView<const FTypedElementHandle> InElements);
 	void NotifyMovementEnded(TArrayView<const FTypedElementHandle> InElements);
diff --git a/Engine/Source/Editor/Experimental/EditorInteractiveToolsFramework/Private/EdModeInteractiveToolsContext.cpp b/Engine/Source/Editor/Experimental/EditorInteractiveToolsFramework/Private/EdModeInteractiveToolsContext.cpp
--- a/Engine/Source/Editor/Experimental/EditorInteractiveToolsFramework/Private/EdModeInteractiveToolsContext.cpp
+++ b/Engine/Source/Editor/Experimental/EditorInteractiveToolsFramework/Private/EdModeInteractiveToolsContext.cpp
@@ -1159,6 +1159,13 @@ void UEdModeInteractiveToolsContext::InitializeContextFromModeManagerContext(UMo
 
 	ParentModeManagerToolsContext = ModeManagerToolsContext;
 
+
+	SetCreateContextStoreFunc([this, ModeManagerToolsContext](const FContextInitInfo& ContextInfo)
+	{
+		UContextObjectStore* NewContextStore = NewObject<UContextObjectStore>(ModeManagerToolsContext->ContextObjectStore);
+		return NewContextStore;
+	});
+
 	InitializeContextWithEditorModeManager(ModeManager, ModeManagerToolsContext->InputRouter);
 }
 
diff --git a/Engine/Source/Editor/UnrealEd/Private/EditorActor.cpp b/Engine/Source/Editor/UnrealEd/Private/EditorActor.cpp
--- a/Engine/Source/Editor/UnrealEd/Private/EditorActor.cpp
+++ b/Engine/Source/Editor/UnrealEd/Private/EditorActor.cpp
@@ -875,7 +875,15 @@ bool UUnrealEdEngine::DeleteActors(const TArray<AActor*>& InActorsToDelete, UWor
 		FText ErrorMsg;
 		if (!CanDeleteActor(ActorToDelete, &ErrorMsg))
 		{
-			UE_LOG(LogEditorActor, Log, TEXT("Cannot delete actor %s: %s"), *ActorToDelete->GetFullName(), *ErrorMsg.ToString());
+			// If this was a temporary editor preview actor, go ahead and delete it now
+			if (ActorToDelete->bIsEditorPreviewActor)
+			{
+				InWorld->DestroyActor(ActorToDelete);
+			}
+			else
+			{
+				UE_LOG(LogEditorActor, Log, TEXT("Cannot delete actor %s: %s"), *ActorToDelete->GetFullName(), *ErrorMsg.ToString());
+			}
 			continue;
 		}
 
diff --git a/Engine/Source/Editor/UnrealEd/Private/Elements/Actor/ActorElementEditorWorldInterface.cpp b/Engine/Source/Editor/UnrealEd/Private/Elements/Actor/ActorElementEditorWorldInterface.cpp
--- a/Engine/Source/Editor/UnrealEd/Private/Elements/Actor/ActorElementEditorWorldInterface.cpp
+++ b/Engine/Source/Editor/UnrealEd/Private/Elements/Actor/ActorElementEditorWorldInterface.cpp
@@ -64,7 +64,7 @@ void UActorElementEditorWorldInterface::NotifyMovementEnded(const FTypedElementH
 bool UActorElementEditorWorldInterface::CanDeleteElement(const FTypedElementHandle& InElementHandle)
 {
 	AActor* Actor = ActorElementDataUtil::GetActorFromHandle(InElementHandle);
-	return Actor && GUnrealEd->CanDeleteActor(Actor);
+	return Actor && (GUnrealEd->CanDeleteActor(Actor) || Actor->bIsEditorPreviewActor);
 }
 
 bool UActorElementEditorWorldInterface::DeleteElements(TArrayView<const FTypedElementHandle> InElementHandles, UWorld* InWorld, UTypedElementSelectionSet* InSelectionSet, const FTypedElementDeletionOptions& InDeletionOptions)
diff --git a/Engine/Source/Runtime/InteractiveToolsFramework/Public/ContextObjectStore.h b/Engine/Source/Runtime/InteractiveToolsFramework/Public/ContextObjectStore.h
--- a/Engine/Source/Runtime/InteractiveToolsFramework/Public/ContextObjectStore.h
+++ b/Engine/Source/Runtime/InteractiveToolsFramework/Public/ContextObjectStore.h
@@ -34,6 +34,11 @@ public:
 			}
 		}
 
+		if (UContextObjectStore* ParentStore = Cast<UContextObjectStore>(GetOuter()))
+		{
+			return ParentStore->FindContext<TObjectType>();
+		}
+
 		return nullptr;
 	}
 
