-Make manipulation code modify/invalidatepackage for all attached actors in addition the selected one.
-In addition to making sure all modifed attached actors are saved, this also fixes a bunch of issues with the WorldPartitionActorDesc being invalid for attached actors that were moved from a cell to another

diff --git a/Engine/Source/Editor/UnrealEd/Private/Elements/Actor/ActorElementEditorViewportInteractionCustomization.cpp b/Engine/Source/Editor/UnrealEd/Private/Elements/Actor/ActorElementEditorViewportInteractionCustomization.cpp
--- a/Engine/Source/Editor/UnrealEd/Private/Elements/Actor/ActorElementEditorViewportInteractionCustomization.cpp
+++ b/Engine/Source/Editor/UnrealEd/Private/Elements/Actor/ActorElementEditorViewportInteractionCustomization.cpp
@@ -136,9 +136,18 @@ void FActorElementEditorViewportInteractionCustomization::ApplyDeltaToActor(AAct
 {
 	const bool bIsSimulatingInEditor = GEditor->IsSimulatingInEditor();
 
+	TArray<AActor*> AffectedActors;
+	AffectedActors.Add(InActor);
+	
+	// grab all attached actors to invalidate everybody 	
+	InActor->GetAttachedActors(AffectedActors, /*bResetArray*/ false );
+
 	if (GEditor->IsDeltaModificationEnabled())
 	{
-		InActor->Modify();
+		for (AActor* Actor : AffectedActors)
+		{
+			Actor->Modify();
+		}
 	}
 
 	FNavigationLockContext LockNavigationUpdates(InActor->GetWorld(), ENavigationLockReason::ContinuousEditorMove);
@@ -290,10 +299,21 @@ void FActorElementEditorViewportInteractionCustomization::ApplyDeltaToActor(AAct
 	}
 
 	// Update the actor before leaving.
-	InActor->MarkPackageDirty();
+	for (AActor* Actor : AffectedActors)
+	{
+		Actor->MarkPackageDirty();
+	}
+	
 	if (!GIsDemoMode)
 	{
-		InActor->InvalidateLightingCacheDetailed(bTranslationOnly);
+		for (AActor* Actor : AffectedActors)
+		{
+			Actor->InvalidateLightingCacheDetailed(bTranslationOnly);
+		}
+	}
+	
+	for (AActor* Actor : AffectedActors)
+	{
+		Actor->PostEditMove(false);
 	}
-	InActor->PostEditMove(false);
 }
