Improved the UX of the option to enable the interchange framework for the texture import only.

Changed a bit the implementation of UEditorExperimentalSettings::PostEditChangeProperty to make it more resiliant to a name change to one of the class properties.

diff --git a/Engine/Source/Editor/UnrealEd/Private/Settings/SettingsClasses.cpp b/Engine/Source/Editor/UnrealEd/Private/Settings/SettingsClasses.cpp
--- a/Engine/Source/Editor/UnrealEd/Private/Settings/SettingsClasses.cpp
+++ b/Engine/Source/Editor/UnrealEd/Private/Settings/SettingsClasses.cpp
@@ -29,6 +29,7 @@
 #include "CrashReporterSettings.h"
 #include "AutoReimport/AutoReimportUtilities.h"
 #include "Misc/ConfigCacheIni.h" // for FConfigCacheIni::GetString()
+#include "Misc/AssertionMacros.h"
 #include "SourceCodeNavigation.h"
 #include "Interfaces/IProjectManager.h"
 #include "ProjectDescriptor.h"
@@ -201,18 +202,33 @@ void UEditorExperimentalSettings::PostEditChangeProperty( struct FPropertyChange
 
 	const FName Name = (PropertyChangedEvent.Property != nullptr) ? PropertyChangedEvent.Property->GetFName() : NAME_None;
 
-	if (Name == FName(TEXT("ConsoleForGamepadLabels")))
+	if (Name == GET_MEMBER_NAME_CHECKED(UEditorExperimentalSettings, ConsoleForGamepadLabels))
 	{
 		EKeys::SetConsoleForGamepadLabels(ConsoleForGamepadLabels);
 	}
-	else if (Name == FName(TEXT("bHDREditor")))
+	else if (Name == GET_MEMBER_NAME_CHECKED(UEditorExperimentalSettings, bHDREditor))
 	{
 		CVarEditorHDRSupport->Set(bHDREditor ? 1 : 0, ECVF_SetByProjectSetting);
 	}
-	else if (Name == FName(TEXT("HDREditorNITLevel")))
+	else if (Name == GET_MEMBER_NAME_CHECKED(UEditorExperimentalSettings, HDREditorNITLevel))
 	{
 		CVarEditorHDRNITLevel->Set(HDREditorNITLevel, ECVF_SetByProjectSetting);
 	}
+	else if (Name == GET_MEMBER_NAME_CHECKED(UEditorExperimentalSettings, bEnableInterchangeFramework))
+	{
+		if (bEnableInterchangeFramework)
+		{
+			bEnableInterchangeFrameworkForTextureOnly = false;
+		}
+	}
+	else if (Name == GET_MEMBER_NAME_CHECKED(UEditorExperimentalSettings, bEnableInterchangeFrameworkForTextureOnly))
+	{
+		if (bEnableInterchangeFrameworkForTextureOnly)
+		{
+			bEnableInterchangeFramework = false;
+		}
+	}
+
 	if (!FUnrealEdMisc::Get().IsDeletePreferences())
 	{
 		SaveConfig();
