Fix pushing the wrong object in to the replaced references map. Previously the replaced object was stored, but it should be the object that has the property for which the object is replaced.
Removed unnecessary FArchiveReplaceFieldReferences parameters.

diff --git a/Engine/Source/Editor/UnrealEd/Private/Kismet2/KismetReinstanceUtilities.cpp b/Engine/Source/Editor/UnrealEd/Private/Kismet2/KismetReinstanceUtilities.cpp
--- a/Engine/Source/Editor/UnrealEd/Private/Kismet2/KismetReinstanceUtilities.cpp
+++ b/Engine/Source/Editor/UnrealEd/Private/Kismet2/KismetReinstanceUtilities.cpp
@@ -836,38 +836,21 @@ public:
 	 *
 	 * @param InSearchObject		The object to start the search on
 	 * @param ReplacementMap		Map of objects to find -> objects to replace them with (null zeros them)
-	 * @param bNullPrivateRefs		Whether references to non-public objects not contained within the SearchObject
-	 *								should be set to null
-	 * @param bIgnoreOuterRef		Whether we should replace Outer pointers on Objects.
-	 * @param bIgnoreArchetypeRef	Whether we should replace the ObjectArchetype reference on Objects.
-	 * @param bDelayStart			Specify true to prevent the constructor from starting the process.  Allows child classes' to do initialization stuff in their ctor
 	 */
-	FArchiveReplaceFieldReferences
-	(
-		UObject* InSearchObject,
-		const TMap<FFieldVariant, FFieldVariant>& InReplacementMap,
-		bool bNullPrivateRefs,
-		bool bIgnoreOuterRef,
-		bool bIgnoreArchetypeRef,
-		bool bDelayStart = false,
-		bool bIgnoreClassGeneratedByRef = true
-	)
+	FArchiveReplaceFieldReferences(UObject* InSearchObject, const TMap<FFieldVariant, FFieldVariant>& InReplacementMap)
 		: ReplacementMap(InReplacementMap)
 	{
 		SearchObject = InSearchObject;
 		Count = 0;
-		bNullPrivateReferences = bNullPrivateRefs;
+		bNullPrivateReferences = false;
 
 		ArIsObjectReferenceCollector = true;
 		ArIsModifyingWeakAndStrongReferences = true;		// Also replace weak references too!
-		ArIgnoreArchetypeRef = bIgnoreArchetypeRef;
-		ArIgnoreOuterRef = bIgnoreOuterRef;
-		ArIgnoreClassGeneratedByRef = bIgnoreClassGeneratedByRef;
+		ArIgnoreArchetypeRef = true;
+		ArIgnoreOuterRef = true;
+		ArIgnoreClassGeneratedByRef = true;
 
-		if (!bDelayStart)
-		{
-			SerializeSearchObject();
-		}
+		SerializeSearchObject();
 	}
 
 	/**
@@ -875,17 +858,19 @@ public:
 	 */
 	void SerializeSearchObject()
 	{
-		ReplacedReferences.Empty();
+		ReplacedReferences.Reset();
 
 		if (SearchObject != NULL && !SerializedObjects.Find(SearchObject)
 			&& (ReplacementMap.Num() > 0 || bNullPrivateReferences))
 		{
 			// start the initial serialization
 			SerializedObjects.Add(SearchObject);
+			SerializingObject = SearchObject;
 			SerializeObject(SearchObject);
 			for (int32 Iter = 0; Iter < PendingSerializationObjects.Num(); Iter++)
 			{
-				SerializeObject(PendingSerializationObjects[Iter]);
+				SerializingObject = PendingSerializationObjects[Iter];
+				SerializeObject(SerializingObject);
 			}
 			PendingSerializationObjects.Reset();
 		}
@@ -905,7 +890,7 @@ public:
 				Obj = ReplaceWith->ToUObject();
 				if (bTrackReplacedReferences)
 				{
-					ReplacedReferences.FindOrAdd(Obj).AddUnique(GetSerializedProperty());
+					ReplacedReferences.FindOrAdd(SerializingObject).AddUnique(GetSerializedProperty());
 				}
 				Count++;
 			}
@@ -1037,7 +1022,7 @@ void FBlueprintCompileReinstancer::UpdateBytecodeReferences()
 			for( TFieldIterator<UFunction> FuncIter(BPClass, EFieldIteratorFlags::ExcludeSuper); FuncIter; ++FuncIter )
 			{
 				UFunction* CurrentFunction = *FuncIter;
-				FArchiveReplaceFieldReferences ReplaceAr(CurrentFunction, FieldMappings, /*bNullPrivateRefs=*/ false, /*bIgnoreOuterRef=*/ true, /*bIgnoreArchetypeRef=*/ true);
+				FArchiveReplaceFieldReferences ReplaceAr(CurrentFunction, FieldMappings);
 				bBPWasChanged |= (0 != ReplaceAr.GetCount());
 			}
 
@@ -1059,7 +1044,7 @@ void FBlueprintCompileReinstancer::UpdateBytecodeReferences()
 				}
 			}
 
-			FArchiveReplaceFieldReferences ReplaceInBPAr(DependentBP, FieldMappings, false, true, true);
+			FArchiveReplaceFieldReferences ReplaceInBPAr(DependentBP, FieldMappings);
 			if (ReplaceInBPAr.GetCount())
 			{
 #if VALIDATE_BYTECODE_REFERENCE_DEPENDENCY_CACHE
diff --git a/Engine/Source/Runtime/CoreUObject/Public/Serialization/ArchiveReplaceObjectRef.h b/Engine/Source/Runtime/CoreUObject/Public/Serialization/ArchiveReplaceObjectRef.h
--- a/Engine/Source/Runtime/CoreUObject/Public/Serialization/ArchiveReplaceObjectRef.h
+++ b/Engine/Source/Runtime/CoreUObject/Public/Serialization/ArchiveReplaceObjectRef.h
@@ -98,6 +98,9 @@ protected:
 	/** Initial object to start the reference search from */
 	UObject* SearchObject = nullptr;
 
+	/** Object that SerializeObject was most recently called on */
+	UObject* SerializingObject = nullptr;
+
 	/** The number of times encountered */
 	int32 Count = 0;
 
@@ -217,17 +220,19 @@ public:
 	 */
 	void SerializeSearchObject()
 	{
-		ReplacedReferences.Empty();
+		ReplacedReferences.Reset();
 
 		if (SearchObject != NULL && !SerializedObjects.Find(SearchObject)
 		&&	(ReplacementMap.Num() > 0 || bNullPrivateReferences))
 		{
 			// start the initial serialization
 			SerializedObjects.Add(SearchObject);
+			SerializingObject = SearchObject;
 			SerializeObject(SearchObject);
 			for (int32 Iter = 0; Iter < PendingSerializationObjects.Num(); Iter++)
 			{
-				SerializeObject(PendingSerializationObjects[Iter]);
+				SerializingObject = PendingSerializationObjects[Iter];
+				SerializeObject(SerializingObject);
 			}
 			PendingSerializationObjects.Reset();
 		}
@@ -246,7 +251,7 @@ public:
 				Obj = *ReplaceWith;
 				if (bTrackReplacedReferences)
 				{
-					ReplacedReferences.FindOrAdd(Obj).AddUnique(GetSerializedProperty());
+					ReplacedReferences.FindOrAdd(SerializingObject).AddUnique(GetSerializedProperty());
 				}
 				Count++;
 			}
