Cooker: Add fallback to looking for the package on disk if it does not exist in the asset registry.
#rnx

diff --git a/Engine/Source/Editor/UnrealEd/Private/Cooker/CookPackageData.cpp b/Engine/Source/Editor/UnrealEd/Private/Cooker/CookPackageData.cpp
--- a/Engine/Source/Editor/UnrealEd/Private/Cooker/CookPackageData.cpp
+++ b/Engine/Source/Editor/UnrealEd/Private/Cooker/CookPackageData.cpp
@@ -1842,6 +1842,11 @@ bool FPackageDatas::TryLookupFileNameOnDisk(FName PackageName, FString& OutFileN
 
 		if (Assets.Num() <= 0)
 		{
+			if (FPackageName::DoesPackageExist(PackageNameStr, &OutFileName, false /* InAllowTextFormats */))
+			{
+				UE_LOG(LogCook, Warning, TEXT("Package %s exists on disk but does not exist in the AssetRegistry"), *PackageNameStr);
+				return true;
+			}
 			return false;
 		}
 
