Fixed some notification issues with the type element selection set.
The mecanism to select something wihtout emiting a notification was clearing the existing pending notifications and it wasn't interacting properly with the old bacth notification system.

Also fixed a batch selection notification issue when unselecting all the selected actors from the scene outliner.

diff --git a/Engine/Source/Editor/UnrealEd/Private/EditorSelectUtils.cpp b/Engine/Source/Editor/UnrealEd/Private/EditorSelectUtils.cpp
--- a/Engine/Source/Editor/UnrealEd/Private/EditorSelectUtils.cpp
+++ b/Engine/Source/Editor/UnrealEd/Private/EditorSelectUtils.cpp
@@ -418,6 +418,12 @@ void UUnrealEdEngine::SelectGroup(AGroupActor* InGroupActor, bool bForceSelectio
 
 	if (UTypedElementSelectionSet* SelectionSet = ActorSelection->GetElementSelectionSet())
 	{
+		FTypedElementList::FScopedClearNewPendingChange ClearNewPendingChange;
+		if (!bNotify)
+		{
+			ClearNewPendingChange = SelectionSet->GetScopedClearNewPendingChange();
+		}
+
 		const FTypedElementSelectionOptions SelectionOptions = FTypedElementSelectionOptions()
 			.SetWarnIfLocked(true)
 			.SetAllowGroups(false)
@@ -455,10 +461,6 @@ void UUnrealEdEngine::SelectGroup(AGroupActor* InGroupActor, bool bForceSelectio
 			{
 				SelectionSet->NotifyPendingChanges();
 			}
-			else
-			{
-				SelectionSet->ClearPendingChanges();
-			}
 		}
 	}
 }
@@ -497,6 +499,12 @@ void UUnrealEdEngine::SelectActor(AActor* Actor, bool bInSelected, bool bNotify,
 
 	if (UTypedElementSelectionSet* SelectionSet = ActorSelection->GetElementSelectionSet())
 	{
+		FTypedElementList::FScopedClearNewPendingChange ClearNewPendingChange;
+		if (!bNotify)
+		{
+			ClearNewPendingChange = SelectionSet->GetScopedClearNewPendingChange();
+		}
+
 		const FTypedElementSelectionOptions SelectionOptions = FTypedElementSelectionOptions()
 			.SetAllowHidden(bSelectEvenIfHidden)
 			.SetWarnIfLocked(true)
@@ -512,10 +520,6 @@ void UUnrealEdEngine::SelectActor(AActor* Actor, bool bInSelected, bool bNotify,
 			{
 				SelectionSet->NotifyPendingChanges();
 			}
-			else
-			{
-				SelectionSet->ClearPendingChanges();
-			}
 		}
 		else if (bNotify || bForceRefresh)
 		{
@@ -536,6 +540,12 @@ void UUnrealEdEngine::SelectComponent(UActorComponent* Component, bool bInSelect
 
 	if (UTypedElementSelectionSet* SelectionSet = ComponentSelection->GetElementSelectionSet())
 	{
+		FTypedElementList::FScopedClearNewPendingChange ClearNewPendingChange;
+		if (!bNotify)
+		{
+			ClearNewPendingChange = SelectionSet->GetScopedClearNewPendingChange();
+		}
+
 		const FTypedElementSelectionOptions SelectionOptions = FTypedElementSelectionOptions()
 			.SetAllowHidden(bSelectEvenIfHidden)
 			.SetWarnIfLocked(true)
@@ -551,10 +561,6 @@ void UUnrealEdEngine::SelectComponent(UActorComponent* Component, bool bInSelect
 			{
 				SelectionSet->NotifyPendingChanges();
 			}
-			else
-			{
-				SelectionSet->ClearPendingChanges();
-			}
 		}
 	}
 }
@@ -667,8 +673,14 @@ void UUnrealEdEngine::SelectNone(bool bNoteSelectionChange, bool bDeselectBSPSur
 	USelection* ActorSelection = GetSelectedActors();
 	UTypedElementSelectionSet* SelectionSet = ActorSelection->GetElementSelectionSet();
 
+	FTypedElementList::FScopedClearNewPendingChange ClearNewPendingChange;
+
 	if (SelectionSet)
 	{
+		if (!bNoteSelectionChange)
+		{ 
+			ClearNewPendingChange = SelectionSet->GetScopedClearNewPendingChange();
+		}
 		bSelectionChanged |= SelectionSet->ClearSelection(FTypedElementSelectionOptions().SetAllowLegacyNotifications(false));
 	}
 
@@ -687,10 +699,6 @@ void UUnrealEdEngine::SelectNone(bool bNoteSelectionChange, bool bDeselectBSPSur
 			{
 				SelectionSet->NotifyPendingChanges();
 			}
-			else
-			{
-				SelectionSet->ClearPendingChanges();
-			}
 		}
 	}
 }
diff --git a/Engine/Source/Editor/UnrealEd/Private/EditorServer.cpp b/Engine/Source/Editor/UnrealEd/Private/EditorServer.cpp
--- a/Engine/Source/Editor/UnrealEd/Private/EditorServer.cpp
+++ b/Engine/Source/Editor/UnrealEd/Private/EditorServer.cpp
@@ -3660,7 +3660,9 @@ void UEditorEngine::PasteSelectedActorsFromClipboard( UWorld* InWorld, const FTe
 
 		const FScopedTransaction Transaction( TransDescription );
 
-		SelectNone( true, false );
+		const bool bNoteSelectionChange = false;
+		const bool bDeselectBSPSurfs = false;
+		SelectNone( bNoteSelectionChange, bDeselectBSPSurfs );
 		ABrush::SetSuppressBSPRegeneration(true);
 		edactPasteSelected( InWorld, false, false, true );
 		ABrush::SetSuppressBSPRegeneration(false);
diff --git a/Engine/Source/Runtime/TypedElementFramework/Private/Elements/Framework/TypedElementList.cpp b/Engine/Source/Runtime/TypedElementFramework/Private/Elements/Framework/TypedElementList.cpp
--- a/Engine/Source/Runtime/TypedElementFramework/Private/Elements/Framework/TypedElementList.cpp
+++ b/Engine/Source/Runtime/TypedElementFramework/Private/Elements/Framework/TypedElementList.cpp
@@ -364,9 +364,47 @@ bool FTypedElementList::NotifyPendingChanges()
 	return false;
 }
 
-void FTypedElementList::ClearPendingChanges()
+FTypedElementList::FScopedClearNewPendingChange::FScopedClearNewPendingChange(FTypedElementList& InTypeElementList)
 {
-	bHasPendingNotify = false;
+	bool bCanClearNewPendingChange = !InTypeElementList.bHasPendingNotify;
+
+	if (InTypeElementList.LegacySync)
+	{
+		bCanClearNewPendingChange &= !InTypeElementList.LegacySync->IsRunningBatchOperation();
+	}
+	
+	if (bCanClearNewPendingChange)
+	{
+		TypedElementList = &InTypeElementList;
+	}
+}
+
+
+FTypedElementList::FScopedClearNewPendingChange::FScopedClearNewPendingChange(FScopedClearNewPendingChange&& Other)
+	: TypedElementList(Other.TypedElementList)
+{
+	Other.TypedElementList = nullptr;
+}
+
+FTypedElementList::FScopedClearNewPendingChange& FTypedElementList::FScopedClearNewPendingChange::operator=(FScopedClearNewPendingChange&& Other)
+{
+	TypedElementList = Other.TypedElementList;
+	Other.TypedElementList = nullptr;
+
+	return *this;
+}
+
+FTypedElementList::FScopedClearNewPendingChange::~FScopedClearNewPendingChange()
+{
+	if (TypedElementList)
+	{
+		TypedElementList->bHasPendingNotify = false;
+	}
+}
+
+FTypedElementList::FScopedClearNewPendingChange FTypedElementList::GetScopedClearNewPendingChange()
+{
+	return FScopedClearNewPendingChange(*this);
 }
 
 void FTypedElementList::NoteListMayChange()
diff --git a/Engine/Source/Runtime/TypedElementFramework/Public/Elements/Framework/TypedElementList.h b/Engine/Source/Runtime/TypedElementFramework/Public/Elements/Framework/TypedElementList.h
--- a/Engine/Source/Runtime/TypedElementFramework/Public/Elements/Framework/TypedElementList.h
+++ b/Engine/Source/Runtime/TypedElementFramework/Public/Elements/Framework/TypedElementList.h
@@ -400,11 +400,22 @@ public:
 	 */
 	FORCEINLINE void Empty(const int32 InSlack = 0)
 	{
-		NoteListMayChange();
+		// Avoid creating unnecessary notifications
+		const bool bWasEmpty = ElementHandles.IsEmpty();
+		if (!bWasEmpty)
+		{
+			NoteListMayChange();
+		}
+
 		ElementCombinedIds.Empty(InSlack);
 		ElementHandles.Empty(InSlack);
 		ElementCounts.ClearCounters();
-		NoteListChanged(EChangeType::Cleared);
+
+		if (!bWasEmpty)
+		{
+			ElementCounts.ClearCounters();
+			NoteListChanged(EChangeType::Cleared);
+		}
 	}
 
 	/**
@@ -412,11 +423,15 @@ public:
 	 */
 	FORCEINLINE void Reset()
 	{
-		NoteListMayChange();
-		ElementCombinedIds.Reset();
-		ElementHandles.Reset();
-		ElementCounts.ClearCounters();
-		NoteListChanged(EChangeType::Cleared);
+		// Avoid creating unnecessary notifications
+		if (!ElementHandles.IsEmpty())
+		{
+			NoteListMayChange();
+			ElementCombinedIds.Reset();
+			ElementHandles.Reset();
+			ElementCounts.ClearCounters();
+			NoteListChanged(EChangeType::Cleared);
+		}
 	}
 
 	/**
@@ -615,9 +630,30 @@ public:
 	bool NotifyPendingChanges();
 
 	/**
-	 * Clear whether there are pending changes for OnChangedDelegate to notify for, without emitting a notification.
+	 * A utility struct that help to cancel any new pending notification that happened in a scope.
+	 * Note: it won't cancel a notification if there is a legacy batch operation ongoing
+	 */
+	struct TYPEDELEMENTFRAMEWORK_API FScopedClearNewPendingChange
+	{
+		FScopedClearNewPendingChange() = default;
+		FScopedClearNewPendingChange(FTypedElementList& InTypeElementList);
+
+		FScopedClearNewPendingChange(const FScopedClearNewPendingChange&) = delete;
+		FScopedClearNewPendingChange& operator=(const FScopedClearNewPendingChange&) = delete;
+
+		FScopedClearNewPendingChange(FScopedClearNewPendingChange&& Other);
+		FScopedClearNewPendingChange& operator=(FScopedClearNewPendingChange&& Other);
+
+		~FScopedClearNewPendingChange();
+
+	private:
+		FTypedElementList* TypedElementList = nullptr;
+	};
+
+	/**
+	 * Get a scoped object that when destroyed it clear a pending change notification without emitting the notification if it happened during its lifecycle.
 	 */
-	void ClearPendingChanges();
+	FScopedClearNewPendingChange GetScopedClearNewPendingChange();
 
 	/**
 	 * Access the interface to allow external systems (such as USelection) to receive immediate sync notifications as an element list is changed.
diff --git a/Engine/Source/Runtime/TypedElementRuntime/Public/Elements/Framework/TypedElementSelectionSet.h b/Engine/Source/Runtime/TypedElementRuntime/Public/Elements/Framework/TypedElementSelectionSet.h
--- a/Engine/Source/Runtime/TypedElementRuntime/Public/Elements/Framework/TypedElementSelectionSet.h
+++ b/Engine/Source/Runtime/TypedElementRuntime/Public/Elements/Framework/TypedElementSelectionSet.h
@@ -495,11 +495,11 @@ public:
 	}
 
 	/**
-	 * Clear whether there are pending changes for OnChangedDelegate to notify for, without emitting a notification.
+	 * Get a scoped object that when destroyed it clear a pending change notification without emitting the notification if it happened during its lifecycle.
 	 */
-	void ClearPendingChanges()
+	FTypedElementList::FScopedClearNewPendingChange GetScopedClearNewPendingChange()
 	{
-		ElementList->ClearPendingChanges();
+		return ElementList->GetScopedClearNewPendingChange();
 	}
 
 	/**
