WorldPartition SaveAs: Remove need for Uninitialize/Duplicate/Initialize by rebasing the WorldPartition in UWorld::Rename

diff --git a/Engine/Source/Editor/UnrealEd/Private/FileHelpers.cpp b/Engine/Source/Editor/UnrealEd/Private/FileHelpers.cpp
--- a/Engine/Source/Editor/UnrealEd/Private/FileHelpers.cpp
+++ b/Engine/Source/Editor/UnrealEd/Private/FileHelpers.cpp
@@ -847,19 +847,6 @@ static bool SaveWorld(UWorld* World,
 			{
 				bSuccess = SaveExternalPackages(Package, /*bCheckDirty*/true);
 			}
-			// If we are dealing with a renamed partitioned world we need to uninit the existing partition as it is pointing to the wrong path and also make sure 
-			// we create a duplicate so that any code pointing to the old partition doesn't get broken until we send out a OnWorldPartitionCreated later in this method.
-			if (bSuccess && bRenamedWorldPartition && !DuplicatedWorld)
-			{
-				UWorldPartition* WorldPartition = SaveWorld->GetWorldPartition();
-				check(WorldPartition);
-				// If we don't have a DuplicatedWorld, create a new world partition object to invalidate any existing pointers to the previous object.
-				WorldPartition->Uninitialize();
-				UWorldPartition* DuplicatedPartition = DuplicateObject<UWorldPartition>(WorldPartition, WorldPartition->GetOuter());
-				WorldPartition->MarkAsGarbage();
-
-				SaveWorld->GetWorldSettings()->SetWorldPartition(DuplicatedPartition);
-			}
 		}
 				
 		SlowTask.EnterProgressFrame(50);
@@ -891,18 +878,6 @@ static bool SaveWorld(UWorld* World,
 
 				// Save Snapshot of loaded Editor Cells
 				GetMutableDefault<UWorldPartitionEditorPerProjectUserSettings>()->SetEditorGridLoadedCells(SaveWorld, LoadedEditorCells);
-
-				// No need to initialize if we have a DuplicatedWorld since the map will get loaded by a subsequent LoadMap call
-				if (!DuplicatedWorld)
-				{
-					UWorldPartition* WorldPartition = SaveWorld->GetWorldPartition();
-					check(WorldPartition);
-					WorldPartition->Initialize(SaveWorld, FTransform::Identity);
-					check(WorldPartition->GetContainerPackage() == FName(*NewPackageName));
-
-					IWorldPartitionEditorModule& WorldPartitionEditorModule = FModuleManager::LoadModuleChecked<IWorldPartitionEditorModule>("WorldPartitionEditor");
-					WorldPartitionEditorModule.OnWorldPartitionCreated().Broadcast(SaveWorld);
-				}
 			}
 		}
 
diff --git a/Engine/Source/Runtime/Engine/Classes/Engine/World.h b/Engine/Source/Runtime/Engine/Classes/Engine/World.h
--- a/Engine/Source/Runtime/Engine/Classes/Engine/World.h
+++ b/Engine/Source/Runtime/Engine/Classes/Engine/World.h
@@ -4045,7 +4045,8 @@ public:
 	DECLARE_MULTICAST_DELEGATE_FourParams(FWorldPostDuplicateEvent, UWorld* /*World*/, bool /*bDuplicateForPIE*/, FReplacementMap& /*ReplacementMap*/, TArray<UObject*>& /*ObjectsToFixReferences*/);
 
 #if WITH_EDITOR
-	DECLARE_MULTICAST_DELEGATE_FiveParams(FWorldRenameEvent, UWorld* /*World*/, const TCHAR* /*InName*/, UObject* /*NewOuter*/, ERenameFlags /*Flags*/, bool& /*bShouldFailRename*/);
+	DECLARE_MULTICAST_DELEGATE_FiveParams(FWorldPreRenameEvent, UWorld* /*World*/, const TCHAR* /*InName*/, UObject* /*NewOuter*/, ERenameFlags /*Flags*/, bool& /*bShouldFailRename*/);
+	DECLARE_MULTICAST_DELEGATE_OneParam(FWorldPostRenameEvent, UWorld*);
 #endif // WITH_EDITOR
 
 	// Delegate type for level change events
@@ -4078,7 +4079,10 @@ public:
 
 #if WITH_EDITOR
 	// Callback for world rename event (pre)
-	static FWorldRenameEvent OnPreWorldRename;
+	static FWorldPreRenameEvent OnPreWorldRename;
+
+	// Callback for world rename event (post)
+	static FWorldPostRenameEvent OnPostWorldRename;
 #endif // WITH_EDITOR
 
 	// Post duplication event.
diff --git a/Engine/Source/Runtime/Engine/Private/World.cpp b/Engine/Source/Runtime/Engine/Private/World.cpp
--- a/Engine/Source/Runtime/Engine/Private/World.cpp
+++ b/Engine/Source/Runtime/Engine/Private/World.cpp
@@ -335,7 +335,8 @@ FWorldDelegates::FWorldEvent FWorldDelegates::OnPostWorldCreation;
 FWorldDelegates::FWorldInitializationEvent FWorldDelegates::OnPreWorldInitialization;
 FWorldDelegates::FWorldInitializationEvent FWorldDelegates::OnPostWorldInitialization;
 #if WITH_EDITOR
-FWorldDelegates::FWorldRenameEvent FWorldDelegates::OnPreWorldRename;
+FWorldDelegates::FWorldPreRenameEvent FWorldDelegates::OnPreWorldRename;
+FWorldDelegates::FWorldPostRenameEvent FWorldDelegates::OnPostWorldRename;
 #endif // WITH_EDITOR
 FWorldDelegates::FWorldPostDuplicateEvent FWorldDelegates::OnPostDuplicate;
 FWorldDelegates::FWorldCleanupEvent FWorldDelegates::OnWorldCleanup;
@@ -759,6 +760,11 @@ bool UWorld::Rename(const TCHAR* InName, UObject* NewOuter, ERenameFlags Flags)
 	}
 	
 
+	if (!bTestRename)
+	{
+		FWorldDelegates::OnPostWorldRename.Broadcast(this);
+	}
+
 	return true;
 }
 #endif
diff --git a/Engine/Source/Runtime/Engine/Private/WorldPartition/ActorDescContainer.cpp b/Engine/Source/Runtime/Engine/Private/WorldPartition/ActorDescContainer.cpp
--- a/Engine/Source/Runtime/Engine/Private/WorldPartition/ActorDescContainer.cpp
+++ b/Engine/Source/Runtime/Engine/Private/WorldPartition/ActorDescContainer.cpp
@@ -9,6 +9,7 @@
 #include "Misc/Base64.h"
 #include "UObject/ObjectSaveContext.h"
 #include "UObject/CoreRedirects.h"
+#include "Engine/World.h"
 #endif
 
 UActorDescContainer::UActorDescContainer(const FObjectInitializer& ObjectInitializer)
@@ -162,6 +163,16 @@ void UActorDescContainer::BeginDestroy()
 }
 
 #if WITH_EDITOR
+
+void UActorDescContainer::OnWorldRenamed(UWorld* RenamedWorld)
+{
+	if (GetWorld() == RenamedWorld)
+	{
+		// Update container package
+		ContainerPackageName = RenamedWorld->GetPackage()->GetFName();
+	}
+}
+
 bool UActorDescContainer::ShouldHandleActorEvent(const AActor* Actor)
 {
 	return Actor && Actor->IsMainPackageActor() && (Actor->GetLevel() != nullptr) && (Actor->GetLevel()->GetPackage()->GetFName() == ContainerPackageName);
@@ -299,6 +310,7 @@ void UActorDescContainer::RegisterEditorDelegates()
 {
 	if (GEditor && !IsTemplate() && World && !World->IsGameWorld())
 	{
+		FWorldDelegates::OnPostWorldRename.AddUObject(this, &UActorDescContainer::OnWorldRenamed);
 		FCoreUObjectDelegates::OnObjectPreSave.AddUObject(this, &UActorDescContainer::OnObjectPreSave);
 		FEditorDelegates::OnPackageDeleted.AddUObject(this, &UActorDescContainer::OnPackageDeleted);
 		FCoreUObjectDelegates::OnObjectsReplaced.AddUObject(this, &UActorDescContainer::OnObjectsReplaced);
@@ -309,6 +321,7 @@ void UActorDescContainer::UnregisterEditorDelegates()
 {
 	if (GEditor && !IsTemplate() && World && !World->IsGameWorld())
 	{
+		FWorldDelegates::OnPostWorldRename.RemoveAll(this);
 		FCoreUObjectDelegates::OnObjectPreSave.RemoveAll(this);
 		FEditorDelegates::OnPackageDeleted.RemoveAll(this);
 		FCoreUObjectDelegates::OnObjectsReplaced.RemoveAll(this);
diff --git a/Engine/Source/Runtime/Engine/Public/WorldPartition/ActorDescContainer.h b/Engine/Source/Runtime/Engine/Public/WorldPartition/ActorDescContainer.h
--- a/Engine/Source/Runtime/Engine/Public/WorldPartition/ActorDescContainer.h
+++ b/Engine/Source/Runtime/Engine/Public/WorldPartition/ActorDescContainer.h
@@ -25,6 +25,7 @@ public:
 
 #if WITH_EDITOR
 	// Events
+	virtual void OnWorldRenamed(UWorld* RenamedWorld);
 	virtual void OnObjectPreSave(UObject* Object, FObjectPreSaveContext SaveContext);
 	virtual void OnPackageDeleted(UPackage* Package);
 	virtual void OnObjectsReplaced(const TMap<UObject*, UObject*>& OldToNewObjectMap);
