Don't reinitialize world subsystems when doing a map rebuild as we don't uninitialize them anymore.
This is essentially a backout of 11776301 which was invalidated by 13063957.

diff --git a/Engine/Source/Editor/UnrealEd/Private/EditorServer.cpp b/Engine/Source/Editor/UnrealEd/Private/EditorServer.cpp
--- a/Engine/Source/Editor/UnrealEd/Private/EditorServer.cpp
+++ b/Engine/Source/Editor/UnrealEd/Private/EditorServer.cpp
@@ -1706,9 +1706,6 @@ void UEditorEngine::RebuildMap(UWorld* InWorld, EMapRebuildType RebuildType)
 	FEditorDelegates::MapChange.Broadcast(MapChangeEventFlags::MapRebuild);
 	GEngine->BroadcastLevelActorListChanged();
 	
-	// Need to reinitialize world subsystems since they are torn down as part of the lighting build
-	GWorld->InitializeSubsystems();
-
 	GWarn->EndSlowTask();
 }
 
