Don't try to change spatially loaded flag on actors that don't allow changing it.

diff --git a/Engine/Source/Editor/UnrealEd/Private/Commandlets/WorldPartitionConvertCommandlet.cpp b/Engine/Source/Editor/UnrealEd/Private/Commandlets/WorldPartitionConvertCommandlet.cpp
--- a/Engine/Source/Editor/UnrealEd/Private/Commandlets/WorldPartitionConvertCommandlet.cpp
+++ b/Engine/Source/Editor/UnrealEd/Private/Commandlets/WorldPartitionConvertCommandlet.cpp
@@ -286,7 +286,7 @@ bool UWorldPartitionConvertCommandlet::PrepareStreamingLevelForConversion(ULevel
 
 		for (AActor* Actor: SubLevel->Actors)
 		{
-			if (Actor)
+			if (Actor && Actor->CanChangeIsSpatiallyLoadedFlag())
 			{
 				Actor->SetIsSpatiallyLoaded(false);
 			}
@@ -978,7 +978,7 @@ int32 UWorldPartitionConvertCommandlet::Main(const FString& Params)
 						LandscapeInfos.Add(LandscapeInfo);
 					}
 					// Only override default grid placement on actors that are spatially loaded
-					else if (Actor->GetIsSpatiallyLoaded())
+					else if (Actor->GetIsSpatiallyLoaded() && Actor->CanChangeIsSpatiallyLoadedFlag())
 					{
 						const FBox ActorBounds = Actor->GetStreamingBounds();
 
