Fix Linux compile error

diff --git a/Engine/Source/Editor/UnrealEd/Private/Kismet2/KismetDebugUtilities.cpp b/Engine/Source/Editor/UnrealEd/Private/Kismet2/KismetDebugUtilities.cpp
--- a/Engine/Source/Editor/UnrealEd/Private/Kismet2/KismetDebugUtilities.cpp
+++ b/Engine/Source/Editor/UnrealEd/Private/Kismet2/KismetDebugUtilities.cpp
@@ -2215,9 +2215,9 @@ TSharedPtr<FPropertyInstanceInfo> FPropertyInstanceInfo::ResolvePathToProperty(c
 		if (InPropertyInfo)
 		{
 			const FString ChildNameStr = InChildName.ToString();
-			FProperty* Property = InPropertyInfo->Property.Get();
+			FProperty* Prop = InPropertyInfo->Property.Get();
 
-			if (Property->IsA<FSetProperty>() || Property->IsA<FArrayProperty>() || Property->IsA<FMapProperty>())
+			if (Prop->IsA<FSetProperty>() || Prop->IsA<FArrayProperty>() || Prop->IsA<FMapProperty>())
 			{
 				return InPropertyInfo->Children.FindByPredicate(
 					[&ChildNameStr](const TSharedPtr<FPropertyInstanceInfo>& Child)
