Add non-imported class type filtering to the "Add Component" type selection UI in the Blueprint editor.

Additionally auto-import any namespace associated with the selected component class type when a new component is added.

Change summary:
- Modified the component type registry to internally create/store unloaded BP data to assist with unloaded BP component type filtering.
- Added FEditorClassUtils::GetImplementedInterfaceClassPathsFromAsset() to move BP interface asset tag processing into a shared utility API.
- Modified SComponentClassCombo to work with IClassViewFilter-style filtering and added a filter options menu (currently only visible when namespace filtering features are enabled via CVar).
- A few minor modifications/additions to SSubobjectEditor/SSubobjectBlueprintEditor to assist with interfacing with the owning FBlueprintEditor context for filtering, updating etc.

diff --git a/Engine/Source/Editor/ClassViewer/Private/SClassViewer.cpp b/Engine/Source/Editor/ClassViewer/Private/SClassViewer.cpp
--- a/Engine/Source/Editor/ClassViewer/Private/SClassViewer.cpp
+++ b/Engine/Source/Editor/ClassViewer/Private/SClassViewer.cpp
@@ -1362,36 +1362,12 @@ void FClassHierarchy::SetAssetDataFields(TSharedPtr<FClassViewerNode>& InOutClas
 	const uint32 ClassFlags = InAssetData.GetTagValueRef<uint32>(FBlueprintTags::ClassFlags);
 	InOutClassViewerNode->UnloadedBlueprintData->SetClassFlags(ClassFlags);
 
-	const FString ImplementedInterfaces = InAssetData.GetTagValueRef<FString>(FBlueprintTags::ImplementedInterfaces);
-	if(!ImplementedInterfaces.IsEmpty())
+	// Get interface class paths.
+	TArray<FString> ImplementedInterfaces;
+	FEditorClassUtils::GetImplementedInterfaceClassPathsFromAsset(InAssetData, ImplementedInterfaces);
+	for (const FString& InterfacePath : ImplementedInterfaces)
 	{
-		// Parse string like "((Interface=Class'"/Script/VPBookmark.VPBookmarkProvider"'),(Interface=Class'"/Script/VPUtilities.VPContextMenuProvider"'))"
-		// We don't want to actually resolve the hard ref so do some manual parsing
-
-		FString FullInterface;
-		FString CurrentString = *ImplementedInterfaces;
-		while(CurrentString.Split(TEXT("Interface="), nullptr, &FullInterface))
-		{
-			// Cutoff at next )
-			int32 RightParen = INDEX_NONE;			
-			if (FullInterface.FindChar(TCHAR(')'), RightParen))
-			{
-				// Keep parsing
-				CurrentString = FullInterface.Mid(RightParen);
-
-				// Strip class name
-				FullInterface = *FPackageName::ExportTextPathToObjectPath(FullInterface.Left(RightParen));
-
-				// Handle quotes
-				FString InterfacePath;
-				const TCHAR* NewBuffer = FPropertyHelpers::ReadToken(*FullInterface, InterfacePath, true);
-
-				if (NewBuffer)
-				{
-					UnloadedBlueprintData->AddImplementedInterface(InterfacePath);
-				}
-			}
-		}
+		UnloadedBlueprintData->AddImplementedInterface(InterfacePath);
 	}
 }
 
@@ -1524,8 +1500,8 @@ void SClassViewer::Construct(const FArguments& InArgs, const FClassViewerInitial
 		CustomFilter->GetFilterOptions(FilterOptions);
 		CustomClassFilterOptions.Append(FilterOptions);
 
-		// Clear out the temp array for the next pass.
-		FilterOptions.Empty();
+		// Reset the temp array for the next pass.
+		FilterOptions.Reset();
 	}
 
 	TSharedRef<SWidget> FiltersWidget = SNullWidget::NullWidget;
@@ -2051,8 +2027,6 @@ TSharedRef<SWidget> SClassViewer::GetViewButtonContent()
 	MenuBuilder.EndSection();
 
 	return MenuBuilder.MakeWidget();
-
-	return MenuBuilder.MakeWidget();
 }
 
 void SClassViewer::SetCurrentDeveloperViewType(EClassViewerDeveloperType NewType)
diff --git a/Engine/Source/Editor/Kismet/Private/BlueprintEditor.cpp b/Engine/Source/Editor/Kismet/Private/BlueprintEditor.cpp
--- a/Engine/Source/Editor/Kismet/Private/BlueprintEditor.cpp
+++ b/Engine/Source/Editor/Kismet/Private/BlueprintEditor.cpp
@@ -298,13 +298,16 @@ namespace BlueprintEditorImpl
 				return;
 			}
 
-			TSharedRef<FClassViewerFilterOption> FilterOption = MakeShared<FClassViewerFilterOption>();
-			FilterOption->bEnabled = bIsFilterEnabled;
-			FilterOption->LabelText = LOCTEXT("ClassViewerNamespaceFilterMenuOptionLabel", "Show Only Imported Types");
-			FilterOption->ToolTipText = LOCTEXT("ClassViewerNamespaceFilterMenuOptionToolTip", "Don't include non-imported class types.");
-			FilterOption->OnOptionChanged = FOnClassViewerFilterOptionChanged::CreateSP(this, &FImportedClassViewerFilter::OnFilterOptionChanged);
+			if (!ToggleFilterOption.IsValid())
+			{
+				ToggleFilterOption = MakeShared<FClassViewerFilterOption>();
+				ToggleFilterOption->bEnabled = bIsFilterEnabled;
+				ToggleFilterOption->LabelText = LOCTEXT("ClassViewerNamespaceFilterMenuOptionLabel", "Show Only Imported Types");
+				ToggleFilterOption->ToolTipText = LOCTEXT("ClassViewerNamespaceFilterMenuOptionToolTip", "Don't include non-imported class types.");
+				ToggleFilterOption->OnOptionChanged = FOnClassViewerFilterOptionChanged::CreateSP(this, &FImportedClassViewerFilter::OnFilterOptionChanged);
+			}
 
-			OutFilterOptions.Add(FilterOption);
+			OutFilterOptions.Add(ToggleFilterOption.ToSharedRef());
 		}
 
 	protected:
@@ -317,6 +320,9 @@ namespace BlueprintEditorImpl
 		/** Class viewer filters (one per edited BP asset). */
 		TArray<TSharedRef<IClassViewerFilter>> ClassViewerFilters;
 
+		/** Filter option for the class viewer settings menu. */
+		TSharedPtr<FClassViewerFilterOption> ToggleFilterOption;
+
 		/** Whether or not the filter is enabled. */
 		bool bIsFilterEnabled;
 	};
@@ -957,6 +963,7 @@ void FBlueprintEditor::RefreshEditors(ERefreshBlueprintEditorReason::Type Reason
 
 	if(SubobjectEditor.IsValid())
 	{
+		SubobjectEditor->RefreshComponentTypesList();
 		SubobjectEditor->UpdateTree();
 
 		// Note: Don't pass 'true' here because we don't want the preview actor to be reconstructed until after Blueprint modification is complete.
@@ -2696,12 +2703,20 @@ void FBlueprintEditor::CreateDefaultTabContents(const TArray<UBlueprint*>& InBlu
 
 void FBlueprintEditor::CreateSubobjectEditors()
 {
+	TArray<TSharedRef<IClassViewerFilter>> ClassFilters;
+	if (ImportedClassViewerFilter.IsValid())
+	{
+		ClassFilters.Add(ImportedClassViewerFilter.ToSharedRef());
+	}
+
 	SubobjectEditor = SAssignNew(SubobjectEditor, SSubobjectBlueprintEditor)
 		.ObjectContext(this, &FBlueprintEditor::GetSubobjectEditorObjectContext)
 		.PreviewActor(this, &FBlueprintEditor::GetPreviewActor)
 		.AllowEditing(this, &FBlueprintEditor::InEditingMode)
 		.OnSelectionUpdated(this, &FBlueprintEditor::OnSelectionUpdated)
-		.OnItemDoubleClicked(this, &FBlueprintEditor::OnComponentDoubleClicked);
+		.OnItemDoubleClicked(this, &FBlueprintEditor::OnComponentDoubleClicked)
+		.OnImportNamespaceToEditorContext(this, &FBlueprintEditor::ImportNamespace)
+		.SubobjectClassListFilters(ClassFilters);
 	
 	SubobjectViewport = SAssignNew(SubobjectViewport, SSCSEditorViewport)
 		.BlueprintEditor(SharedThis(this));
diff --git a/Engine/Source/Editor/Kismet/Private/BlueprintNamespaceHelper.cpp b/Engine/Source/Editor/Kismet/Private/BlueprintNamespaceHelper.cpp
--- a/Engine/Source/Editor/Kismet/Private/BlueprintNamespaceHelper.cpp
+++ b/Engine/Source/Editor/Kismet/Private/BlueprintNamespaceHelper.cpp
@@ -78,7 +78,7 @@ static void UpdateNamespaceFeatureSettingsCVarSinkFunction()
 					if (AssetEditorPtr.IsValid() && AssetEditorPtr->IsBlueprintEditor())
 					{
 						TSharedPtr<IBlueprintEditor> BlueprintEditorPtr = StaticCastSharedPtr<IBlueprintEditor>(AssetEditorPtr);
-						BlueprintEditorPtr->RefreshMyBlueprint();
+						BlueprintEditorPtr->RefreshEditors();
 						BlueprintEditorPtr->RefreshInspector();
 					}
 				}
diff --git a/Engine/Source/Editor/Kismet/Private/BlueprintNamespaceUtilities.cpp b/Engine/Source/Editor/Kismet/Private/BlueprintNamespaceUtilities.cpp
--- a/Engine/Source/Editor/Kismet/Private/BlueprintNamespaceUtilities.cpp
+++ b/Engine/Source/Editor/Kismet/Private/BlueprintNamespaceUtilities.cpp
@@ -85,16 +85,25 @@ FString FBlueprintNamespaceUtilities::GetObjectNamespace(const UObject* InObject
 		}
 		else
 		{
-			Namespace = GetObjectNamespace(Field->GetPackage());
+			const UBlueprint* Blueprint = nullptr;
+			if (const UClass* Class = Cast<UClass>(Field))
+			{
+				Blueprint = UBlueprint::GetBlueprintFromClass(Class);
+			}
+
+			if (Blueprint)
+			{
+				Namespace = GetObjectNamespace(Blueprint);
+			}
+			else
+			{
+				Namespace = GetObjectNamespace(Field->GetPackage());
+			}
 		}
 	}
 	else if (const UBlueprint* Blueprint = Cast<UBlueprint>(InObject))
 	{
-		if (Blueprint->GeneratedClass)
-		{
-			Namespace = GetObjectNamespace(Blueprint->GeneratedClass);
-		}
-		else if (!Blueprint->BlueprintNamespace.IsEmpty())
+		if (!Blueprint->BlueprintNamespace.IsEmpty())
 		{
 			Namespace = Blueprint->BlueprintNamespace;
 		}
diff --git a/Engine/Source/Editor/SubobjectEditor/Private/SSubobjectBlueprintEditor.cpp b/Engine/Source/Editor/SubobjectEditor/Private/SSubobjectBlueprintEditor.cpp
--- a/Engine/Source/Editor/SubobjectEditor/Private/SSubobjectBlueprintEditor.cpp
+++ b/Engine/Source/Editor/SubobjectEditor/Private/SSubobjectBlueprintEditor.cpp
@@ -27,6 +27,8 @@
 #include "Kismet2/ChildActorComponentEditorUtils.h"
 #include "ObjectTools.h"						// ThumbnailTools::CacheEmptyThumbnail
 #include "K2Node_ComponentBoundEvent.h"
+#include "BlueprintEditorSettings.h"
+#include "BlueprintNamespaceUtilities.h"
 
 #define LOCTEXT_NAMESPACE "SSubobjectBlueprintEditor"
 
@@ -44,6 +46,7 @@ void SSubobjectBlueprintEditor::Construct(const FArguments& InArgs)
 	OnSelectionUpdated = InArgs._OnSelectionUpdated;
 	OnItemDoubleClicked = InArgs._OnItemDoubleClicked;
 	OnHighlightPropertyInDetailsView = InArgs._OnHighlightPropertyInDetailsView;
+	OnImportNamespaceToEditorContext = InArgs._OnImportNamespaceToEditorContext;
 	AllowEditing = InArgs._AllowEditing;
 	HideComponentClassCombo = InArgs._HideComponentClassCombo;
 	bAllowTreeUpdates = true;
@@ -84,6 +87,7 @@ void SSubobjectBlueprintEditor::Construct(const FArguments& InArgs)
 		.OnSubobjectClassSelected(this, &SSubobjectBlueprintEditor::PerformComboAddClass)
 		.ToolTipText(LOCTEXT("AddComponent_Tooltip", "Adds a new component to this actor"))
 		.IsEnabled(true)
+		.CustomClassFilters(InArgs._SubobjectClassListFilters)
 	]
 	+ SHorizontalBox::Slot()
 	.VAlign(VAlign_Center)
@@ -514,6 +518,12 @@ TSharedPtr<SWidget> SSubobjectBlueprintEditor::BuildSceneRootDropActionMenu(FSub
 
 FSubobjectDataHandle SSubobjectBlueprintEditor::AddNewSubobject(const FSubobjectDataHandle& ParentHandle, UClass* NewClass, UObject* AssetOverride, FText& OutFailReason, TUniquePtr<FScopedTransaction> InOngoingTransaction)
 {
+	if(GetDefault<UBlueprintEditorSettings>()->bEnableNamespaceImportingFeatures)
+	{
+		// If bound, invoke the delegate to import the namespace associated with the given class.
+		OnImportNamespaceToEditorContext.ExecuteIfBound(FBlueprintNamespaceUtilities::GetObjectNamespace(NewClass));
+	}
+
 	FAddNewSubobjectParams Params;
     Params.ParentHandle = ParentHandle;
     Params.NewClass = NewClass;
diff --git a/Engine/Source/Editor/SubobjectEditor/Private/SSubobjectEditor.cpp b/Engine/Source/Editor/SubobjectEditor/Private/SSubobjectEditor.cpp
--- a/Engine/Source/Editor/SubobjectEditor/Private/SSubobjectEditor.cpp
+++ b/Engine/Source/Editor/SubobjectEditor/Private/SSubobjectEditor.cpp
@@ -3515,6 +3515,14 @@ void SSubobjectEditor::Utils::PopulateHandlesArray(const TArray<FSubobjectEditor
     }
 }
 
+void SSubobjectEditor::RefreshComponentTypesList()
+{
+	if (ComponentClassCombo.IsValid())
+	{
+		ComponentClassCombo->UpdateComponentClassList();
+	}
+}
+
 TArray<UObject*> USubobjectEditorMenuContext::GetSelectedObjects() const
 {
 	TArray<UObject*> Result;
diff --git a/Engine/Source/Editor/SubobjectEditor/Public/SSubobjectBlueprintEditor.h b/Engine/Source/Editor/SubobjectEditor/Public/SSubobjectBlueprintEditor.h
--- a/Engine/Source/Editor/SubobjectEditor/Public/SSubobjectBlueprintEditor.h
+++ b/Engine/Source/Editor/SubobjectEditor/Public/SSubobjectBlueprintEditor.h
@@ -4,6 +4,8 @@
 
 #include "SSubobjectEditor.h"
 
+class IClassViewerFilter;
+
 /**
 * This is the editor for subobjects within the blueprint editor that
 */
@@ -11,6 +13,7 @@ class SUBOBJECTEDITOR_API SSubobjectBlueprintEditor final : public SSubobjectEdi
 {
 public:
 	DECLARE_DELEGATE_OneParam(FOnHighlightPropertyInDetailsView, const class FPropertyPath&);
+	DECLARE_DELEGATE_OneParam(FOnImportNamespaceToEditorContext, const FString& InNamespace);
 
 private:	
 	SLATE_BEGIN_ARGS(SSubobjectBlueprintEditor)
@@ -20,6 +23,7 @@ private:
 		, _HideComponentClassCombo(false)
         , _OnSelectionUpdated()
 		, _OnHighlightPropertyInDetailsView()
+		, _OnImportNamespaceToEditorContext()
 		{}
 
 		SLATE_ATTRIBUTE(UObject*, ObjectContext)
@@ -29,6 +33,8 @@ private:
 	    SLATE_EVENT(FOnSelectionUpdated, OnSelectionUpdated)
 	    SLATE_EVENT(FOnItemDoubleClicked, OnItemDoubleClicked)
 		SLATE_EVENT(FOnHighlightPropertyInDetailsView, OnHighlightPropertyInDetailsView)
+		SLATE_EVENT(FOnImportNamespaceToEditorContext, OnImportNamespaceToEditorContext)
+		SLATE_ARGUMENT(TArray<TSharedRef<IClassViewerFilter>>, SubobjectClassListFilters)
 	    
 	SLATE_END_ARGS()
 	
@@ -67,6 +73,9 @@ public:
 
 	/** Delegate to invoke when the given property should be highlighted in the details view (e.g. diff). */
 	FOnHighlightPropertyInDetailsView OnHighlightPropertyInDetailsView;
+
+	/** Delegate to invoke when a component class is selected in order to ensure the namespace is imported into the current editor context. */
+	FOnImportNamespaceToEditorContext OnImportNamespaceToEditorContext;
 	
 	/**
 	* Fills out an events section in ui.
@@ -115,5 +124,4 @@ protected:
 
 	/** Get the current preview actor for this blueprint editor. */
 	AActor* GetActorPreview() const { return PreviewActor.Get(nullptr); }
-	
 };
\ No newline at end of file
diff --git a/Engine/Source/Editor/SubobjectEditor/Public/SSubobjectEditor.h b/Engine/Source/Editor/SubobjectEditor/Public/SSubobjectEditor.h
--- a/Engine/Source/Editor/SubobjectEditor/Public/SSubobjectEditor.h
+++ b/Engine/Source/Editor/SubobjectEditor/Public/SSubobjectEditor.h
@@ -557,6 +557,9 @@ public:
 
 	/** SubobjectHandle of the current object that is represented by the subobject editor */
 	FSubobjectDataHandle GetObjectContextHandle() const;
+
+	/** Refresh the type list presented by the add component button when clicked */
+	void RefreshComponentTypesList();
 	
 protected:
 
@@ -611,6 +614,9 @@ protected:
 	/** The tools buttons box */
 	TSharedPtr<SHorizontalBox> ButtonBox;
 
+	/** The add component button / type selector */
+	TSharedPtr<SComponentClassCombo> ComponentClassCombo;
+
 	/** Gate to prevent changing the selection while selection change is being broadcast. */
 	bool bUpdatingSelection;
 
diff --git a/Engine/Source/Editor/UnrealEd/Private/ComponentTypeRegistry.cpp b/Engine/Source/Editor/UnrealEd/Private/ComponentTypeRegistry.cpp
--- a/Engine/Source/Editor/UnrealEd/Private/ComponentTypeRegistry.cpp
+++ b/Engine/Source/Editor/UnrealEd/Private/ComponentTypeRegistry.cpp
@@ -19,9 +19,229 @@
 #include "Kismet2/KismetEditorUtilities.h"
 #include "SComponentClassCombo.h"
 #include "Settings/ClassViewerSettings.h"
+#include "ClassViewerFilter.h"
+#include "EditorClassUtils.h"
 
 #define LOCTEXT_NAMESPACE "ComponentTypeRegistry"
 
+namespace UE::Editor::ComponentTypeRegistry::Private
+{
+	class FUnloadedBlueprintData : public IUnloadedBlueprintData
+	{
+	public:
+		FUnloadedBlueprintData(const FAssetData& InAssetData)
+			:ClassPath(NAME_None)
+			,ParentClassPath(NAME_None)
+			,ClassFlags(CLASS_None)
+			,bIsNormalBlueprintType(false)
+		{
+			ClassName = MakeShared<FString>(InAssetData.AssetName.ToString());
+
+			FString GeneratedClassPath;
+			const UClass* AssetClass = InAssetData.GetClass();
+			if (AssetClass && AssetClass->IsChildOf(UBlueprintGeneratedClass::StaticClass()))
+			{
+				ClassPath = InAssetData.ObjectPath;
+			}
+			else if (InAssetData.GetTagValue(FBlueprintTags::GeneratedClassPath, GeneratedClassPath))
+			{
+				ClassPath = FName(*FPackageName::ExportTextPathToObjectPath(GeneratedClassPath));
+			}
+
+			FString ParentClassPathString;
+			if (InAssetData.GetTagValue(FBlueprintTags::ParentClassPath, ParentClassPathString))
+			{
+				ParentClassPath = FName(*FPackageName::ExportTextPathToObjectPath(ParentClassPathString));
+			}
+
+			FEditorClassUtils::GetImplementedInterfaceClassPathsFromAsset(InAssetData, ImplementedInterfaces);
+		}
+
+		virtual ~FUnloadedBlueprintData()
+		{
+		}
+
+		// Begin IUnloadedBlueprintData interface
+		virtual bool HasAnyClassFlags(uint32 InFlagsToCheck) const
+		{
+			return (ClassFlags & InFlagsToCheck) != 0;
+		}
+
+		virtual bool HasAllClassFlags(uint32 InFlagsToCheck) const
+		{
+			return ((ClassFlags & InFlagsToCheck) == InFlagsToCheck);
+		}
+
+		virtual void SetClassFlags(uint32 InFlags)
+		{
+			ClassFlags = InFlags;
+		}
+
+		virtual bool ImplementsInterface(const UClass* InInterface) const
+		{
+			FString InterfacePath = InInterface->GetPathName();
+			for (const FString& ImplementedInterface : ImplementedInterfaces)
+			{
+				if (ImplementedInterface == InterfacePath)
+				{
+					return true;
+				}
+			}
+
+			FComponentClassComboEntryPtr CurrentEntry = FComponentTypeRegistry::Get().FindClassEntryForObjectPath(ParentClassPath);
+			while (CurrentEntry.IsValid())
+			{
+				if (const UClass* CurrentClass = CurrentEntry->GetComponentClass())
+				{
+					return CurrentClass->ImplementsInterface(InInterface);
+				}
+
+				TSharedPtr<FUnloadedBlueprintData> UnloadedBlueprintData = StaticCastSharedPtr<FUnloadedBlueprintData>(CurrentEntry->GetUnloadedBlueprintData());
+				if (UnloadedBlueprintData.IsValid())
+				{
+					for (const FString& ImplementedInterface : UnloadedBlueprintData->ImplementedInterfaces)
+					{
+						if (ImplementedInterface == InterfacePath)
+						{
+							return true;
+						}
+					}
+
+					CurrentEntry = FComponentTypeRegistry::Get().FindClassEntryForObjectPath(UnloadedBlueprintData->ParentClassPath);
+				}
+				else
+				{
+					CurrentEntry.Reset();
+				}
+			}
+
+			return false;
+		}
+
+		virtual bool IsChildOf(const UClass* InClass) const
+		{
+			FComponentClassComboEntryPtr CurrentEntry = FComponentTypeRegistry::Get().FindClassEntryForObjectPath(ParentClassPath);
+			while (CurrentEntry.IsValid())
+			{
+				if (const UClass* CurrentClass = CurrentEntry->GetComponentClass())
+				{
+					return CurrentClass->IsChildOf(InClass);
+				}
+
+				TSharedPtr<FUnloadedBlueprintData> UnloadedBlueprintData = StaticCastSharedPtr<FUnloadedBlueprintData>(CurrentEntry->GetUnloadedBlueprintData());
+				if (UnloadedBlueprintData.IsValid())
+				{
+					CurrentEntry = FComponentTypeRegistry::Get().FindClassEntryForObjectPath(UnloadedBlueprintData->ParentClassPath);
+				}
+				else
+				{
+					CurrentEntry.Reset();
+				}
+			}
+
+			return false;
+		}
+
+		virtual bool IsA(const UClass* InClass) const
+		{
+			// Unloaded blueprint classes should always be a BPGC, so this just checks against the expected type.
+			return UBlueprintGeneratedClass::StaticClass()->UObject::IsA(InClass);
+		}
+
+		virtual const UClass* GetClassWithin() const
+		{
+			FComponentClassComboEntryPtr CurrentEntry = FComponentTypeRegistry::Get().FindClassEntryForObjectPath(ParentClassPath);
+			while (CurrentEntry.IsValid())
+			{
+				if (const UClass* CurrentClass = CurrentEntry->GetComponentClass())
+				{
+					return CurrentClass->ClassWithin;
+				}
+
+				TSharedPtr<FUnloadedBlueprintData> UnloadedBlueprintData = StaticCastSharedPtr<FUnloadedBlueprintData>(CurrentEntry->GetUnloadedBlueprintData());
+				if (UnloadedBlueprintData.IsValid())
+				{
+					CurrentEntry = FComponentTypeRegistry::Get().FindClassEntryForObjectPath(UnloadedBlueprintData->ParentClassPath);
+				}
+				else
+				{
+					CurrentEntry.Reset();
+				}
+			}
+
+			return nullptr;
+		}
+
+		virtual const UClass* GetNativeParent() const
+		{
+			const UClass* CurrentClass = nullptr;
+			FComponentClassComboEntryPtr CurrentEntry = FComponentTypeRegistry::Get().FindClassEntryForObjectPath(ParentClassPath);
+			while (CurrentEntry.IsValid() || CurrentClass)
+			{
+				if (!CurrentClass)
+				{
+					CurrentClass = CurrentEntry->GetComponentClass();
+				}
+
+				if (CurrentClass)
+				{
+					if (CurrentClass->HasAnyClassFlags(CLASS_Native))
+					{
+						return CurrentClass;
+					}
+					else
+					{
+						CurrentClass = CurrentClass->GetSuperClass();
+					}
+				}
+				else
+				{
+					TSharedPtr<FUnloadedBlueprintData> UnloadedBlueprintData = StaticCastSharedPtr<FUnloadedBlueprintData>(CurrentEntry->GetUnloadedBlueprintData());
+					if (UnloadedBlueprintData.IsValid())
+					{
+						CurrentEntry = FComponentTypeRegistry::Get().FindClassEntryForObjectPath(UnloadedBlueprintData->ParentClassPath);
+					}
+					else
+					{
+						CurrentEntry.Reset();
+					}
+				}
+			}
+
+			return nullptr;
+		}
+
+		virtual void SetNormalBlueprintType(bool bInNormalBPType)
+		{
+			bIsNormalBlueprintType = bInNormalBPType;
+		}
+
+		virtual bool IsNormalBlueprintType() const
+		{
+			return bIsNormalBlueprintType;
+		}
+
+		virtual TSharedPtr<FString> GetClassName() const
+		{
+			return ClassName;
+		}
+
+		virtual FName GetClassPath() const
+		{
+			return ClassPath;
+		}
+		// End IUnloadedBlueprintData interface
+
+	private:
+		TSharedPtr<FString> ClassName;
+		FName ClassPath;
+		FName ParentClassPath;
+		uint32 ClassFlags;
+		TArray<FString> ImplementedInterfaces;
+		bool bIsNormalBlueprintType;
+	};
+}
+
 //////////////////////////////////////////////////////////////////////////
 // FComponentTypeRegistryData
 
@@ -58,6 +278,7 @@ public:
 	TArray<FComponentClassComboEntryPtr> ComponentClassList;
 	TArray<FComponentTypeEntry> ComponentTypeList;
 	TArray<FAssetData> PendingAssetData;
+	TMap<FName, int32> ClassPathToClassListIndexMap;
 	FOnComponentTypeListChanged ComponentListChanged;
 	bool bNeedsRefreshNextTick;
 };
@@ -213,6 +434,7 @@ void FComponentTypeRegistryData::ForceRefreshComponentList()
 	bNeedsRefreshNextTick = false;
 	ComponentClassList.Empty();
 	ComponentTypeList.Empty();
+	ClassPathToClassListIndexMap.Empty();
 
 	struct SortComboEntry
 	{
@@ -370,7 +592,23 @@ void FComponentTypeRegistryData::ForceRefreshComponentList()
 					const UClass* BlueprintIconClass = FClassIconFinder::GetIconClassForAssetData(AssetData);
 
 					const bool bIncludeInFilter = true;
-					FComponentClassComboEntryPtr NewEntry(new FComponentClassComboEntry(BlueprintComponents, FixedString, AssetData.ObjectPath, BlueprintIconClass, bIncludeInFilter));
+					FComponentClassComboEntryPtr NewEntry = MakeShared<FComponentClassComboEntry>(BlueprintComponents, FixedString, AssetData.ObjectPath, BlueprintIconClass, bIncludeInFilter);
+					
+					// Create an unloaded blueprint data object to assist with class filtering
+					TSharedPtr<IUnloadedBlueprintData> UnloadedBlueprintData;
+					{
+						using namespace UE::Editor::ComponentTypeRegistry;
+						UnloadedBlueprintData = MakeShared<Private::FUnloadedBlueprintData>(AssetData);
+
+						const uint32 ClassFlags = AssetData.GetTagValueRef<uint32>(FBlueprintTags::ClassFlags);
+						UnloadedBlueprintData->SetClassFlags(ClassFlags);
+
+						const FString BlueprintType = AssetData.GetTagValueRef<FString>(FBlueprintTags::BlueprintType);
+						UnloadedBlueprintData->SetNormalBlueprintType(BlueprintType == TEXT("BPType_Normal"));
+
+						NewEntry->SetUnloadedBlueprintData(UnloadedBlueprintData);
+					}
+					
 					SortedClassList.Add(NewEntry);
 				}
 			}
@@ -401,7 +639,11 @@ void FComponentTypeRegistryData::ForceRefreshComponentList()
 				PreviousHeading = CurrentHeadingText;
 			}
 
-			ComponentClassList.Add(CurrentEntry);
+			int32 EntryIndex = ComponentClassList.Add(CurrentEntry);
+			if (CurrentEntry->IsClass())
+			{
+				ClassPathToClassListIndexMap.FindOrAdd(FName(*CurrentEntry->GetComponentPath()), EntryIndex);
+			}
 		}
 	}
 
@@ -517,5 +759,18 @@ void FComponentTypeRegistry::InvalidateClass(TSubclassOf<UActorComponent> /*Clas
 	Data->Invalidate();
 }
 
+FComponentClassComboEntryPtr FComponentTypeRegistry::FindClassEntryForObjectPath(FName InObjectPath) const
+{
+	if (int32* ClassListIndexPtr = Data->ClassPathToClassListIndexMap.Find(InObjectPath))
+	{
+		const int32 ClassListIndex = *ClassListIndexPtr;
+		if (Data->ComponentClassList.IsValidIndex(ClassListIndex))
+		{
+			return Data->ComponentClassList[ClassListIndex];
+		}
+	}
+
+	return nullptr;
+}
 
 #undef LOCTEXT_NAMESPACE
diff --git a/Engine/Source/Editor/UnrealEd/Private/EditorClassUtils.cpp b/Engine/Source/Editor/UnrealEd/Private/EditorClassUtils.cpp
--- a/Engine/Source/Editor/UnrealEd/Private/EditorClassUtils.cpp
+++ b/Engine/Source/Editor/UnrealEd/Private/EditorClassUtils.cpp
@@ -162,3 +162,43 @@ UClass* FEditorClassUtils::GetClassFromString(const FString& ClassName)
 	}
 	return Class;
 }
+
+void FEditorClassUtils::GetImplementedInterfaceClassPathsFromAsset(const struct FAssetData& InAssetData, TArray<FString>& OutClassPaths)
+{
+	if (!InAssetData.IsValid())
+	{
+		return;
+	}
+
+	const FString ImplementedInterfaces = InAssetData.GetTagValueRef<FString>(FBlueprintTags::ImplementedInterfaces);
+	if (!ImplementedInterfaces.IsEmpty())
+	{
+		// Parse string like "((Interface=Class'"/Script/VPBookmark.VPBookmarkProvider"'),(Interface=Class'"/Script/VPUtilities.VPContextMenuProvider"'))"
+		// We don't want to actually resolve the hard ref so do some manual parsing
+
+		FString FullInterface;
+		FString CurrentString = *ImplementedInterfaces;
+		while (CurrentString.Split(TEXT("Interface="), nullptr, &FullInterface))
+		{
+			// Cutoff at next )
+			int32 RightParen = INDEX_NONE;
+			if (FullInterface.FindChar(TCHAR(')'), RightParen))
+			{
+				// Keep parsing
+				CurrentString = FullInterface.Mid(RightParen);
+
+				// Strip class name
+				FullInterface = *FPackageName::ExportTextPathToObjectPath(FullInterface.Left(RightParen));
+
+				// Handle quotes
+				FString InterfacePath;
+				const TCHAR* NewBuffer = FPropertyHelpers::ReadToken(*FullInterface, InterfacePath, true);
+
+				if (NewBuffer)
+				{
+					OutClassPaths.Add(InterfacePath);
+				}
+			}
+		}
+	}
+}
\ No newline at end of file
diff --git a/Engine/Source/Editor/UnrealEd/Private/SComponentClassCombo.cpp b/Engine/Source/Editor/UnrealEd/Private/SComponentClassCombo.cpp
--- a/Engine/Source/Editor/UnrealEd/Private/SComponentClassCombo.cpp
+++ b/Engine/Source/Editor/UnrealEd/Private/SComponentClassCombo.cpp
@@ -20,6 +20,10 @@
 #include "Misc/TextFilterExpressionEvaluator.h"
 #include "SPositiveActionButton.h"
 #include "Widgets/Layout/SSeparator.h"
+#include "Modules/ModuleManager.h"
+#include "ClassViewerModule.h"
+#include "ClassViewerFilter.h"
+#include "Framework/MultiBox/MultiBoxBuilder.h"
 
 #define LOCTEXT_NAMESPACE "ComponentClassCombo"
 
@@ -50,6 +54,17 @@ void SComponentClassCombo::Construct(const FArguments& InArgs)
 	OnComponentClassSelected = InArgs._OnComponentClassSelected;
 	OnSubobjectClassSelected = InArgs._OnSubobjectClassSelected;
 	TextFilter = MakeShared<FTextFilterExpressionEvaluator>(ETextFilterExpressionEvaluatorMode::BasicString);
+	
+	if(InArgs._CustomClassFilters.Num() > 0)
+	{
+		ComponentClassFilterData = MakeUnique<FComponentClassFilterData>();
+
+		ComponentClassFilterData->InitOptions = MakeShared<FClassViewerInitializationOptions>();
+		ComponentClassFilterData->InitOptions->ClassFilters.Append(InArgs._CustomClassFilters);
+
+		ComponentClassFilterData->ClassFilter = FModuleManager::LoadModuleChecked<FClassViewerModule>("ClassViewer").CreateClassFilter(*ComponentClassFilterData->InitOptions);
+		ComponentClassFilterData->FilterFuncs = FModuleManager::LoadModuleChecked<FClassViewerModule>("ClassViewer").CreateFilterFuncs();
+	}
 
 	FComponentTypeRegistry::Get().SubscribeToComponentList(ComponentClassList).AddRaw(this, &SComponentClassCombo::UpdateComponentClassList);
 
@@ -88,7 +103,29 @@ void SComponentClassCombo::Construct(const FArguments& InArgs)
 						.Padding(1.f)
 						.AutoHeight()
 						[
-							SearchBox.ToSharedRef()
+							SNew(SHorizontalBox)
+							+SHorizontalBox::Slot()
+							[
+								SearchBox.ToSharedRef()
+							]
+							+SHorizontalBox::Slot()
+							.AutoWidth()
+							.Padding(2.0f, 2.0f)
+							[
+								SNew(SComboButton)
+								.ContentPadding(0)
+								.ForegroundColor(FSlateColor::UseForeground())
+								.ComboButtonStyle(FAppStyle::Get(), "SimpleComboButton")
+								.HasDownArrow(false)
+								.Visibility(this, &SComponentClassCombo::GetFilterOptionsButtonVisibility)
+								.OnGetMenuContent(this, &SComponentClassCombo::GetFilterOptionsMenuContent)
+								.ButtonContent()
+								[
+									SNew(SImage)
+									.Image(FAppStyle::Get().GetBrush("Icons.Settings"))
+									.ColorAndOpacity(FSlateColor::UseForeground())
+								]
+							]
 						]
 						+SVerticalBox::Slot()
 						.MaxHeight(400)
@@ -131,7 +168,10 @@ void SComponentClassCombo::ClearSelection()
 
 void SComponentClassCombo::GenerateFilteredComponentList()
 {
-	if ( TextFilter->GetFilterText().IsEmpty() )
+	const bool bNoClassFilter = !ComponentClassFilterData.IsValid();
+	const bool bHasFilterText = !TextFilter->GetFilterText().IsEmpty();
+
+	if (bNoClassFilter && !bHasFilterText)
 	{
 		FilteredComponentClassList = *ComponentClassList;
 	}
@@ -142,6 +182,9 @@ void SComponentClassCombo::GenerateFilteredComponentList()
 		int32 LastHeadingIndex = INDEX_NONE;
 		FComponentClassComboEntryPtr* LastHeadingPtr = nullptr;
 
+		int32 LastSeparatorIndex = INDEX_NONE;
+		FComponentClassComboEntryPtr* LastSeparatorPtr = nullptr;
+
 		for (int32 ComponentIndex = 0; ComponentIndex < ComponentClassList->Num(); ComponentIndex++)
 		{
 			FComponentClassComboEntryPtr& CurrentEntry = (*ComponentClassList)[ComponentIndex];
@@ -151,11 +194,28 @@ void SComponentClassCombo::GenerateFilteredComponentList()
 				LastHeadingIndex = FilteredComponentClassList.Num();
 				LastHeadingPtr = &CurrentEntry;
 			}
-			else if (CurrentEntry->IsClass() && CurrentEntry->IsIncludedInFilter())
+			else if (CurrentEntry->IsSeparator())
 			{
-				FString FriendlyComponentName = GetSanitizedComponentName( CurrentEntry );
+				LastSeparatorIndex = FilteredComponentClassList.Num();
+				LastSeparatorPtr = &CurrentEntry;
+			}
+			else if(CurrentEntry->IsClass())
+			{
+				// Disallow class entries that are not to be seen when searching via text.
+				bool bAllowEntry = !bHasFilterText || CurrentEntry->IsIncludedInFilter();
+				if(bAllowEntry)
+				{
+					// Disallow class entries that don't match the custom class filter, if set.
+					bAllowEntry = bNoClassFilter || IsComponentClassAllowed(CurrentEntry);
+					if (bAllowEntry && bHasFilterText)
+					{
+						// Finally, disallow class entries that don't match the search box text.
+						FString FriendlyComponentName = GetSanitizedComponentName(CurrentEntry);
+						bAllowEntry = TextFilter->TestTextFilter(FBasicStringFilterExpressionContext(FriendlyComponentName));
+					}
+				}
 
-				if ( TextFilter->TestTextFilter(FBasicStringFilterExpressionContext(FriendlyComponentName)) )
+				if (bAllowEntry)
 				{
 					// Add the heading first if it hasn't already been added
 					if (LastHeadingPtr && LastHeadingIndex != INDEX_NONE)
@@ -165,19 +225,30 @@ void SComponentClassCombo::GenerateFilteredComponentList()
 						LastHeadingPtr = nullptr;
 					}
 
+					// Add the separator next so that it will precede the heading
+					if (LastSeparatorPtr && LastSeparatorIndex != INDEX_NONE)
+					{
+						FilteredComponentClassList.Insert(*LastSeparatorPtr, LastSeparatorIndex);
+						LastSeparatorIndex = INDEX_NONE;
+						LastSeparatorPtr = nullptr;
+					}
+
 					// Add the class
-					FilteredComponentClassList.Add( CurrentEntry );
+					FilteredComponentClassList.Add(CurrentEntry);
 				}
 			}
 		}
 
-		// Select the first non-category item that passed the filter
-		for (FComponentClassComboEntryPtr& TestEntry : FilteredComponentClassList)
+		if (ComponentClassListView.IsValid())
 		{
-			if (TestEntry->IsClass())
+			// Select the first non-category item that passed the filter
+			for (FComponentClassComboEntryPtr& TestEntry : FilteredComponentClassList)
 			{
-				ComponentClassListView->SetSelection(TestEntry, ESelectInfo::OnNavigation);
-				break;
+				if (TestEntry->IsClass())
+				{
+					ComponentClassListView->SetSelection(TestEntry, ESelectInfo::OnNavigation);
+					break;
+				}
 			}
 		}
 	}
@@ -193,11 +264,7 @@ void SComponentClassCombo::OnSearchBoxTextChanged( const FText& InSearchText )
 	TextFilter->SetFilterText(InSearchText);
 	SearchBox->SetError(TextFilter->GetFilterErrorText());
 
-	// Generate a filtered list
-	GenerateFilteredComponentList();
-
-	// Ask the combo to update its contents on next tick
-	ComponentClassListView->RequestListRefresh();
+	UpdateComponentClassList();
 }
 
 void SComponentClassCombo::OnSearchBoxTextCommitted(const FText& NewText, ETextCommit::Type CommitInfo)
@@ -373,7 +440,14 @@ TSharedRef<ITableRow> SComponentClassCombo::GenerateAddComponentRow( FComponentC
 
 void SComponentClassCombo::UpdateComponentClassList()
 {
+	// Regenerate the filtered list
 	GenerateFilteredComponentList();
+
+	// Ask the combo to update its contents on next tick
+	if (ComponentClassListView.IsValid())
+	{
+		ComponentClassListView->RequestListRefresh();
+	}
 }
 
 FText SComponentClassCombo::GetFriendlyComponentName(FComponentClassComboEntryPtr Entry) const
@@ -497,4 +571,111 @@ TSharedRef<SToolTip> SComponentClassCombo::GetComponentToolTip(FComponentClassCo
 		.Text(FText::FromString(Entry->GetClassName()));
 }
 
+bool SComponentClassCombo::IsComponentClassAllowed(FComponentClassComboEntryPtr Entry) const
+{
+	if (Entry.IsValid() && Entry->IsClass() && ComponentClassFilterData.IsValid())
+	{
+		check(ComponentClassFilterData->InitOptions.IsValid());
+		check(ComponentClassFilterData->ClassFilter.IsValid());
+		check(ComponentClassFilterData->FilterFuncs.IsValid());
+
+		if (const UClass* ComponentClass = Entry->GetComponentClass())
+		{
+			return ComponentClassFilterData->ClassFilter->IsClassAllowed(*ComponentClassFilterData->InitOptions, ComponentClass, ComponentClassFilterData->FilterFuncs.ToSharedRef());
+		}
+		else
+		{
+			TSharedPtr<IUnloadedBlueprintData> UnloadedBlueprintData = Entry->GetUnloadedBlueprintData();
+			if (UnloadedBlueprintData.IsValid())
+			{
+				return ComponentClassFilterData->ClassFilter->IsUnloadedClassAllowed(*ComponentClassFilterData->InitOptions, UnloadedBlueprintData.ToSharedRef(), ComponentClassFilterData->FilterFuncs.ToSharedRef());
+			}
+		}
+	}
+
+	// Allow all entries to otherwise pass by default.
+	return true;
+}
+
+void SComponentClassCombo::GetComponentClassFilterOptions(TArray<TSharedRef<FClassViewerFilterOption>>& OutFilterOptions) const
+{
+	if (ComponentClassFilterData.IsValid() && ComponentClassFilterData->InitOptions.IsValid())
+	{
+		TArray<TSharedRef<FClassViewerFilterOption>> FilterOptions;
+		if(ComponentClassFilterData->ClassFilter.IsValid())
+		{
+			ComponentClassFilterData->ClassFilter->GetFilterOptions(FilterOptions);
+			OutFilterOptions.Append(FilterOptions);
+		}
+
+		for (const TSharedRef<IClassViewerFilter>& ClassFilter : ComponentClassFilterData->InitOptions->ClassFilters)
+		{
+			FilterOptions.Reset();
+			ClassFilter->GetFilterOptions(FilterOptions);
+
+			OutFilterOptions.Append(FilterOptions);
+		}
+	}
+}
+
+TSharedRef<SWidget> SComponentClassCombo::GetFilterOptionsMenuContent()
+{
+	const bool bCloseSelfOnly = true;
+	const bool bShouldCloseWindowAfterMenuSelection = true;
+	FMenuBuilder MenuBuilder(bShouldCloseWindowAfterMenuSelection, /* InCommandList = */nullptr, /* InExtender = */nullptr, bCloseSelfOnly);
+
+	TArray<TSharedRef<FClassViewerFilterOption>> ClassFilterOptions;
+	GetComponentClassFilterOptions(ClassFilterOptions);
+
+	if (ClassFilterOptions.Num() > 0)
+	{
+		MenuBuilder.BeginSection("ClassFilterOptions", LOCTEXT("ClassFilterOptionsHeading", "Class Filters"));
+		{
+			for (const TSharedRef<FClassViewerFilterOption>& ClassFilterOption : ClassFilterOptions)
+			{
+				MenuBuilder.AddMenuEntry(
+					ClassFilterOption->LabelText,
+					ClassFilterOption->ToolTipText,
+					FSlateIcon(),
+					FUIAction(
+						FExecuteAction::CreateSP(this, &SComponentClassCombo::ToggleFilterOption, ClassFilterOption),
+						FCanExecuteAction(),
+						FIsActionChecked::CreateSP(this, &SComponentClassCombo::IsFilterOptionEnabled, ClassFilterOption)
+					),
+					NAME_None,
+					EUserInterfaceActionType::ToggleButton
+				);
+			}
+		}
+		MenuBuilder.EndSection();
+	}
+
+	return MenuBuilder.MakeWidget();
+}
+
+void SComponentClassCombo::ToggleFilterOption(TSharedRef<FClassViewerFilterOption> FilterOption)
+{
+	FilterOption->bEnabled = !FilterOption->bEnabled;
+
+	if (FilterOption->OnOptionChanged.IsBound())
+	{
+		FilterOption->OnOptionChanged.Execute(FilterOption->bEnabled);
+	}
+
+	UpdateComponentClassList();
+}
+
+bool SComponentClassCombo::IsFilterOptionEnabled(TSharedRef<FClassViewerFilterOption> FilterOption) const
+{
+	return FilterOption->bEnabled;
+}
+
+EVisibility SComponentClassCombo::GetFilterOptionsButtonVisibility() const
+{
+	TArray<TSharedRef<FClassViewerFilterOption>> FilterOptions;
+	GetComponentClassFilterOptions(FilterOptions);
+
+	return FilterOptions.Num() > 0 ? EVisibility::Visible : EVisibility::Collapsed;
+}
+
 #undef LOCTEXT_NAMESPACE
diff --git a/Engine/Source/Editor/UnrealEd/Public/ComponentTypeRegistry.h b/Engine/Source/Editor/UnrealEd/Public/ComponentTypeRegistry.h
--- a/Engine/Source/Editor/UnrealEd/Public/ComponentTypeRegistry.h
+++ b/Engine/Source/Editor/UnrealEd/Public/ComponentTypeRegistry.h
@@ -42,6 +42,12 @@ struct UNREALED_API FComponentTypeRegistry
 	 * Called when a specific class has been updated and should force the component type registry to update as well
 	 */
 	void InvalidateClass(TSubclassOf<UActorComponent> ClassToUpdate);
+
+	/**
+	 * Attempts to locate the class entry corresponding to the given object path.
+	 */
+	FComponentClassComboEntryPtr FindClassEntryForObjectPath(FName InObjectPath) const;
+
 private:
 	void OnReloadComplete(EReloadCompleteReason Reason);
 
diff --git a/Engine/Source/Editor/UnrealEd/Public/EditorClassUtils.h b/Engine/Source/Editor/UnrealEd/Public/EditorClassUtils.h
--- a/Engine/Source/Editor/UnrealEd/Public/EditorClassUtils.h
+++ b/Engine/Source/Editor/UnrealEd/Public/EditorClassUtils.h
@@ -99,4 +99,12 @@ namespace FEditorClassUtils
 	 * @return					UClass pointer if it exists
 	 */
 	UNREALED_API UClass* GetClassFromString(const FString& ClassName);
+
+	/**
+	 * Fetches the set of interface class object paths from an asset data entry containing the appropriate asset tag(s).
+	 * 
+	 * @param	InAssetData		Reference to an asset data entry.
+	 * @param	OutClassPaths	One or more interface class object paths, or empty if the corresponding asset tag(s) were not found.
+	 */
+	UNREALED_API void GetImplementedInterfaceClassPathsFromAsset(const struct FAssetData& InAssetData, TArray<FString>& OutClassPaths);
 };
diff --git a/Engine/Source/Editor/UnrealEd/Public/SComponentClassCombo.h b/Engine/Source/Editor/UnrealEd/Public/SComponentClassCombo.h
--- a/Engine/Source/Editor/UnrealEd/Public/SComponentClassCombo.h
+++ b/Engine/Source/Editor/UnrealEd/Public/SComponentClassCombo.h
@@ -15,6 +15,11 @@
 class FComponentClassComboEntry;
 class SToolTip;
 class FTextFilterExpressionEvaluator;
+class IClassViewerFilter;
+class IUnloadedBlueprintData;
+class FClassViewerFilterFuncs;
+class FClassViewerFilterOption;
+class FClassViewerInitializationOptions;
 
 typedef TSharedPtr<class FComponentClassComboEntry> FComponentClassComboEntryPtr;
 
@@ -154,6 +159,16 @@ public:
 	int32 GetSortPriority() const { return CustomizationArgs.SortPriority; }
 
 	void AddReferencedObjects(FReferenceCollector& Collector);
+
+	// If the component type is not loaded, this stores data that can be used for class filtering.
+	TSharedPtr<IUnloadedBlueprintData> GetUnloadedBlueprintData() const { return UnloadedBlueprintData; }
+
+	// Can optionally be called to set unloaded component type data to assist with class filtering.
+	void SetUnloadedBlueprintData(TSharedPtr<IUnloadedBlueprintData> InUnloadedBlueprintData)
+	{
+		UnloadedBlueprintData = InUnloadedBlueprintData;
+	}
+
 private:
 	TSubclassOf<UActorComponent> ComponentClass;
 	const UClass* IconClass;
@@ -166,6 +181,7 @@ private:
 	bool bIncludedInFilter;
 	EComponentCreateAction::Type ComponentCreateAction;
 	FComponentEntryCustomizationArgs CustomizationArgs;
+	TSharedPtr<IUnloadedBlueprintData> UnloadedBlueprintData;
 };
 
 //////////////////////////////////////////////////////////////////////////
@@ -178,8 +194,9 @@ public:
 	{}
 
 		SLATE_ATTRIBUTE(bool, IncludeText)
-		SLATE_EVENT( FComponentClassSelected, OnComponentClassSelected )
-		SLATE_EVENT( FSubobjectClassSelected, OnSubobjectClassSelected )
+		SLATE_EVENT(FComponentClassSelected, OnComponentClassSelected)
+		SLATE_EVENT(FSubobjectClassSelected, OnSubobjectClassSelected)
+		SLATE_ARGUMENT(TArray<TSharedRef<IClassViewerFilter>>, CustomClassFilters)
 
 	SLATE_END_ARGS()
 
@@ -190,11 +207,6 @@ public:
 	/** Clear the current combo list selection */
 	void ClearSelection();
 
-	/**
-	 * Updates the filtered list of component names.
-	 */
-	void GenerateFilteredComponentList();
-
 	FText GetCurrentSearchString() const;
 
 	/**
@@ -216,11 +228,37 @@ public:
 	/** Returns a component name without the substring "Component" and sanitized for display */
 	static FString GetSanitizedComponentName(FComponentClassComboEntryPtr Entry);
 
-private:
+protected:
+	/** Internal data used to facilitate component class filtering */
+	struct FComponentClassFilterData
+	{
+		TSharedPtr<FClassViewerInitializationOptions> InitOptions;
+		TSharedPtr<IClassViewerFilter> ClassFilter;
+		TSharedPtr<FClassViewerFilterFuncs> FilterFuncs;
+	};
 
+	/**
+	 * Updates the filtered list of component names.
+	 */
+	void GenerateFilteredComponentList();
+
+private:
+	
 	FText GetFriendlyComponentName(FComponentClassComboEntryPtr Entry) const;
 
 	TSharedRef<SToolTip> GetComponentToolTip(FComponentClassComboEntryPtr Entry) const;
+
+	bool IsComponentClassAllowed(FComponentClassComboEntryPtr Entry) const;
+
+	void GetComponentClassFilterOptions(TArray<TSharedRef<FClassViewerFilterOption>>& OutFilterOptions) const;
+
+	TSharedRef<SWidget> GetFilterOptionsMenuContent();
+
+	void ToggleFilterOption(TSharedRef<FClassViewerFilterOption> FilterOption);
+
+	bool IsFilterOptionEnabled(TSharedRef<FClassViewerFilterOption> FilterOption) const;
+
+	EVisibility GetFilterOptionsButtonVisibility() const;
 	
 	FComponentClassSelected OnComponentClassSelected;
 
@@ -244,6 +282,9 @@ private:
 	/** The component list control - part of the combo drop down */
 	TSharedPtr< SListView<FComponentClassComboEntryPtr> > ComponentClassListView;
 
+	/** Internal data that facilitates custom class filtering */
+	TUniquePtr<FComponentClassFilterData> ComponentClassFilterData;
+
 	/** Cached selection index used to skip over unselectable items */
 	int32 PrevSelectedIndex;
 };
