Fixed usage of TCHAR_TO_ANSI in IoStoreUtilities.cpp, PakFileUtilities.cpp and in GeForceNOWWrapper.cpp.
Fixed usage of TCHAR_TO_UTF8 in FbxMainExport (ExportStaticMeshToFbx, ExportSplineMeshToFbx).

diff --git a/Engine/Source/Developer/IoStoreUtilities/Private/IoStoreUtilities.cpp b/Engine/Source/Developer/IoStoreUtilities/Private/IoStoreUtilities.cpp
--- a/Engine/Source/Developer/IoStoreUtilities/Private/IoStoreUtilities.cpp
+++ b/Engine/Source/Developer/IoStoreUtilities/Private/IoStoreUtilities.cpp
@@ -251,9 +251,9 @@ static void LoadKeyChain(const TCHAR* CmdLine, FKeyChain& OutCryptoSettings)
 				UE_LOG(LogIoStore, Fatal, TEXT("AES encryption key must be a pure ANSI string!"));
 			}
 
-			ANSICHAR* AsAnsi = TCHAR_TO_ANSI(*EncryptionKeyString);
-			check(TCString<ANSICHAR>::Strlen(AsAnsi) == RequiredKeyLength);
-			FMemory::Memcpy(NewKey.Key.Key, AsAnsi, RequiredKeyLength);
+			const auto AsAnsi = StringCast<ANSICHAR>(*EncryptionKeyString);
+			check(AsAnsi.Length() == RequiredKeyLength);
+			FMemory::Memcpy(NewKey.Key.Key, AsAnsi.Get(), RequiredKeyLength);
 			OutCryptoSettings.EncryptionKeys.Add(NewKey.Guid, NewKey);
 			UE_LOG(LogIoStore, Display, TEXT("Parsed AES encryption key from command line."));
 		}
diff --git a/Engine/Source/Developer/PakFileUtilities/Private/PakFileUtilities.cpp b/Engine/Source/Developer/PakFileUtilities/Private/PakFileUtilities.cpp
--- a/Engine/Source/Developer/PakFileUtilities/Private/PakFileUtilities.cpp
+++ b/Engine/Source/Developer/PakFileUtilities/Private/PakFileUtilities.cpp
@@ -1643,9 +1643,9 @@ void LoadKeyChain(const TCHAR* CmdLine, FKeyChain& OutCryptoSettings)
 				UE_LOG(LogPakFile, Fatal, TEXT("AES encryption key must be a pure ANSI string!"));
 			}
 
-			ANSICHAR* AsAnsi = TCHAR_TO_ANSI(*EncryptionKeyString);
-			check(TCString<ANSICHAR>::Strlen(AsAnsi) == RequiredKeyLength);
-			FMemory::Memcpy(NewKey.Key.Key, AsAnsi, RequiredKeyLength);
+			const auto AsAnsi = StringCast<ANSICHAR>(*EncryptionKeyString);
+			check(AsAnsi.Length() == RequiredKeyLength);
+			FMemory::Memcpy(NewKey.Key.Key, AsAnsi.Get(), RequiredKeyLength);
 			OutCryptoSettings.EncryptionKeys.Add(NewKey.Guid, NewKey);
 			UE_LOG(LogPakFile, Display, TEXT("Parsed AES encryption key from command line."));
 		}
diff --git a/Engine/Source/Editor/UnrealEd/Private/Fbx/FbxMainExport.cpp b/Engine/Source/Editor/UnrealEd/Private/Fbx/FbxMainExport.cpp
--- a/Engine/Source/Editor/UnrealEd/Private/Fbx/FbxMainExport.cpp
+++ b/Engine/Source/Editor/UnrealEd/Private/Fbx/FbxMainExport.cpp
@@ -4588,7 +4588,8 @@ FbxNode* FFbxExporter::ExportStaticMeshToFbx(const UStaticMesh* StaticMesh, int3
 			check(UVsLayer);
 
 			FString UVChannelNameBuilder = TEXT("UVmap_") + FString::FromInt(TexCoordSourceIndex);
-			const char* UVChannelName = TCHAR_TO_UTF8(*UVChannelNameBuilder); // actually UTF8 as required by Fbx, but can't use UE's UTF8CHAR type because that's a uint8 aka *unsigned* char
+			const auto UVChannelNameUTF8 = TStringConversion<FTCHARToUTF8_Convert>(*UVChannelNameBuilder); // Do not inline it! The lifetime of this object needs to extend over the usage of the converted buffer.
+			const char* UVChannelName = UVChannelNameUTF8.Get(); // actually UTF8 as required by Fbx, but can't use UE's UTF8CHAR type because that's a uint8 aka *unsigned* char
 			if ((LightmapUVChannel >= 0) || ((LightmapUVChannel == -1) && (TexCoordSourceIndex == StaticMesh->GetLightMapCoordinateIndex())))
 			{
 				UVChannelName = "LightMapUV";
@@ -4910,7 +4911,8 @@ void FFbxExporter::ExportSplineMeshToFbx(const USplineMeshComponent* SplineMeshC
 			UVsLayer = Mesh->GetLayer(TexCoordSourceIndex);
 		}
 		FString UVChannelNameBuilder = TEXT("UVmap_") + FString::FromInt(TexCoordSourceIndex);
-		const char* UVChannelName = TCHAR_TO_UTF8(*UVChannelNameBuilder); // actually UTF8 as required by Fbx, but can't use UE's UTF8CHAR type because that's a uint8 aka *unsigned* char
+		const auto UVChannelNameUTF8 = TStringConversion<FTCHARToUTF8_Convert>(*UVChannelNameBuilder); // Do not inline it! The lifetime of this object needs to extend over the usage of the converted buffer.
+		const char* UVChannelName = UVChannelNameUTF8.Get(); // actually UTF8 as required by Fbx, but can't use UE's UTF8CHAR type because that's a uint8 aka *unsigned* char
 		if (TexCoordSourceIndex == StaticMesh->GetLightMapCoordinateIndex())
 		{
 			UVChannelName = "LightMapUV";
diff --git a/Engine/Source/Runtime/NVidia/GeForceNOW/Private/GeForceNOWWrapper.cpp b/Engine/Source/Runtime/NVidia/GeForceNOW/Private/GeForceNOWWrapper.cpp
--- a/Engine/Source/Runtime/NVidia/GeForceNOW/Private/GeForceNOWWrapper.cpp
+++ b/Engine/Source/Runtime/NVidia/GeForceNOW/Private/GeForceNOWWrapper.cpp
@@ -135,8 +135,7 @@ GfnRuntimeError GeForceNOWWrapper::IsRunningInCloudSecure(GfnIsRunningInCloudAss
 
 GfnRuntimeError GeForceNOWWrapper::SetupTitle(const FString& InPlatformAppId) const
 {
-	char* PlatformAppId = TCHAR_TO_ANSI(*InPlatformAppId);
-	return GfnSetupTitle(PlatformAppId);
+	return GfnSetupTitle(TCHAR_TO_ANSI(*InPlatformAppId));
 }
 
 GfnRuntimeError GeForceNOWWrapper::NotifyAppReady(bool bSuccess, const FString& InStatus) const
@@ -146,9 +145,7 @@ GfnRuntimeError GeForceNOWWrapper::NotifyAppReady(bool bSuccess, const FString&
 
 GfnRuntimeError GeForceNOWWrapper::NotifyTitleExited(const FString& InPlatformId, const FString& InPlatformAppId) const
 {
-	char* PlatformId = TCHAR_TO_ANSI(*InPlatformId);
-	char* PlatformAppId = TCHAR_TO_ANSI(*InPlatformAppId);
-	return GfnTitleExited(PlatformId, PlatformAppId);
+	return GfnTitleExited(TCHAR_TO_ANSI(*InPlatformId), TCHAR_TO_ANSI(*InPlatformAppId));
 }
 
 GfnRuntimeError GeForceNOWWrapper::StartStream(StartStreamInput& InStartStreamInput, StartStreamResponse& OutResponse) const
@@ -269,8 +266,7 @@ GfnRuntimeError GeForceNOWWrapper::GetTitlesAvailable(FString& OutAvailableTitle
 
 GfnRuntimeError GeForceNOWWrapper::IsTitleAvailable(const FString& InTitleID, bool& OutbIsAvailable) const
 {
-	char* TitleID = TCHAR_TO_ANSI(*InTitleID);
-	GfnRuntimeError ErrorCode = GfnIsTitleAvailable(TitleID, &OutbIsAvailable);
+	GfnRuntimeError ErrorCode = GfnIsTitleAvailable(TCHAR_TO_ANSI(*InTitleID), &OutbIsAvailable);
 	return ErrorCode;
 }
 
