Refactor ReloadPackages to use a TWeakObjPtr<UPackage> prior to creating a new map in editor.

This change also attempts to reduce some of the complexity of the function by breaking out some of the filtering that happens at the start.

diff --git a/Engine/Source/Editor/UnrealEd/Private/PackageTools.cpp b/Engine/Source/Editor/UnrealEd/Private/PackageTools.cpp
--- a/Engine/Source/Editor/UnrealEd/Private/PackageTools.cpp
+++ b/Engine/Source/Editor/UnrealEd/Private/PackageTools.cpp
@@ -2,6 +2,7 @@
 
 
 #include "PackageTools.h"
+#include "Algo/Transform.h"
 #include "BlueprintCompilationManager.h"
 #include "UObject/PackageReload.h"
 #include "Misc/MessageDialog.h"
@@ -38,6 +39,7 @@
 #include "FileHelpers.h"
 
 #include "Framework/Notifications/NotificationManager.h"
+#include "UObject/WeakObjectPtrTemplates.h"
 #include "Widgets/Notifications/SNotificationList.h"
 
 #include "AssetData.h"
@@ -523,7 +525,7 @@ UPackageTools::UPackageTools(const FObjectInitializer& ObjectInitializer)
 	{
 		FText ErrorMessage;
 		const bool bResult = ReloadPackages(TopLevelPackages, ErrorMessage, EReloadPackagesInteractionMode::Interactive);
-		
+
 		if (!ErrorMessage.IsEmpty())
 		{
 			FMessageDialog::Open(EAppMsgType::Ok, ErrorMessage);
@@ -539,133 +541,197 @@ UPackageTools::UPackageTools(const FObjectInitializer& ObjectInitializer)
 	}
 
 
-	bool UPackageTools::ReloadPackages( const TArray<UPackage*>& TopLevelPackages, FText& OutErrorMessage, const EReloadPackagesInteractionMode InteractionMode )
+	namespace UE::PackageTools::Private
 	{
-		bool bResult = false;
-
-		FTextBuilder ErrorMessageBuilder;
-
-		// Split the set of selected top level packages into packages which are dirty or in-memory (and thus cannot be reloaded) and packages that are not dirty (and thus can be reloaded).
-		TArray<UPackage*> PackagesToReload;
+	struct FFilteredPackages
+	{
+		void Add(UPackage* RealPackage)
 		{
-			TArray<UPackage*> DirtyPackages;
-			TArray<UPackage*> InMemoryPackages;
-			for (UPackage* TopLevelPackage : TopLevelPackages)
+			if (RealPackage->HasAnyPackageFlags(PKG_InMemoryOnly))
 			{
-				if (TopLevelPackage)
-				{
-					// Get outermost packages, in case groups were selected.
-					UPackage* RealPackage = TopLevelPackage->GetOutermost() ? TopLevelPackage->GetOutermost() : TopLevelPackage;
+				InMemoryPackages.AddUnique(RealPackage);
+			}
+			else if (RealPackage->IsDirty())
+			{
+				DirtyPackages.AddUnique(RealPackage);
+			}
+			else
+			{
+				PackagesToReload.AddUnique(RealPackage);
+			}
+		}
 
-					if (RealPackage->HasAnyPackageFlags(PKG_InMemoryOnly))
-					{
-						InMemoryPackages.AddUnique(RealPackage);
-					}
-					else if (RealPackage->IsDirty())
-					{
-						DirtyPackages.AddUnique(RealPackage);
-					}
-					else
-					{
-						PackagesToReload.AddUnique(RealPackage);
-					}
-				}
+		void Remove(UPackage* InPackage)
+		{
+			PackagesToReload.Remove(InPackage);
+		}
 
-				// How should we handle locally dirty packages?
-				if (DirtyPackages.Num() > 0)
-				{
-					EAppReturnType::Type ReloadDirtyPackagesResult = EAppReturnType::No;
+		bool Contains(UPackage* InPackage) const
+		{
+			return PackagesToReload.Contains(InPackage);
+		}
 
-					// Ask the user whether dirty packages should be reloaded.
-					if (InteractionMode == EReloadPackagesInteractionMode::Interactive)
-					{
-						FTextBuilder ReloadDirtyPackagesMsgBuilder;
-						ReloadDirtyPackagesMsgBuilder.AppendLine(NSLOCTEXT("UnrealEd", "ShouldReloadDirtyPackagesHeader", "The following packages have been modified:"));
-						{
-							ReloadDirtyPackagesMsgBuilder.Indent();
-							for (UPackage* DirtyPackage : DirtyPackages)
-							{
-								ReloadDirtyPackagesMsgBuilder.AppendLine(DirtyPackage->GetFName());
-							}
-							ReloadDirtyPackagesMsgBuilder.Unindent();
-						}
-						ReloadDirtyPackagesMsgBuilder.AppendLine(NSLOCTEXT("UnrealEd", "ShouldReloadDirtyPackagesFooter", "Would you like to reload these packages? This will revert any changes you have made."));
+		TArray<UPackage*> PackagesToReload;
+		TArray<UPackage*> DirtyPackages;
+		TArray<UPackage*> InMemoryPackages;
+	};
 
-						ReloadDirtyPackagesResult = FMessageDialog::Open(EAppMsgType::YesNo, ReloadDirtyPackagesMsgBuilder.ToText());
-					}
-					else if (InteractionMode == EReloadPackagesInteractionMode::AssumePositive)
-					{
-						ReloadDirtyPackagesResult = EAppReturnType::Yes;
-					}
+	FFilteredPackages GetPackagesToReload(const TArray<UPackage*>& TopLevelPackages)
+	{
+		FFilteredPackages Filtered;
+		Algo::TransformIf(TopLevelPackages, Filtered,
+						  [](UPackage* TopLevelPackage) {return TopLevelPackage;},
+						  [](UPackage* TopLevelPackage)
+						  {
+							  return TopLevelPackage;
+						  });
+		return Filtered;
+	}
 
-					if (ReloadDirtyPackagesResult == EAppReturnType::Yes)
-					{
-						for (UPackage* DirtyPackage : DirtyPackages)
-						{
-							DirtyPackage->SetDirtyFlag(false);
-							PackagesToReload.AddUnique(DirtyPackage);
-						}
-						DirtyPackages.Reset();
-					}
-				}
-			}
+	void PromptUserForDirtyPackages(FFilteredPackages& Filtered, const EReloadPackagesInteractionMode InteractionMode)
+	{
+		// How should we handle locally dirty packages?
+		if (Filtered.DirtyPackages.Num() == 0)
+		{
+			return;
+		}
 
-			// Inform the user that dirty packages won't be reloaded.
-			if (DirtyPackages.Num() > 0)
+		EAppReturnType::Type ReloadDirtyPackagesResult = EAppReturnType::No;
+
+		// Ask the user whether dirty packages should be reloaded.
+		if (InteractionMode == EReloadPackagesInteractionMode::Interactive)
+		{
+			FTextBuilder ReloadDirtyPackagesMsgBuilder;
+			ReloadDirtyPackagesMsgBuilder.AppendLine(NSLOCTEXT("UnrealEd", "ShouldReloadDirtyPackagesHeader", "The following packages have been modified:"));
 			{
-				if (!ErrorMessageBuilder.IsEmpty())
+				ReloadDirtyPackagesMsgBuilder.Indent();
+				for (UPackage* DirtyPackage : Filtered.DirtyPackages)
 				{
-					ErrorMessageBuilder.AppendLine();
+					ReloadDirtyPackagesMsgBuilder.AppendLine(DirtyPackage->GetFName());
 				}
+				ReloadDirtyPackagesMsgBuilder.Unindent();
+			}
+			ReloadDirtyPackagesMsgBuilder.AppendLine(NSLOCTEXT("UnrealEd", "ShouldReloadDirtyPackagesFooter", "Would you like to reload these packages? This will revert any changes you have made."));
 
-				ErrorMessageBuilder.AppendLine(NSLOCTEXT("UnrealEd", "Error_ReloadDirtyPackagesHeader", "The following packages have been modified and cannot be reloaded:"));
-				{
-					ErrorMessageBuilder.Indent();
-					for (UPackage* DirtyPackage : DirtyPackages)
-					{
-						ErrorMessageBuilder.AppendLine(DirtyPackage->GetFName());
-					}
-					ErrorMessageBuilder.Unindent();
-				}
-				ErrorMessageBuilder.AppendLine(NSLOCTEXT("UnrealEd", "Error_ReloadDirtyPackagesFooter", "Saving these packages will allow them to be reloaded."));
+			ReloadDirtyPackagesResult = FMessageDialog::Open(EAppMsgType::YesNo, ReloadDirtyPackagesMsgBuilder.ToText());
+		}
+		else if (InteractionMode == EReloadPackagesInteractionMode::AssumePositive)
+		{
+			ReloadDirtyPackagesResult = EAppReturnType::Yes;
+		}
+
+		if (ReloadDirtyPackagesResult == EAppReturnType::Yes)
+		{
+			for (UPackage* DirtyPackage : Filtered.DirtyPackages)
+			{
+				DirtyPackage->SetDirtyFlag(false);
+				Filtered.PackagesToReload.AddUnique(DirtyPackage);
 			}
+			Filtered.DirtyPackages.Reset();
+		}
+	}
+
+	void InformUserAboutDirtyPackages(const FFilteredPackages& Filtered, FTextBuilder& ErrorMessageBuilder)
+	{
+		if (Filtered.DirtyPackages.Num() == 0)
+		{
+			return;
+		}
+		// Inform the user that dirty packages won't be reloaded.
+		if (!ErrorMessageBuilder.IsEmpty())
+		{
+			ErrorMessageBuilder.AppendLine();
+		}
 
-			// Inform the user that in-memory packages won't be reloaded.
-			if (InMemoryPackages.Num() > 0)
+		ErrorMessageBuilder.AppendLine(NSLOCTEXT("UnrealEd", "Error_ReloadDirtyPackagesHeader", "The following packages have been modified and cannot be reloaded:"));
+		{
+			ErrorMessageBuilder.Indent();
+			for (UPackage* DirtyPackage : Filtered.DirtyPackages)
 			{
-				if (!ErrorMessageBuilder.IsEmpty())
-				{
-					ErrorMessageBuilder.AppendLine();
-				}
+				ErrorMessageBuilder.AppendLine(DirtyPackage->GetFName());
+			}
+			ErrorMessageBuilder.Unindent();
+		}
+		ErrorMessageBuilder.AppendLine(NSLOCTEXT("UnrealEd", "Error_ReloadDirtyPackagesFooter", "Saving these packages will allow them to be reloaded."));
+	}
 
-				ErrorMessageBuilder.AppendLine(NSLOCTEXT("UnrealEd", "Error_ReloadInMemoryPackagesHeader", "The following packages are in-memory only and cannot be reloaded:"));
-				{
-					ErrorMessageBuilder.Indent();
-					for (UPackage* InMemoryPackage : InMemoryPackages)
-					{
-						ErrorMessageBuilder.AppendLine(InMemoryPackage->GetFName());
-					}
-					ErrorMessageBuilder.Unindent();
-				}
+	void InformUsersAboutInMemoryPackages(const FFilteredPackages& Filtered, FTextBuilder& ErrorMessageBuilder)
+	{
+		if (Filtered.InMemoryPackages.Num() == 0)
+		{
+			return;
+		}
+		if (!ErrorMessageBuilder.IsEmpty())
+		{
+			ErrorMessageBuilder.AppendLine();
+		}
+		ErrorMessageBuilder.AppendLine(NSLOCTEXT("UnrealEd", "Error_ReloadInMemoryPackagesHeader", "The following packages are in-memory only and cannot be reloaded:"));
+		{
+			ErrorMessageBuilder.Indent();
+			for (UPackage* InMemoryPackage : Filtered.InMemoryPackages)
+			{
+				ErrorMessageBuilder.AppendLine(InMemoryPackage->GetFName());
 			}
+			ErrorMessageBuilder.Unindent();
+		}
+	}
+
+	struct FScopedTrackFilteredPackages
+	{
+		FScopedTrackFilteredPackages(FFilteredPackages& InFiltered) :
+			Filtered(InFiltered)
+		{
+			Algo::Transform(Filtered.PackagesToReload, WeakPackagesToReload, [](UPackage* InPackage) -> TWeakObjectPtr<UPackage> { return InPackage; });
 		}
 
-		// Get the current world.
-		TWeakObjectPtr<UWorld> CurrentWorld;
+		~FScopedTrackFilteredPackages()
+		{
+			Filtered.PackagesToReload.Reset();
+			Algo::TransformIf(WeakPackagesToReload, Filtered.PackagesToReload,
+							  [](TWeakObjectPtr<UPackage> InPackage) {return InPackage.IsValid();},
+							  [](TWeakObjectPtr<UPackage> InPackage) -> UPackage* {return InPackage.Get();});
+
+		}
+
+		TArray<TWeakObjectPtr<UPackage>> WeakPackagesToReload;
+		FFilteredPackages& Filtered;
+	};
+
+	TWeakObjectPtr<UWorld> GetCurrentWorld()
+	{
 		if (GIsEditor)
 		{
 			if (UWorld* EditorWorld = GEditor->GetEditorWorldContext().World())
 			{
-				CurrentWorld = EditorWorld;
+				return EditorWorld;
 			}
 		}
 		else if (UGameEngine* GameEngine = Cast<UGameEngine>(GEngine))
 		{
 			if (UWorld* GameWorld = GameEngine->GetGameWorld())
 			{
-				CurrentWorld = GameWorld;
+				return GameWorld;
 			}
 		}
+		return nullptr;
+	}
+
+	}
+
+	bool UPackageTools::ReloadPackages( const TArray<UPackage*>& TopLevelPackages, FText& OutErrorMessage, const EReloadPackagesInteractionMode InteractionMode )
+	{
+		bool bResult = false;
+
+		FTextBuilder ErrorMessageBuilder;
+
+		using namespace UE::PackageTools::Private;
+		FFilteredPackages Filtered = GetPackagesToReload(TopLevelPackages);
+
+		PromptUserForDirtyPackages(Filtered, InteractionMode);
+		InformUserAboutDirtyPackages(Filtered, ErrorMessageBuilder);
+		InformUsersAboutInMemoryPackages(Filtered, ErrorMessageBuilder);
+
+		TWeakObjectPtr<UWorld> CurrentWorld = GetCurrentWorld();
 
 		// Check to see if we need to reload the current world.
 		FName WorldNameToReload;
@@ -675,44 +741,25 @@ UPackageTools::UPackageTools(const FObjectInitializer& ObjectInitializer)
 			// Is the current world being reloaded? If so, we just reset the current world and load it again at the end rather than let it go through ReloadPackage 
 			// (which doesn't work for the current world due to some assumptions about worlds, and their lifetimes).
 			// We also need to skip the build data package as that will also be destroyed by the transition.
-			if (PackagesToReload.Contains(CurrentWorldPtr->GetOutermost()))
+			if (Filtered.Contains(CurrentWorldPtr->GetOutermost()))
 			{
 				// Cache this so we can reload the world later
 				WorldNameToReload = *CurrentWorldPtr->GetPathName();
 
 				// Remove the world package from the reload list
-				PackagesToReload.Remove(CurrentWorldPtr->GetOutermost());
-
-				// Remove the level build data package from the reload list as creating a new map will unload build data for the current world
-				for (int32 LevelIndex = 0; LevelIndex < CurrentWorldPtr->GetNumLevels(); ++LevelIndex)
-				{
-					ULevel* Level = CurrentWorldPtr->GetLevel(LevelIndex);
-					if (Level->MapBuildData)
-					{
-						PackagesToReload.Remove(Level->MapBuildData->GetOutermost());
-					}
-				}
-
-				// Remove any streaming levels from the reload list as creating a new map will unload streaming levels for the current world
-				for (ULevelStreaming* StreamingLevel : CurrentWorldPtr->GetStreamingLevels())
-				{
-					if (StreamingLevel->IsLevelLoaded())
-					{
-						UPackage* StreamingLevelPackage = StreamingLevel->GetLoadedLevel()->GetOutermost();
-						PackagesToReload.Remove(StreamingLevelPackage);
-					}
-				}
+				Filtered.Remove(CurrentWorldPtr->GetOutermost());
 
 				// Unload the current world
 				if (GIsEditor)
 				{
+					FScopedTrackFilteredPackages TrackPackages(Filtered);
 					const bool bPromptForSave = InteractionMode == UPackageTools::EReloadPackagesInteractionMode::Interactive;
 					GEditor->CreateNewMapForEditing(bPromptForSave);
 				}
 				else if (UGameEngine* GameEngine = Cast<UGameEngine>(GEngine))
 				{
 					// Outside of the editor we need to keep the packages alive to stop the world transition from GC'ing them
-					TGCObjectsScopeGuard<UPackage> KeepPackagesAlive(PackagesToReload);
+					TGCObjectsScopeGuard<UPackage> KeepPackagesAlive(Filtered.PackagesToReload);
 
 					FString LoadMapError;
 					GameEngine->LoadMap(GameEngine->GetWorldContextFromWorldChecked(CurrentWorldPtr), FURL(TEXT("/Engine/Maps/Templates/Template_Default")), nullptr, LoadMapError);
@@ -728,7 +775,7 @@ UPackageTools::UPackageTools(const FObjectInitializer& ObjectInitializer)
 
 					Level->ReleaseRenderingResources();
 
-					if (PackagesToReload.Contains(Level->GetOutermost()))
+					if (Filtered.Contains(Level->GetOutermost()))
 					{
 						for (ULevelStreaming* StreamingLevel : CurrentWorldPtr->GetStreamingLevels())
 						{
@@ -745,6 +792,7 @@ UPackageTools::UPackageTools(const FObjectInitializer& ObjectInitializer)
 			}
 		}
 
+		TArray<UPackage*>& PackagesToReload = Filtered.PackagesToReload;
 		if (PackagesToReload.Num() > 0)
 		{
 			const FScopedBusyCursor BusyCursor;
