DDC: Added -ExecuteBuildsLocally to use the local build executor

#rnx

diff --git a/Engine/Source/Editor/UnrealEd/Private/EditorEngine.cpp b/Engine/Source/Editor/UnrealEd/Private/EditorEngine.cpp
--- a/Engine/Source/Editor/UnrealEd/Private/EditorEngine.cpp
+++ b/Engine/Source/Editor/UnrealEd/Private/EditorEngine.cpp
@@ -238,6 +238,7 @@
 #include "Engine/TextureCube.h"
 #include "Misc/PackageAccessTracking.h"
 #include "Misc/PackageAccessTrackingOps.h"
+#include "DerivedDataBuildLocalExecutor.h"
 #include "DerivedDataBuildRemoteExecutor.h"
 #include "DerivedDataBuildWorkers.h"
 
@@ -634,7 +635,14 @@ void UEditorEngine::InitEditor(IEngineLoop* InEngineLoop)
 {
 	// Allow remote execution of derived data builds from this point
 	// TODO: This needs to be enabled earlier to allow early data builds to be remote executed.
-	InitDerivedDataBuildRemoteExecutor();
+	if (FParse::Param(FCommandLine::Get(), TEXT("ExecuteBuildsLocally")))
+	{
+		InitDerivedDataBuildLocalExecutor();
+	}
+	else
+	{
+		InitDerivedDataBuildRemoteExecutor();
+	}
 	InitDerivedDataBuildWorkers();
 
 	// Call base.
