Display Status Icon next to Pin if any property is being watched

diff --git a/Engine/Source/Editor/GraphEditor/Private/SGraphPin.cpp b/Engine/Source/Editor/GraphEditor/Private/SGraphPin.cpp
--- a/Engine/Source/Editor/GraphEditor/Private/SGraphPin.cpp
+++ b/Engine/Source/Editor/GraphEditor/Private/SGraphPin.cpp
@@ -1133,7 +1133,7 @@ const FSlateBrush* SGraphPin::GetPinStatusIcon() const
 		{
 			UBlueprint* Blueprint = FBlueprintEditorUtils::FindBlueprintForNodeChecked(GraphNode);
 
-			if (FKismetDebugUtilities::IsPinBeingWatched(Blueprint, WatchedPin))
+			if (FKismetDebugUtilities::DoesPinHaveWatches(Blueprint, WatchedPin))
 			{
 				return FEditorStyle::GetBrush(TEXT("Graph.WatchedPinIcon_Pinned"));
 			}
@@ -1151,10 +1151,16 @@ EVisibility SGraphPin::GetPinStatusIconVisibility() const
 		return EVisibility::Collapsed;
 	}
 
+	UEdGraphNode* GraphNode = GraphPin->GetOwningNode();
+	UBlueprint* Blueprint = FBlueprintEditorUtils::FindBlueprintForNode(GraphNode);
+	if (!Blueprint)
+	{
+		return EVisibility::Collapsed;
+	}
+
 	UEdGraphPin const* WatchedPin = ((GraphPin->Direction == EGPD_Input) && (GraphPin->LinkedTo.Num() > 0)) ? GraphPin->LinkedTo[0] : GraphPin;
 
-	UEdGraphSchema const* Schema = GraphPin->GetSchema();
-	return Schema && Schema->IsPinBeingWatched(WatchedPin) ? EVisibility::Visible : EVisibility::Collapsed;
+	return FKismetDebugUtilities::DoesPinHaveWatches(Blueprint, WatchedPin) ? EVisibility::Visible : EVisibility::Collapsed;
 }
 
 FReply SGraphPin::ClickedOnPinStatusIcon()
diff --git a/Engine/Source/Editor/UnrealEd/Private/Kismet2/KismetDebugUtilities.cpp b/Engine/Source/Editor/UnrealEd/Private/Kismet2/KismetDebugUtilities.cpp
--- a/Engine/Source/Editor/UnrealEd/Private/Kismet2/KismetDebugUtilities.cpp
+++ b/Engine/Source/Editor/UnrealEd/Private/Kismet2/KismetDebugUtilities.cpp
@@ -1413,6 +1413,21 @@ bool FKismetDebugUtilities::IsPinBeingWatched(const UBlueprint* Blueprint, const
 	return false;
 }
 
+bool FKismetDebugUtilities::DoesPinHaveWatches(const UBlueprint* Blueprint, const UEdGraphPin* Pin)
+{
+	if (TArray<FBlueprintWatchedPin>* WatchedPins = GetWatchedPins(Blueprint))
+	{
+		for (const FBlueprintWatchedPin& WatchedPin : *WatchedPins)
+		{
+			if (WatchedPin.Get() == Pin)
+			{
+				return true;
+			}
+		}
+	}
+	return false;
+}
+
 bool FKismetDebugUtilities::RemovePinWatch(const UBlueprint* Blueprint, const UEdGraphPin* Pin, const TArray<FName>& InPathToProperty)
 {
 	return RemovePinPropertyWatchesByPredicate(
diff --git a/Engine/Source/Editor/UnrealEd/Public/Kismet2/KismetDebugUtilities.h b/Engine/Source/Editor/UnrealEd/Public/Kismet2/KismetDebugUtilities.h
--- a/Engine/Source/Editor/UnrealEd/Public/Kismet2/KismetDebugUtilities.h
+++ b/Engine/Source/Editor/UnrealEd/Public/Kismet2/KismetDebugUtilities.h
@@ -317,12 +317,56 @@ public:
 	// Notifies listeners when a watched pin is added or removed
 	static FOnWatchedPinsListChanged WatchedPinsListChangedEvent;
 
+	/**
+	 * Returns whether a pin property can be watched
+	 * @param Blueprint The blueprint that owns the pin
+	 * @param Pin The Pin to check if it can be watched
+	 * @param InPathToProperty Path to the property on the pin to watch
+	 */
 	static bool CanWatchPin(const UBlueprint* Blueprint, const UEdGraphPin* Pin, const TArray<FName>& InPathToProperty = TArray<FName>());
+
+	/** 
+	 * Returns whether a pin property is being watched
+	 * @param Blueprint The blueprint that owns the pin
+	 * @param Pin The Pin to check if it's being watched
+	 * @param InPathToProperty Path to the property on the pin to watch
+	 */
 	static bool IsPinBeingWatched(const UBlueprint* Blueprint, const UEdGraphPin* Pin, const TArray<FName>& InPathToProperty = TArray<FName>());
+	
+	/**
+	 * Returns whether there are any watched properties for a given pin 
+	 * @param Blueprint The blueprint that owns the pin
+	 * @param Pin The Pin to check if it has any watches
+	 */
+	static bool DoesPinHaveWatches(const UBlueprint* Blueprint, const UEdGraphPin* Pin);
+
+	/**
+	 * Toggles whether a pin is being watched 
+	 * @param Blueprint The blueprint that owns the pin
+	 * @param Pin The Pin to watch
+	 */
 	static void TogglePinWatch(const UBlueprint* Blueprint, const UEdGraphPin* Pin);
+
+	/**
+	 * Removes a pin property Watch 
+	 * @param Blueprint The blueprint that owns the pin
+	 * @param Pin The Pin to stop watching
+	 * @param InPathToProperty Path to the property on the pin to stop watching
+	 * @return true if a watch was found and removed
+	 */
 	static bool RemovePinWatch(const UBlueprint* Blueprint, const UEdGraphPin* Pin, const TArray<FName>& InPathToProperty = TArray<FName>());
+
+	/** 
+	 * Adds a pin property watch 
+	 * @param Blueprint The blueprint that owns the pin
+	 * @param WatchedPin The Pin to watch
+	 */
 	static void AddPinWatch(const UBlueprint* Blueprint, FBlueprintWatchedPin&& WatchedPin);
+
+	/** Removes all Watched pins from a blueprint */
 	static void ClearPinWatches(const UBlueprint* Blueprint);
+
+	/** Returns whether any pins are watched for a Blueprint */
 	static bool BlueprintHasPinWatches(const UBlueprint* Blueprint);
 
 	/**
