#UE Add -SearchAllAssets option for resave commandlets to allow soft references to redirectors to be fixed up at save time.

diff --git a/Engine/Source/Editor/UnrealEd/Private/Commandlets/ContentCommandlets.cpp b/Engine/Source/Editor/UnrealEd/Private/Commandlets/ContentCommandlets.cpp
--- a/Engine/Source/Editor/UnrealEd/Private/Commandlets/ContentCommandlets.cpp
+++ b/Engine/Source/Editor/UnrealEd/Private/Commandlets/ContentCommandlets.cpp
@@ -340,7 +340,10 @@ int32 UResavePackagesCommandlet::InitializeResaveParameters( const TArray<FStrin
 	// This option will filter the package list and only save packages that are redirectors, or that reference redirectors
 	const bool bFixupRedirects = (Switches.Contains(TEXT("FixupRedirects")) || Switches.Contains(TEXT("FixupRedirectors")));
 
-	if (bResaveDirectRefsAndDeps || bFixupRedirects || bOnlyMaterials)
+	// This option allows the dependency graph and soft object path redirect map to be populated. This is useful if you want soft object references to redirectors to be followed to the destination asset at save time.
+	const bool bSearchAllAssets = Switches.Contains(TEXT("SearchAllAssets"));
+
+	if (bResaveDirectRefsAndDeps || bFixupRedirects || bOnlyMaterials || bSearchAllAssets)
 	{
 		AssetRegistry.SearchAllAssets(true);
 
