Fixed not being able to place cooked anim BPs in level viewports

Turns off some high-level class checks that only let AActor, ABrush and AVolumes be place-able prior to this change
Exports target skeleton asset registry tag for the anim BP class so that it can be subsequently used by factory methods
Modifies actor factory to also interrogate anim BP generated classes for relevant spawn data

diff --git a/Engine/Source/Editor/UnrealEd/Private/Factories/ActorFactory.cpp b/Engine/Source/Editor/UnrealEd/Private/Factories/ActorFactory.cpp
--- a/Engine/Source/Editor/UnrealEd/Private/Factories/ActorFactory.cpp
+++ b/Engine/Source/Editor/UnrealEd/Private/Factories/ActorFactory.cpp
@@ -1163,7 +1163,7 @@ bool UActorFactorySkeletalMesh::CanCreateActorFrom( const FAssetData& AssetData,
 		SkeletalMeshData = AssetData;
 	}
 
-	if ( !SkeletalMeshData.IsValid() && AssetData.GetClass()->IsChildOf( UAnimBlueprint::StaticClass() ) )
+	if ( !SkeletalMeshData.IsValid() && (AssetData.GetClass()->IsChildOf(UAnimBlueprint::StaticClass()) || AssetData.GetClass()->IsChildOf(UAnimBlueprintGeneratedClass::StaticClass())))
 	{
 		const FString TargetSkeletonPath = AssetData.GetTagValueRef<FString>( "TargetSkeleton" );
 		if ( TargetSkeletonPath.IsEmpty() )
@@ -1250,14 +1250,26 @@ bool UActorFactorySkeletalMesh::CanCreateActorFrom( const FAssetData& AssetData,
 
 USkeletalMesh* UActorFactorySkeletalMesh::GetSkeletalMeshFromAsset( UObject* Asset )
 {
-	USkeletalMesh*SkeletalMesh = Cast<USkeletalMesh>( Asset );
-	UAnimBlueprint* AnimBlueprint = Cast<UAnimBlueprint>( Asset );
+	USkeletalMesh* SkeletalMesh = Cast<USkeletalMesh>( Asset );
 	USkeleton* Skeleton = Cast<USkeleton>( Asset );
 
-	if( SkeletalMesh == nullptr && AnimBlueprint != nullptr && AnimBlueprint->TargetSkeleton )
+	if(SkeletalMesh == nullptr)
 	{
 		// base it on preview skeletal mesh, just to have something
-		SkeletalMesh = AnimBlueprint->TargetSkeleton->GetPreviewMesh(true);
+		if(UAnimBlueprint* AnimBlueprint = Cast<UAnimBlueprint>(Asset))
+		{
+			if(AnimBlueprint->TargetSkeleton)
+			{
+				SkeletalMesh = AnimBlueprint->TargetSkeleton->GetPreviewMesh(true);
+			}
+		}
+		if(UAnimBlueprintGeneratedClass* AnimBlueprintGeneratedClass = Cast<UAnimBlueprintGeneratedClass>(Asset))
+		{
+			if(AnimBlueprintGeneratedClass->TargetSkeleton)
+			{
+				SkeletalMesh = AnimBlueprintGeneratedClass->TargetSkeleton->GetPreviewMesh(true);
+			}
+		}
 	}
 
 	if( SkeletalMesh == nullptr && Skeleton != nullptr )
diff --git a/Engine/Source/Editor/UnrealEd/Private/ObjectTools.cpp b/Engine/Source/Editor/UnrealEd/Private/ObjectTools.cpp
--- a/Engine/Source/Editor/UnrealEd/Private/ObjectTools.cpp
+++ b/Engine/Source/Editor/UnrealEd/Private/ObjectTools.cpp
@@ -4395,7 +4395,7 @@ namespace ObjectTools
 	{
 		check(InClass);
 
-		const bool bIsPlaceable = !InClass->HasAllClassFlags(CLASS_NotPlaceable) && (InClass->IsChildOf(AActor::StaticClass()) || InClass->IsChildOf(ABrush::StaticClass()) || InClass->IsChildOf(AVolume::StaticClass()));
+		const bool bIsPlaceable = !InClass->HasAllClassFlags(CLASS_NotPlaceable);
 		const bool bIsAbstractOrDeprecated = InClass->HasAnyClassFlags(CLASS_Abstract | CLASS_Deprecated | CLASS_NewerVersionExists);
 		const bool bIsSkeletonClass = FKismetEditorUtilities::IsClassABlueprintSkeleton(InClass);
 
diff --git a/Engine/Source/Runtime/Engine/Classes/Animation/AnimBlueprintGeneratedClass.h b/Engine/Source/Runtime/Engine/Classes/Animation/AnimBlueprintGeneratedClass.h
--- a/Engine/Source/Runtime/Engine/Classes/Animation/AnimBlueprintGeneratedClass.h
+++ b/Engine/Source/Runtime/Engine/Classes/Animation/AnimBlueprintGeneratedClass.h
@@ -362,7 +362,7 @@ class ENGINE_API UAnimBlueprintGeneratedClass : public UBlueprintGeneratedClass,
 	TArray<FBakedAnimationStateMachine> BakedStateMachines;
 
 	/** Target skeleton for this blueprint class */
-	UPROPERTY()
+	UPROPERTY(AssetRegistrySearchable)
 	TObjectPtr<USkeleton> TargetSkeleton;
 
 	/** A list of anim notifies that state machines (or anything else) may reference */
