Fix reflect visibility to game to use the proper SetHidden function

diff --git a/Engine/Source/Editor/UnrealEd/Private/EditorLevelUtils.cpp b/Engine/Source/Editor/UnrealEd/Private/EditorLevelUtils.cpp
--- a/Engine/Source/Editor/UnrealEd/Private/EditorLevelUtils.cpp
+++ b/Engine/Source/Editor/UnrealEd/Private/EditorLevelUtils.cpp
@@ -1241,7 +1241,7 @@ void SetLevelVisibilityNoGlobalUpdateInternal(ULevel* Level, const bool bShouldB
 				}
 				if (bReflectVisibilityToGame)
 				{
-					Actor->SetHidden(Actor->bHiddenEdLevel);
+					Actor->SetActorHiddenInGame(Actor->bHiddenEdLevel);
 				}
 			}
 		}
