Open Asset Panel now closes on Escape and selecting an asset

diff --git a/Engine/Source/Editor/UnrealEd/Private/Toolkits/SGlobalOpenAssetDialog.cpp b/Engine/Source/Editor/UnrealEd/Private/Toolkits/SGlobalOpenAssetDialog.cpp
--- a/Engine/Source/Editor/UnrealEd/Private/Toolkits/SGlobalOpenAssetDialog.cpp
+++ b/Engine/Source/Editor/UnrealEd/Private/Toolkits/SGlobalOpenAssetDialog.cpp
@@ -6,6 +6,7 @@
 #include "Widgets/SBoxPanel.h"
 #include "Framework/Application/SlateApplication.h"
 #include "Widgets/Layout/SBox.h"
+#include "Widgets/Docking/SDockTab.h"
 
 #include "IContentBrowserSingleton.h"
 #include "ContentBrowserModule.h"
@@ -54,7 +55,7 @@ FReply SGlobalOpenAssetDialog::OnPreviewKeyDown(const FGeometry& MyGeometry, con
 {
 	if (InKeyEvent.GetKey() == EKeys::Escape)
 	{
-		FSlateApplication::Get().DismissAllMenus();
+		RequestCloseAssetPicker();
 		return FReply::Handled();
 	}
 
@@ -67,6 +68,8 @@ void SGlobalOpenAssetDialog::OnAssetSelectedFromPicker(const FAssetData& AssetDa
 	{
 		GEditor->GetEditorSubsystem<UAssetEditorSubsystem>()->OpenEditorForAsset(ObjectToEdit);
 	}
+
+	RequestCloseAssetPicker();
 }
 
 void SGlobalOpenAssetDialog::OnPressedEnterOnAssetsInPicker(const TArray<FAssetData>& SelectedAssets)
@@ -78,8 +81,19 @@ void SGlobalOpenAssetDialog::OnPressedEnterOnAssetsInPicker(const TArray<FAssetD
 			GEditor->GetEditorSubsystem<UAssetEditorSubsystem>()->OpenEditorForAsset(ObjectToEdit);
 		}
 	}
+
+	RequestCloseAssetPicker();
+}
+
+void SGlobalOpenAssetDialog::RequestCloseAssetPicker()
+{
+	if (TSharedPtr<SDockTab> AssetPickerTab = FGlobalTabmanager::Get()->FindExistingLiveTab(FTabId("GlobalAssetPicker")))
+	{
+		AssetPickerTab->RequestCloseTab();
+	}
 }
 
+
 //////////////////////////////////////////////////////////////////////////
 
 #undef LOCTEXT_NAMESPACE
diff --git a/Engine/Source/Editor/UnrealEd/Private/Toolkits/SGlobalOpenAssetDialog.h b/Engine/Source/Editor/UnrealEd/Private/Toolkits/SGlobalOpenAssetDialog.h
--- a/Engine/Source/Editor/UnrealEd/Private/Toolkits/SGlobalOpenAssetDialog.h
+++ b/Engine/Source/Editor/UnrealEd/Private/Toolkits/SGlobalOpenAssetDialog.h
@@ -26,4 +26,5 @@ public:
 protected:
 	void OnAssetSelectedFromPicker(const struct FAssetData& AssetData);
 	void OnPressedEnterOnAssetsInPicker(const TArray<struct FAssetData>& SelectedAssets);
+	void RequestCloseAssetPicker();
 };
