EditorDomain: Prevent saving of packages that require resaved bulk data.

diff --git a/Engine/Source/Editor/UnrealEd/Private/EditorDomain/BulkDataRegistryEditorDomain.cpp b/Engine/Source/Editor/UnrealEd/Private/EditorDomain/BulkDataRegistryEditorDomain.cpp
--- a/Engine/Source/Editor/UnrealEd/Private/EditorDomain/BulkDataRegistryEditorDomain.cpp
+++ b/Engine/Source/Editor/UnrealEd/Private/EditorDomain/BulkDataRegistryEditorDomain.cpp
@@ -166,7 +166,7 @@ void FBulkDataRegistryEditorDomain::Register(UPackage* Owner, const UE::Virtuali
 	UE::Virtualization::FVirtualizedUntypedBulkData CopyBulk(BulkData.CopyTornOff());
 	if (Owner
 		&& Owner->GetFileSize() // We only record the BulkDataList for disk packages
-		&& !Owner->GetHasBeenEndLoaded() // We only record BulkDats that are loaded before the package finishes loading
+		&& !Owner->GetHasBeenEndLoaded() // We only record BulkDatas that are loaded before the package finishes loading
 		&& CopyBulk.CanSaveForRegistry()
 		)
 	{
@@ -174,9 +174,12 @@ void FBulkDataRegistryEditorDomain::Register(UPackage* Owner, const UE::Virtuali
 		AddPendingPackageBulkData(PackageName, CopyBulk);
 	}
 
-	FWriteScopeLock RegistryScopeLock(RegistryLock);
-	check(bActive); // Registrations should not come in after we destruct
-	Registry.Add(BulkData.GetIdentifier(), FRegisteredBulk(MoveTemp(CopyBulk), PackageName));
+	{
+		FWriteScopeLock RegistryScopeLock(RegistryLock);
+		check(bActive); // Registrations should not come in after we destruct
+		Registry.Add(BulkData.GetIdentifier(), FRegisteredBulk(MoveTemp(CopyBulk), PackageName));
+	}
+	ResaveSizeTracker.Register(Owner, BulkData);
 }
 
 void FBulkDataRegistryEditorDomain::OnExitMemory(const UE::Virtualization::FVirtualizedUntypedBulkData& BulkData)
@@ -319,6 +322,11 @@ TFuture<UE::BulkDataRegistry::FData> FBulkDataRegistryEditorDomain::GetData(cons
 		});
 }
 
+uint64 FBulkDataRegistryEditorDomain::GetBulkDataResaveSize(FName PackageName)
+{
+	return ResaveSizeTracker.GetBulkDataResaveSize(PackageName);
+}
+
 void FBulkDataRegistryEditorDomain::TickCook(float DeltaTime, bool bTickComplete)
 {
 	bool bWaitForCooldown = !bTickComplete;
diff --git a/Engine/Source/Editor/UnrealEd/Private/EditorDomain/BulkDataRegistryEditorDomain.h b/Engine/Source/Editor/UnrealEd/Private/EditorDomain/BulkDataRegistryEditorDomain.h
--- a/Engine/Source/Editor/UnrealEd/Private/EditorDomain/BulkDataRegistryEditorDomain.h
+++ b/Engine/Source/Editor/UnrealEd/Private/EditorDomain/BulkDataRegistryEditorDomain.h
@@ -169,6 +169,7 @@ public:
 	virtual void OnExitMemory(const UE::Virtualization::FVirtualizedUntypedBulkData& BulkData) override;
 	virtual TFuture<UE::BulkDataRegistry::FMetaData> GetMeta(const FGuid& BulkDataId) override;
 	virtual TFuture<UE::BulkDataRegistry::FData> GetData(const FGuid& BulkDataId) override;
+	virtual uint64 GetBulkDataResaveSize(FName PackageName) override;
 
 	// FTickableEditorObject/FTickableCookObject interface
 	virtual void Tick(float DeltaTime) override;
@@ -195,6 +196,7 @@ private:
 	FCriticalSection PendingPackageLock;
 
 	TMap<FGuid, FRegisteredBulk> Registry;
+	FResaveSizeTracker ResaveSizeTracker;
 	TMap<FGuid, FUpdatingPayload> UpdatingPayloads;
 	TMap<FName, TUniquePtr<FPendingPackage>> PendingPackages;
 	TMap<FGuid, TRefCountPtr<FPendingPayloadId>> PendingPayloadIds;
diff --git a/Engine/Source/Editor/UnrealEd/Private/EditorDomain/EditorDomain.cpp b/Engine/Source/Editor/UnrealEd/Private/EditorDomain/EditorDomain.cpp
--- a/Engine/Source/Editor/UnrealEd/Private/EditorDomain/EditorDomain.cpp
+++ b/Engine/Source/Editor/UnrealEd/Private/EditorDomain/EditorDomain.cpp
@@ -42,15 +42,13 @@ public:
 
 	static IPackageResourceManager* SetPackageResourceManager()
 	{
-		if (GIsEditor && (!IsRunningCommandlet() || IsRunningCookCommandlet()))
+		bool bEditorDomainEnabled = IsEditorDomainEnabled();
+		if (GIsEditor)
 		{
-			bool bEditorDomainEnabled = true;
-			GConfig->GetBool(TEXT("EditorDomain"), TEXT("EditorDomainEnabled"), bEditorDomainEnabled, GEditorIni);
-			if (GConfig->GetBool(TEXT("CookSettings"), TEXT("EditorDomainEnabled"), bEditorDomainEnabled, GEditorIni))
-			{
-				UE_LOG(LogEditorDomain, Error, TEXT("Editor.ini:[CookSettings]:EditorDomainEnabled is deprecated, use Editor.ini:[EditorDomain]:EditorDomainEnabled instead."));
-			}
 			UE_LOG(LogEditorDomain, Display, TEXT("EditorDomain is %s"), bEditorDomainEnabled ? TEXT("Enabled") : TEXT("Disabled"));
+		}
+		if (bEditorDomainEnabled)
+		{
 			UE::EditorDomain::UtilsInitialize();
 			UE::TargetDomain::UtilsInitialize(bEditorDomainEnabled);
 			if (bEditorDomainEnabled)
diff --git a/Engine/Source/Editor/UnrealEd/Private/EditorDomain/EditorDomainUtils.cpp b/Engine/Source/Editor/UnrealEd/Private/EditorDomain/EditorDomainUtils.cpp
--- a/Engine/Source/Editor/UnrealEd/Private/EditorDomain/EditorDomainUtils.cpp
+++ b/Engine/Source/Editor/UnrealEd/Private/EditorDomain/EditorDomainUtils.cpp
@@ -2,6 +2,7 @@
 
 #include "EditorDomain/EditorDomainUtils.h"
 
+#include "Algo/AnyOf.h"
 #include "Algo/IsSorted.h"
 #include "Algo/Unique.h"
 #include "AssetRegistry/AssetData.h"
@@ -14,6 +15,7 @@
 #include "DerivedDataPayload.h"
 #include "DerivedDataRequestOwner.h"
 #include "Editor.h"
+#include "Engine/StaticMesh.h"
 #include "HAL/FileManager.h"
 #include "Hash/Blake3.h"
 #include "Memory/SharedBuffer.h"
@@ -22,6 +24,7 @@
 #include "Misc/PackagePath.h"
 #include "Misc/ScopeRWLock.h"
 #include "Misc/StringBuilder.h"
+#include "Serialization/BulkDataRegistry.h"
 #include "Serialization/CompactBinaryWriter.h"
 #include "Serialization/MemoryReader.h"
 #include "Serialization/PackageWriterToSharedBuffer.h"
@@ -73,9 +76,10 @@ bool GTargetDomainClassUseAllowList = true;
 bool GTargetDomainClassEmptyAllowList = false;
 TArray<int32> GGlobalAddedCustomVersions;
 bool bGGlobalAddedCustomVersionsInitialized = false;
+int64 GMaxBulkDataSize = -1;
 
 // Change to a new guid when EditorDomain needs to be invalidated
-const TCHAR* EditorDomainVersion = TEXT("30E58214A4A84D638FAA8826B81338A1");
+const TCHAR* EditorDomainVersion = TEXT("D0A667ED7E374A60962A4AA61B63E6B5");
 
 // Identifier of the CacheBuckets for EditorDomain tables
 const TCHAR* EditorDomainPackageBucketName = TEXT("EditorDomainPackage");
@@ -871,6 +875,10 @@ void UtilsInitialize()
 		GTargetDomainClassBlockList = ConstructTargetIterativeClassBlockList();
 	}
 
+	double MaxBulkDataSize = 64 * 1024;
+	GConfig->GetDouble(TEXT("EditorDomain"), TEXT("MaxBulkDataSize"), MaxBulkDataSize, GEditorIni);
+	GMaxBulkDataSize = static_cast<uint64>(MaxBulkDataSize);
+
 	// Constructing allowlists requires use of UStructs, and the early SetPackageResourceManager
 	// where UtilsInitialize is called is too early; trying to call UStruct->GetSchemaHash at that
 	// time will break the UClass. Defer the construction of allowlist-based data until OnPostEngineInit
@@ -930,11 +938,6 @@ void RequestEditorDomainPackage(const FPackagePath& PackagePath,
 class FEditorDomainPackageWriter final : public TPackageWriterToSharedBuffer<IPackageWriter>
 {
 public:
-	FEditorDomainPackageWriter(uint64& InFileSize)
-		: FileSize(InFileSize)
-	{
-	}
-
 	// IPackageWriter
 	virtual FCapabilities GetCapabilities() const override
 	{
@@ -968,6 +971,16 @@ public:
 		return Attachments;
 	}
 
+	uint64 GetBulkDataSize() const
+	{
+		return BulkDataSize;
+	}
+
+	uint64 GetFileSize() const
+	{
+		return FileSize;
+	}
+
 protected:
 	virtual TFuture<FMD5Hash> CommitPackageInternal(FPackageWriterRecords::FPackage&& Record,
 		const FCommitPackageInfo& Info) override
@@ -1000,6 +1013,7 @@ protected:
 		check(Record.Package->Buffer.GetSize() > 0); // Header+Exports segment is non-zero in length
 		AttachmentBuffers.Add(Record.Package->Buffer);
 
+		BulkDataSize = 0;
 		for (const FBulkDataRecord& BulkRecord : Record.BulkDatas)
 		{
 			checkf(BulkRecord.Info.BulkDataType == IPackageWriter::FBulkDataInfo::AppendToExports,
@@ -1019,7 +1033,8 @@ protected:
 					FileRegion.Length, BulkRecord.Buffer));
 				SizeFromRegions += FileRegion.Length;
 			}
-			checkf(SizeFromRegions == BulkRecord.Buffer.GetSize(), TEXT("Expects all BulkData to be in a region."))
+			checkf(SizeFromRegions == BulkRecord.Buffer.GetSize(), TEXT("Expects all BulkData to be in a region."));
+			BulkDataSize += BulkRecord.Buffer.GetSize();
 		}
 		for (const FLinkerAdditionalDataRecord& AdditionalRecord : Record.LinkerAdditionalDatas)
 		{
@@ -1038,7 +1053,8 @@ protected:
 				SizeFromRegions += FileRegion.Length;
 			}
 			checkf(SizeFromRegions == AdditionalRecord.Buffer.GetSize(),
-				TEXT("Expects all LinkerAdditionalData to be in a region."))
+				TEXT("Expects all LinkerAdditionalData to be in a region."));
+			BulkDataSize += AdditionalRecord.Buffer.GetSize();
 		}
 
 		// We use a counter for PayloadIds rather than hashes of the Attachments. We do this because
@@ -1074,7 +1090,8 @@ protected:
 private:
 	TArray<FAttachment> Attachments;
 	FPackageWriterRecords::FWritePackage WritePackageRecord;
-	uint64& FileSize;
+	uint64 FileSize = 0;
+	uint64 BulkDataSize = 0;
 };
 
 
@@ -1087,8 +1104,9 @@ bool TrySavePackage(UPackage* Package)
 		return false;
 	}
 
+	FName PackageName = Package->GetFName();
 	FString ErrorMessage;
-	FPackageDigest PackageDigest = EditorDomain->GetPackageDigest(Package->GetFName());
+	FPackageDigest PackageDigest = EditorDomain->GetPackageDigest(PackageName);
 	if (!PackageDigest.IsSuccessful())
 	{
 		UE_LOG(LogEditorDomain, Warning, TEXT("Could not save package to EditorDomain: %s."),
@@ -1098,10 +1116,17 @@ bool TrySavePackage(UPackage* Package)
 	if (!EnumHasAnyFlags(PackageDigest.DomainUse, EDomainUse::SaveEnabled))
 	{
 		UE_LOG(LogEditorDomain, Verbose, TEXT("Skipping save of blocked package to EditorDomain: %s."),
-			*Package->GetName());
+			*WriteToString<256>(PackageName));
 		return false;
 	}
-	UE_LOG(LogEditorDomain, Verbose, TEXT("Saving to EditorDomain: %s."), *Package->GetName())
+	if (GMaxBulkDataSize >= 0 &&
+		IBulkDataRegistry::Get().GetBulkDataResaveSize(PackageName) > static_cast<uint64>(GMaxBulkDataSize))
+	{
+		UE_LOG(LogEditorDomain, Verbose, TEXT("Skipping save of package with new BulkData to EditorDomain: %s."),
+			*WriteToString<256>(PackageName));
+		return false;
+	}
+	UE_LOG(LogEditorDomain, Verbose, TEXT("Saving to EditorDomain: %s."), *WriteToString<256>(PackageName));
 
 	uint32 SaveFlags = SAVE_NoError // Do not crash the SaveServer on an error
 		| SAVE_BulkDataByReference	// EditorDomain saves reference bulkdata from the WorkspaceDomain rather than duplicating it
@@ -1110,14 +1135,21 @@ bool TrySavePackage(UPackage* Package)
 		// | SAVE_KeepGUID;			// Prevent indeterminism by keeping the Guid
 		;
 
+	TArray<UObject*> PackageObjects;
+	auto GetPackageObjects = [&PackageObjects, Package]() -> TArrayView<UObject*>
+	{
+		if (PackageObjects.IsEmpty())
+		{
+			GetObjectsWithPackage(Package, PackageObjects);
+		}
+		return PackageObjects;
+	};
 	if (GetEditorDomainSaveUnversioned())
 	{
 		// With some exceptions, EditorDomain packages are saved unversioned; 
 		// editors request the appropriate version of the EditorDomain package matching their serialization version
 		bool bSaveUnversioned = true;
-		TArray<UObject*> PackageObjects;
-		GetObjectsWithPackage(Package, PackageObjects);
-		for (UObject* Object : PackageObjects)
+		for (UObject* Object : GetPackageObjects())
 		{
 			UClass* Class = Object ? Object->GetClass() : nullptr;
 			if (Class && Class->HasAnyClassFlags(CLASS_CompiledFromBlueprint))
@@ -1132,8 +1164,7 @@ bool TrySavePackage(UPackage* Package)
 		SaveFlags |= bSaveUnversioned ? SAVE_Unversioned_Properties : 0;
 	}
 
-	uint64 FileSize = 0;
-	FEditorDomainPackageWriter* PackageWriter = new FEditorDomainPackageWriter(FileSize);
+	FEditorDomainPackageWriter* PackageWriter = new FEditorDomainPackageWriter();
 	IPackageWriter::FBeginPackageInfo BeginInfo;
 	BeginInfo.PackageName = Package->GetFName();
 	PackageWriter->BeginPackage(BeginInfo);
@@ -1253,10 +1284,36 @@ bool TrySavePackage(UPackage* Package)
 		return false;
 	}
 
+	if (GMaxBulkDataSize >= 0 && PackageWriter->GetBulkDataSize() > static_cast<uint64>(GMaxBulkDataSize))
+	{
+		// TODO EditorDomain MeshDescription: We suppress the warning if the package includes a MeshDescription;
+		// MeshDescriptions do not yet detect the need for a bulkdata resave until save serialization.
+		bool bDisplayAsWarning = !Algo::AnyOf(GetPackageObjects(), [](UObject* Object)
+			{
+				return Object->IsA<UStaticMesh>();
+			});
+
+		FString Message = FString::Printf(TEXT("Could not save package to EditorDomain because BulkData size is too large. ")
+			TEXT("Package=%s, BulkDataSize=%d, EditorDomain.MaxBulkDataSize=%d")
+			TEXT("\n\tWe did not detect this until after trying to save the package, which is bad for performance. Resave the package ")
+			TEXT("or debug why the size was not detected by the BulkDataRegistry."),
+			*WriteToString<256>(PackageName), PackageWriter->GetBulkDataSize(), GMaxBulkDataSize);
+		if (bDisplayAsWarning)
+		{
+			UE_LOG(LogEditorDomain, Display, TEXT("%s"), *Message);
+		}
+		else
+		{
+			UE_LOG(LogEditorDomain, Verbose, TEXT("%s"), *Message);
+		}
+		return false;
+	}
+
 	ICache& Cache = GetCache();
 
 	TCbWriter<16> MetaData;
 	MetaData.BeginObject();
+	uint64 FileSize = PackageWriter->GetFileSize();
 	MetaData << "FileSize" << FileSize;
 	MetaData.EndObject();
 
@@ -1418,4 +1475,4 @@ void PutBulkDataPayloadId(FName PackageName, const FGuid& BulkDataId, FSharedBuf
 	Owner.KeepAlive();
 }
 
-}
\ No newline at end of file
+}
diff --git a/Engine/Source/Runtime/CoreUObject/Private/Serialization/BulkDataRegistry.cpp b/Engine/Source/Runtime/CoreUObject/Private/Serialization/BulkDataRegistry.cpp
--- a/Engine/Source/Runtime/CoreUObject/Private/Serialization/BulkDataRegistry.cpp
+++ b/Engine/Source/Runtime/CoreUObject/Private/Serialization/BulkDataRegistry.cpp
@@ -4,7 +4,17 @@
 
 #if WITH_EDITOR
 
+#include "Containers/Array.h"
+#include "Containers/Set.h"
+#include "HAL/CriticalSection.h"
+#include "Interfaces/IPluginManager.h"
 #include "Misc/ConfigCacheIni.h"
+#include "Misc/CoreDelegates.h"
+#include "Misc/ScopeRWLock.h"
+#include "Serialization/VirtualizedBulkData.h"
+#include "UObject/Package.h"
+#include "UObject/PackageResourceManager.h"
+#include "UObject/UObjectGlobals.h"
 
 namespace UE::BulkDataRegistry::Private
 {
@@ -33,6 +43,29 @@ public:
 		Promise.SetValue(UE::BulkDataRegistry::FData{ false, FCompressedBuffer() });
 		return Promise.GetFuture();
 	}
+	virtual uint64 GetBulkDataResaveSize(FName PackageName) override
+	{
+		return 0;
+	}
+};
+
+class FBulkDataRegistryTrackBulkDataToResave : public FBulkDataRegistryNull
+{
+public:
+	FBulkDataRegistryTrackBulkDataToResave() = default;
+	virtual ~FBulkDataRegistryTrackBulkDataToResave() {}
+
+	virtual void Register(UPackage* Owner, const UE::Virtualization::FVirtualizedUntypedBulkData& BulkData) override
+	{
+		ResaveSizeTracker.Register(Owner, BulkData);
+	}
+	virtual uint64 GetBulkDataResaveSize(FName PackageName) override
+	{
+		return ResaveSizeTracker.GetBulkDataResaveSize(PackageName);
+	}
+
+private:
+	FResaveSizeTracker ResaveSizeTracker;
 };
 
 }
@@ -63,6 +96,10 @@ void IBulkDataRegistry::Initialize()
 			GBulkDataRegistry = SetRegistryDelegate.Execute();
 		}
 	}
+	else if (IsEditorDomainEnabled())
+	{
+		GBulkDataRegistry = new FBulkDataRegistryTrackBulkDataToResave();
+	}
 
 	// Assign the null BulkDataRegistry if it was disabled or not set by a higher level source
 	if (!GBulkDataRegistry)
@@ -97,4 +134,93 @@ void SetGlobalBuildInputResolver(UE::DerivedData::IBuildInputResolver* InResolve
 	UE::GlobalBuildInputResolver::Private::GGlobalBuildInputResolver = InResolver;
 }
 
+namespace UE::BulkDataRegistry::Private
+{
+
+FResaveSizeTracker::FResaveSizeTracker()
+{
+	FCoreUObjectDelegates::OnEndLoadPackage.AddRaw(this, &FResaveSizeTracker::OnEndLoadPackage);
+	ELoadingPhase::Type CurrentPhase = IPluginManager::Get().GetLastCompletedLoadingPhase();
+	if (CurrentPhase == ELoadingPhase::None || CurrentPhase < ELoadingPhase::PostEngineInit)
+	{
+		FCoreDelegates::OnPostEngineInit.AddRaw(this, &FResaveSizeTracker::OnPostEngineInit);
+	}
+	else
+	{
+		OnPostEngineInit();
+	}
+	FCoreUObjectDelegates::OnEndLoadPackage.AddRaw(this, &FResaveSizeTracker::OnEndLoadPackage);
+}
+
+FResaveSizeTracker::~FResaveSizeTracker()
+{
+	FCoreUObjectDelegates::OnEndLoadPackage.RemoveAll(this);
+	FCoreDelegates::OnPostEngineInit.RemoveAll(this);
+}
+
+void FResaveSizeTracker::OnPostEngineInit()
+{
+	bPostEngineInitComplete = true;
+
+	FReadScopeLock ScopeLock(Lock);
+	DeferredRemove.Reserve(PackageBulkResaveSize.Num());
+	for (const TPair<FName, uint64>& Pair : PackageBulkResaveSize)
+	{
+		DeferredRemove.Add(Pair.Key);
+	}
+}
+
+void FResaveSizeTracker::Register(UPackage* Owner, const UE::Virtualization::FVirtualizedUntypedBulkData& BulkData)
+{
+	if (!BulkData.GetIdentifier().IsValid() || !BulkData.IsMemoryOnlyPayload())
+	{
+		return;
+	}
+
+	if (!Owner
+		|| !Owner->GetFileSize() // We only track disk packages
+		|| (bPostEngineInitComplete && Owner->GetHasBeenEndLoaded()) // We only record BulkDatas that are loaded before the package finishes loading
+		)
+	{
+		return;
+	}
+
+	FWriteScopeLock ScopeLock(Lock);
+	PackageBulkResaveSize.FindOrAdd(Owner->GetFName()) += BulkData.GetPayloadSize();
+}
+
+uint64 FResaveSizeTracker::GetBulkDataResaveSize(FName PackageName)
+{
+	FReadScopeLock ScopeLock(Lock);
+	return PackageBulkResaveSize.FindRef(PackageName);
+}
+
+void FResaveSizeTracker::OnEndLoadPackage(TConstArrayView<UPackage*> LoadedPackages)
+{
+	if (!bPostEngineInitComplete)
+	{
+		return;
+	}
+
+	TArray<FName> PackageNames;
+	PackageNames.Reserve(LoadedPackages.Num());
+	for (UPackage* LoadedPackage : LoadedPackages)
+	{
+		PackageNames.Add(LoadedPackage->GetFName());
+	}
+
+	FWriteScopeLock ScopeLock(Lock);
+	// The contract for GetBulkDataResaveSize specifies that we must answer correctly until
+	// OnEndLoadPackage is complete. This includes being called from other subscribers to OnEndLoadPackage
+	// that might run after us. So we defer the removals from PackageBulkResaveSize until the next call 
+	// to OnEndLoadPackage.
+	for (FName PackageName : DeferredRemove)
+	{
+		PackageBulkResaveSize.Remove(PackageName);
+	}
+	DeferredRemove.Reset();
+	DeferredRemove.Append(MoveTemp(PackageNames));
+}
+
+}
 #endif // WITH_EDITOR
\ No newline at end of file
diff --git a/Engine/Source/Runtime/CoreUObject/Private/UObject/PackageResourceManager.cpp b/Engine/Source/Runtime/CoreUObject/Private/UObject/PackageResourceManager.cpp
--- a/Engine/Source/Runtime/CoreUObject/Private/UObject/PackageResourceManager.cpp
+++ b/Engine/Source/Runtime/CoreUObject/Private/UObject/PackageResourceManager.cpp
@@ -3,6 +3,7 @@
 #include "UObject/PackageResourceManager.h"
 
 #include "Async/AsyncFileHandle.h"
+#include "Misc/ConfigCacheIni.h"
 #include "Misc/PackageSegment.h"
 #include "Misc/PreloadableFile.h"
 #include "Misc/ScopeLock.h"
@@ -259,3 +260,25 @@ FOpenAsyncPackageResult::~FOpenAsyncPackageResult()
 {
 	// Defined in CPP so we can forward declare IAsyncReadFileHandle
 }
+
+#if WITH_EDITOR
+bool IsEditorDomainEnabled()
+{
+	static bool bEnabled = []()
+	{
+		if (!GIsEditor || (IsRunningCommandlet() && !IsRunningCookCommandlet()))
+		{
+			return false;
+		}
+		bool bEnabledByConfig = true;
+		GConfig->GetBool(TEXT("EditorDomain"), TEXT("EditorDomainEnabled"), bEnabledByConfig, GEditorIni);
+		if (GConfig->GetBool(TEXT("CookSettings"), TEXT("EditorDomainEnabled"), bEnabledByConfig, GEditorIni))
+		{
+			UE_LOG(LogPackageResourceManager, Error,
+				TEXT("Editor.ini:[CookSettings]:EditorDomainEnabled is deprecated, use Editor.ini:[EditorDomain]:EditorDomainEnabled instead."));
+		}
+		return bEnabledByConfig;
+	}();
+	return bEnabled;
+}
+#endif
\ No newline at end of file
diff --git a/Engine/Source/Runtime/CoreUObject/Public/Serialization/BulkDataRegistry.h b/Engine/Source/Runtime/CoreUObject/Public/Serialization/BulkDataRegistry.h
--- a/Engine/Source/Runtime/CoreUObject/Public/Serialization/BulkDataRegistry.h
+++ b/Engine/Source/Runtime/CoreUObject/Public/Serialization/BulkDataRegistry.h
@@ -5,9 +5,12 @@
 #if WITH_EDITOR
 
 #include "Async/Future.h"
+#include "Containers/Array.h"
+#include "Containers/Set.h"
 #include "Compression/CompressedBuffer.h"
 #include "Containers/ArrayView.h"
 #include "Delegates/Delegate.h"
+#include "HAL/CriticalSection.h"
 #include "HAL/Platform.h"
 #include "IO/IoHash.h"
 
@@ -66,6 +69,14 @@ public:
 	 */
 	virtual TFuture<UE::BulkDataRegistry::FData> GetData(const FGuid& BulkDataId) = 0;
 
+	/**
+	 * Report whether the Package had BulkDatas during load that upgrade or otherwise exist in memoryonly and
+	 * cannot save all its BulkDatas by reference when resaved. This function only returns the correct
+	 * information until OnEndLoadPackage is called for the given package; after that it can return an arbitrary
+	 * value.
+	 */
+	virtual uint64 GetBulkDataResaveSize(FName PackageName) = 0;
+
 	/** Set and intialize global IBulkDataRegistry; Get fatally fails before. */
 	COREUOBJECT_API static void Initialize();
 	/** Shutdown and deallocate global IBulkDataRegistry; Get fatally fails afterwards. */
@@ -77,6 +88,32 @@ protected:
 	virtual ~IBulkDataRegistry() {};
 };
 
+namespace UE::BulkDataRegistry::Private
+{
+
+/** Implements behavior needed across multiple BulkDataRegistry implementations for GetBulkDataResaveSize */
+class FResaveSizeTracker
+{
+public:
+	COREUOBJECT_API FResaveSizeTracker();
+	COREUOBJECT_API ~FResaveSizeTracker();
+
+	COREUOBJECT_API void Register(UPackage* Owner, const UE::Virtualization::FVirtualizedUntypedBulkData& BulkData);
+	COREUOBJECT_API uint64 GetBulkDataResaveSize(FName PackageName);
+
+private:
+	COREUOBJECT_API void OnEndLoadPackage(TConstArrayView<UPackage*> LoadedPackages);
+	COREUOBJECT_API void OnPostEngineInit();
+
+	FRWLock Lock;
+	TMap<FName, uint64> PackageBulkResaveSize;
+	TArray<FName> DeferredRemove;
+	bool bPostEngineInitComplete = false;
+};
+
+}
+
+
 // Temporary interface for tunneling the EditorBuildInputResolver into CoreUObject.
 // In the future this will be implemented as part of the BuildAPI
 namespace UE::DerivedData { class IBuildInputResolver; }
diff --git a/Engine/Source/Runtime/CoreUObject/Public/UObject/PackageResourceManager.h b/Engine/Source/Runtime/CoreUObject/Public/UObject/PackageResourceManager.h
--- a/Engine/Source/Runtime/CoreUObject/Public/UObject/PackageResourceManager.h
+++ b/Engine/Source/Runtime/CoreUObject/Public/UObject/PackageResourceManager.h
@@ -421,4 +421,10 @@ private:
 #endif
 };
 
-
+#if WITH_EDITOR
+/**
+ * Report whether the EditorDomain is enabled by config, for systems that need to behave differently if editordomain
+ * will be enabled and need to know before the PackageResourceManager is constructed.
+ */
+COREUOBJECT_API bool IsEditorDomainEnabled();
+#endif
