Cooker: Fix FRequestCluster's search of dependencies to consider CoreRedirects, so that it checks for existence of the redirected file rather than the original reference which will likely no longer exist.

diff --git a/Engine/Source/Editor/UnrealEd/Private/Cooker/CookRequestCluster.cpp b/Engine/Source/Editor/UnrealEd/Private/Cooker/CookRequestCluster.cpp
--- a/Engine/Source/Editor/UnrealEd/Private/Cooker/CookRequestCluster.cpp
+++ b/Engine/Source/Editor/UnrealEd/Private/Cooker/CookRequestCluster.cpp
@@ -22,6 +22,7 @@
 #include "Interfaces/ITargetPlatform.h"
 #include "Misc/RedirectCollector.h"
 #include "Misc/StringBuilder.h"
+#include "UObject/CoreRedirects.h"
 #include "UObject/SavePackage.h"
 
 namespace UE::Cook
@@ -795,6 +796,11 @@ void FRequestCluster::FGraphSearch::VisitVertex(const FVertexData& VertexData)
 				bool bHardDependency = Dependencies == &HardDependencies;
 				for (FName Dependency : *Dependencies)
 				{
+					// Process any CoreRedirects before checking whether the package exists
+					FName Redirected = FCoreRedirects::GetRedirectedName(ECoreRedirectFlags::Type_Package,
+						FCoreRedirectObjectName(NAME_None, NAME_None, Dependency)).PackageName;
+					Dependency = Redirected;
+
 					TUniquePtr<FVertexData> NewVertexData;
 					EInstigator InstigatorType = bHardDependency ? EInstigator::HardDependency : EInstigator::SoftDependency;
 					FPackageData* DependencyPackageData = FindOrAddVertex(Dependency, nullptr,
@@ -1167,4 +1173,4 @@ bool FRequestCluster::IsRequestCookable(FName PackageName, FPackageData*& InOutP
 	return true;
 }
 
-} // namespace UE::Cook
\ No newline at end of file
+} // namespace UE::Cook
