- Fix crash if we are sending EOF updates and NewObject is called by an update

#rnx

diff --git a/Engine/Source/Editor/UnrealEd/Private/EditorEngine.cpp b/Engine/Source/Editor/UnrealEd/Private/EditorEngine.cpp
--- a/Engine/Source/Editor/UnrealEd/Private/EditorEngine.cpp
+++ b/Engine/Source/Editor/UnrealEd/Private/EditorEngine.cpp
@@ -1890,19 +1890,28 @@ void UEditorEngine::Tick( float DeltaSeconds, bool bIdleMode )
 	EditorContext.World()->CommitModelSurfaces();	
 
 	bool bUpdateLinkedOrthoViewports = false;
+	
 	/////////////////////////////
 	// Redraw viewports.
-
-	ForEachObjectOfClass(UWorld::StaticClass(), [](UObject* WorldObj)
 	{
-		UWorld* World = CastChecked<UWorld>(WorldObj);
+		// Gather worlds that need EOF updates
+		// This must be done in two steps as the object hash table is locked during ForEachObjectOfClass so any NewObject calls would fail
+		TArray<UWorld*, TInlineAllocator<4>> WorldsToEOFUpdate;
+		ForEachObjectOfClass(UWorld::StaticClass(), [&WorldsToEOFUpdate](UObject* WorldObj)
+		{
+			UWorld* World = CastChecked<UWorld>(WorldObj);
+			if (World->HasEndOfFrameUpdates())
+			{
+				WorldsToEOFUpdate.Add(World);
+			}
+		});
 
-		if (World->HasEndOfFrameUpdates())
+		// Make sure deferred component updates have been sent to the rendering thread.
+		for (UWorld* World : WorldsToEOFUpdate)
 		{
-			// Make sure deferred component updates have been sent to the rendering thread.
 			World->SendAllEndOfFrameUpdates();
 		}
-	});
+	}
 
 	// Do not redraw if the application is hidden
 	bool bAllWindowsHidden = !bHasFocus && AreAllWindowsHidden();
