DDC: Removed the Put/Get compatibility functions from ICache

#rnx

diff --git a/Engine/Source/Developer/DerivedDataCache/Private/DerivedDataBuildJob.cpp b/Engine/Source/Developer/DerivedDataCache/Private/DerivedDataBuildJob.cpp
--- a/Engine/Source/Developer/DerivedDataCache/Private/DerivedDataBuildJob.cpp
+++ b/Engine/Source/Developer/DerivedDataCache/Private/DerivedDataBuildJob.cpp
@@ -23,6 +23,7 @@
 #include "DerivedDataCacheRecord.h"
 #include "DerivedDataPayload.h"
 #include "DerivedDataRequestOwner.h"
+#include "DerivedDataSharedString.h"
 #include "HAL/CriticalSection.h"
 #include "HAL/Event.h"
 #include "HAL/FileManager.h"
@@ -193,7 +194,7 @@ private:
 	static TArray<FName> ParseExportBuildTypes(bool& bOutExportAll);
 
 private:
-	FString Name;
+	FSharedString Name;
 	FString FunctionName{TEXT("Unknown")};
 
 	/** Active state for the job. Always moves through states in order. */
@@ -285,8 +286,8 @@ const TCHAR* LexToString(EBuildJobState State)
 
 FBuildJob::FBuildJob(
 	const FBuildJobCreateParams& Params,
-	FStringView InName,
-	FStringView InFunctionName,
+	const FStringView InName,
+	const FStringView InFunctionName,
 	FOnBuildComplete&& InOnComplete)
 	: Name(InName)
 	, FunctionName(InFunctionName)
@@ -508,7 +509,7 @@ void FBuildJob::BeginCacheQuery()
 {
 	TRACE_CPUPROFILER_EVENT_SCOPE(FBuildJob::CacheQuery);
 	EnumAddFlags(BuildStatus, EBuildStatus::CacheQuery);
-	Cache.Get({Context->GetCacheKey()}, Name, MakeCacheRecordQueryPolicy(BuildPolicy, *Context), Owner,
+	Cache.Get({{Name, Context->GetCacheKey(), MakeCacheRecordQueryPolicy(BuildPolicy, *Context)}}, Owner,
 		[this](FCacheGetCompleteParams&& Params) { EndCacheQuery(MoveTemp(Params)); });
 }
 
@@ -574,7 +575,7 @@ void FBuildJob::BeginCacheStore()
 	EnumAddFlags(BuildStatus, EBuildStatus::CacheStore);
 	FCacheRecordBuilder RecordBuilder(Context->GetCacheKey());
 	Output.Get().Save(RecordBuilder);
-	Cache.Put({RecordBuilder.Build()}, Name, MakeCacheStorePolicy(BuildPolicy.GetCombinedPolicy(), *Context), Owner,
+	Cache.Put({{Name, RecordBuilder.Build(), MakeCacheStorePolicy(BuildPolicy.GetCombinedPolicy(), *Context)}}, Owner,
 		[this](FCachePutCompleteParams&& Params) { EndCacheStore(MoveTemp(Params)); });
 }
 
diff --git a/Engine/Source/Developer/DerivedDataCache/Private/DerivedDataCache.cpp b/Engine/Source/Developer/DerivedDataCache/Private/DerivedDataCache.cpp
--- a/Engine/Source/Developer/DerivedDataCache/Private/DerivedDataCache.cpp
+++ b/Engine/Source/Developer/DerivedDataCache/Private/DerivedDataCache.cpp
@@ -237,40 +237,6 @@ FCacheRecordPolicy FCacheRecordPolicyBuilder::Build()
 	return Policy;
 }
 
-void ICache::Put(
-	const TConstArrayView<FCacheRecord> Records,
-	const FStringView Context,
-	const ECachePolicy Policy,
-	IRequestOwner& Owner,
-	FOnCachePutComplete&& OnComplete)
-{
-	const FCacheRequestName Name(Context);
-	TArray<FCachePutRequest, TInlineAllocator<16>> Requests;
-	Requests.Reserve(Records.Num());
-	for (const FCacheRecord& Record : Records)
-	{
-		Requests.Add({Name, Record, Policy});
-	}
-	return Put(Requests, Owner, MoveTemp(OnComplete));
-}
-
-void ICache::Get(
-	const TConstArrayView<FCacheKey> Keys,
-	const FStringView Context,
-	const FCacheRecordPolicy Policy,
-	IRequestOwner& Owner,
-	FOnCacheGetComplete&& OnComplete)
-{
-	const FCacheRequestName Name(Context);
-	TArray<FCacheGetRequest, TInlineAllocator<16>> Requests;
-	Requests.Reserve(Keys.Num());
-	for (const FCacheKey& Key : Keys)
-	{
-		Requests.Add({Name, Key, Policy});
-	}
-	return Get(Requests, Owner, MoveTemp(OnComplete));
-}
-
 } // UE::DerivedData
 
 namespace UE::DerivedData::Private
diff --git a/Engine/Source/Developer/DerivedDataCache/Public/DerivedDataCache.h b/Engine/Source/Developer/DerivedDataCache/Public/DerivedDataCache.h
--- a/Engine/Source/Developer/DerivedDataCache/Public/DerivedDataCache.h
+++ b/Engine/Source/Developer/DerivedDataCache/Public/DerivedDataCache.h
@@ -302,59 +302,6 @@ class ICache : public ICacheStore
 public:
 	virtual ~ICache() = default;
 
-	using ICacheStore::Put;
-	using ICacheStore::Get;
-
-	/**
-	 * Asynchronous request to put cache records according to the policy.
-	 *
-	 * The callback will always be called for every key, and may be called from an arbitrary thread.
-	 * Records may finish storing in any order, and from multiple threads concurrently.
-	 *
-	 * A cache store is free to interpret a record containing only a key as a request to delete that
-	 * record from the store. Records may contain payloads that do not have data, and these payloads
-	 * must reference an existing copy of the payload in the store, if available, and must otherwise
-	 * be stored as a partial record that can attempt recovery of missing payloads when fetched.
-	 *
-	 * @param Records      The cache records to store. Must have a key.
-	 * @param Context      A description of the request. An object path is typically sufficient.
-	 * @param Policy       Flags to control the behavior of the request. See ECachePolicy.
-	 * @param Owner        The owner to execute the request within.
-	 * @param OnComplete   A callback invoked for every record as it completes or is canceled.
-	 */
-	UE_API void Put(
-		TConstArrayView<FCacheRecord> Records,
-		FStringView Context,
-		ECachePolicy Policy,
-		IRequestOwner& Owner,
-		FOnCachePutComplete&& OnComplete = FOnCachePutComplete());
-
-	/**
-	 * Asynchronous request to get cache records according to the policy.
-	 *
-	 * The callback will always be called for every key, and may be called from an arbitrary thread.
-	 * Records may become available in any order, and from multiple threads concurrently.
-	 *
-	 * Records may propagate into other cache stores, in accordance with the policy. A propagated
-	 * record may be a partial record, with some payloads missing data, depending on the policy.
-	 *
-	 * When payloads are required by the policy, but not available, the status must be Error. The
-	 * cache store must produce a partial record with the available payloads when the policy flag
-	 * PartialOnError is set for the missing payloads.
-	 *
-	 * @param Keys         The keys identifying the cache records to query.
-	 * @param Context      A description of the request. An object path is typically sufficient.
-	 * @param Policy       Flags to control the behavior of the request. See FCacheRecordPolicy.
-	 * @param Owner        The owner to execute the request within.
-	 * @param OnComplete   A callback invoked for every key as it completes or is canceled.
-	 */
-	UE_API void Get(
-		TConstArrayView<FCacheKey> Keys,
-		FStringView Context,
-		FCacheRecordPolicy Policy,
-		IRequestOwner& Owner,
-		FOnCacheGetComplete&& OnComplete);
-
 	/**
 	 * Cancel all queued and active cache requests and invoke their callbacks.
 	 *
diff --git a/Engine/Source/Developer/Virtualization/Private/VirtualizationDDCBackend.cpp b/Engine/Source/Developer/Virtualization/Private/VirtualizationDDCBackend.cpp
--- a/Engine/Source/Developer/Virtualization/Private/VirtualizationDDCBackend.cpp
+++ b/Engine/Source/Developer/Virtualization/Private/VirtualizationDDCBackend.cpp
@@ -103,11 +103,8 @@ EPushResult FDDCBackend::PushData(const FPayloadId& Id, const FCompressedBuffer&
 		Result = Params;
 	};
 
-	Cache.Put(	{RecordBuilder.Build()},
-				TEXT("Mirage"), // TODO: Improve this when we start passing more context to push
-				TransferPolicy,
-				Owner, 
-				MoveTemp(Callback));
+	// TODO: Improve the name when we start passing more context to this function
+	Cache.Put({{{TEXT("Mirage")}, RecordBuilder.Build(), TransferPolicy}}, Owner, MoveTemp(Callback));
 
 	Owner.Wait();
 
@@ -145,11 +142,8 @@ FCompressedBuffer FDDCBackend::PullData(const FPayloadId& Id)
 		}
 	};
 
-	Cache.Get(	{Key}, 
-				TEXT("Mirage"), // TODO: Improve this when we start passing more context to push
-				TransferPolicy,
-				Owner,
-				MoveTemp(Callback));
+	// TODO: Improve the name when we start passing more context to this function
+	Cache.Get({{{TEXT("Mirage")}, Key, TransferPolicy}}, Owner, MoveTemp(Callback));
 
 	Owner.Wait();
 
@@ -174,11 +168,8 @@ bool FDDCBackend::DoesPayloadExist(const FPayloadId& Id)
 		ResultStatus = Params.Status;
 	};
 
-	Cache.Get(	{ Key },
-				TEXT("Mirage"), // TODO: Improve this when we start passing more context to push
-				QueryPolicy | UE::DerivedData::ECachePolicy::SkipData,
-				Owner,
-				MoveTemp(Callback));
+	// TODO: Improve the name when we start passing more context to this function
+	Cache.Get({{{TEXT("Mirage")}, Key, QueryPolicy | UE::DerivedData::ECachePolicy::SkipData}}, Owner, MoveTemp(Callback));
 
 	Owner.Wait();
 
diff --git a/Engine/Source/Editor/UnrealEd/Private/Cooker/CookRequestCluster.cpp b/Engine/Source/Editor/UnrealEd/Private/Cooker/CookRequestCluster.cpp
--- a/Engine/Source/Editor/UnrealEd/Private/Cooker/CookRequestCluster.cpp
+++ b/Engine/Source/Editor/UnrealEd/Private/Cooker/CookRequestCluster.cpp
@@ -383,6 +383,7 @@ bool FRequestCluster::TryTakeOwnership(FPackageData& PackageData, bool bUrgent,
 
 void FRequestCluster::StartAsync(const FCookerTimer& CookerTimer, bool& bOutComplete)
 {
+	using namespace UE::DerivedData;
 	using namespace UE::EditorDomain;
 
 	if (bStartAsyncComplete)
@@ -399,25 +400,27 @@ void FRequestCluster::StartAsync(const FCookerTimer& CookerTimer, bool& bOutComp
 		if (bBatchDownloadEnabled)
 		{
 			// If the EditorDomain is active, then batch-download all packages to cook from remote cache into local
-			TArray<UE::DerivedData::FCacheKey> CacheKeys;
-			CacheKeys.Reserve(Requests.Num());
+			TArray<FCacheGetRequest> CacheRequests;
+			CacheRequests.Reserve(Requests.Num());
+
+			ECachePolicy CachePolicy = ECachePolicy::Default | UE::DerivedData::ECachePolicy::SkipData;
 			for (FPackageData* PackageData : Requests)
 			{
-				FPackageDigest PackageDigest = EditorDomain->GetPackageDigest(PackageData->GetPackageName());
+				FName PackageName = PackageData->GetPackageName();
+				FPackageDigest PackageDigest = EditorDomain->GetPackageDigest(PackageName);
 				if (PackageDigest.IsSuccessful() && EnumHasAnyFlags(PackageDigest.DomainUse, EDomainUse::LoadEnabled))
 				{
-					CacheKeys.Add(GetEditorDomainPackageKey(PackageDigest.Hash));
+					TStringBuilder<256> RequestName;
+					RequestName << PackageName << TEXT(" [RequestCluster]");
+					CacheRequests.Add({{RequestName}, GetEditorDomainPackageKey(PackageDigest.Hash), CachePolicy});
 				}
 			}
 
-			if (CacheKeys.Num() > 0)
+			if (!CacheRequests.IsEmpty())
 			{
-				UE::DerivedData::ICache& Cache = UE::DerivedData::GetCache();
-				UE::DerivedData::ECachePolicy CachePolicy = UE::DerivedData::ECachePolicy::Default | UE::DerivedData::ECachePolicy::SkipData;
-				UE::DerivedData::FRequestOwner Owner(UE::DerivedData::EPriority::Highest);
+				FRequestOwner Owner(EPriority::Highest);
 				Owner.KeepAlive();
-				Cache.Get(CacheKeys, TEXT("RequestClusterEditorDomainDownload"_SV), CachePolicy, Owner,
-					[](UE::DerivedData::FCacheGetCompleteParams&& Params) {});
+				GetCache().Get(CacheRequests, Owner, [](FCacheGetCompleteParams&& Params) {});
 			}
 		}
 	}
diff --git a/Engine/Source/Editor/UnrealEd/Private/EditorDomain/EditorDomainUtils.cpp b/Engine/Source/Editor/UnrealEd/Private/EditorDomain/EditorDomainUtils.cpp
--- a/Engine/Source/Editor/UnrealEd/Private/EditorDomain/EditorDomainUtils.cpp
+++ b/Engine/Source/Editor/UnrealEd/Private/EditorDomain/EditorDomainUtils.cpp
@@ -930,8 +930,7 @@ void RequestEditorDomainPackage(const FPackagePath& PackagePath,
 	// But set the CachePolicy to store into remote. This will cause the CacheStore to push
 	// any existing local value into upstream storage and refresh the last-used time in the upstream.
 	ECachePolicy CachePolicy = SkipFlags | ECachePolicy::Local | ECachePolicy::StoreRemote;
-	Cache.Get({ GetEditorDomainPackageKey(EditorDomainHash) }, PackagePath.GetDebugName(),
-		CachePolicy, Owner, MoveTemp(Callback));
+	Cache.Get({{{PackagePath.GetDebugName()}, GetEditorDomainPackageKey(EditorDomainHash), CachePolicy}}, Owner, MoveTemp(Callback));
 }
 
 /** Stores data from SavePackage in accessible fields */
@@ -1324,7 +1323,7 @@ bool TrySavePackage(UPackage* Package)
 	}
 	RecordBuilder.SetMeta(MetaData.Save().AsObject());
 	FRequestOwner Owner(EPriority::Normal);
-	Cache.Put({RecordBuilder.Build()}, Package->GetName(), ECachePolicy::Default, Owner);
+	Cache.Put({{{Package->GetName()}, RecordBuilder.Build()}}, Owner);
 	Owner.KeepAlive();
 
 	// TODO_BuildDefinitionList: Calculate and store BuildDefinitionList on the PackageData, or collect it here from some other source.
@@ -1366,8 +1365,7 @@ void GetBulkDataList(FName PackageName, UE::DerivedData::IRequestOwner& Owner, T
 
 	using namespace UE::DerivedData;
 	ICache& Cache = GetCache();
-	Cache.Get({ GetBulkDataListKey(PackageDigest.Hash) },
-		WriteToString<128>(PackageName), ECachePolicy::Default, Owner,
+	Cache.Get({{{WriteToString<128>(PackageName)}, GetBulkDataListKey(PackageDigest.Hash)}}, Owner,
 		[InnerCallback = MoveTemp(Callback)](FCacheGetCompleteParams&& Params)
 		{
 			bool bOk = Params.Status == EStatus::Ok;
@@ -1399,7 +1397,7 @@ void PutBulkDataList(FName PackageName, FSharedBuffer Buffer)
 	FRequestOwner Owner(EPriority::Normal);
 	FCacheRecordBuilder RecordBuilder(GetBulkDataListKey(PackageDigest.Hash));
 	RecordBuilder.SetValue(Buffer);
-	Cache.Put({RecordBuilder.Build()}, WriteToString<128>(PackageName), ECachePolicy::Default, Owner);
+	Cache.Put({{{WriteToString<128>(PackageName)}, RecordBuilder.Build()}}, Owner);
 	Owner.KeepAlive();
 }
 
@@ -1437,13 +1435,13 @@ void GetBulkDataPayloadId(FName PackageName, const FGuid& BulkDataId, UE::Derive
 
 	using namespace UE::DerivedData;
 	ICache& Cache = GetCache();
-	Cache.Get({ GetBulkDataPayloadIdKey(PackageAndGuidHash) },
-		WriteToString<192>(PackageName, TEXT("/"), BulkDataId), ECachePolicy::Default, Owner,
+	Cache.Get({{{WriteToString<192>(PackageName, TEXT("/"), BulkDataId)}, GetBulkDataPayloadIdKey(PackageAndGuidHash)}},
+		Owner,
 		[InnerCallback = MoveTemp(Callback)](FCacheGetCompleteParams&& Params)
-	{
-		bool bOk = Params.Status == EStatus::Ok;
-		InnerCallback(bOk ? Params.Record.GetValue().GetData().Decompress() : FSharedBuffer());
-	});
+		{
+			bool bOk = Params.Status == EStatus::Ok;
+			InnerCallback(bOk ? Params.Record.GetValue().GetData().Decompress() : FSharedBuffer());
+		});
 }
 
 void PutBulkDataPayloadId(FName PackageName, const FGuid& BulkDataId, FSharedBuffer Buffer)
@@ -1471,7 +1469,7 @@ void PutBulkDataPayloadId(FName PackageName, const FGuid& BulkDataId, FSharedBuf
 	FRequestOwner Owner(EPriority::Normal);
 	FCacheRecordBuilder RecordBuilder(GetBulkDataPayloadIdKey(PackageAndGuidHash));
 	RecordBuilder.SetValue(Buffer);
-	Cache.Put({RecordBuilder.Build()}, WriteToString<128>(PackageName), ECachePolicy::Default, Owner);
+	Cache.Put({{{WriteToString<128>(PackageName)}, RecordBuilder.Build()}}, Owner);
 	Owner.KeepAlive();
 }
 
