Correct and clarify tooltips for 2D layer snapping options.

diff --git a/Engine/Source/Editor/UnrealEd/Classes/Settings/LevelEditorViewportSettings.h b/Engine/Source/Editor/UnrealEd/Classes/Settings/LevelEditorViewportSettings.h
--- a/Engine/Source/Editor/UnrealEd/Classes/Settings/LevelEditorViewportSettings.h
+++ b/Engine/Source/Editor/UnrealEd/Classes/Settings/LevelEditorViewportSettings.h
@@ -389,11 +389,11 @@ private:
 
 public:
 
-	/** If enabled, actor rotations will snap to the grid. */
+	/** If enabled, new Actors that you drag into the viewport snap to the active 2D layer. */
 	UPROPERTY(EditAnywhere, config, Category=GridSnapping, meta=(DisplayName = "Enable 2D Layer Snapping"))
 	uint32 bEnableLayerSnap:1;
 
-	/** The index of the snap plane to use when bEnableLayerSnap is true (from the project SnapLayers array) */
+	/** The index of the snap plane to use when 2D Layer Snapping is enabled, from the array of Snap Layers set for the project. */
 	UPROPERTY(config)
 	int32 ActiveSnapLayerIndex;
 
diff --git a/Engine/Source/Editor/UnrealEd/Public/Settings/EditorProjectSettings.h b/Engine/Source/Editor/UnrealEd/Public/Settings/EditorProjectSettings.h
--- a/Engine/Source/Editor/UnrealEd/Public/Settings/EditorProjectSettings.h
+++ b/Engine/Source/Editor/UnrealEd/Public/Settings/EditorProjectSettings.h
@@ -109,11 +109,11 @@ struct UNREALED_API FMode2DLayer
 		, Depth(InDepth)
 	{ }
 
-	/** Whether snapping to surfaces in the world is enabled */
+	/** A descriptive name for this snap layer. */
 	UPROPERTY(EditAnywhere, config, Category = Layer)
 	FString Name;
 
-	/** The amount of depth to apply when snapping to surfaces */
+	/** The position of this snap layer's plane along the Snap Axis. */
 	UPROPERTY(EditAnywhere, config, Category = Layer)
 	float Depth;
 };
@@ -135,19 +135,19 @@ class UNREALED_API ULevelEditor2DSettings : public UDeveloperSettings
 	GENERATED_UCLASS_BODY()
 
 public:
-	/** If enabled will allow 2D mode */
+	/** If enabled, shows the 2D combined translate and rotate tool in the viewport toolbar. */
 	UPROPERTY(EditAnywhere, config, Category=General, meta=(DisplayName="Enable 2D combined translate + rotate widget"))
 	bool bEnable2DWidget;
 
-	/** If enabled will allow 2D mode */
+	/** If enabled, shows the 2D layer snapping controls in the viewport toolbar. */
 	UPROPERTY(EditAnywhere, config, Category=LayerSnapping)
 	bool bEnableSnapLayers;
 
-	/** Snap axis */
+	/** Sets the world space axis for 2D snap layers. */
 	UPROPERTY(EditAnywhere, config, Category=LayerSnapping, meta=(EditCondition=bEnableSnapLayers))
 	ELevelEditor2DAxis SnapAxis;
 
-	/** Snap layers that are displayed in the viewport toolbar */
+	/** Snap layers that are displayed in the viewport toolbar. */
 	UPROPERTY(EditAnywhere, config, Category=LayerSnapping, meta=(EditCondition=bEnableSnapLayers))
 	TArray<FMode2DLayer> SnapLayers;
 
