StringBuilder: Deprecated AppendAnsi

diff --git a/Engine/Source/Editor/UnrealEd/Private/DerivedDataBuildRemoteExecutor.cpp b/Engine/Source/Editor/UnrealEd/Private/DerivedDataBuildRemoteExecutor.cpp
--- a/Engine/Source/Editor/UnrealEd/Private/DerivedDataBuildRemoteExecutor.cpp
+++ b/Engine/Source/Editor/UnrealEd/Private/DerivedDataBuildRemoteExecutor.cpp
@@ -620,7 +620,7 @@ void FRemoteBuildExecutionRequest::DetermineIfWorkerExists_Async()
 			State.WorkerDescriptionId = State.WorkerDescriptor.GetHash();
 
 			TStringBuilder<128> WorkerUri;
-			WorkerUri.AppendAnsi("/apply/workers/");
+			WorkerUri.Append("/apply/workers/");
 			WorkerUri << State.WorkerDescriptionId;
 
 			UE::Zen::FZenScopedRequestPtr Request(Executor.RequestPool.Get(), false);
@@ -637,7 +637,7 @@ void FRemoteBuildExecutionRequest::PostWorkerObject_Async()
 	UE::Tasks::Launch(TEXT("FRemoteBuildExecutionRequest::PostWorkerObject"), [this]
 		{
 			TStringBuilder<128> WorkerUri;
-			WorkerUri.AppendAnsi("/apply/workers/");
+			WorkerUri.Append("/apply/workers/");
 			WorkerUri << State.WorkerDescriptionId;
 
 			Executor.Stats.TotalWorkerObjectsUploaded.AddBlob(State.WorkerDescriptor.GetSize());
@@ -669,7 +669,7 @@ void FRemoteBuildExecutionRequest::PostWorkerPackage_Async()
 	UE::Tasks::Launch(TEXT("FRemoteBuildExecutionRequest::PostWorkerPackage"), [this]
 		{
 			TStringBuilder<128> WorkerUri;
-			WorkerUri.AppendAnsi("/apply/workers/");
+			WorkerUri.Append("/apply/workers/");
 			WorkerUri << State.WorkerDescriptionId;
 
 			uint64_t AttachmentBytes{};
@@ -739,7 +739,7 @@ void FRemoteBuildExecutionRequest::PostActionObject_Async()
 			}
 
 			TStringBuilder<128> ActionUri;
-			ActionUri.AppendAnsi("/apply/jobs/");
+			ActionUri.Append("/apply/jobs/");
 			ActionUri << State.WorkerDescriptionId;
 
 			State.NeedHashes.Empty();
@@ -794,7 +794,7 @@ void FRemoteBuildExecutionRequest::PostActionPackage_Async()
 			ActionPackage.SetObject(State.Action);
 
 			TStringBuilder<128> ActionUri;
-			ActionUri.AppendAnsi("/apply/jobs/");
+			ActionUri.Append("/apply/jobs/");
 			ActionUri << State.WorkerDescriptionId;
 
 			Executor.Stats.TotalActionPackagesUploaded.AddBlob(AttachmentBytes + State.Action.GetSize());
@@ -813,7 +813,7 @@ void FRemoteBuildExecutionRequest::GetResultPackage_Async()
 	UE::Tasks::Launch(TEXT("FRemoteBuildExecutionRequest::GetResultPackage"), [this]
 		{
 			TStringBuilder<128> JobGetUri;
-			JobGetUri.AppendAnsi("/apply/jobs/");
+			JobGetUri.Append("/apply/jobs/");
 			JobGetUri << State.WorkerDescriptionId;
 			JobGetUri << "/";
 			JobGetUri << State.ActionId;
diff --git a/Engine/Source/Runtime/Core/Private/Misc/Guid.cpp b/Engine/Source/Runtime/Core/Private/Misc/Guid.cpp
--- a/Engine/Source/Runtime/Core/Private/Misc/Guid.cpp
+++ b/Engine/Source/Runtime/Core/Private/Misc/Guid.cpp
@@ -198,14 +198,14 @@ void FGuid::AppendString(FUtf8StringBuilderBase& Builder, EGuidFormats Format) c
 {
 	TAnsiStringBuilder<70> AnsiBuilder;
 	AppendString(AnsiBuilder, Format);
-	Builder.AppendAnsi(AnsiBuilder);
+	Builder.Append(AnsiBuilder);
 }
 
 void FGuid::AppendString(FWideStringBuilderBase& Builder, EGuidFormats Format) const
 {
 	TAnsiStringBuilder<70> AnsiBuilder;
 	AppendString(AnsiBuilder, Format);
-	Builder.AppendAnsi(AnsiBuilder);
+	Builder.Append(AnsiBuilder);
 }
 
 
diff --git a/Engine/Source/Runtime/Core/Private/Unix/UnixPlatformCrashContext.cpp b/Engine/Source/Runtime/Core/Private/Unix/UnixPlatformCrashContext.cpp
--- a/Engine/Source/Runtime/Core/Private/Unix/UnixPlatformCrashContext.cpp
+++ b/Engine/Source/Runtime/Core/Private/Unix/UnixPlatformCrashContext.cpp
@@ -476,7 +476,7 @@ void FUnixCrashContext::AddPlatformSpecificProperties() const
 	if (AnsiSignalName != nullptr)
 	{
 		TStringBuilder<32> SignalName;
-		SignalName.AppendAnsi(AnsiSignalName);
+		SignalName.Append(AnsiSignalName);
 
 		AddCrashProperty(TEXT("CrashSignalName"), *SignalName);
 	}
@@ -937,13 +937,13 @@ void PlatformCrashHandler(int32 Signal, siginfo_t* Info, void* Context)
 		}
 
 		DefaultErrorMessage.Append(TEXT("Caught signal "));
-		DefaultErrorMessage.AppendAnsi(ItoANSI(Signal, 10));
+		DefaultErrorMessage.Append(ItoANSI(Signal, 10));
 
 		ANSICHAR* SignalName = strsignal(Signal);
 		if (SignalName != nullptr)
 		{
 			DefaultErrorMessage.Append(TEXT(" "));
-			DefaultErrorMessage.AppendAnsi(SignalName);
+			DefaultErrorMessage.Append(SignalName);
 		}
 
 		ErrorMessage = *DefaultErrorMessage;
diff --git a/Engine/Source/Runtime/Core/Public/Misc/StringBuilder.h b/Engine/Source/Runtime/Core/Public/Misc/StringBuilder.h
--- a/Engine/Source/Runtime/Core/Public/Misc/StringBuilder.h
+++ b/Engine/Source/Runtime/Core/Public/Misc/StringBuilder.h
@@ -167,12 +167,13 @@ public:
 	}
 
 	inline BuilderType& Append(CharType Char) { return AppendChar(Char); }
+
+	UE_DEPRECATED(5.0, "Use Append instead of AppendAnsi.")
 	inline BuilderType& AppendAnsi(const FAnsiStringView String) { return Append(String); }
-	inline BuilderType& AppendUtf8(const FUtf8StringView String) { return Append(String); }
+	UE_DEPRECATED(5.0, "Use Append instead of AppendAnsi.")
 	inline BuilderType& AppendAnsi(const ANSICHAR* const String) { return Append(String); }
-	inline BuilderType& AppendUtf8(const UTF8CHAR* const String) { return Append(String); }
+	UE_DEPRECATED(5.0, "Use Append instead of AppendAnsi.")
 	inline BuilderType& AppendAnsi(const ANSICHAR* const String, const int32 Length) { return Append(String, Length); }
-	inline BuilderType& AppendUtf8(const UTF8CHAR* const String, const int32 Length) { return Append(String, Length); }
 
 	/** Replace characters at given position and length with substring */
 	void ReplaceAt(int32 Pos, int32 RemoveLen, ViewType Str)
@@ -181,7 +182,7 @@ public:
 		check(RemoveLen >= 0);
 		check(Pos + RemoveLen <= Len());
 
-		const int DeltaLen = Str.Len() - RemoveLen;		
+		const int DeltaLen = Str.Len() - RemoveLen;
 		if (DeltaLen < 0)
 		{
 			CurPos += DeltaLen;
@@ -201,7 +202,7 @@ public:
 				*It = *(It - DeltaLen);
 			}
 		}
-		
+
 		FMemory::Memcpy(Base + Pos, Str.GetData(), Str.Len() * sizeof(CharType));
 	}
 
