Fix a crash with the thumbnail that render before a skeletalmesh is ready to be rendered
#rnx

diff --git a/Engine/Source/Editor/UnrealEd/Classes/ThumbnailRendering/SkeletalMeshThumbnailRenderer.h b/Engine/Source/Editor/UnrealEd/Classes/ThumbnailRendering/SkeletalMeshThumbnailRenderer.h
--- a/Engine/Source/Editor/UnrealEd/Classes/ThumbnailRendering/SkeletalMeshThumbnailRenderer.h
+++ b/Engine/Source/Editor/UnrealEd/Classes/ThumbnailRendering/SkeletalMeshThumbnailRenderer.h
@@ -24,6 +24,7 @@ class UNREALED_API USkeletalMeshThumbnailRenderer : public UDefaultSizedThumbnai
 
 	// Begin UThumbnailRenderer Object
 	virtual void Draw(UObject* Object, int32 X, int32 Y, uint32 Width, uint32 Height, FRenderTarget*, FCanvas* Canvas, bool bAdditionalViewFamily) override;
+	virtual bool AllowsRealtimeThumbnails(UObject* Object) const override;
 	// End UThumbnailRenderer Object
 
 	// UObject implementation
diff --git a/Engine/Source/Editor/UnrealEd/Private/ThumbnailRendering/SkeletalMeshThumbnailRenderer.cpp b/Engine/Source/Editor/UnrealEd/Private/ThumbnailRendering/SkeletalMeshThumbnailRenderer.cpp
--- a/Engine/Source/Editor/UnrealEd/Private/ThumbnailRendering/SkeletalMeshThumbnailRenderer.cpp
+++ b/Engine/Source/Editor/UnrealEd/Private/ThumbnailRendering/SkeletalMeshThumbnailRenderer.cpp
@@ -36,6 +36,16 @@ void USkeletalMeshThumbnailRenderer::Draw(UObject* Object, int32 X, int32 Y, uin
 	ThumbnailScene->SetSkeletalMesh(nullptr);
 }
 
+bool USkeletalMeshThumbnailRenderer::AllowsRealtimeThumbnails(UObject* Object) const
+{
+	if (!Super::AllowsRealtimeThumbnails(Object))
+	{
+		return false;
+	}
+	USkeletalMesh* SkeletalMesh = Cast<USkeletalMesh>(Object);
+	return SkeletalMesh && SkeletalMesh->GetResourceForRendering();
+}
+
 void USkeletalMeshThumbnailRenderer::BeginDestroy()
 {
 	ThumbnailSceneCache.Clear();
