Suppress spurious warning when the cooker finds a reference to an empty package.
#rnx

diff --git a/Engine/Source/Editor/UnrealEd/Private/Cooker/CookPackageData.cpp b/Engine/Source/Editor/UnrealEd/Private/Cooker/CookPackageData.cpp
--- a/Engine/Source/Editor/UnrealEd/Private/Cooker/CookPackageData.cpp
+++ b/Engine/Source/Editor/UnrealEd/Private/Cooker/CookPackageData.cpp
@@ -1844,7 +1844,14 @@ bool FPackageDatas::TryLookupFileNameOnDisk(FName PackageName, FString& OutFileN
 		{
 			if (FPackageName::DoesPackageExist(PackageNameStr, &OutFileName, false /* InAllowTextFormats */))
 			{
-				UE_LOG(LogCook, Warning, TEXT("Package %s exists on disk but does not exist in the AssetRegistry"), *PackageNameStr);
+				if (AssetRegistry->GetAssetPackageDataCopy(PackageName))
+				{
+					// The AssetRegistry knows the package exists, but it has 0 assets.
+				}
+				else
+				{
+					UE_LOG(LogCook, Warning, TEXT("Package %s exists on disk but does not exist in the AssetRegistry"), *PackageNameStr);
+				}
 				return true;
 			}
 			return false;
@@ -2323,4 +2330,4 @@ FPoppedPackageDataScope::~FPoppedPackageDataScope()
 }
 #endif
 
-}
\ No newline at end of file
+}
