Prevent packages being saved to the editor domain from being added to the list of potential files to add to source control.

#rnx

- When a package was being saved for the editor domain,it was being pushed to the list of files that should be added to source control with the filename "EditorDomainPackageWriter" and when that list of files was processed (if the editor had source control enabled) it would result in a number of errors being written to the log as the source control provider would fail to add multiple instances of the file "EditorDomainPackageWriter"

diff --git a/Engine/Source/Editor/UnrealEd/Private/EditorEngine.cpp b/Engine/Source/Editor/UnrealEd/Private/EditorEngine.cpp
--- a/Engine/Source/Editor/UnrealEd/Private/EditorEngine.cpp
+++ b/Engine/Source/Editor/UnrealEd/Private/EditorEngine.cpp
@@ -4317,7 +4317,7 @@ FSavePackageResultStruct UEditorEngine::Save(UPackage* InOuter, UObject* InAsset
 	// See if the package is a valid candidate for being auto-added to the default changelist.
 	// Only allows the addition of newly created packages while in the editor and then only if the user has the option enabled.
 	bool bAutoAddPkgToSCC = false;
-	if (!bIsCooking)
+	if (!ObjectSaveContext.bProceduralSave)
 	{
 		bAutoAddPkgToSCC = IsPackageValidForAutoAdding( InOuter, Filename );
 	}
