World Partition HLOD Builder - Fix HLOD Stats TCHAR issue when writing the CSV report

diff --git a/Engine/Source/Editor/UnrealEd/Private/WorldPartition/WorldPartitionHLODsBuilder.cpp b/Engine/Source/Editor/UnrealEd/Private/WorldPartition/WorldPartitionHLODsBuilder.cpp
--- a/Engine/Source/Editor/UnrealEd/Private/WorldPartition/WorldPartitionHLODsBuilder.cpp
+++ b/Engine/Source/Editor/UnrealEd/Private/WorldPartition/WorldPartitionHLODsBuilder.cpp
@@ -709,14 +709,16 @@ bool UWorldPartitionHLODsBuilder::DumpStats()
 	const FString Filename = FString::Printf(TEXT("%s/HLODStats/HLODStats-%s.csv"), *FPaths::ProjectLogDir(), *FDateTime::Now().ToString());
 	FArchive* CSVFile = IFileManager::Get().CreateFileWriter(*Filename, FILEWRITE_AllowRead);
 
-	FString CSVHeader = FHLODStat::GetHeaderCSVString();
-	FString Newline = TEXT("\n");
+	FTCHARToUTF8 CSVHeader = FHLODStat::GetHeaderCSVString();
+	FString Newline = "\n";
 
-	*CSVFile << CSVHeader << Newline;
+	CSVFile->Serialize((UTF8CHAR*)CSVHeader.Get(), CSVHeader.Length() * sizeof(UTF8CHAR));
+	*CSVFile << Newline;
 	for (const FHLODStat& HLODStat : HLODStats)
 	{
-		FString CSVString = HLODStat.ToCSVString();
-		*CSVFile << CSVString << Newline;
+		FTCHARToUTF8 CSVString = HLODStat.ToCSVString();
+		CSVFile->Serialize((UTF8CHAR*)CSVString.Get(), CSVString.Length() * sizeof(UTF8CHAR));
+		*CSVFile << Newline;
 	}
 
 	delete CSVFile;
