CookRequestClusters: suppress the timeout warning on reading vertex data if we're not reading vertexes from the oplog. In the non-oplog case, the vertex data is empty and is just passed through on an async thread rather than waiting for the oplog. It's unknown why that async thread could take up to 10 seconds (hitting a breakpoint maybe), but we don't need a warning about it.
#rnx

diff --git a/Engine/Source/Editor/UnrealEd/Private/Cooker/CookRequestCluster.cpp b/Engine/Source/Editor/UnrealEd/Private/Cooker/CookRequestCluster.cpp
--- a/Engine/Source/Editor/UnrealEd/Private/Cooker/CookRequestCluster.cpp
+++ b/Engine/Source/Editor/UnrealEd/Private/Cooker/CookRequestCluster.cpp
@@ -667,7 +667,7 @@ void FRequestCluster::FGraphSearch::Poll(TArray<TUniquePtr<FVertexData>>& OutCom
 void FRequestCluster::FGraphSearch::UpdateDisplay()
 {
 	constexpr double WarningTimeout = 10.0;
-	if (FPlatformTime::Seconds() > LastActivityTime + WarningTimeout)
+	if (FPlatformTime::Seconds() > LastActivityTime + WarningTimeout && bCookAttachmentsEnabled)
 	{
 		FScopeLock ScopeLock(&Lock);
 		int32 NumVertices = 0;
@@ -677,7 +677,7 @@ void FRequestCluster::FGraphSearch::UpdateDisplay()
 			NumVertices += Pair.Key->PendingVertices;
 		}
 
-		UE_LOG(LogCook, Warning, TEXT("FRequestCluster waited more than %.0lfs for previous build results from DDC. ")
+		UE_LOG(LogCook, Warning, TEXT("FRequestCluster waited more than %.0lfs for previous build results from the oplog. ")
 			TEXT("NumPendingBatches == %d, NumPendingVertices == %d. Continuing to wait..."),
 			WarningTimeout, NumBatches, NumVertices);
 		LastActivityTime = FPlatformTime::Seconds();
