FBX: add support for importing enum-typed FBX properties as string-typed custom attributes

Enum-typed properties in FBX are converted to string-typed custom attributes using the string
value that corresponds to the enum index.

diff --git a/Engine/Source/Editor/UnrealEd/Private/FbxAnimUtils.cpp b/Engine/Source/Editor/UnrealEd/Private/FbxAnimUtils.cpp
--- a/Engine/Source/Editor/UnrealEd/Private/FbxAnimUtils.cpp
+++ b/Engine/Source/Editor/UnrealEd/Private/FbxAnimUtils.cpp
@@ -269,7 +269,7 @@ namespace FbxAnimUtils
 			FbxAnimCurveNode* CurveNode = Property.GetCurveNode();
 			const FString PropertyName = UTF8_TO_TCHAR(Property.GetName());
 			const EFbxType PropertyType = Property.GetPropertyDataType().GetType();
-			const bool bIsTypeSupported = FbxAnimUtils::IsSupportedCurveDataType(PropertyType) || eFbxString == PropertyType;
+			const bool bIsTypeSupported = FbxAnimUtils::IsSupportedCurveDataType(PropertyType) || PropertyType == eFbxString || PropertyType == eFbxEnum;
 
 			if (Property.GetFlag(FbxPropertyFlags::eUserDefined) && bIsTypeSupported 
 				&& (bImportAllCustomAttributes || CustomAttributeNamesToImport.Contains(PropertyName)))
diff --git a/Engine/Source/Editor/UnrealEd/Private/SkeletalMeshEdit.cpp b/Engine/Source/Editor/UnrealEd/Private/SkeletalMeshEdit.cpp
--- a/Engine/Source/Editor/UnrealEd/Private/SkeletalMeshEdit.cpp
+++ b/Engine/Source/Editor/UnrealEd/Private/SkeletalMeshEdit.cpp
@@ -1469,6 +1469,37 @@ bool UnFbx::FFbxImporter::ImportCustomAttributeToBone(UAnimSequence* TargetSeque
 				UE::Anim::AddTypedCustomAttribute<FStringAnimationAttribute, FString>(FName(CurveName), BoneName, TargetSequence, MakeArrayView(TimeArray), MakeArrayView(StringValues));
 				break;
 			}
+			case EFbxType::eFbxEnum:
+			{
+				// Enum-typed properties in FBX are converted to string-typed custom attributes using the string value
+				// that corresponds to the enum index.
+				TArray<FString> StringValues;
+				FillCurveAttributeToBone<FString>(TimeArray, StringValues, FbxCurve, AnimTimeSpan,
+					[&ValueScale, &InProperty](const FbxAnimCurveKey* Key, const FbxTime* KeyTime) {
+					int32 EnumIndex = -1;
+
+					if (KeyTime)
+					{
+						FbxPropertyValue& EvaluatedValue = InProperty.EvaluateValue(*KeyTime);
+						EvaluatedValue.Get(&EnumIndex, EFbxType::eFbxEnum);
+					}
+					else
+					{
+						EnumIndex = InProperty.Get<FbxEnum>();
+					}
+
+					if (EnumIndex < 0 || EnumIndex >= InProperty.GetEnumCount())
+					{
+						return FString();
+					}
+
+					const char* EnumValue = InProperty.GetEnumValue(EnumIndex);
+					return FString(UTF8_TO_TCHAR(EnumValue));
+				});
+
+				UE::Anim::AddTypedCustomAttribute<FStringAnimationAttribute, FString>(FName(CurveName), BoneName, TargetSequence, MakeArrayView(TimeArray), MakeArrayView(StringValues));
+				break;
+			}
 			default:
 			{
 				AddTokenizedErrorMessage(FTokenizedMessage::Create(EMessageSeverity::Warning, FText::Format(LOCTEXT("Warning_CustomAttributeTypeNotSupported", "Custom Attribute ({0}) could not be imported on bone, as its type {1} is not supported."), FText::FromString(CurveName), FText::AsNumber((int32)InProperty.GetPropertyDataType().GetType()))), FFbxErrors::Animation_InvalidData);
