CookPackageData: Fix typo causing instigator not to be set for packages first requested via UpdateRequestData.
#rnx

diff --git a/Engine/Source/Editor/UnrealEd/Private/Cooker/CookPackageData.cpp b/Engine/Source/Editor/UnrealEd/Private/Cooker/CookPackageData.cpp
--- a/Engine/Source/Editor/UnrealEd/Private/Cooker/CookPackageData.cpp
+++ b/Engine/Source/Editor/UnrealEd/Private/Cooker/CookPackageData.cpp
@@ -226,7 +226,7 @@ void FPackageData::UpdateRequestData(const TConstArrayView<const ITargetPlatform
 	}
 	else if (InRequestedPlatforms.Num() > 0)
 	{
-		SetRequestData(InRequestedPlatforms, bInIsUrgent, MoveTemp(InCompletionCallback), MoveTemp(Instigator));
+		SetRequestData(InRequestedPlatforms, bInIsUrgent, MoveTemp(InCompletionCallback), MoveTemp(InInstigator));
 		SendToState(EPackageState::Request, ESendFlags::QueueAddAndRemove);
 	}
 }
