EditorDomain: Reduce costs from TrySavePackage into the EditorDomain on iterative cooks on the buildfarm (which empties its local ZenServer before each iterative cook). Do this by skipping the save of any package that is loaded before we have heard back from the upstream about whether the upstream has the package.
#rnx

diff --git a/Engine/Source/Editor/UnrealEd/Private/Cooker/CookRequestCluster.cpp b/Engine/Source/Editor/UnrealEd/Private/Cooker/CookRequestCluster.cpp
--- a/Engine/Source/Editor/UnrealEd/Private/Cooker/CookRequestCluster.cpp
+++ b/Engine/Source/Editor/UnrealEd/Private/Cooker/CookRequestCluster.cpp
@@ -394,34 +394,18 @@ void FRequestCluster::StartAsync(const FCookerTimer& CookerTimer, bool& bOutComp
 	FEditorDomain* EditorDomain = FEditorDomain::Get();
 	if (EditorDomain)
 	{
-		// Disable BatchDownload until cache storage implements fetch-head requests in response to a Get with SkipData
 		bool bBatchDownloadEnabled = true;
 		GConfig->GetBool(TEXT("EditorDomain"), TEXT("BatchDownloadEnabled"), bBatchDownloadEnabled, GEditorIni);
 		if (bBatchDownloadEnabled)
 		{
 			// If the EditorDomain is active, then batch-download all packages to cook from remote cache into local
-			TArray<FCacheGetRequest> CacheRequests;
-			CacheRequests.Reserve(Requests.Num());
-
-			ECachePolicy CachePolicy = ECachePolicy::Default | UE::DerivedData::ECachePolicy::SkipData;
+			TArray<FName> BatchDownload;
+			BatchDownload.Reserve(Requests.Num());
 			for (FPackageData* PackageData : Requests)
 			{
-				FName PackageName = PackageData->GetPackageName();
-				FPackageDigest PackageDigest = EditorDomain->GetPackageDigest(PackageName);
-				if (PackageDigest.IsSuccessful() && EnumHasAnyFlags(PackageDigest.DomainUse, EDomainUse::LoadEnabled))
-				{
-					TStringBuilder<256> RequestName;
-					RequestName << PackageName << TEXT(" [RequestCluster]");
-					CacheRequests.Add({{RequestName}, GetEditorDomainPackageKey(PackageDigest.Hash), CachePolicy});
-				}
-			}
-
-			if (!CacheRequests.IsEmpty())
-			{
-				FRequestOwner Owner(EPriority::Highest);
-				Owner.KeepAlive();
-				GetCache().Get(CacheRequests, Owner, [](FCacheGetCompleteParams&& Params) {});
-			}
+				BatchDownload.Add(PackageData->GetPackageName());
+			};
+			EditorDomain->BatchDownload(BatchDownload);
 		}
 	}
 
diff --git a/Engine/Source/Editor/UnrealEd/Private/EditorDomain/EditorDomain.cpp b/Engine/Source/Editor/UnrealEd/Private/EditorDomain/EditorDomain.cpp
--- a/Engine/Source/Editor/UnrealEd/Private/EditorDomain/EditorDomain.cpp
+++ b/Engine/Source/Editor/UnrealEd/Private/EditorDomain/EditorDomain.cpp
@@ -7,6 +7,7 @@
 #include "AssetRegistry/IAssetRegistry.h"
 #include "Async/AsyncFileHandleNull.h"
 #include "Containers/UnrealString.h"
+#include "CoreGlobals.h"
 #include "DerivedDataCache.h"
 #include "DerivedDataCacheRecord.h"
 #include "EditorDomain/EditorDomainArchive.h"
@@ -77,6 +78,8 @@ FEditorDomain::FEditorDomain()
 	{
 		SaveClient.Reset(new FEditorDomainSaveClient());
 	}
+	bSkipSavesUntilCatalogLoaded = GIsBuildMachine;
+
 	AssetRegistry = IAssetRegistry::Get();
 	// We require calling SearchAllAssets, because we rely on being able to call WaitOnAsset
 	// without needing to call ScanPathsSynchronous
@@ -100,6 +103,14 @@ FEditorDomain::FEditorDomain()
 
 FEditorDomain::~FEditorDomain()
 {
+	TUniquePtr<UE::DerivedData::FRequestOwner> LocalBatchDownloadOwner;
+	{
+		FScopeLock ScopeLock(&Locks->Lock);
+		LocalBatchDownloadOwner = MoveTemp(BatchDownloadOwner);
+	}
+	// BatchDownloadOwner must be deleted (which calls Cancel) outside of the lock, since its callback takes the lock
+	LocalBatchDownloadOwner.Reset();
+
 	FScopeLock ScopeLock(&Locks->Lock);
 	// AssetRegistry has already been destructed by this point, do not try to access it.
 	// AssetRegistry->OnAssetUpdatedOnDisk().RemoveAll(this);
@@ -194,11 +205,17 @@ bool FEditorDomain::TryFindOrAddPackageSource(FName PackageName,
 
 UE::EditorDomain::FPackageDigest FEditorDomain::GetPackageDigest(FName PackageName)
 {
+	FScopeLock ScopeLock(&Locks->Lock);
+	return GetPackageDigest_WithinLock(PackageName);
+}
+
+UE::EditorDomain::FPackageDigest FEditorDomain::GetPackageDigest_WithinLock(FName PackageName)
+{
+	// Called within &Locks->Lock
 	using namespace UE::EditorDomain;
 
 	TRefCountPtr<FPackageSource> PackageSource;
 	FPackageDigest ErrorDigest;
-	FScopeLock ScopeLock(&Locks->Lock);
 	if (!TryFindOrAddPackageSource(PackageName, PackageSource, &ErrorDigest))
 	{
 		return ErrorDigest;
@@ -271,6 +288,7 @@ int64 FEditorDomain::FileSize(const FPackagePath& PackagePath, EPackageSegment P
 		{
 			return Workspace->FileSize(PackagePath, PackageSegment, OutUpdatedPath);
 		}
+		PackageSource->SetHasLoaded();
 
 		auto MetaDataGetComplete =
 			[&FileSize, &PackageSource, &PackagePath, PackageSegment, Locks=this->Locks, OutUpdatedPath]
@@ -334,6 +352,7 @@ FOpenPackageResult FEditorDomain::OpenReadPackage(const FPackagePath& PackagePat
 	{
 		return Workspace->OpenReadPackage(PackagePath, PackageSegment, OutUpdatedPath);
 	}
+	PackageSource->SetHasLoaded();
 
 	// TODO: Change priority to High instead of Blocking once we have removed the GetPackageFormat below
 	// and don't need to block on the result before exiting this function
@@ -391,6 +410,7 @@ FOpenAsyncPackageResult FEditorDomain::OpenAsyncReadPackage(const FPackagePath&
 	{
 		return Workspace->OpenAsyncReadPackage(PackagePath, PackageSegment);
 	}
+	PackageSource->SetHasLoaded();
 
 	// TODO: Change priority to Normal instead of Blocking once we have removed the GetPackageFormat below
 	// and don't need to block on the result before exiting this function
@@ -552,7 +572,7 @@ void FEditorDomain::FilterKeepPackagesToSave(TArray<UPackage*>& InOutPackagesToS
 		if (FPackagePath::TryFromPackageName(Package->GetFName(), PackagePath))
 		{
 			TRefCountPtr<FPackageSource> PackageSource = FindPackageSource(PackagePath);
-			if (PackageSource && PackageSource->NeedsEditorDomainSave())
+			if (PackageSource && PackageSource->NeedsEditorDomainSave(*this))
 			{
 				PackageSource->bHasSaved = true;
 				bKeep = true;
@@ -569,6 +589,63 @@ void FEditorDomain::FilterKeepPackagesToSave(TArray<UPackage*>& InOutPackagesToS
 	}
 }
 
+bool FEditorDomain::FPackageSource::NeedsEditorDomainSave(FEditorDomain& EditorDomain) const
+{
+	return !bHasSaved && Source == EPackageSource::Workspace &&
+		(!EditorDomain.bSkipSavesUntilCatalogLoaded || bLoadedAfterCatalogLoaded);
+}
+
+void FEditorDomain::FPackageSource::SetHasLoaded()
+{
+	if (bHasLoaded)
+	{
+		return;
+	}
+	bHasLoaded = true;
+	bLoadedAfterCatalogLoaded = bHasQueriedCatalog;
+}
+
+void FEditorDomain::BatchDownload(TArrayView<FName> PackageNames)
+{
+	using namespace UE::EditorDomain;
+	using namespace UE::DerivedData;
+
+	FScopeLock ScopeLock(&Locks->Lock);
+	if (!BatchDownloadOwner)
+	{
+		BatchDownloadOwner = MakeUnique<FRequestOwner>(EPriority::Highest);
+	}
+
+	TArray<FCacheGetRequest> CacheRequests;
+	CacheRequests.Reserve(PackageNames.Num());
+	ECachePolicy CachePolicy = ECachePolicy::Default | ECachePolicy::SkipData;
+	for (FName PackageName : PackageNames)
+	{
+		FPackageDigest PackageDigest = GetPackageDigest_WithinLock(PackageName);
+		if (PackageDigest.IsSuccessful() && EnumHasAnyFlags(PackageDigest.DomainUse, EDomainUse::LoadEnabled))
+		{
+			CacheRequests.Add({ { WriteToString<256>(PackageName) }, GetEditorDomainPackageKey(PackageDigest.Hash),
+				CachePolicy });
+		}
+	}
+
+	if (!CacheRequests.IsEmpty())
+	{
+		FRequestBarrier Barrier(*BatchDownloadOwner);
+		GetCache().Get(CacheRequests, *BatchDownloadOwner, [this](FCacheGetCompleteParams&& Params)
+			{
+				FScopeLock ScopeLock(&Locks->Lock);
+				TRefCountPtr<FPackageSource> PackageSource;
+				FPackageDigest ErrorDigest;
+				FName PackageName = FName(*Params.Name);
+				if (TryFindOrAddPackageSource(PackageName, PackageSource, &ErrorDigest))
+				{
+					PackageSource->bHasQueriedCatalog = true;
+				}
+			});
+	}
+}
+
 void FEditorDomain::OnPackageSavedWithContext(const FString& PackageFileName, UPackage* Package,
 	FObjectPostSaveContext ObjectSaveContext)
 {
@@ -636,4 +713,4 @@ FString LexToString(UE::EditorDomain::FPackageDigest::EStatus Status, FName Stat
 	default:
 		return TEXT("Unknown result code.");
 	}
-}
\ No newline at end of file
+}
diff --git a/Engine/Source/Editor/UnrealEd/Private/EditorDomain/EditorDomain.h b/Engine/Source/Editor/UnrealEd/Private/EditorDomain/EditorDomain.h
--- a/Engine/Source/Editor/UnrealEd/Private/EditorDomain/EditorDomain.h
+++ b/Engine/Source/Editor/UnrealEd/Private/EditorDomain/EditorDomain.h
@@ -10,6 +10,7 @@
 #include "IO/IoHash.h"
 #include "Logging/LogMacros.h"
 #include "Templates/RefCounting.h"
+#include "Templates/UniquePtr.h"
 #include "TickableEditorObject.h"
 #include "UObject/NameTypes.h"
 #include "UObject/ObjectSaveContext.h"
@@ -21,6 +22,7 @@ class IAssetRegistry;
 class UObject;
 class UPackage;
 struct FAssetData;
+namespace UE::DerivedData { class FRequestOwner; }
 
 namespace UE::EditorDomain
 {
@@ -148,6 +150,9 @@ public:
 	 */
 	UE::EditorDomain::FPackageDigest GetPackageDigest(FName PackageDigest);
 
+	/** Request the download of the given packages from the upstream DDC server. */
+	void BatchDownload(TArrayView<FName> PackageNames);
+
 private:
 	/**
 	 * Reference-counted struct holding the locks used for multithreaded synchronization.
@@ -166,13 +171,22 @@ private:
 	 */
 	struct FPackageSource : public FThreadSafeRefCountedObject
 	{
-		UE::EditorDomain::FPackageDigest Digest;
-		EPackageSource Source = EPackageSource::Undecided;
-		bool bHasSaved = false;
-		bool NeedsEditorDomainSave() const
+		FPackageSource()
+			:bHasSaved(false), bHasLoaded(false), bHasQueriedCatalog(false), bLoadedAfterCatalogLoaded(false)
 		{
-			return !bHasSaved && Source == EPackageSource::Workspace;
 		}
+
+		/** Return whether we should call TrySavePackage into the EditorDomain. */
+		bool NeedsEditorDomainSave(FEditorDomain& EditorDomain) const;
+		/** Mark that a load is being attempted, which can impact whether we need to save. */
+		void SetHasLoaded();
+
+		UE::EditorDomain::FPackageDigest Digest;
+		EPackageSource Source = EPackageSource::Undecided;
+		bool bHasSaved : 1;
+		bool bHasLoaded : 1;
+		bool bHasQueriedCatalog : 1;
+		bool bLoadedAfterCatalogLoaded : 1;
 	};
 
 	/** Disallow copy constructors */
@@ -197,6 +211,9 @@ private:
 		FObjectPostSaveContext ObjectSaveContext);
 	/** AssetUpdated event to invalidate our information about where it should be loaded from. */
 	void OnAssetUpdatedOnDisk(const FAssetData& AssetData);
+	/** Same As GetPackageDigest, but assumes lock is already held. */
+	UE::EditorDomain::FPackageDigest GetPackageDigest_WithinLock(FName PackageDigest);
+
 
 	/** Subsystem used to request the save of missing packages into the EditorDomain from a separate process. */
 	TUniquePtr<FEditorDomainSaveClient> SaveClient;
@@ -208,12 +225,20 @@ private:
 	TRefCountPtr<FLocks> Locks;
 	/** Digests previously found for a package. Used for optimization, but also to record loaded-from-domain. */
 	TMap<FName, TRefCountPtr<FPackageSource>> PackageSources;
+	/** RequestOwner used for BatchDownloads from upstream server. */
+	TUniquePtr<UE::DerivedData::FRequestOwner> BatchDownloadOwner;
+
 	/** True by default, set to false when reading is disabled for testing. */
 	bool bEditorDomainReadEnabled = true;
 	/** If true, use an out-of-process EditorDomainSaveServer for saves, else save in process in EndLoad */
 	bool bExternalSave = false;
 	/** Marker for whether our PostEngineInit callback has been called */
 	bool bHasPassedPostEngineInit = false;
+	/**
+	 * Configuration value for whether saves into the EditorDomain should be skipped for packages loaded before
+	 * the DDC server reports whether they exist in the upstream DDC.
+	 */
+	bool bSkipSavesUntilCatalogLoaded = false;
 
 	static FEditorDomain* RegisteredEditorDomain;
 
