CookPackageData FindOrAddPackageData: Improve check statement information.
#rnx

diff --git a/Engine/Source/Editor/UnrealEd/Private/Cooker/CookPackageData.cpp b/Engine/Source/Editor/UnrealEd/Private/Cooker/CookPackageData.cpp
--- a/Engine/Source/Editor/UnrealEd/Private/Cooker/CookPackageData.cpp
+++ b/Engine/Source/Editor/UnrealEd/Private/Cooker/CookPackageData.cpp
@@ -1628,7 +1628,12 @@ FPackageData& FPackageDatas::FindOrAddPackageData(const FName& PackageName, cons
 		if (PackageDataMapAddr != nullptr)
 		{
 			FPackageData** FileNameMapAddr = FileNameToPackageData.Find(NormalizedFileName);
-			check(FileNameMapAddr && *FileNameMapAddr == *PackageDataMapAddr);
+			checkf(FileNameMapAddr, TEXT("Package %s is being added with filename %s, but it already exists with filename %s, ")
+				TEXT("and it is not present in FileNameToPackageData map under the new name."),
+				*PackageName.ToString(), *NormalizedFileName.ToString(), *(*PackageDataMapAddr)->GetFileName().ToString());
+			checkf(*FileNameMapAddr == *PackageDataMapAddr,
+				TEXT("Package %s is being added with filename %s, but that filename maps to a different package %s."),
+				*PackageName.ToString(), *NormalizedFileName.ToString(), *(*FileNameMapAddr)->GetPackageName().ToString());
 			return **PackageDataMapAddr;
 		}
 
