WorldPartition:
- Reset InstancingContext on UWorld::Rename (SaveAs)
- Update Loaded editor cells after rename

diff --git a/Engine/Source/Editor/UnrealEd/Private/FileHelpers.cpp b/Engine/Source/Editor/UnrealEd/Private/FileHelpers.cpp
--- a/Engine/Source/Editor/UnrealEd/Private/FileHelpers.cpp
+++ b/Engine/Source/Editor/UnrealEd/Private/FileHelpers.cpp
@@ -739,7 +739,8 @@ static bool SaveWorld(UWorld* World,
 
 		// Rename the package and the object, as necessary
 		UWorld* DuplicatedWorld = nullptr;
-		bool bRenamedWorldPartition = false;
+		UWorldPartition* RenamedWorldPartition = nullptr;
+	
 
 		// Save Loaded cells
 		TArray<FName> LoadedEditorCells;
@@ -754,11 +755,11 @@ static bool SaveWorld(UWorld* World,
 					return false;
 				}
 
+				RenamedWorldPartition = World->GetWorldPartition();
 				// Load all unloaded actors before rename. If this is causing issues (oom or other) map will need to be renamed through a provided builder commandlet
-				if (UWorldPartition* WorldPartition = World->GetWorldPartition())
+				if (RenamedWorldPartition)
 				{
-					bRenamedWorldPartition = true;
-					LoadedEditorCells = WorldPartition->GetUserLoadedEditorGridCells();
+					LoadedEditorCells = RenamedWorldPartition->GetUserLoadedEditorGridCells();
 					if (AWorldDataLayers* WorldDataLayers = World->GetWorldDataLayers())
 					{
 						WorldDataLayers->ForEachDataLayer([](UDataLayer* DataLayer)
@@ -771,10 +772,10 @@ static bool SaveWorld(UWorld* World,
 						});
 					}
 					// Make sure AlwaysLoaded DL cells get loaded
-					WorldPartition->RefreshLoadedEditorCells(false);
+					RenamedWorldPartition->RefreshLoadedEditorCells(false);
 					// Load all the rest
 					const FBox All(FVector(-WORLD_MAX), FVector(WORLD_MAX));
-					WorldPartition->LoadEditorCells(All, false);
+					RenamedWorldPartition->LoadEditorCells(All, false);
 				}
 
 				// If we are doing a SaveAs on a world that already exists on disk, we need to duplicate it:
@@ -880,10 +881,11 @@ static bool SaveWorld(UWorld* World,
 				AssetRegistry.ScanPathsSynchronous( ULevel::GetExternalObjectsPaths(NewPackageName) , true);
 			}
 
-			if (bRenamedWorldPartition)
+			if (RenamedWorldPartition)
 			{
 				// Save Snapshot of loaded Editor Cells
 				GetMutableDefault<UWorldPartitionEditorPerProjectUserSettings>()->SetEditorGridLoadedCells(SaveWorld, LoadedEditorCells);
+				RenamedWorldPartition->LoadEditorCells(LoadedEditorCells, /*bIsFromUserChange=*/true);
 			}
 		}
 
diff --git a/Engine/Source/Runtime/Engine/Private/WorldPartition/ActorDescContainer.cpp b/Engine/Source/Runtime/Engine/Private/WorldPartition/ActorDescContainer.cpp
--- a/Engine/Source/Runtime/Engine/Private/WorldPartition/ActorDescContainer.cpp
+++ b/Engine/Source/Runtime/Engine/Private/WorldPartition/ActorDescContainer.cpp
@@ -169,11 +169,16 @@ void UActorDescContainer::OnWorldRenamed(UWorld* RenamedWorld)
 {
 	if (GetWorld() == RenamedWorld)
 	{
-		// Update container package
-		ContainerPackageName = RenamedWorld->GetPackage()->GetFName();
+		OnWorldRenamed();
 	}
 }
 
+void UActorDescContainer::OnWorldRenamed()
+{
+	// Update container package
+	ContainerPackageName = GetWorld()->GetPackage()->GetFName();
+}
+
 bool UActorDescContainer::ShouldHandleActorEvent(const AActor* Actor)
 {
 	return Actor && Actor->IsMainPackageActor() && (Actor->GetLevel() != nullptr) && (Actor->GetLevel()->GetPackage()->GetFName() == ContainerPackageName);
diff --git a/Engine/Source/Runtime/Engine/Private/WorldPartition/WorldPartition.cpp b/Engine/Source/Runtime/Engine/Private/WorldPartition/WorldPartition.cpp
--- a/Engine/Source/Runtime/Engine/Private/WorldPartition/WorldPartition.cpp
+++ b/Engine/Source/Runtime/Engine/Private/WorldPartition/WorldPartition.cpp
@@ -477,13 +477,7 @@ void UWorldPartition::Initialize(UWorld* InWorld, const FTransform& InTransform)
 			{
 				TArray<FName> EditorGridLastLoadedCells = GetMutableDefault<UWorldPartitionEditorPerProjectUserSettings>()->GetEditorGridLoadedCells(InWorld);
 
-				for (FName EditorGridLastLoadedCell : EditorGridLastLoadedCells)
-				{
-					if (UWorldPartitionEditorCell* Cell = FindObject<UWorldPartitionEditorCell>(EditorHash, *EditorGridLastLoadedCell.ToString()))
-					{
-						UpdateLoadingEditorCell(Cell, true, true);
-					}
-				}
+				LoadEditorCells(EditorGridLastLoadedCells, /*bIsFromUserChange=*/true);
 			}
 		}
 		
@@ -774,6 +768,17 @@ void UWorldPartition::LoadEditorCells(const FBox& Box, bool bIsFromUserChange)
 	}
 }
 
+void UWorldPartition::LoadEditorCells(const TArray<FName>& CellNames, bool bIsFromUserChange)
+{
+	for (FName CellName : CellNames)
+	{
+		if (UWorldPartitionEditorCell* Cell = FindObject<UWorldPartitionEditorCell>(EditorHash, *CellName.ToString()))
+		{
+			UpdateLoadingEditorCell(Cell, true, bIsFromUserChange);
+		}
+	}
+}
+
 bool UWorldPartition::RefreshLoadedEditorCells(bool bIsFromUserChange)
 {
 	auto GetCellsToRefresh = [this](TArray<UWorldPartitionEditorCell*>& CellsToRefresh)
@@ -1333,6 +1338,15 @@ void UWorldPartition::SetEditorWantedCellSize(uint32 InCellSize)
 	EditorHash->SetEditorWantedCellSize(InCellSize);
 }
 
+void UWorldPartition::OnWorldRenamed()
+{
+	UActorDescContainer::OnWorldRenamed();
+
+	// World was renamed so existing context is invalid.
+	InstancingContext = FLinkerInstancingContext();
+	InstancingSoftObjectPathFixupArchive.Reset();
+}
+
 void UWorldPartition::RemapSoftObjectPath(FSoftObjectPath& ObjectPath)
 {
 	StreamingPolicy->RemapSoftObjectPath(ObjectPath);
diff --git a/Engine/Source/Runtime/Engine/Public/WorldPartition/ActorDescContainer.h b/Engine/Source/Runtime/Engine/Public/WorldPartition/ActorDescContainer.h
--- a/Engine/Source/Runtime/Engine/Public/WorldPartition/ActorDescContainer.h
+++ b/Engine/Source/Runtime/Engine/Public/WorldPartition/ActorDescContainer.h
@@ -24,19 +24,10 @@ public:
 	virtual UWorld* GetWorld() const override;
 
 #if WITH_EDITOR
-	// Events
-	virtual void OnWorldRenamed(UWorld* RenamedWorld);
 	virtual void OnObjectPreSave(UObject* Object, FObjectPreSaveContext SaveContext);
 	virtual void OnPackageDeleted(UPackage* Package);
 	virtual void OnObjectsReplaced(const TMap<UObject*, UObject*>& OldToNewObjectMap);
 
-	void SetContainerPackage(const FName& InContainerPackageName)
-	{
-		// Should only be called on empty container
-		check(GetActorDescCount() == 0);
-		ContainerPackageName = InContainerPackageName;
-	}
-
 	FName GetContainerPackage() const { return ContainerPackageName; }
 
 	/** Removes an actor desc without the need to load a package */
@@ -66,6 +57,10 @@ protected:
 	//~ End UObject Interface
 
 #if WITH_EDITOR
+	// Events
+	void OnWorldRenamed(UWorld* RenamedWorld);
+	virtual void OnWorldRenamed();
+
 	virtual void OnActorDescAdded(FWorldPartitionActorDesc* NewActorDesc);
 	virtual void OnActorDescRemoved(FWorldPartitionActorDesc* ActorDesc);
 	virtual void OnActorDescUpdating(FWorldPartitionActorDesc* ActorDesc) {}
diff --git a/Engine/Source/Runtime/Engine/Public/WorldPartition/WorldPartition.h b/Engine/Source/Runtime/Engine/Public/WorldPartition/WorldPartition.h
--- a/Engine/Source/Runtime/Engine/Public/WorldPartition/WorldPartition.h
+++ b/Engine/Source/Runtime/Engine/Public/WorldPartition/WorldPartition.h
@@ -97,6 +97,8 @@ private:
 	void OnEndPlay();
 
 	//~ Begin UActorDescContainer Interface
+	virtual void OnWorldRenamed() override;
+
 	virtual void OnActorDescAdded(FWorldPartitionActorDesc* NewActorDesc) override;
 	virtual void OnActorDescRemoved(FWorldPartitionActorDesc* ActorDesc) override;
 	virtual void OnActorDescUpdating(FWorldPartitionActorDesc* ActorDesc) override;
@@ -120,6 +122,7 @@ public:
 	FName GetWorldPartitionEditorName() const;
 
 	void LoadEditorCells(const FBox& Box, bool bIsFromUserChange);
+	void LoadEditorCells(const TArray<FName>& CellNames, bool bIsFromUserChange);
 	void UnloadEditorCells(const FBox& Box, bool bIsFromUserChange);
 	bool AreEditorCellsLoaded(const FBox& Box);
 	bool RefreshLoadedEditorCells(bool bIsFromUserChange);
