New Empty map -> Deactivates all Editor modes (same as New map from template)

diff --git a/Engine/Source/Editor/UnrealEd/Private/EditorServer.cpp b/Engine/Source/Editor/UnrealEd/Private/EditorServer.cpp
--- a/Engine/Source/Editor/UnrealEd/Private/EditorServer.cpp
+++ b/Engine/Source/Editor/UnrealEd/Private/EditorServer.cpp
@@ -2253,24 +2253,15 @@ void UEditorEngine::CreateNewMapForEditing(bool bPromptUserToSave, bool bIsParti
 
 	const FScopedBusyCursor BusyCursor;
 
-	// Also change out of Landscape mode to ensure all references are cleared.
-	if( GLevelEditorModeTools().IsModeActive( FBuiltinEditorModes::EM_Landscape ) )
+	// Deactivate any editor modes when loading a new map
+	if (ULevelEditorSubsystem* LevelEditorSubsystem = GEditor->GetEditorSubsystem<ULevelEditorSubsystem>())
 	{
-		GLevelEditorModeTools().DeactivateMode( FBuiltinEditorModes::EM_Landscape );
-	}
-
-	// Also change out of Foliage mode to ensure all references are cleared.
-	if( GLevelEditorModeTools().IsModeActive( FBuiltinEditorModes::EM_Foliage ) )
-	{
-		GLevelEditorModeTools().DeactivateMode( FBuiltinEditorModes::EM_Foliage );
-	}
-
-	// Change out of mesh paint mode when opening a new map.
-	if( GLevelEditorModeTools().IsModeActive( FBuiltinEditorModes::EM_MeshPaint ) )
-	{
-		GLevelEditorModeTools().DeactivateMode( FBuiltinEditorModes::EM_MeshPaint );
+		if (FEditorModeTools* ModeManager = LevelEditorSubsystem->GetLevelEditorModeManager())
+		{
+			ModeManager->DeactivateAllModes();
+		}
 	}
-
+		
 	NewMap(bIsPartitionedWorld);
 
 	FEditorFileUtils::ResetLevelFilenames();
