Keep the 'const' Blueprint function attribute in sync on function overrides inherited from a parent class.

diff --git a/Engine/Source/Editor/UnrealEd/Private/Kismet2/BlueprintEditorUtils.cpp b/Engine/Source/Editor/UnrealEd/Private/Kismet2/BlueprintEditorUtils.cpp
--- a/Engine/Source/Editor/UnrealEd/Private/Kismet2/BlueprintEditorUtils.cpp
+++ b/Engine/Source/Editor/UnrealEd/Private/Kismet2/BlueprintEditorUtils.cpp
@@ -6716,6 +6716,21 @@ struct FConformCallsToParentFunctionUtils
 							Args.Add(TEXT("ParentClass"), FText::FromString(SignatureClass->GetName()));
 							InBlueprint->Message_Note(FText::Format(LOCTEXT("ConvertedToLocalMemberFunction_Note", "Function '{NodeTitle}' was previously implemented as an override, but the function is no longer found in class '{ParentClass}'. As a result, it has been converted to a full member function."), Args).ToString());
 						}
+						else
+						{
+							if (FunctionEntryNode->bEnforceConstCorrectness)
+							{
+								// Sync the 'const' attribute with the original function, in case it has been changed
+								const bool bIsConstFunction = Function->HasAllFunctionFlags(FUNC_Const);
+								if (bIsConstFunction != FunctionEntryNode->HasAllExtraFlags(FUNC_Const))
+								{
+									int32 ExtraFlags = FunctionEntryNode->GetExtraFlags();
+
+									FunctionEntryNode->Modify();
+									FunctionEntryNode->SetExtraFlags(ExtraFlags ^ FUNC_Const);
+								}
+							}
+						}
 					}
 				}
 
