Support adding an unloaded cooked BP-based component to a BP

diff --git a/Engine/Source/Editor/UnrealEd/Private/ComponentTypeRegistry.cpp b/Engine/Source/Editor/UnrealEd/Private/ComponentTypeRegistry.cpp
--- a/Engine/Source/Editor/UnrealEd/Private/ComponentTypeRegistry.cpp
+++ b/Engine/Source/Editor/UnrealEd/Private/ComponentTypeRegistry.cpp
@@ -561,14 +561,25 @@ void FComponentTypeRegistryData::ForceRefreshComponentList()
 
 		if (OnDiskClasses.Num() > 0)
 		{
-			// GetAssetsByClass call is a kludge to get the full asset paths for the blueprints we care about, Bob T. thinks 
-			// that the Asset Registry could just keep asset paths:
-			TArray<FAssetData> BlueprintAssetData;
-			AssetRegistryModule.Get().GetAssetsByClass(UBlueprint::StaticClass()->GetFName(), BlueprintAssetData, true);
 			TMap<FString, FAssetData> BlueprintNames;
-			for (const FAssetData& Blueprint : BlueprintAssetData)
 			{
-				BlueprintNames.Add(Blueprint.AssetName.ToString(), Blueprint);
+				// GetAssetsByClass call is a kludge to get the full asset paths for the blueprints we care about
+				// Bob T. thinks that the Asset Registry could just keep asset paths
+				TArray<FAssetData> Assets;
+				AssetRegistryModule.Get().GetAssetsByClass(UBlueprint::StaticClass()->GetFName(), Assets, true);
+				for (FAssetData& Asset : Assets)
+				{
+					BlueprintNames.Add(Asset.AssetName.ToString(), MoveTemp(Asset));
+				}
+
+				Assets.Reset();
+				AssetRegistryModule.Get().GetAssetsByClass(UBlueprintGeneratedClass::StaticClass()->GetFName(), Assets, true);
+				for (FAssetData& Asset : Assets)
+				{
+					FString BlueprintName = Asset.AssetName.ToString();
+					BlueprintName.RemoveFromEnd(TEXT("_C"));
+					BlueprintNames.Add(MoveTemp(BlueprintName), MoveTemp(Asset));
+				}
 			}
 
 			for (const FName& OnDiskClass : OnDiskClasses)
diff --git a/Engine/Source/Editor/UnrealEd/Private/SComponentClassCombo.cpp b/Engine/Source/Editor/UnrealEd/Private/SComponentClassCombo.cpp
--- a/Engine/Source/Editor/UnrealEd/Private/SComponentClassCombo.cpp
+++ b/Engine/Source/Editor/UnrealEd/Private/SComponentClassCombo.cpp
@@ -316,9 +316,17 @@ void SComponentClassCombo::OnAddComponentSelectionChanged( FComponentClassComboE
 				if (ComponentClass == nullptr)
 				{
 					// The class is not loaded yet, so load it:
-					const ELoadFlags LoadFlags = LOAD_None;
-					UBlueprint* LoadedObject = LoadObject<UBlueprint>(nullptr, *InItem->GetComponentPath(), nullptr, LoadFlags, nullptr);
-					ComponentClass = GetAuthoritativeBlueprintClass(LoadedObject);
+					if (UObject* LoadedObject = LoadObject<UObject>(nullptr, *InItem->GetComponentPath()))
+					{
+						if (UClass* LoadedClass = Cast<UClass>(LoadedObject))
+						{
+							ComponentClass = LoadedClass;
+						}
+						else if (UBlueprint* LoadedBP = Cast<UBlueprint>(LoadedObject))
+						{
+							ComponentClass = GetAuthoritativeBlueprintClass(LoadedBP);
+						}
+					}
 				}
 				
 				FSubobjectDataHandle NewActorCompHandle =
