- Replacing retarget manager with "Retarget Sources"
- Remove right-click retarget features from ABPs, anims and skeletons
- Deprecating URig asset

diff --git a/Engine/Source/Developer/AssetTools/Private/AssetTypeActions/AssetTypeActions_AnimBlueprint.cpp b/Engine/Source/Developer/AssetTools/Private/AssetTypeActions/AssetTypeActions_AnimBlueprint.cpp
--- a/Engine/Source/Developer/AssetTools/Private/AssetTypeActions/AssetTypeActions_AnimBlueprint.cpp
+++ b/Engine/Source/Developer/AssetTools/Private/AssetTypeActions/AssetTypeActions_AnimBlueprint.cpp
@@ -141,58 +141,6 @@ void FAssetTypeActions_AnimBlueprint::GetActions(const TArray<UObject*>& InObjec
 				)
 			);
 	}
-	
-	Section.AddSubMenu(
-		"RetargetBlueprintSubmenu",
-		LOCTEXT("RetargetBlueprintSubmenu", "Retarget Anim Blueprints"),
-		LOCTEXT("RetargetBlueprintSubmenu_ToolTip", "Opens the retarget blueprints menu"),
-		FNewToolMenuDelegate::CreateSP( this, &FAssetTypeActions_AnimBlueprint::FillRetargetMenu, InObjects),
-		false,
-		FSlateIcon(FAppStyle::GetAppStyleSetName(), "Persona.RetargetManager")
-		);
-}
-
-void FAssetTypeActions_AnimBlueprint::FillRetargetMenu(UToolMenu* MenuBuilder, const TArray<UObject*> InObjects)
-{
-	bool bAllSkeletonsNull = true;
-
-	for(auto Iter = InObjects.CreateConstIterator(); Iter; ++Iter)
-	{
-		if(UAnimBlueprint* AnimBlueprint = Cast<UAnimBlueprint>(*Iter))
-		{
-			if(AnimBlueprint->TargetSkeleton)
-			{
-				bAllSkeletonsNull = false;
-				break;
-			}
-		}
-	}
-
-	FToolMenuSection& Section = MenuBuilder->AddSection("Section");
-	if(bAllSkeletonsNull)
-	{
-		Section.AddMenuEntry(
-			"AnimBlueprint_RetargetSkeletonInPlace",
-			LOCTEXT("AnimBlueprint_RetargetSkeletonInPlace", "Retarget skeleton on existing Anim Blueprints"),
-			LOCTEXT("AnimBlueprint_RetargetSkeletonInPlaceTooltip", "Retargets the selected Anim Blueprints to a new skeleton (and optionally all referenced animations too)"),
-			FSlateIcon(FAppStyle::GetAppStyleSetName(), "Persona.RetargetManager"),
-			FUIAction(
-			FExecuteAction::CreateSP( this, &FAssetTypeActions_AnimBlueprint::RetargetAssets, InObjects, false ), // false = do not duplicate assets first
-			FCanExecuteAction()
-			)
-			);
-	}
-
-	Section.AddMenuEntry(
-		"AnimBlueprint_DuplicateAndRetargetSkeleton",
-		LOCTEXT("AnimBlueprint_DuplicateAndRetargetSkeleton", "Duplicate Anim Blueprints and Retarget"),
-		LOCTEXT("AnimBlueprint_DuplicateAndRetargetSkeletonTooltip", "Duplicates and then retargets the selected Anim Blueprints to a new skeleton (and optionally all referenced animations too)"),
-		FSlateIcon(FAppStyle::GetAppStyleSetName(), "Persona.RetargetManager"),
-		FUIAction(
-		FExecuteAction::CreateSP( this, &FAssetTypeActions_AnimBlueprint::RetargetAssets, InObjects, true ), // true = duplicate assets and retarget them
-		FCanExecuteAction()
-		)
-		);
 }
 
 UThumbnailInfo* FAssetTypeActions_AnimBlueprint::GetThumbnailInfo(UObject* Asset) const
@@ -234,35 +182,17 @@ void FAssetTypeActions_AnimBlueprint::OpenAssetEditor( const TArray<UObject*>& I
 		auto AnimBlueprint = Cast<UAnimBlueprint>(*ObjIt);
 		if (AnimBlueprint != NULL && AnimBlueprint->SkeletonGeneratedClass && AnimBlueprint->GeneratedClass)
 		{
-			if(AnimBlueprint->BlueprintType != BPTYPE_Interface && !AnimBlueprint->TargetSkeleton && !AnimBlueprint->bIsTemplate)
+			const bool bBringToFrontIfOpen = true;
+#if WITH_EDITOR
+			if (IAssetEditorInstance* EditorInstance = GEditor->GetEditorSubsystem<UAssetEditorSubsystem>()->FindEditorForAsset(AnimBlueprint, bBringToFrontIfOpen))
 			{
-				FText ShouldRetargetMessage = LOCTEXT("ShouldRetarget_Message", "Could not find the skeleton for Anim Blueprint '{BlueprintName}' Would you like to choose a new one?");
-				
-				FFormatNamedArguments Arguments;
-				Arguments.Add( TEXT("BlueprintName"), FText::FromString(AnimBlueprint->GetName()));
-
-				if ( FMessageDialog::Open(EAppMsgType::YesNo, FText::Format(ShouldRetargetMessage, Arguments)) == EAppReturnType::Yes )
-				{
-					bool bDuplicateAssets = false;
-					TArray<UObject*> AnimBlueprints;
-					AnimBlueprints.Add(AnimBlueprint);
-					RetargetAssets(AnimBlueprints, bDuplicateAssets);
-				}
+				EditorInstance->FocusWindow(AnimBlueprint);
 			}
 			else
-			{
-				const bool bBringToFrontIfOpen = true;
-#if WITH_EDITOR
-				if (IAssetEditorInstance* EditorInstance = GEditor->GetEditorSubsystem<UAssetEditorSubsystem>()->FindEditorForAsset(AnimBlueprint, bBringToFrontIfOpen))
-				{
-					EditorInstance->FocusWindow(AnimBlueprint);
-				}
-				else
 #endif
-				{
-					IAnimationBlueprintEditorModule& AnimationBlueprintEditorModule = FModuleManager::LoadModuleChecked<IAnimationBlueprintEditorModule>("AnimationBlueprintEditor");
-					AnimationBlueprintEditorModule.CreateAnimationBlueprintEditor(Mode, EditWithinLevelEditor, AnimBlueprint);
-				}
+			{
+				IAnimationBlueprintEditorModule& AnimationBlueprintEditorModule = FModuleManager::LoadModuleChecked<IAnimationBlueprintEditorModule>("AnimationBlueprintEditor");
+				AnimationBlueprintEditorModule.CreateAnimationBlueprintEditor(Mode, EditWithinLevelEditor, AnimBlueprint);
 			}
 		}
 		else
@@ -332,52 +262,6 @@ void FAssetTypeActions_AnimBlueprint::ExecuteFindSkeleton(TArray<TWeakObjectPtr<
 	}
 }
 
-void FAssetTypeActions_AnimBlueprint::RetargetAnimationHandler(USkeleton* OldSkeleton, USkeleton* NewSkeleton, bool bRemapReferencedAssets, bool bAllowRemapToExisting, bool bConvertSpaces, const EditorAnimUtils::FNameDuplicationRule* NameRule, TArray<TWeakObjectPtr<UObject>> AnimBlueprints)
-{
-	if(!OldSkeleton || OldSkeleton->GetPreviewMesh(true))
-	{
-		FAnimationRetargetContext RetargetContext(AnimBlueprints, bRemapReferencedAssets, bConvertSpaces);
-
-		if(bAllowRemapToExisting)
-		{
-			SAnimationRemapAssets::ShowWindow(RetargetContext, NewSkeleton);
-		}
-
-		EditorAnimUtils::RetargetAnimations(OldSkeleton, NewSkeleton, RetargetContext, bRemapReferencedAssets, NameRule);
-	}
-	else
-	{
-		FFormatNamedArguments Args;
-		Args.Add(TEXT("OldSkeletonName"), FText::FromString(GetNameSafe(OldSkeleton)));
-		Args.Add(TEXT("NewSkeletonName"), FText::FromString(GetNameSafe(NewSkeleton)));
-		FNotificationInfo Info(FText::Format(LOCTEXT("Retarget Failed", "Old Skeleton {OldSkeletonName} and New Skeleton {NewSkeletonName} need to have Preview Mesh set up to convert animation"), Args));
-		Info.ExpireDuration = 5.0f;
-		Info.bUseLargeFont = false;
-		TSharedPtr<SNotificationItem> Notification = FSlateNotificationManager::Get().AddNotification(Info);
-		if(Notification.IsValid())
-		{
-			Notification->SetCompletionState(SNotificationItem::CS_Fail);
-		}
-	}
-}
-
-void FAssetTypeActions_AnimBlueprint::RetargetAssets(TArray<UObject*> InAnimBlueprints, bool bDuplicateAssets)
-{
-	bool bRemapReferencedAssets = false;
-	USkeleton* OldSkeleton = NULL;
-
-	if ( InAnimBlueprints.Num() > 0 )
-	{
-		UAnimBlueprint * AnimBP = CastChecked<UAnimBlueprint>(InAnimBlueprints[0]);
-		OldSkeleton = AnimBP->TargetSkeleton;
-	}
-
-	const FText Message = LOCTEXT("RemapSkeleton_Warning", "Select the skeleton to remap this asset to.");
-	auto AnimBlueprints = GetTypedWeakObjectPtrs<UObject>(InAnimBlueprints);
-
-	SAnimationRemapSkeleton::ShowWindow(OldSkeleton, Message, bDuplicateAssets, FOnRetargetAnimation::CreateSP(this, &FAssetTypeActions_AnimBlueprint::RetargetAnimationHandler, AnimBlueprints));
-}
-
 TSharedPtr<SWidget> FAssetTypeActions_AnimBlueprint::GetThumbnailOverlay(const FAssetData& AssetData) const
 {
 	const FSlateBrush* Icon = FSlateIconFinder::FindIconBrushForClass(UAnimBlueprint::StaticClass());
diff --git a/Engine/Source/Developer/AssetTools/Private/AssetTypeActions/AssetTypeActions_AnimBlueprint.h b/Engine/Source/Developer/AssetTools/Private/AssetTypeActions/AssetTypeActions_AnimBlueprint.h
--- a/Engine/Source/Developer/AssetTools/Private/AssetTypeActions/AssetTypeActions_AnimBlueprint.h
+++ b/Engine/Source/Developer/AssetTools/Private/AssetTypeActions/AssetTypeActions_AnimBlueprint.h
@@ -29,18 +29,9 @@ public:
 
 private:
 
-	/** Handler to fill the retarget submenu */
-	void FillRetargetMenu(class UToolMenu* Menu, const TArray<UObject*> InObjects);
-
 	/** Handler for when FindSkeleton is selected */
 	void ExecuteFindSkeleton(TArray<TWeakObjectPtr<UAnimBlueprint>> Objects);
-
-	/** Context menu item handler for changing the supplied assets skeletons */ 
-	void RetargetAssets(TArray<UObject*> InAnimBlueprints, bool bDuplicateAssets);
-
-	/** Handler for retargeting */
-	void RetargetAnimationHandler(USkeleton* OldSkeleton, USkeleton* NewSkeleton, bool bRemapReferencedAssets, bool bAllowRemapToExisting, bool bConvertSpaces, const EditorAnimUtils::FNameDuplicationRule* NameRule, TArray<TWeakObjectPtr<UObject>> AnimBlueprints);
-
+	
 	/** Certain options are not available for template anim BPs */
 	bool AreAnyNonTemplateAnimBlueprintsSelected(TArray<TWeakObjectPtr<UAnimBlueprint>> Objects) const;
 };
diff --git a/Engine/Source/Developer/AssetTools/Private/AssetTypeActions/AssetTypeActions_AnimationAsset.cpp b/Engine/Source/Developer/AssetTools/Private/AssetTypeActions/AssetTypeActions_AnimationAsset.cpp
--- a/Engine/Source/Developer/AssetTools/Private/AssetTypeActions/AssetTypeActions_AnimationAsset.cpp
+++ b/Engine/Source/Developer/AssetTools/Private/AssetTypeActions/AssetTypeActions_AnimationAsset.cpp
@@ -48,55 +48,6 @@ void FAssetTypeActions_AnimationAsset::GetActions(const TArray<UObject*>& InObje
 			FCanExecuteAction()
 			)
 		);
-
-	Section.AddSubMenu(
-		"RetargetAnimSubmenu", 
-		LOCTEXT("RetargetAnimSubmenu", "Retarget Anim Assets"),
-		LOCTEXT("RetargetAnimSubmenu_ToolTip", "Opens the retarget anim assets menu"),
-		FNewMenuDelegate::CreateSP( this, &FAssetTypeActions_AnimationAsset::FillRetargetMenu, InObjects ),
-		false,
-		FSlateIcon(FAppStyle::GetAppStyleSetName(), "Persona.RetargetManager")
-		);
-}
-
-void FAssetTypeActions_AnimationAsset::FillRetargetMenu( FMenuBuilder& MenuBuilder, const TArray<UObject*> InObjects )
-{
-	bool bAllSkeletonsNull = true;
-
-	for(auto Iter = InObjects.CreateConstIterator(); Iter; ++Iter)
-	{
-		if(UAnimationAsset* AnimAsset = Cast<UAnimationAsset>(*Iter))
-		{
-			if(AnimAsset->GetSkeleton())
-			{
-				bAllSkeletonsNull = false;
-				break;
-			}
-		}
-	}
-
-	if(bAllSkeletonsNull)
-	{
-		MenuBuilder.AddMenuEntry(
-			LOCTEXT("AnimAsset_RetargetSkeletonInPlace", "Retarget skeleton on existing Anim Assets"),
-			LOCTEXT("AnimAsset_RetargetSkeletonInPlaceTooltip", "Retargets the selected Anim Assets to a new skeleton (and optionally all referenced animations too)"),
-			FSlateIcon(FAppStyle::GetAppStyleSetName(), "Persona.RetargetManager"),
-			FUIAction(
-			FExecuteAction::CreateSP( this, &FAssetTypeActions_AnimationAsset::RetargetAssets, InObjects, false, false, TSharedPtr<IToolkitHost>() ), // false = do not duplicate assets first
-			FCanExecuteAction()
-			)
-			);
-	}
-
-	MenuBuilder.AddMenuEntry(
-		LOCTEXT("AnimAsset_DuplicateAndRetargetSkeleton", "Duplicate Anim Assets and Retarget"),
-		LOCTEXT("AnimAsset_DuplicateAndRetargetSkeletonTooltip", "Duplicates and then retargets the selected Anim Assets to a new skeleton (and optionally all referenced animations too)"),
-		FSlateIcon(FAppStyle::GetAppStyleSetName(), "Persona.RetargetManager"),
-		FUIAction(
-		FExecuteAction::CreateSP( this, &FAssetTypeActions_AnimationAsset::RetargetAssets, InObjects, true, false, TSharedPtr<IToolkitHost>() ), // true = duplicate assets and retarget them
-		FCanExecuteAction()
-		)
-		);
 }
 
 UThumbnailInfo* FAssetTypeActions_AnimationAsset::GetThumbnailInfo(UObject* Asset) const
@@ -139,87 +90,68 @@ void FAssetTypeActions_AnimationAsset::OpenAnimAssetEditor(const TArray<UObject*
 	// For each one..
 	for (UAnimationAsset* AnimAsset : AnimAssets)
 	{
-		USkeleton* AnimSkeleton = AnimAsset->GetSkeleton();
-		if (!AnimSkeleton)
+		// First see if we already have it open
+		const bool bBringToFrontIfOpen = true;
+#if WITH_EDITOR
+		if (IAssetEditorInstance* EditorInstance = GEditor->GetEditorSubsystem<UAssetEditorSubsystem>()->FindEditorForAsset(AnimAsset, bBringToFrontIfOpen))
 		{
-			FText ShouldRetargetMessage = LOCTEXT("ShouldRetargetAnimAsset_Message", "Could not find the skeleton for Anim '{AnimName}' Would you like to choose a new one?");
-
-			FFormatNamedArguments Arguments;
-			Arguments.Add(TEXT("AnimName"), FText::FromString(AnimAsset->GetName()));
-
-			if (FMessageDialog::Open(EAppMsgType::YesNo, FText::Format(ShouldRetargetMessage, Arguments)) == EAppReturnType::Yes)
-			{
-				bool bDuplicateAssets = false;
-				TArray<UObject*> AssetsToRetarget;
-				AssetsToRetarget.Add(AnimAsset);
-				RetargetAssets(AssetsToRetarget, bDuplicateAssets, true, EditWithinLevelEditor);
-			}
+			EditorInstance->FocusWindow();
 		}
 		else
-		{
-			// First see if we already have it open
-			const bool bBringToFrontIfOpen = true;
-#if WITH_EDITOR
-			if (IAssetEditorInstance* EditorInstance = GEditor->GetEditorSubsystem<UAssetEditorSubsystem>()->FindEditorForAsset(AnimAsset, bBringToFrontIfOpen))
-			{
-				EditorInstance->FocusWindow();
-			}
-			else
 #endif
+		{
+			// See if we are trying to open a single asset. If we are, we re-use a compatible anim editor.
+			bool bSingleAsset = AnimAssets.Num() == 1;
+			bool bFoundEditor = false;
+			if (bSingleAsset && !bForceNewEditor)
 			{
-				// See if we are trying to open a single asset. If we are, we re-use a compatible anim editor.
-				bool bSingleAsset = AnimAssets.Num() == 1;
-				bool bFoundEditor = false;
-				if (bSingleAsset && !bForceNewEditor)
-				{
-					// See if there is an animation asset with the same skeleton already being edited
-					TArray<UObject*> AllEditedAssets;
+				// See if there is an animation asset with the same skeleton already being edited
+				TArray<UObject*> AllEditedAssets;
 #if WITH_EDITOR
-					AllEditedAssets = GEditor->GetEditorSubsystem<UAssetEditorSubsystem>()->GetAllEditedAssets();
+				AllEditedAssets = GEditor->GetEditorSubsystem<UAssetEditorSubsystem>()->GetAllEditedAssets();
 #endif
-					UAnimationAsset* CompatibleEditedAsset = nullptr;
-					for (UObject* EditedAsset : AllEditedAssets)
+				UAnimationAsset* CompatibleEditedAsset = nullptr;
+				for (UObject* EditedAsset : AllEditedAssets)
+				{
+					UAnimationAsset* EditedAnimAsset = Cast<UAnimationAsset>(EditedAsset);
+					if (EditedAnimAsset && EditedAnimAsset->GetSkeleton() == AnimAsset->GetSkeleton())
 					{
-						UAnimationAsset* EditedAnimAsset = Cast<UAnimationAsset>(EditedAsset);
-						if (EditedAnimAsset && EditedAnimAsset->GetSkeleton() == AnimAsset->GetSkeleton())
-						{
-							CompatibleEditedAsset = EditedAnimAsset;
-							break;
-						}
+						CompatibleEditedAsset = EditedAnimAsset;
+						break;
 					}
+				}
 
-					// If there is..
-					if(CompatibleEditedAsset)
-					{
-						// Find the anim editors that are doing it
-						TArray<IAssetEditorInstance*> AssetEditors;
+				// If there is..
+				if(CompatibleEditedAsset)
+				{
+					// Find the anim editors that are doing it
+					TArray<IAssetEditorInstance*> AssetEditors;
 #if WITH_EDITOR
-						AssetEditors = GEditor->GetEditorSubsystem<UAssetEditorSubsystem>()->FindEditorsForAsset(CompatibleEditedAsset);
+					AssetEditors = GEditor->GetEditorSubsystem<UAssetEditorSubsystem>()->FindEditorsForAsset(CompatibleEditedAsset);
 #endif
-						for (IAssetEditorInstance* ExistingEditor : AssetEditors)
+					for (IAssetEditorInstance* ExistingEditor : AssetEditors)
+					{
+						if (ExistingEditor->GetEditorName() == FName("AnimationEditor"))
 						{
-							if (ExistingEditor->GetEditorName() == FName("AnimationEditor"))
+							// Change the current anim to this one
+							IAnimationEditor* AnimEditor = static_cast<IAnimationEditor*>(ExistingEditor);
+							if(AnimEditor->GetPersonaToolkit()->GetSkeleton() == AnimAsset->GetSkeleton())
 							{
-								// Change the current anim to this one
-								IAnimationEditor* AnimEditor = static_cast<IAnimationEditor*>(ExistingEditor);
-								if(AnimEditor->GetPersonaToolkit()->GetSkeleton() == AnimAsset->GetSkeleton())
-								{
-									AnimEditor->SetAnimationAsset(AnimAsset);
-									AnimEditor->FocusWindow();
-									bFoundEditor = true;
-									break;
-								}
+								AnimEditor->SetAnimationAsset(AnimAsset);
+								AnimEditor->FocusWindow();
+								bFoundEditor = true;
+								break;
 							}
 						}
 					}
 				}
+			}
 
-				// We didn't find an editor, make a new one
-				if (!bFoundEditor)
-				{
-					IAnimationEditorModule& AnimationEditorModule = FModuleManager::LoadModuleChecked<IAnimationEditorModule>("AnimationEditor");
-					AnimationEditorModule.CreateAnimationEditor(Mode, EditWithinLevelEditor, AnimAsset);
-				}
+			// We didn't find an editor, make a new one
+			if (!bFoundEditor)
+			{
+				IAnimationEditorModule& AnimationEditorModule = FModuleManager::LoadModuleChecked<IAnimationEditorModule>("AnimationEditor");
+				AnimationEditorModule.CreateAnimationEditor(Mode, EditWithinLevelEditor, AnimAsset);
 			}
 		}
 	}
@@ -262,84 +194,4 @@ void FAssetTypeActions_AnimationAsset::ExecuteFindSkeleton(TArray<TWeakObjectPtr
 	}
 }
 
-void FAssetTypeActions_AnimationAsset::RetargetAnimationHandler(USkeleton* OldSkeleton, USkeleton* NewSkeleton, bool bRemapReferencedAssets, bool bAllowRemapToExisting, bool bConvertSpaces, const EditorAnimUtils::FNameDuplicationRule* NameRule, TArray<TWeakObjectPtr<UObject>> InAnimAssets)
-{
-	if(!OldSkeleton || OldSkeleton->GetPreviewMesh(false))
-	{
-		EditorAnimUtils::RetargetAnimations(OldSkeleton, NewSkeleton, InAnimAssets, bRemapReferencedAssets, NameRule, bConvertSpaces);
-	}
-	else
-	{
-		FFormatNamedArguments Args;
-		Args.Add(TEXT("OldSkeletonName"), FText::FromString(GetNameSafe(OldSkeleton)));
-		Args.Add(TEXT("NewSkeletonName"), FText::FromString(GetNameSafe(NewSkeleton)));
-		FNotificationInfo Info(FText::Format(LOCTEXT("Retarget Failed", "Old Skeleton {OldSkeletonName} and New Skeleton {NewSkeletonName} need to have Preview Mesh set up to convert animation"), Args));
-		Info.ExpireDuration = 5.0f;
-		Info.bUseLargeFont = false;
-		TSharedPtr<SNotificationItem> Notification = FSlateNotificationManager::Get().AddNotification(Info);
-		if(Notification.IsValid())
-		{
-			Notification->SetCompletionState(SNotificationItem::CS_Fail);
-		}
-	}
-}
-
-void FAssetTypeActions_AnimationAsset::RetargetNonSkeletonAnimationHandler(USkeleton* OldSkeleton, USkeleton* NewSkeleton, bool bRemapReferencedAssets, bool bAllowRemapToExisting, bool bConvertSpaces, const EditorAnimUtils::FNameDuplicationRule* NameRule, TArray<TWeakObjectPtr<UObject>> InAnimAssets, TWeakPtr<IToolkitHost> EditWithinLevelEditor)
-{
-	RetargetAnimationHandler(OldSkeleton, NewSkeleton, bRemapReferencedAssets, bAllowRemapToExisting, bConvertSpaces, NameRule, InAnimAssets);
-
-#if WITH_EDITOR
-	if(NewSkeleton)
-	{
-		for(auto Asset : InAnimAssets)
-		{
-			if (Asset.IsValid())
-			{
-				if (EditWithinLevelEditor.IsValid())
-				{
-
-					GEditor->GetEditorSubsystem<UAssetEditorSubsystem>()->OpenEditorForAsset(Asset.Get(), EToolkitMode::WorldCentric, EditWithinLevelEditor.Pin());
-				}
-				else
-				{
-					GEditor->GetEditorSubsystem<UAssetEditorSubsystem>()->OpenEditorForAsset(Asset.Get());
-				}
-			}
-		}
-	}
-	else
-	{
-		FMessageDialog::Open(EAppMsgType::Ok, LOCTEXT("FailedToLoadSkeletonlessAnimAsset", "The Anim Asset could not be loaded because it's skeleton is missing."));
-	}
-#endif
-
-}
-
-
-void FAssetTypeActions_AnimationAsset::RetargetAssets(TArray<UObject*> InAnimAssets, bool bDuplicateAssets, bool bOpenEditor, TSharedPtr<class IToolkitHost> EditWithinLevelEditor)
-{
-	bool bRemapReferencedAssets = false;
-	USkeleton* NewSkeleton = NULL;
-	USkeleton* OldSkeleton = NULL;
-	if(InAnimAssets.Num() > 0)
-	{
-		UAnimationAsset * AnimAsset = CastChecked<UAnimationAsset>(InAnimAssets[0]);
-		OldSkeleton = AnimAsset->GetSkeleton();
-	}
-
-	const FText Message = LOCTEXT("SelectSkeletonToRemap", "Select the skeleton to remap this asset to.");
-
-	auto AnimAssets = GetTypedWeakObjectPtrs<UObject>(InAnimAssets);
-
-	if (bOpenEditor)
-	{
-		SAnimationRemapSkeleton::ShowWindow(OldSkeleton, Message, bDuplicateAssets, FOnRetargetAnimation::CreateSP(this, &FAssetTypeActions_AnimationAsset::RetargetNonSkeletonAnimationHandler, AnimAssets, TWeakPtr<class IToolkitHost> (EditWithinLevelEditor)) );
-	}
-	else
-	{
-		SAnimationRemapSkeleton::ShowWindow(OldSkeleton, Message, bDuplicateAssets, FOnRetargetAnimation::CreateSP(this, &FAssetTypeActions_AnimationAsset::RetargetAnimationHandler, AnimAssets) );
-	}
-}
-
-
 #undef LOCTEXT_NAMESPACE
diff --git a/Engine/Source/Developer/AssetTools/Private/AssetTypeActions/AssetTypeActions_Skeleton.cpp b/Engine/Source/Developer/AssetTools/Private/AssetTypeActions/AssetTypeActions_Skeleton.cpp
--- a/Engine/Source/Developer/AssetTools/Private/AssetTypeActions/AssetTypeActions_Skeleton.cpp
+++ b/Engine/Source/Developer/AssetTools/Private/AssetTypeActions/AssetTypeActions_Skeleton.cpp
@@ -442,57 +442,10 @@ void FAssetTypeActions_Skeleton::GetActions(const TArray<UObject*>& InObjects, F
 			false, 
 			FSlateIcon(FEditorStyle::GetStyleSetName(), "Persona.AssetActions.CreateAnimAsset")
 			);
-
-	// only show if one is selected. It won't work since I changed the window to be normal window
-	if (Skeletons.Num() == 1)
-	{
-		Section.AddMenuEntry(
-			"Skeleton_Retarget",
-			LOCTEXT("Skeleton_Retarget", "Retarget to Another Skeleton"),
-			LOCTEXT("Skeleton_RetargetTooltip", "Allow all animation assets for this skeleton retarget to another skeleton."),
-			FSlateIcon(FEditorStyle::GetStyleSetName(), "Persona.AssetActions.RetargetSkeleton"),
-			FUIAction(
-				FExecuteAction::CreateSP(this, &FAssetTypeActions_Skeleton::ExecuteRetargetSkeleton, Skeletons),
-				FCanExecuteAction()
-				)
-			);
-	}
-
-	// @todo ImportAnimation
-	/*
-	Section.AddMenuEntry(
-		"Skeleton_ImportAnimation",
-		LOCTEXT("Skeleton_ImportAnimation", "Import Animation"),
-		LOCTEXT("Skeleton_ImportAnimationTooltip", "Imports an animation for the selected skeleton."),
-		FSlateIcon(),
-		FUIAction(
-			FExecuteAction::CreateSP( this, &FAssetTypeActions_Skeleton::ExecuteImportAnimation, Skeletons ),
-			FCanExecuteAction()
-			)
-		);
-	*/;
 }
 
 void FAssetTypeActions_Skeleton::FillCreateMenu(FMenuBuilder& MenuBuilder, TArray<TWeakObjectPtr<USkeleton>> Skeletons) const
 {
-	// create rig
-	if (Skeletons.Num() == 1)
-	{
-		MenuBuilder.BeginSection("CreateRig", LOCTEXT("CreateRigMenuHeading", "Rig"));
-		{
-			MenuBuilder.AddMenuEntry(
-				LOCTEXT("Skeleton_CreateRig", "Create Rig"),
-				LOCTEXT("Skeleton_CreateRigTooltip", "Create Rig from this skeleton."),
-				FSlateIcon(),
-				FUIAction(
-				FExecuteAction::CreateSP(const_cast<FAssetTypeActions_Skeleton*>(this), &FAssetTypeActions_Skeleton::ExecuteCreateRig, Skeletons),
-				FCanExecuteAction()
-				)
-				);
-		}
-		MenuBuilder.EndSection();
-	}
-
 	TArray<TWeakObjectPtr<UObject>> Objects;
 	Algo::Transform(Skeletons, Objects, [](const TWeakObjectPtr<USkeleton>& Skeleton) { return Skeleton; });
 	AnimationEditorUtils::FillCreateAssetMenu(MenuBuilder, Objects, FAnimAssetCreated::CreateSP(const_cast<FAssetTypeActions_Skeleton*>(this), &FAssetTypeActions_Skeleton::OnAssetCreated));
@@ -523,469 +476,6 @@ void FAssetTypeActions_Skeleton::OpenAssetEditor( const TArray<UObject*>& InObje
 	}
 }
 
-/** Handler for when Create Rig is selected */
-void FAssetTypeActions_Skeleton::ExecuteCreateRig(TArray<TWeakObjectPtr<USkeleton>> Skeletons)
-{
-	if (Skeletons.Num() == 1)
-	{
-		CreateRig(Skeletons[0]);
-	}
-}
-
-void FAssetTypeActions_Skeleton::CreateRig(const TWeakObjectPtr<USkeleton> Skeleton)
-{
-	if(Skeleton.IsValid())
-	{
-		FCreateRigDlg CreateRigDlg(Skeleton.Get());
-		if(CreateRigDlg.ShowModal() == FCreateRigDlg::Confirm)
-		{
-			check(CreateRigDlg.RequiredBones.Num() > 0);
-
-			// Determine an appropriate name
-			FString Name;
-			FString PackageName;
-			CreateUniqueAssetName(Skeleton->GetOutermost()->GetName(), TEXT("Rig"), PackageName, Name);
-
-			// Create the asset, and assign its skeleton
-			FAssetToolsModule& AssetToolsModule = FModuleManager::GetModuleChecked<FAssetToolsModule>("AssetTools");
-			URig* NewAsset = Cast<URig>(AssetToolsModule.Get().CreateAsset(Name, FPackageName::GetLongPackagePath(PackageName), URig::StaticClass(), NULL));
-
-			if(NewAsset)
-			{
-				NewAsset->CreateFromSkeleton(Skeleton.Get(), CreateRigDlg.RequiredBones);
-				NewAsset->MarkPackageDirty();
-
-
-				TArray<UObject*> ObjectsToSync;
-				ObjectsToSync.Add(NewAsset);
-				FAssetTools::Get().SyncBrowserToAssets(ObjectsToSync);
-			}
-		}
-	}
-}
-
-void FAssetTypeActions_Skeleton::RetargetAnimationHandler(USkeleton* OldSkeleton, USkeleton* NewSkeleton, bool bRemapReferencedAssets, bool bAllowRemapToExisting, bool bConvertSpaces, const EditorAnimUtils::FNameDuplicationRule* NameRule)
-{
-	if(OldSkeleton == nullptr || OldSkeleton->GetPreviewMesh(true) == nullptr)
-	{
-		FFormatNamedArguments Args;
-		Args.Add(TEXT("OldSkeletonName"), FText::FromString(GetNameSafe(OldSkeleton)));
-		Args.Add(TEXT("NewSkeletonName"), FText::FromString(GetNameSafe(NewSkeleton)));
-		FNotificationInfo Info(FText::Format(LOCTEXT("Retarget Failed", "Old Skeleton {OldSkeletonName} and New Skeleton {NewSkeletonName} need to have Preview Mesh set up to convert animation"), Args));
-		Info.ExpireDuration = 5.0f;
-		Info.bUseLargeFont = false;
-		TSharedPtr<SNotificationItem> Notification = FSlateNotificationManager::Get().AddNotification(Info);
-		if(Notification.IsValid())
-		{
-			Notification->SetCompletionState(SNotificationItem::CS_Fail);
-		}
-
-		return;
-	}
-
-	// namerule should be null
-	// find all assets who references old skeleton
-	TArray<FName> Packages;
-
-	// If the asset registry is still loading assets, we cant check for referencers, so we must open the rename dialog
-	FAssetRegistryModule& AssetRegistryModule = FModuleManager::LoadModuleChecked<FAssetRegistryModule>("AssetRegistry");
-	AssetRegistryModule.Get().GetReferencers(OldSkeleton->GetOutermost()->GetFName(), Packages);
-
-	if(AssetRegistryModule.Get().IsLoadingAssets())
-	{
-		// Open a dialog asking the user to wait while assets are being discovered
-		SDiscoveringAssetsDialog::OpenDiscoveringAssetsDialog(
-			SDiscoveringAssetsDialog::FOnAssetsDiscovered::CreateSP(this, &FAssetTypeActions_Skeleton::PerformRetarget, OldSkeleton, NewSkeleton, Packages, bConvertSpaces)
-			);
-	}
-	else
-	{
-		PerformRetarget(OldSkeleton, NewSkeleton, Packages, bConvertSpaces);
-	}
-}
-void FAssetTypeActions_Skeleton::ExecuteRetargetSkeleton(TArray<TWeakObjectPtr<USkeleton>> Skeletons)
-{
-	TArray<UObject*> AllEditedAssets;
-#if WITH_EDITOR
-	AllEditedAssets = GEditor->GetEditorSubsystem<UAssetEditorSubsystem>()->GetAllEditedAssets();
-#endif
-
-	for (auto SkelIt = Skeletons.CreateConstIterator(); SkelIt; ++SkelIt)
-	{	
-		USkeleton* OldSkeleton = (*SkelIt).Get();
-
-		// Close any assets that reference this skeleton
-		for(UObject* EditedAsset : AllEditedAssets)
-		{
-			bool bCloseAssetEditor = false;
-			if (USkeleton* Skeleton = Cast<USkeleton>(EditedAsset))
-			{
-				bCloseAssetEditor = OldSkeleton == Skeleton;
-			}
-			else if (UAnimationAsset* AnimationAsset = Cast<UAnimationAsset>(EditedAsset))
-			{
-				bCloseAssetEditor = OldSkeleton == AnimationAsset->GetSkeleton();
-			}
-			else if (USkeletalMesh* SkeletalMesh = Cast<USkeletalMesh>(EditedAsset))
-			{
-				bCloseAssetEditor = OldSkeleton == SkeletalMesh->GetSkeleton();
-			}
-			else if (UAnimBlueprint* AnimBlueprint = Cast<UAnimBlueprint>(EditedAsset))
-			{
-				bCloseAssetEditor = OldSkeleton == AnimBlueprint->TargetSkeleton;
-			}
-			else if (UPhysicsAsset* PhysicsAsset = Cast<UPhysicsAsset>(EditedAsset))
-			{
-				USkeletalMesh* PhysicsSkeletalMesh = PhysicsAsset->PreviewSkeletalMesh.LoadSynchronous();
-				if(PhysicsSkeletalMesh != nullptr)
-				{
-					bCloseAssetEditor = OldSkeleton == PhysicsSkeletalMesh->GetSkeleton();
-				}
-			}
-
-			if(bCloseAssetEditor)
-			{
-#if WITH_EDITOR
-				GEditor->GetEditorSubsystem<UAssetEditorSubsystem>()->CloseAllEditorsForAsset(EditedAsset);
-#endif
-			}
-		}
-
-		const FText Message = LOCTEXT("RetargetSkeleton_Warning", "This only converts animation data -i.e. animation assets and Anim Blueprints. \nIf you'd like to convert SkeletalMesh, use the context menu (Assign Skeleton) for each mesh. \n\nIf you'd like to convert mesh as well, please do so before converting animation data. \nOtherwise you will lose any extra track that is in the new mesh.");
-		// ask user what they'd like to change to 
-		SAnimationRemapSkeleton::ShowWindow(OldSkeleton, Message, false, FOnRetargetAnimation::CreateSP(this, &FAssetTypeActions_Skeleton::RetargetAnimationHandler));
-	}
-}
-
-void FAssetTypeActions_Skeleton::PerformRetarget(USkeleton *OldSkeleton, USkeleton *NewSkeleton, TArray<FName> Packages, bool bConvertSpaces) const
-{
-	TArray<FAssetToRemapSkeleton> AssetsToRemap;
-
-	// populate AssetsToRemap
-	AssetsToRemap.Empty(Packages.Num());
-
-	for ( int32 AssetIndex=0; AssetIndex<Packages.Num(); ++AssetIndex)
-	{
-		AssetsToRemap.Add( FAssetToRemapSkeleton(Packages[AssetIndex]) );
-	}
-
-	// Load all packages 
-	TArray<UPackage*> PackagesToSave;
-	LoadPackages(AssetsToRemap, PackagesToSave);
-
-	// Update the source control state for the packages containing the assets we are remapping
-	ISourceControlProvider& SourceControlProvider = ISourceControlModule::Get().GetProvider();
-	if ( ISourceControlModule::Get().IsEnabled() )
-	{
-		SourceControlProvider.Execute(ISourceControlOperation::Create<FUpdateStatus>(), PackagesToSave);
-	}
-
-	// Prompt to check out all referencing packages, leave redirectors for assets referenced by packages that are not checked out and remove those packages from the save list.
-	const bool bUserAcceptedCheckout = CheckOutPackages(AssetsToRemap, PackagesToSave);
-
-	if ( bUserAcceptedCheckout )
-	{
-		// If any referencing packages are left read-only, the checkout failed or SCC was not enabled. Trim them from the save list and leave redirectors.
-		DetectReadOnlyPackages(AssetsToRemap, PackagesToSave);
-
-		// retarget skeleton
-		RetargetSkeleton(AssetsToRemap, OldSkeleton, NewSkeleton, bConvertSpaces);
-
-		// Save all packages that were referencing any of the assets that were moved without redirectors
-		SavePackages(PackagesToSave);
-
-		// Finally, report any failures that happened during the rename
-		ReportFailures(AssetsToRemap);
-	}
-}
-
-void FAssetTypeActions_Skeleton::LoadPackages(TArray<FAssetToRemapSkeleton>& AssetsToRemap, TArray<UPackage*>& OutPackagesToSave) const
-{
-	const FText StatusUpdate = LOCTEXT("RemapSkeleton_LoadPackage", "Loading Packages");
-	GWarn->BeginSlowTask( StatusUpdate, true );
-
-	FAssetToolsModule& AssetToolsModule = FModuleManager::GetModuleChecked<FAssetToolsModule>("AssetTools");
-	// go through all assets try load
-	for ( int32 AssetIdx = 0; AssetIdx < AssetsToRemap.Num(); ++AssetIdx )
-	{
-		GWarn->StatusUpdate( AssetIdx, AssetsToRemap.Num(), StatusUpdate );
-
-		FAssetToRemapSkeleton& RemapData = AssetsToRemap[AssetIdx];
-
-		const FString PackageName = RemapData.PackageName.ToString();
-
-		// load package
-		UPackage* Package = LoadPackage(NULL, *PackageName, LOAD_None);
-		if ( !Package )
-		{
-			RemapData.ReportFailed( LOCTEXT("RemapSkeletonFailed_LoadPackage", "Could not load the package."));
-			continue;
-		}
-
-		// get all the objects
-		TArray<UObject*> Objects;
-		GetObjectsWithOuter(Package, Objects);
-
-		// see if we have skeletalmesh
-		bool bSkeletalMeshPackage = false;
-		for (auto Iter=Objects.CreateIterator(); Iter; ++Iter)
-		{
-			// we only care animation asset or animation blueprint
-			if ((*Iter)->GetClass()->IsChildOf(UAnimationAsset::StaticClass()) ||
-				(*Iter)->GetClass()->IsChildOf(UAnimBlueprint::StaticClass()))
-			{
-				// add to asset 
-				RemapData.Asset = (*Iter);
-				break;
-			}
-			else if ((*Iter)->GetClass()->IsChildOf(USkeletalMesh::StaticClass()) )
-			{
-				bSkeletalMeshPackage = true;
-				break;
-			}
-		}
-
-		// if we have skeletalmesh, we ignore this package, do not report as error
-		if ( bSkeletalMeshPackage )
-		{
-			continue;
-		}
-
-		// if none was relevant - skeletal mesh is going to get here
-		if ( !RemapData.Asset.IsValid() )
-		{
-			RemapData.ReportFailed( LOCTEXT("RemapSkeletonFailed_LoadObject", "Could not load any related object."));
-			continue;
-		}
-
-		OutPackagesToSave.Add(Package);
-	}
-
-	GWarn->EndSlowTask();
-}
-
-bool FAssetTypeActions_Skeleton::CheckOutPackages(TArray<FAssetToRemapSkeleton>& AssetsToRemap, TArray<UPackage*>& InOutPackagesToSave) const
-{
-	bool bUserAcceptedCheckout = true;
-	
-	if ( InOutPackagesToSave.Num() > 0 )
-	{
-		if ( ISourceControlModule::Get().IsEnabled() )
-		{
-			TArray<UPackage*> PackagesCheckedOutOrMadeWritable;
-			TArray<UPackage*> PackagesNotNeedingCheckout;
-			bUserAcceptedCheckout = FEditorFileUtils::PromptToCheckoutPackages( false, InOutPackagesToSave, &PackagesCheckedOutOrMadeWritable, &PackagesNotNeedingCheckout );
-			if ( bUserAcceptedCheckout )
-			{
-				TArray<UPackage*> PackagesThatCouldNotBeCheckedOut = InOutPackagesToSave;
-
-				for ( auto PackageIt = PackagesCheckedOutOrMadeWritable.CreateConstIterator(); PackageIt; ++PackageIt )
-				{
-					PackagesThatCouldNotBeCheckedOut.Remove(*PackageIt);
-				}
-
-				for ( auto PackageIt = PackagesNotNeedingCheckout.CreateConstIterator(); PackageIt; ++PackageIt )
-				{
-					PackagesThatCouldNotBeCheckedOut.Remove(*PackageIt);
-				}
-
-				for ( auto PackageIt = PackagesThatCouldNotBeCheckedOut.CreateConstIterator(); PackageIt; ++PackageIt )
-				{
-					const FName NonCheckedOutPackageName = (*PackageIt)->GetFName();
-
-					for ( auto RenameDataIt = AssetsToRemap.CreateIterator(); RenameDataIt; ++RenameDataIt )
-					{
-						FAssetToRemapSkeleton& RemapData = *RenameDataIt;
-
-						if (RemapData.Asset.IsValid() && RemapData.Asset.Get()->GetOutermost() == *PackageIt)
-						{
-							RemapData.ReportFailed( LOCTEXT("RemapSkeletonFailed_CheckOutFailed", "Check out failed"));
-						}
-					}
-
-					InOutPackagesToSave.Remove(*PackageIt);
-				}
-			}
-		}
-	}
-
-	return bUserAcceptedCheckout;
-}
-
-void FAssetTypeActions_Skeleton::DetectReadOnlyPackages(TArray<FAssetToRemapSkeleton>& AssetsToRemap, TArray<UPackage*>& InOutPackagesToSave) const
-{
-	// For each valid package...
-	for ( int32 PackageIdx = InOutPackagesToSave.Num() - 1; PackageIdx >= 0; --PackageIdx )
-	{
-		UPackage* Package = InOutPackagesToSave[PackageIdx];
-
-		if ( Package )
-		{
-			// Find the package filename
-			FString Filename;
-			if ( FPackageName::DoesPackageExist(Package->GetName(), &Filename) )
-			{
-				// If the file is read only
-				if ( IFileManager::Get().IsReadOnly(*Filename) )
-				{
-					FName PackageName = Package->GetFName();
-
-					for ( auto RenameDataIt = AssetsToRemap.CreateIterator(); RenameDataIt; ++RenameDataIt )
-					{
-						FAssetToRemapSkeleton& RenameData = *RenameDataIt;
-
-						if (RenameData.Asset.IsValid() && RenameData.Asset.Get()->GetOutermost() == Package)
-						{
-							RenameData.ReportFailed( LOCTEXT("RemapSkeletonFailed_FileReadOnly", "File still read only"));
-						}
-					}				
-
-					// Remove the package from the save list
-					InOutPackagesToSave.RemoveAt(PackageIdx);
-				}
-			}
-		}
-	}
-}
-
-void FAssetTypeActions_Skeleton::SavePackages(const TArray<UPackage*> PackagesToSave) const
-{
-	if ( PackagesToSave.Num() > 0 )
-	{
-		const bool bCheckDirty = false;
-		const bool bPromptToSave = false;
-		FEditorFileUtils::PromptForCheckoutAndSave(PackagesToSave, bCheckDirty, bPromptToSave);
-
-		ISourceControlModule::Get().QueueStatusUpdate(PackagesToSave);
-	}
-}
-
-void FAssetTypeActions_Skeleton::ReportFailures(const TArray<FAssetToRemapSkeleton>& AssetsToRemap) const
-{
-	TArray<FText> FailedToRemap;
-	for ( auto RenameIt = AssetsToRemap.CreateConstIterator(); RenameIt; ++RenameIt )
-	{
-		const FAssetToRemapSkeleton& RemapData = *RenameIt;
-		if ( RemapData.bRemapFailed )
-		{
-			UObject* Asset = RemapData.Asset.Get();
-			if ( Asset )
-			{
-				FFormatNamedArguments Args;
-				Args.Add(TEXT("FailureReason"), RemapData.FailureReason);
-				Args.Add(TEXT("AssetName"), FText::FromString(Asset->GetOutermost()->GetName()));
-
-				FailedToRemap.Add(FText::Format(LOCTEXT("AssetRemapFailure", "{AssetName} - {FailureReason}"), Args));
-			}
-			else
-			{
-				FailedToRemap.Add(LOCTEXT("RemapSkeletonFailed_InvalidAssetText", "Invalid Asset"));
-			}
-		}
-	}
-
-	if ( FailedToRemap.Num() > 0 )
-	{
-		SRemapFailures::OpenRemapFailuresDialog(FailedToRemap);
-	}
-}
-
-void FAssetTypeActions_Skeleton::RetargetSkeleton(TArray<FAssetToRemapSkeleton>& AssetsToRemap, USkeleton* OldSkeleton, USkeleton* NewSkeleton, bool bConvertSpaces) const
-{
-	TArray<UAnimBlueprint*>	AnimBlueprints;
-
-	// first we convert all individual assets
-	for ( auto RenameIt = AssetsToRemap.CreateIterator(); RenameIt; ++RenameIt )
-	{
-		FAssetToRemapSkeleton& RenameData = *RenameIt;
-		if ( RenameData.bRemapFailed == false  )
-		{
-			UObject* Asset = RenameData.Asset.Get();
-			if ( Asset )
-			{
-				if ( UAnimationAsset* AnimAsset = Cast<UAnimationAsset>(Asset) )
-				{
-					TUniquePtr<IAnimationDataController::FScopedBracket> ScopedBracket;
-
-					if (UAnimSequenceBase* SequenceBase = Cast<UAnimSequenceBase>(AnimAsset))
-					{
-						IAnimationDataController& Controller = SequenceBase->GetController();
-						ScopedBracket = MakeUnique<IAnimationDataController::FScopedBracket>(Controller, LOCTEXT("FAssetTypeActions_Skeleton_Bracket", "Retargeting Skeleton"));
-
-						Controller.FindOrAddCurveNamesOnSkeleton(NewSkeleton, ERawCurveTrackTypes::RCT_Float);
-						
-						if (UAnimSequence* Sequence = Cast<UAnimSequence>(SequenceBase))
-						{
-							Controller.FindOrAddCurveNamesOnSkeleton(NewSkeleton, ERawCurveTrackTypes::RCT_Transform);
-						}
-					}
-					
-					AnimAsset->ReplaceSkeleton(NewSkeleton, bConvertSpaces);
-				}
-				else if ( UAnimBlueprint* AnimBlueprint = Cast<UAnimBlueprint>(Asset) )
-				{
-					AnimBlueprints.Add(AnimBlueprint);
-				}
-			}
-		}
-	}
-
-	// convert all Animation Blueprints and compile 
-	for ( auto AnimBPIter = AnimBlueprints.CreateIterator(); AnimBPIter; ++AnimBPIter )
-	{
-		UAnimBlueprint * AnimBlueprint = (*AnimBPIter);
-		AnimBlueprint->TargetSkeleton = NewSkeleton;
-		
-		FBlueprintEditorUtils::RefreshAllNodes(AnimBlueprint);
-		FKismetEditorUtilities::CompileBlueprint(AnimBlueprint, EBlueprintCompileOptions::SkipGarbageCollection);
-	}
-
-	// Copy sockets IF the socket doesn't exists on target skeleton and if the joint exists
-	if (OldSkeleton && OldSkeleton->Sockets.Num() > 0 && NewSkeleton)
-	{
-		const FReferenceSkeleton& NewRefSkeleton = NewSkeleton->GetReferenceSkeleton();
-		// if we have sockets from old skeleton, see if we can trasnfer
-		for (const auto& OldSocket : OldSkeleton->Sockets)
-		{
-			bool bExistingOnNewSkeleton = false;
-
-			for (auto& NewSocket : NewSkeleton->Sockets)
-			{
-				if (OldSocket->SocketName == NewSocket->SocketName)
-				{
-					// if name is same, we can't copy over
-					bExistingOnNewSkeleton = true;
-				}
-			}
-
-			if (!bExistingOnNewSkeleton)
-			{
-				// make sure the joint still exists
-				if (NewRefSkeleton.FindBoneIndex(OldSocket->BoneName) != INDEX_NONE)
-				{
-					USkeletalMeshSocket* NewSocketInst = NewObject<USkeletalMeshSocket>(NewSkeleton);
-					if (NewSocketInst)
-					{
-						NewSocketInst->CopyFrom(OldSocket);
-						NewSkeleton->Sockets.Add(NewSocketInst);
-						NewSkeleton->MarkPackageDirty();
-					}
-				}
-			}
-		}
-	}
-
-	// now update any running instance
-	for (FThreadSafeObjectIterator Iter(USkeletalMeshComponent::StaticClass()); Iter; ++Iter)
-	{
-		USkeletalMeshComponent * MeshComponent = Cast<USkeletalMeshComponent>(*Iter);
-		if (MeshComponent->SkeletalMesh && MeshComponent->SkeletalMesh->GetSkeleton() == OldSkeleton)
-		{
-			MeshComponent->InitAnim(true);
-		}
-	}
-}
-
 bool FAssetTypeActions_Skeleton::OnAssetCreated(TArray<UObject*> NewAssets) const
 {
 	if(NewAssets.Num() > 1)
diff --git a/Engine/Source/Developer/AssetTools/Private/AssetTypeActions/AssetTypeActions_Skeleton.h b/Engine/Source/Developer/AssetTools/Private/AssetTypeActions/AssetTypeActions_Skeleton.h
--- a/Engine/Source/Developer/AssetTools/Private/AssetTypeActions/AssetTypeActions_Skeleton.h
+++ b/Engine/Source/Developer/AssetTools/Private/AssetTypeActions/AssetTypeActions_Skeleton.h
@@ -44,38 +44,8 @@ public:
 	virtual void OpenAssetEditor( const TArray<UObject*>& InObjects, TSharedPtr<class IToolkitHost> EditWithinLevelEditor = TSharedPtr<IToolkitHost>() ) override;
 	virtual uint32 GetCategories() override { return EAssetTypeCategories::Animation; }
 
-private:
-	/** Handler for when Skeleton Retarget is selected */
-	void ExecuteRetargetSkeleton(TArray<TWeakObjectPtr<USkeleton>> Skeletons);
-
-	/** Handler for when Create Rig is selected */
-	void ExecuteCreateRig(TArray<TWeakObjectPtr<USkeleton>> Skeletons);
-
 private: // Helper functions
-	/** Creates animation assets using the BaseName+Suffix */
-	void CreateRig(const TWeakObjectPtr<USkeleton> Skeleton);
 
- 	/** 
-	 * Main function for handling retargeting old Skeleton to new Skeleton
-	 * 
-	 * @param OldSkeleton : Old Skeleton that it changes from
-	 * @param NewSkeleton : New Skeleton that it would change to
-	 * @param Referencers : List of asset names that reference Old Skeleton
-	 * @param bConvertSpaces : Whether to convert spaces or not
-	 */
- 	void PerformRetarget(USkeleton *OldSkeleton, USkeleton *NewSkeleton, TArray<FName> Packages, bool bConvertSpaces) const;
-
-	// utility functions for performing retargeting
-	// these codes are from AssetRenameManager workflow
-	void LoadPackages(TArray<FAssetToRemapSkeleton>& AssetsToRemap, TArray<UPackage*>& OutPackagesToSave) const;
-	bool CheckOutPackages(TArray<FAssetToRemapSkeleton>& AssetsToRemap, TArray<UPackage*>& InOutPackagesToSave) const;
-	void ReportFailures(const TArray<FAssetToRemapSkeleton>& AssetsToRemap) const;
-	void RetargetSkeleton(TArray<FAssetToRemapSkeleton>& AssetsToRemap, USkeleton* OldSkeleton, USkeleton* NewSkeleton, bool bConvertSpaces) const;
-	void SavePackages(const TArray<UPackage*> PackagesToSave) const;
-	void DetectReadOnlyPackages(TArray<FAssetToRemapSkeleton>& AssetsToRemap, TArray<UPackage*>& InOutPackagesToSave) const;
 	void FillCreateMenu(FMenuBuilder& MenuBuilder, TArray<TWeakObjectPtr<USkeleton>> Skeletons) const;
 	bool OnAssetCreated(TArray<UObject*> NewAssets) const;
-
-	/** Handler for retargeting */
-	void RetargetAnimationHandler(USkeleton* OldSkeleton, USkeleton* NewSkeleton, bool bRemapReferencedAssets, bool bAllowRemapToExisting, bool bConvertSpaces, const EditorAnimUtils::FNameDuplicationRule* NameRule);
 };
diff --git a/Engine/Source/Developer/AssetTools/Public/AssetTypeActions/AssetTypeActions_AnimationAsset.h b/Engine/Source/Developer/AssetTools/Public/AssetTypeActions/AssetTypeActions_AnimationAsset.h
--- a/Engine/Source/Developer/AssetTools/Public/AssetTypeActions/AssetTypeActions_AnimationAsset.h
+++ b/Engine/Source/Developer/AssetTools/Public/AssetTypeActions/AssetTypeActions_AnimationAsset.h
@@ -23,8 +23,6 @@ public:
 	virtual class UThumbnailInfo* GetThumbnailInfo(UObject* Asset) const override;
 
 private:
-	/** Handler to fill the retarget submenu */
-	void FillRetargetMenu( FMenuBuilder& MenuBuilder, const TArray<UObject*> InObjects );
 
 	/** Handler for when FindSkeleton is selected */
 	void ExecuteFindSkeleton(TArray<TWeakObjectPtr<UAnimationAsset>> Objects);
@@ -34,11 +32,4 @@ private:
 
 	/** Open animation asset, will find existing editor if desired. */
 	void OpenAnimAssetEditor(const TArray<UObject*>& InObjects, bool bForceNewEditor, TSharedPtr<IToolkitHost> EditWithinLevelEditor);
-
-	/** Context menu item handler for changing the supplied assets skeletons */ 
-	void RetargetAssets(TArray<UObject*> InAnimBlueprints, bool bDuplicateAssets, bool bOpenEditor, TSharedPtr<class IToolkitHost> EditWithinLevelEditor);
-
-	/** Handler for retargeting */
-	void RetargetAnimationHandler(USkeleton* OldSkeleton, USkeleton* NewSkeleton, bool bRemapReferencedAssets, bool bAllowRemapToExisting, bool bConvertSpaces, const EditorAnimUtils::FNameDuplicationRule* NameRule, TArray<TWeakObjectPtr<UObject>> InAnimAssets);
-	void RetargetNonSkeletonAnimationHandler(USkeleton* OldSkeleton, USkeleton* NewSkeleton, bool bRemapReferencedAssets, bool bAllowRemapToExisting, bool bConvertSpaces, const EditorAnimUtils::FNameDuplicationRule* NameRule, TArray<TWeakObjectPtr<UObject>> InAnimAssets, TWeakPtr<class IToolkitHost> EditWithinLevelEditor);
 };
diff --git a/Engine/Source/Editor/Persona/Private/PersonaModule.cpp b/Engine/Source/Editor/Persona/Private/PersonaModule.cpp
--- a/Engine/Source/Editor/Persona/Private/PersonaModule.cpp
+++ b/Engine/Source/Editor/Persona/Private/PersonaModule.cpp
@@ -296,9 +296,9 @@ TSharedRef<class FWorkflowTabFactory> FPersonaModule::CreateCurveViewerTabFactor
 	return MakeShareable(new FAnimCurveViewerTabSummoner(InHostingApp, InEditableSkeleton, InPreviewScene, InOnObjectsSelected));
 }
 
-TSharedRef<class FWorkflowTabFactory> FPersonaModule::CreateRetargetManagerTabFactory(const TSharedRef<class FWorkflowCentricApplication>& InHostingApp, const TSharedRef<class IEditableSkeleton>& InEditableSkeleton, const TSharedRef<IPersonaPreviewScene>& InPreviewScene, FSimpleMulticastDelegate& InOnPostUndo) const
+TSharedRef<class FWorkflowTabFactory> FPersonaModule::CreateRetargetSourcesTabFactory(const TSharedRef<class FWorkflowCentricApplication>& InHostingApp, const TSharedRef<class IEditableSkeleton>& InEditableSkeleton, const TSharedRef<IPersonaPreviewScene>& InPreviewScene, FSimpleMulticastDelegate& InOnPostUndo) const
 {
-	return MakeShareable(new FRetargetManagerTabSummoner(InHostingApp, InEditableSkeleton, InPreviewScene, InOnPostUndo));
+	return MakeShareable(new FRetargetSourcesTabSummoner(InHostingApp, InEditableSkeleton, InPreviewScene, InOnPostUndo));
 }
 
 TSharedRef<class FWorkflowTabFactory> FPersonaModule::CreateAdvancedPreviewSceneTabFactory(const TSharedRef<class FWorkflowCentricApplication>& InHostingApp, const TSharedRef<IPersonaPreviewScene>& InPreviewScene) const
diff --git a/Engine/Source/Editor/Persona/Private/SAnimationSequenceBrowser.cpp b/Engine/Source/Editor/Persona/Private/SAnimationSequenceBrowser.cpp
--- a/Engine/Source/Editor/Persona/Private/SAnimationSequenceBrowser.cpp
+++ b/Engine/Source/Editor/Persona/Private/SAnimationSequenceBrowser.cpp
@@ -645,23 +645,6 @@ TSharedPtr<SWidget> SAnimationSequenceBrowser::OnGetAssetContextMenu(const TArra
 	}
 	MenuBuilder.EndSection();
 
-	if (bHasSelectedAnimAsset)
-	{
-		MenuBuilder.BeginSection("AnimationSequenceAdvancedOptions", LOCTEXT("AdvancedOptionsHeading", "Advanced") );
-	{
-		MenuBuilder.AddMenuEntry(
-			LOCTEXT("ChangeSkeleton", "Create a copy for another Skeleton..."),
-			LOCTEXT("ChangeSkeleton_ToolTip", "Create a copy for different skeleton"),
-			FSlateIcon(),
-			FUIAction(
-			FExecuteAction::CreateSP( this, &SAnimationSequenceBrowser::OnCreateCopy, SelectedAssets ),
-			FCanExecuteAction()
-			)
-			);
-	}
-		MenuBuilder.EndSection();
-	}
-
 	return MenuBuilder.MakeWidget();
 }
 
@@ -819,49 +802,6 @@ void SAnimationSequenceBrowser::OnReimportAnimation(TArray<FAssetData> SelectedA
 	}
 }
 
-void SAnimationSequenceBrowser::RetargetAnimationHandler(USkeleton* OldSkeleton, USkeleton* NewSkeleton, bool bRemapReferencedAssets, bool bAllowRemapToExisting, bool bConvertSpaces, const EditorAnimUtils::FNameDuplicationRule* NameRule, TArray<TWeakObjectPtr<UObject>> InAnimAssets)
-{
-	UObject* AssetToOpen = EditorAnimUtils::RetargetAnimations(OldSkeleton, NewSkeleton, InAnimAssets, bRemapReferencedAssets, NameRule, bConvertSpaces);
-
-	if(UAnimationAsset* AnimAsset = Cast<UAnimationAsset>(AssetToOpen))
-	{
-		FAssetRegistryModule::AssetCreated(AssetToOpen);
-
-		// once all success, attempt to open new editor with new skeleton
-		GEditor->GetEditorSubsystem<UAssetEditorSubsystem>()->OpenEditorForAsset(AssetToOpen);
-	}
-}
-
-void SAnimationSequenceBrowser::OnCreateCopy(TArray<FAssetData> Selected)
-{
-	if ( Selected.Num() > 0 )
-	{
-		// ask which skeleton users would like to choose
-		USkeleton* OldSkeleton = PersonaToolkitPtr.Pin()->GetSkeleton();
-		USkeleton* NewSkeleton = NULL;
-		bool		bDuplicateAssets = true;
-
-		const FText Message = LOCTEXT("RemapSkeleton_Warning", "This will duplicate the asset and convert to new skeleton.");
-
-		TArray<UObject *> AnimAssets;
-		for ( auto SelectedAsset : Selected )
-		{
-			UAnimationAsset* Asset = Cast<UAnimationAsset>(SelectedAsset.GetAsset());
-			if (Asset)
-			{
-				AnimAssets.Add(Asset);
-			}
-		}
-
-		if (AnimAssets.Num() > 0)
-		{
-			auto AnimAssetsToConvert = FObjectEditorUtils::GetTypedWeakObjectPtrs<UObject>(AnimAssets);
-			// ask user what they'd like to change to 
-			SAnimationRemapSkeleton::ShowWindow(OldSkeleton, Message, true, FOnRetargetAnimation::CreateSP(this, &SAnimationSequenceBrowser::RetargetAnimationHandler, AnimAssetsToConvert));
-		}
-	}
-}
-
 bool SAnimationSequenceBrowser::CanShowColumnForAssetRegistryTag(FName AssetType, FName TagName) const
 {
 	return !AssetRegistryTagsToIgnore.Contains(TagName);
diff --git a/Engine/Source/Editor/Persona/Private/SAnimationSequenceBrowser.h b/Engine/Source/Editor/Persona/Private/SAnimationSequenceBrowser.h
--- a/Engine/Source/Editor/Persona/Private/SAnimationSequenceBrowser.h
+++ b/Engine/Source/Editor/Persona/Private/SAnimationSequenceBrowser.h
@@ -270,8 +270,6 @@ protected:
 	/** The skeleton notify filter */
 	TSharedPtr<FFrontendFilter_SkeletonNotify> SkeletonNotifyFilter;
 
-	void RetargetAnimationHandler(USkeleton* OldSkeleton, USkeleton* NewSkeleton, bool bRemapReferencedAssets, bool bAllowRemapToExisting, bool bConvertSpaces, const EditorAnimUtils::FNameDuplicationRule* NameRule, TArray<TWeakObjectPtr<UObject>> InAnimAssets);
-
 private:
 	/** Updates the animation preview in the tooltip */
 	EActiveTimerReturnType UpdateTootipPreview( double InCurrentTime, float InDeltaTime );
diff --git a/Engine/Source/Editor/Persona/Private/SRetargetManager.cpp b/Engine/Source/Editor/Persona/Private/SRetargetManager.cpp
deleted file mode 100644
diff --git a/Engine/Source/Editor/Persona/Private/SRetargetManager.h b/Engine/Source/Editor/Persona/Private/SRetargetManager.h
deleted file mode 100644
diff --git a/Engine/Source/Editor/Persona/Private/SRetargetSourceWindow.cpp b/Engine/Source/Editor/Persona/Private/SRetargetSourceWindow.cpp
--- a/Engine/Source/Editor/Persona/Private/SRetargetSourceWindow.cpp
+++ b/Engine/Source/Editor/Persona/Private/SRetargetSourceWindow.cpp
@@ -180,20 +180,6 @@ void SRetargetSourceWindow::Construct(const FArguments& InArgs, const TSharedRef
 		.AutoHeight()
 		[
 			SNew( SHorizontalBox )
-			+SHorizontalBox::Slot()
-			.AutoWidth()
-			[
-				SNew(STextBlock)
-				.Text(LOCTEXT("Skeleton_Label", "Current Skeleton "))
-				.Font(FEditorStyle::GetFontStyle("Persona.RetargetManager.BoldFont"))
-			]
-
-			+SHorizontalBox::Slot()
-			.HAlign(HAlign_Left)
-			[
-				SNew(STextBlock)
-				.Text(SkeletonName)
-			]
 
 			+SHorizontalBox::Slot()
 			.HAlign(HAlign_Right)
@@ -205,7 +191,7 @@ void SRetargetSourceWindow::Construct(const FArguments& InArgs, const TSharedRef
 				.HAlign(HAlign_Center)
 				.VAlign(VAlign_Center)
 				.Text(LOCTEXT("AddRetargetSourceButton_Label", "Add New"))
-				.ToolTipText(LOCTEXT("AddRetargetSourceButton_ToolTip", "Add new retarget source to the list. It won't add if one already exists."))
+				.ToolTipText(LOCTEXT("AddRetargetSourceButton_ToolTip", "Select a Skeletal Mesh asset to become a new retarget source for this Skeleton asset.\n\nRetarget Sources indicate what proportions a sequence was authored with so that animation is correctly retargeted to other proportions.\n\nThese become 'Retarget Source' options on sequences.\n\nRetarget Sources are only needed when an animation sequence is authored on a skeletal mesh with proportions that are different than the default skeleton asset."))
 			]
 
 			+SHorizontalBox::Slot()
@@ -218,7 +204,7 @@ void SRetargetSourceWindow::Construct(const FArguments& InArgs, const TSharedRef
 				.HAlign(HAlign_Center)
 				.VAlign(VAlign_Center)
 				.Text(LOCTEXT("UpdateAllRetargetSourceButton_Label", "Update All"))
-				.ToolTipText(LOCTEXT("UpdateAllRetargetSourceButton_ToolTip", "Use this to update all retarget source poses with latest mesh. If you want to update individual, use the context menu."))
+				.ToolTipText(LOCTEXT("UpdateAllRetargetSourceButton_ToolTip", "Use this to update all retarget source poses with latest mesh proportions. If you want to update individually, use the context menu."))
 			]
 		]
 
diff --git a/Engine/Source/Editor/Persona/Private/SRetargetSources.cpp b/Engine/Source/Editor/Persona/Private/SRetargetSources.cpp
new file mode 100644
--- /dev/null
+++ b/Engine/Source/Editor/Persona/Private/SRetargetSources.cpp
@@ -0,0 +1,35 @@
+// Copyright Epic Games, Inc. All Rights Reserved.
+
+#include "SRetargetSources.h"
+#include "Widgets/SBoxPanel.h"
+#include "SRetargetSourceWindow.h"
+#include "IEditableSkeleton.h"
+#include "PropertyCustomizationHelpers.h"
+
+#define LOCTEXT_NAMESPACE "SRetargetSources"
+
+void SRetargetSources::Construct(
+	const FArguments& InArgs,
+	const TSharedRef<IEditableSkeleton>& InEditableSkeleton,
+	const TSharedRef<IPersonaPreviewScene>& InPreviewScene,
+	FSimpleMulticastDelegate& InOnPostUndo)
+{
+	EditableSkeletonPtr = InEditableSkeleton;
+	PreviewScenePtr = InPreviewScene;
+
+	const FString DocLink = TEXT("Shared/Editors/Persona");
+	ChildSlot
+	[
+		SNew (SVerticalBox)
+		+ SVerticalBox::Slot()
+		.Padding(2, 5)
+		.FillHeight(0.5)
+		[
+			// construct retarget source window
+			SNew(SRetargetSourceWindow, InEditableSkeleton, InOnPostUndo)
+		]
+	];
+}
+
+#undef LOCTEXT_NAMESPACE
+
diff --git a/Engine/Source/Editor/Persona/Private/SRetargetSources.h b/Engine/Source/Editor/Persona/Private/SRetargetSources.h
new file mode 100644
--- /dev/null
+++ b/Engine/Source/Editor/Persona/Private/SRetargetSources.h
@@ -0,0 +1,30 @@
+// Copyright Epic Games, Inc. All Rights Reserved.
+
+#pragma once
+
+#include "CoreMinimal.h"
+#include "Widgets/DeclarativeSyntaxSupport.h"
+#include "Widgets/SCompoundWidget.h"
+#include "IPersonaPreviewScene.h"
+
+class IEditableSkeleton;
+struct FAssetData;
+
+class SRetargetSources : public SCompoundWidget
+{
+public:
+	SLATE_BEGIN_ARGS( SRetargetSources )
+	{}
+
+	SLATE_END_ARGS()
+	
+	void Construct(const FArguments& InArgs, const TSharedRef<class IEditableSkeleton>& InEditableSkeleton, const TSharedRef<class IPersonaPreviewScene>& InPreviewScene, FSimpleMulticastDelegate& InOnPostUndo);
+
+private:
+
+	/** The editable skeleton */
+	TWeakPtr<class IEditableSkeleton> EditableSkeletonPtr;
+
+	/** The preview scene  */
+	TWeakPtr<class IPersonaPreviewScene> PreviewScenePtr;
+};
diff --git a/Engine/Source/Editor/Persona/Private/SRigWindow.cpp b/Engine/Source/Editor/Persona/Private/SRigWindow.cpp
deleted file mode 100644
diff --git a/Engine/Source/Editor/Persona/Private/SRigWindow.h b/Engine/Source/Editor/Persona/Private/SRigWindow.h
deleted file mode 100644
diff --git a/Engine/Source/Editor/Persona/Private/TabSpawners.cpp b/Engine/Source/Editor/Persona/Private/TabSpawners.cpp
--- a/Engine/Source/Editor/Persona/Private/TabSpawners.cpp
+++ b/Engine/Source/Editor/Persona/Private/TabSpawners.cpp
@@ -19,7 +19,7 @@
 #include "SAnimationSequenceBrowser.h"
 #include "SAnimationEditorViewport.h"
 #include "SPoseWatchManager.h"
-#include "SRetargetManager.h"
+#include "SRetargetSources.h"
 #include "SKismetInspector.h"
 #include "Widgets/Input/SButton.h"
 #include "PersonaPreviewSceneDescription.h"
@@ -320,25 +320,25 @@ FTabSpawnerEntry& FPreviewViewportSummoner::RegisterTabSpawner(TSharedRef<FTabMa
 /////////////////////////////////////////////////////
 // FRetargetManagerTabSummoner
 
-FRetargetManagerTabSummoner::FRetargetManagerTabSummoner(TSharedPtr<class FAssetEditorToolkit> InHostingApp, const TSharedRef<IEditableSkeleton>& InEditableSkeleton, const TSharedRef<IPersonaPreviewScene>& InPreviewScene, FSimpleMulticastDelegate& InOnPostUndo)
+FRetargetSourcesTabSummoner::FRetargetSourcesTabSummoner(TSharedPtr<class FAssetEditorToolkit> InHostingApp, const TSharedRef<IEditableSkeleton>& InEditableSkeleton, const TSharedRef<IPersonaPreviewScene>& InPreviewScene, FSimpleMulticastDelegate& InOnPostUndo)
 	: FWorkflowTabFactory(FPersonaTabs::RetargetManagerID, InHostingApp)
 	, EditableSkeleton(InEditableSkeleton)
 	, PreviewScene(InPreviewScene)
 	, OnPostUndo(InOnPostUndo)
 {
-	TabLabel = LOCTEXT("RetargetManagerTabTitle", "Retarget Manager");
+	TabLabel = LOCTEXT("RetargetSourcesTabTitle", "Retarget Sources");
 	TabIcon = FSlateIcon(FEditorStyle::GetStyleSetName(), "Persona.Tabs.RetargetManager");
 
 	EnableTabPadding();
 	bIsSingleton = true;
 
-	ViewMenuDescription = LOCTEXT("RetargetManagerTabView", "Retarget Manager");
-	ViewMenuTooltip = LOCTEXT("RetargetManagerTabView_ToolTip", "Manages different options for retargeting");
+	ViewMenuDescription = LOCTEXT("RetargetSourcesTabView", "Retarget Sources");
+	ViewMenuTooltip = LOCTEXT("RetargetSourcesTabView_ToolTip", "Retarget Sources indicate what proportions a sequence was authored with so that animation is correctly retargeted to other proportions.\n\nThese become 'Retarget Source' options on sequences.\n\nRetarget Sources are only needed when an animation sequence is authored on a skeletal mesh with proportions that are different than the default skeleton asset.");
 }
 
-TSharedRef<SWidget> FRetargetManagerTabSummoner::CreateTabBody(const FWorkflowTabSpawnInfo& Info) const
+TSharedRef<SWidget> FRetargetSourcesTabSummoner::CreateTabBody(const FWorkflowTabSpawnInfo& Info) const
 {
-	return SNew(SRetargetManager, EditableSkeleton.Pin().ToSharedRef(), PreviewScene.Pin().ToSharedRef(), OnPostUndo);
+	return SNew(SRetargetSources, EditableSkeleton.Pin().ToSharedRef(), PreviewScene.Pin().ToSharedRef(), OnPostUndo);
 }
 
 
diff --git a/Engine/Source/Editor/Persona/Private/TabSpawners.h b/Engine/Source/Editor/Persona/Private/TabSpawners.h
--- a/Engine/Source/Editor/Persona/Private/TabSpawners.h
+++ b/Engine/Source/Editor/Persona/Private/TabSpawners.h
@@ -215,17 +215,17 @@ struct FPreviewViewportSummoner : public FWorkflowTabFactory
 /////////////////////////////////////////////////////
 // FRetargetManagerTabSummoner
 
-struct FRetargetManagerTabSummoner : public FWorkflowTabFactory
+struct FRetargetSourcesTabSummoner : public FWorkflowTabFactory
 {
 public:
-	FRetargetManagerTabSummoner(TSharedPtr<class FAssetEditorToolkit> InHostingApp, const TSharedRef<IEditableSkeleton>& InEditableSkeleton, const TSharedRef<IPersonaPreviewScene>& InPreviewScene, FSimpleMulticastDelegate& InOnPostUndo);
+	FRetargetSourcesTabSummoner(TSharedPtr<class FAssetEditorToolkit> InHostingApp, const TSharedRef<IEditableSkeleton>& InEditableSkeleton, const TSharedRef<IPersonaPreviewScene>& InPreviewScene, FSimpleMulticastDelegate& InOnPostUndo);
 
 	virtual TSharedRef<SWidget> CreateTabBody(const FWorkflowTabSpawnInfo& Info) const override;
 
 	// Create a tooltip widget for the tab
 	virtual TSharedPtr<SToolTip> CreateTabToolTipWidget(const FWorkflowTabSpawnInfo& Info) const override
 	{
-		return  IDocumentation::Get()->CreateToolTip(LOCTEXT("RetargetSourceTooltip", "In this panel, you can manage retarget sources for different body types"), NULL, TEXT("Shared/Editors/Persona"), TEXT("RetargetManager"));
+		return  IDocumentation::Get()->CreateToolTip(LOCTEXT("RetargetSourceTooltip", "In this panel, you can manage retarget sources for animations authored with varying proportions."), NULL, TEXT("Shared/Editors/Persona"), TEXT("RetargetSources"));
 	}
 
 private:
diff --git a/Engine/Source/Editor/Persona/Public/PersonaModule.h b/Engine/Source/Editor/Persona/Public/PersonaModule.h
--- a/Engine/Source/Editor/Persona/Public/PersonaModule.h
+++ b/Engine/Source/Editor/Persona/Public/PersonaModule.h
@@ -280,8 +280,8 @@ public:
 	/** Create a skeleton curve viewer tab factory */
 	virtual TSharedRef<FWorkflowTabFactory> CreateCurveViewerTabFactory(const TSharedRef<class FWorkflowCentricApplication>& InHostingApp, const TSharedRef<class IEditableSkeleton>& InEditableSkeleton, const TSharedRef<class IPersonaPreviewScene>& InPreviewScene, FOnObjectsSelected InOnObjectsSelected) const;
 
-	/** Create a retarget manager tab factory */
-	virtual TSharedRef<FWorkflowTabFactory> CreateRetargetManagerTabFactory(const TSharedRef<class FWorkflowCentricApplication>& InHostingApp, const TSharedRef<class IEditableSkeleton>& InEditableSkeleton, const TSharedRef<IPersonaPreviewScene>& InPreviewScene, FSimpleMulticastDelegate& InOnPostUndo) const;
+	/** Create a retarget sources tab factory */
+	virtual TSharedRef<FWorkflowTabFactory> CreateRetargetSourcesTabFactory(const TSharedRef<class FWorkflowCentricApplication>& InHostingApp, const TSharedRef<class IEditableSkeleton>& InEditableSkeleton, const TSharedRef<IPersonaPreviewScene>& InPreviewScene, FSimpleMulticastDelegate& InOnPostUndo) const;
 
 	/** Create a tab factory used to configure preview scene settings */
 	virtual TSharedRef<FWorkflowTabFactory> CreateAdvancedPreviewSceneTabFactory(const TSharedRef<class FWorkflowCentricApplication>& InHostingApp, const TSharedRef<IPersonaPreviewScene>& InPreviewScene) const;
diff --git a/Engine/Source/Editor/SkeletonEditor/Private/SkeletonEditor.cpp b/Engine/Source/Editor/SkeletonEditor/Private/SkeletonEditor.cpp
--- a/Engine/Source/Editor/SkeletonEditor/Private/SkeletonEditor.cpp
+++ b/Engine/Source/Editor/SkeletonEditor/Private/SkeletonEditor.cpp
@@ -214,7 +214,7 @@ void FSkeletonEditor::ExtendToolbar()
 			ToolbarBuilder.BeginSection("Skeleton");
 			{
 				ToolbarBuilder.AddToolBarButton(FSkeletonEditorCommands::Get().AnimNotifyWindow);
-				ToolbarBuilder.AddToolBarButton(FSkeletonEditorCommands::Get().RetargetManager, NAME_None, LOCTEXT("Toolbar_RetargetManager", "Retarget Manager"));
+				ToolbarBuilder.AddToolBarButton(FSkeletonEditorCommands::Get().RetargetManager, NAME_None, LOCTEXT("Toolbar_RetargetSources", "Retarget Sources"));
 				ToolbarBuilder.AddToolBarButton(FSkeletonEditorCommands::Get().ImportMesh);
 			}
 			ToolbarBuilder.EndSection();
diff --git a/Engine/Source/Editor/SkeletonEditor/Private/SkeletonEditorMode.cpp b/Engine/Source/Editor/SkeletonEditor/Private/SkeletonEditorMode.cpp
--- a/Engine/Source/Editor/SkeletonEditor/Private/SkeletonEditorMode.cpp
+++ b/Engine/Source/Editor/SkeletonEditor/Private/SkeletonEditorMode.cpp
@@ -30,7 +30,7 @@ FSkeletonEditorMode::FSkeletonEditorMode(TSharedRef<FWorkflowCentricApplication>
 
 	TabFactories.RegisterFactory(PersonaModule.CreateAnimNotifiesTabFactory(InHostingApp, InSkeletonTree->GetEditableSkeleton(),  OnObjectsSelected));
 	TabFactories.RegisterFactory(PersonaModule.CreateAdvancedPreviewSceneTabFactory(InHostingApp, SkeletonEditor->GetPersonaToolkit()->GetPreviewScene()));
-	TabFactories.RegisterFactory(PersonaModule.CreateRetargetManagerTabFactory(InHostingApp, InSkeletonTree->GetEditableSkeleton(), SkeletonEditor->GetPersonaToolkit()->GetPreviewScene(), SkeletonEditor->OnPostUndo));
+	TabFactories.RegisterFactory(PersonaModule.CreateRetargetSourcesTabFactory(InHostingApp, InSkeletonTree->GetEditableSkeleton(), SkeletonEditor->GetPersonaToolkit()->GetPreviewScene(), SkeletonEditor->OnPostUndo));
 	TabFactories.RegisterFactory(PersonaModule.CreateCurveViewerTabFactory(InHostingApp, InSkeletonTree->GetEditableSkeleton(), SkeletonEditor->GetPersonaToolkit()->GetPreviewScene(), OnObjectsSelected));
 	TabFactories.RegisterFactory(PersonaModule.CreateSkeletonSlotNamesTabFactory(InHostingApp, InSkeletonTree->GetEditableSkeleton(), OnObjectSelected));
 	TabFactories.RegisterFactory(PersonaModule.CreateAssetDetailsTabFactory(InHostingApp, FOnGetAsset::CreateSP(&SkeletonEditor.Get(), &FSkeletonEditor::HandleGetAsset), FOnDetailsCreated()));
diff --git a/Engine/Source/Editor/UnrealEd/Private/EditorAnimUtils.cpp b/Engine/Source/Editor/UnrealEd/Private/EditorAnimUtils.cpp
--- a/Engine/Source/Editor/UnrealEd/Private/EditorAnimUtils.cpp
+++ b/Engine/Source/Editor/UnrealEd/Private/EditorAnimUtils.cpp
@@ -252,104 +252,6 @@ namespace EditorAnimUtils
 		}
 	}
 
-	void FAnimationRetargetContext::RetargetAnimations(USkeleton* OldSkeleton, USkeleton* NewSkeleton)
-	{
-		check (!bConvertAnimationDataInComponentSpaces || OldSkeleton);
-		check (NewSkeleton);
-
-		if (bConvertAnimationDataInComponentSpaces)
-		{
-			// we need to update reference pose before retargeting. 
-			// this is to ensure the skeleton has the latest pose you're looking at. 
-			USkeletalMesh * PreviewMesh = NULL;
-			if (OldSkeleton != NULL)
-			{
-				PreviewMesh = OldSkeleton->GetPreviewMesh(true);
-				if (PreviewMesh)
-				{
-					OldSkeleton->UpdateReferencePoseFromMesh(PreviewMesh);
-				}
-			}
-			
-			PreviewMesh = NewSkeleton->GetPreviewMesh(true);
-			if (PreviewMesh)
-			{
-				NewSkeleton->UpdateReferencePoseFromMesh(PreviewMesh);
-			}
-		}
-
-		// anim sequences will be retargeted first becauseReplaceSkeleton forces it to change skeleton
-		// @todo: please note that I think we can merge two loops 
-		//(without separating two loops - one for AnimSequence and one for everybody else)
-		// but if you have animation asssets that does replace skeleton, it will try fix up internal asset also
-		// so I think you might be doing twice - look at AnimationAsset:ReplaceSkeleton
-		// for safety, I'm doing Sequence first and then everything else
-		// however this can be re-investigated and fixed better in the future
-		for(auto Iter = AnimationAssetsToRetarget.CreateIterator(); Iter; ++Iter)
-		{
-			UAnimSequence* AnimSequenceToRetarget = Cast<UAnimSequence>(*Iter);
-			if (AnimSequenceToRetarget)
-			{
-				// Copy curve data from source asset, preserving data in the target if present.
-				if (OldSkeleton)
-				{
-					UAnimationBlueprintLibrary::CopyAnimationCurveNamesToSkeleton(OldSkeleton, NewSkeleton, AnimSequenceToRetarget, ERawCurveTrackTypes::RCT_Float);
-
-					// clear transform curves since those curves won't work in new skeleton
-					// since we're deleting curves, mark this rebake flag off
-					IAnimationDataController& Controller = AnimSequenceToRetarget->GetController();
-					Controller.RemoveAllCurvesOfType(ERawCurveTrackTypes::RCT_Transform);
-
-					// I can't copy transform curves yet because transform curves need retargeting. 
-					//EditorAnimUtils::CopyAnimCurves(OldSkeleton, NewSkeleton, AssetToRetarget, USkeleton::AnimTrackCurveMappingName, FRawCurveTracks::TransformType);
-				}
-			}
-
-			UAnimationAsset* AssetToRetarget = (*Iter);
-			if (HasDuplicates())
-			{
-				AssetToRetarget->ReplaceReferredAnimations(RemappedAnimAssets);
-			}
-			AssetToRetarget->ReplaceSkeleton(NewSkeleton, bConvertAnimationDataInComponentSpaces);
-			AssetToRetarget->MarkPackageDirty();
-		}
-
-		// convert all Animation Blueprints and compile 
-		for ( auto AnimBPIter = AnimBlueprintsToRetarget.CreateIterator(); AnimBPIter; ++AnimBPIter )
-		{
-			UAnimBlueprint * AnimBlueprint = (*AnimBPIter);
-
-			AnimBlueprint->TargetSkeleton = NewSkeleton;
-
-			// We can directly retarget templates (although not transitively via parents) so we need to make sure we clear the flag here
-			AnimBlueprint->bIsTemplate = false;
-
-			if (HasDuplicates())
-			{
-				// if they have parent blueprint, make sure to re-link to the new one also
-				UAnimBlueprint* CurrentParentBP = Cast<UAnimBlueprint>(AnimBlueprint->ParentClass->ClassGeneratedBy);
-				if (CurrentParentBP)
-				{
-					UAnimBlueprint* const * ParentBP = DuplicatedBlueprints.Find(CurrentParentBP);
-					if (ParentBP)
-					{
-						AnimBlueprint->ParentClass = (*ParentBP)->GeneratedClass;
-					}
-				}
-			}
-
-			if(RemappedAnimAssets.Num() > 0)
-			{
-				ReplaceReferredAnimationsInBlueprint(AnimBlueprint, RemappedAnimAssets);
-			}
-
-			FBlueprintEditorUtils::RefreshAllNodes(AnimBlueprint);
-			FKismetEditorUtilities::CompileBlueprint(AnimBlueprint, EBlueprintCompileOptions::SkipGarbageCollection);
-			AnimBlueprint->PostEditChange();
-			AnimBlueprint->MarkPackageDirty();
-		}
-	}
-
 	void FAnimationRetargetContext::AddRemappedAsset(UAnimationAsset* OriginalAsset, UAnimationAsset* NewAsset)
 	{
 		RemappedAnimAssets.Add(OriginalAsset, NewAsset);
@@ -359,96 +261,7 @@ namespace EditorAnimUtils
 	{
 		GEditor->GetEditorSubsystem<UAssetEditorSubsystem>()->OpenEditorForAsset(AssetToOpen);
 	}
-
-	//////////////////////////////////////////////////////////////////
-	UObject* RetargetAnimations(USkeleton* OldSkeleton, USkeleton* NewSkeleton, TArray<TWeakObjectPtr<UObject>> AssetsToRetarget, bool bRetargetReferredAssets, const FNameDuplicationRule* NameRule, bool bConvertSpace)
-	{
-		FAnimationRetargetContext RetargetContext(AssetsToRetarget, bRetargetReferredAssets, bConvertSpace);
-		return RetargetAnimations(OldSkeleton, NewSkeleton, RetargetContext, bRetargetReferredAssets, NameRule);
-	}
-
-	UObject* RetargetAnimations(USkeleton* OldSkeleton, USkeleton* NewSkeleton, const TArray<FAssetData>& AssetsToRetarget, bool bRetargetReferredAssets, const FNameDuplicationRule* NameRule, bool bConvertSpace)
-	{
-		FAnimationRetargetContext RetargetContext(AssetsToRetarget, bRetargetReferredAssets, bConvertSpace);
-		return RetargetAnimations(OldSkeleton, NewSkeleton, RetargetContext, bRetargetReferredAssets, NameRule);
-	}
-
-	UObject* RetargetAnimations(USkeleton* OldSkeleton, USkeleton* NewSkeleton, FAnimationRetargetContext& RetargetContext, bool bRetargetReferredAssets, const FNameDuplicationRule* NameRule)
-	{
-		check(NewSkeleton);
-		UObject* OriginalObject  = RetargetContext.GetSingleTargetObject();
-		UPackage* DuplicationDestPackage = NewSkeleton->GetOutermost();
-
-		if(	RetargetContext.HasAssetsToRetarget() )
-		{
-			if(NameRule)
-			{
-				RetargetContext.DuplicateAssetsToRetarget(DuplicationDestPackage, NameRule);
-			}
-			RetargetContext.RetargetAnimations(OldSkeleton, NewSkeleton);
-		}
-
-		FNotificationInfo Notification(FText::GetEmpty());
-		Notification.ExpireDuration = 5.f;
-
-		UObject* NotifyLinkObject = OriginalObject;
-		if(OriginalObject && NameRule)
-		{
-			NotifyLinkObject = RetargetContext.GetDuplicate(OriginalObject);
-		}
-
-		if(!NameRule)
-		{
-			if(OriginalObject)
-			{
-				Notification.Text = FText::Format(LOCTEXT("SingleNonDuplicatedAsset", "'{0}' retargeted to new skeleton '{1}'"), FText::FromString(OriginalObject->GetName()), FText::FromString(NewSkeleton->GetName()));
-			}
-			else
-			{
-				Notification.Text = FText::Format(LOCTEXT("MultiNonDuplicatedAsset", "Assets retargeted to new skeleton '{0}'"), FText::FromString(NewSkeleton->GetName()));
-			}
-			
-		}
-		else
-		{
-			if(OriginalObject)
-			{
-				Notification.Text = FText::Format(LOCTEXT("SingleDuplicatedAsset", "'{0}' duplicated to '{1}' and retargeted"), FText::FromString(OriginalObject->GetName()), FText::FromString(DuplicationDestPackage->GetName()));
-			}
-			else
-			{
-				Notification.Text = FText::Format(LOCTEXT("MultiDuplicatedAsset", "Assets duplicated to '{0}' and retargeted"), FText::FromString(DuplicationDestPackage->GetName()));
-			}
-		}
-
-		if(NotifyLinkObject)
-		{
-			Notification.Hyperlink = FSimpleDelegate::CreateStatic(&OpenAssetFromNotify, NotifyLinkObject);
-			Notification.HyperlinkText = LOCTEXT("OpenAssetLink", "Open");
-		}
-
-		FSlateNotificationManager::Get().AddNotification(Notification);
-
-		// sync newly created objects on CB
-		if (NotifyLinkObject)
-		{
-			TArray<UObject*> NewObjects = RetargetContext.GetAllDuplicates();
-			TArray<FAssetData> CurrentSelection;
-			for(auto& NewObject : NewObjects)
-			{
-				CurrentSelection.Add(FAssetData(NewObject));
-			}
-
-			FContentBrowserModule& ContentBrowserModule = FModuleManager::Get().LoadModuleChecked<FContentBrowserModule>("ContentBrowser");
-			ContentBrowserModule.Get().SyncBrowserToAssets(CurrentSelection);
-		}
-		if(OriginalObject && NameRule)
-		{
-			return RetargetContext.GetDuplicate(OriginalObject);
-		}
-		return NULL;
-	}
-
+	
 	FString CreateDesiredName(UObject* Asset, const FNameDuplicationRule* NameRule)
 	{
 		check(Asset);
diff --git a/Engine/Source/Editor/UnrealEd/Private/SSkeletonWidget.cpp b/Engine/Source/Editor/UnrealEd/Private/SSkeletonWidget.cpp
--- a/Engine/Source/Editor/UnrealEd/Private/SSkeletonWidget.cpp
+++ b/Engine/Source/Editor/UnrealEd/Private/SSkeletonWidget.cpp
@@ -346,809 +346,6 @@ void SSkeletonSelectorWindow::ConstructWindow()
 		];
 }
 
-TSharedPtr<SWindow> SAnimationRemapSkeleton::DialogWindow;
-
-bool SAnimationRemapSkeleton::OnShouldFilterAsset(const struct FAssetData& AssetData)
-{
-	USkeleton* AssetSkeleton = nullptr;
-	if (AssetData.IsAssetLoaded())
-	{
-		AssetSkeleton = Cast<USkeleton>(AssetData.GetAsset());
-	}
-
-	// do not show same skeleton
-	if (OldSkeleton && OldSkeleton == AssetSkeleton)
-	{
-		return true;
-	}
-
-	if (bShowOnlyCompatibleSkeletons)
-	{
-		if(OldSkeleton && OldSkeleton->GetRig())
-		{
-			URig * Rig = OldSkeleton->GetRig();
-
-			const FString Value = AssetData.GetTagValueRef<FString>(USkeleton::RigTag);
-
-			if(Rig->GetFullName() == Value)
-			{
-				return false;
-			}
-
-			// if loaded, check to see if it has same rig
-			if (AssetData.IsAssetLoaded())
-			{
-				USkeleton* LoadedSkeleton = Cast<USkeleton>(AssetData.GetAsset());
-
-				if (LoadedSkeleton && LoadedSkeleton->GetRig() == Rig)
-				{
-					return false;
-				}
-			}
-		}
-
-		return true;
-	}
-
-	return false;
-}
-
-void SAnimationRemapSkeleton::UpdateAssetPicker()
-{
-	FAssetPickerConfig AssetPickerConfig;
-	AssetPickerConfig.Filter.ClassNames.Add(USkeleton::StaticClass()->GetFName());
-	AssetPickerConfig.OnAssetSelected = FOnAssetSelected::CreateSP(this, &SAnimationRemapSkeleton::OnAssetSelectedFromPicker);
-	AssetPickerConfig.bAllowNullSelection = false;
-	AssetPickerConfig.InitialAssetViewType = EAssetViewType::Column;
-	AssetPickerConfig.OnShouldFilterAsset = FOnShouldFilterAsset::CreateSP(this, &SAnimationRemapSkeleton::OnShouldFilterAsset);
-	AssetPickerConfig.bShowPathInColumnView = true;
-	AssetPickerConfig.bShowTypeInColumnView = false;
-
-	FContentBrowserModule& ContentBrowserModule = FModuleManager::Get().LoadModuleChecked<FContentBrowserModule>(TEXT("ContentBrowser"));
-
-	if (AssetPickerBox.IsValid())
-	{
-		AssetPickerBox->SetContent(
-			ContentBrowserModule.Get().CreateAssetPicker(AssetPickerConfig)
-		);
-	}
-}
-
-void SAnimationRemapSkeleton::Construct( const FArguments& InArgs )
-{
-	OldSkeleton = InArgs._CurrentSkeleton;
-	NewSkeleton = nullptr;
-	WidgetWindow = InArgs._WidgetWindow;
-	bRemapReferencedAssets = true;
-	bConvertSpaces = false;
-	bShowOnlyCompatibleSkeletons = false;
-	OnRetargetAnimationDelegate = InArgs._OnRetargetDelegate;
-	bShowDuplicateAssetOption = InArgs._ShowDuplicateAssetOption;
-
-	TSharedRef<SVerticalBox> RetargetWidget = SNew(SVerticalBox);
-
-	RetargetWidget->AddSlot()
-	[
-		SNew(STextBlock)
-		.AutoWrapText(true)
-		.Font(FEditorStyle::GetFontStyle("Persona.RetargetManager.SmallBoldFont"))
-		.Text(LOCTEXT("RetargetBasePose_OptionLabel", "Retarget Options"))
-	];
-
-	if(InArgs._ShowRemapOption)
-	{
-		RetargetWidget->AddSlot()
-		[
-			SNew(SCheckBox)
-			.IsChecked(this, &SAnimationRemapSkeleton::IsRemappingReferencedAssets)
-			.OnCheckStateChanged(this, &SAnimationRemapSkeleton::OnRemappingReferencedAssetsChanged)
-			[
-				SNew(STextBlock).Text(LOCTEXT("RemapSkeleton_RemapAssets", "Remap referenced assets "))
-			]
-		];
-
-		bRemapReferencedAssets = true;
-
-		if(InArgs._ShowExistingRemapOption)
-		{
-			RetargetWidget->AddSlot()
-				[
-					SNew(SCheckBox)
-					.IsChecked(this, &SAnimationRemapSkeleton::IsRemappingToExistingAssetsChecked)
-					.IsEnabled(this, &SAnimationRemapSkeleton::IsRemappingToExistingAssetsEnabled)
-					.OnCheckStateChanged(this, &SAnimationRemapSkeleton::OnRemappingToExistingAssetsChanged)
-					[
-						SNew(STextBlock)
-						.Text(LOCTEXT("RemapSkeleton_RemapToExisting", "Allow remapping to existing assets"))
-					]
-				];
-
-			// Not by default, user must specify
-			bAllowRemappingToExistingAssets = false;
-		}
-	}
-
-	if (InArgs._ShowConvertSpacesOption)
-	{
-		TSharedPtr<SToolTip> ConvertSpaceTooltip = IDocumentation::Get()->CreateToolTip(FText::FromString("Check if you'd like to convert animation data to new skeleton space. If this is false, it won't convert any animation data to new space."),
-			nullptr, FString("Shared/Editors/Persona"), FString("AnimRemapSkeleton_ConvertSpace"));
-		RetargetWidget->AddSlot()
-			[
-				SNew(SCheckBox)
-				.IsChecked(this, &SAnimationRemapSkeleton::IsConvertSpacesChecked)
-				.OnCheckStateChanged(this, &SAnimationRemapSkeleton::OnConvertSpacesCheckChanged)
-				[
-					SNew(STextBlock).Text(LOCTEXT("RemapSkeleton_ConvertSpaces", "Convert Spaces to new Skeleton"))
-					.ToolTip(ConvertSpaceTooltip)
-				]
-			];
-
-		bConvertSpaces = true;
-	}
-
-	if(InArgs._ShowCompatibleDisplayOption)
-	{
-		TSharedPtr<SToolTip> ConvertSpaceTooltip = IDocumentation::Get()->CreateToolTip(FText::FromString("Check if you'd like to show only the skeleton that uses the same rig."),
-			nullptr, FString("Shared/Editors/Persona"), FString("AnimRemapSkeleton_ShowCompatbielSkeletons")); // @todo add tooltip
-		RetargetWidget->AddSlot()
-			[
-				SNew(SCheckBox)
-				.IsChecked(this, &SAnimationRemapSkeleton::IsShowOnlyCompatibleSkeletonsChecked)
-				.IsEnabled(this, &SAnimationRemapSkeleton::IsShowOnlyCompatibleSkeletonsEnabled)
-				.OnCheckStateChanged(this, &SAnimationRemapSkeleton::OnShowOnlyCompatibleSkeletonsCheckChanged)
-				[
-					SNew(STextBlock).Text(LOCTEXT("RemapSkeleton_ShowCompatible", "Show Only Compatible Skeletons"))
-					.ToolTip(ConvertSpaceTooltip)
-				]
-			];
-
-		bShowOnlyCompatibleSkeletons = true;
-	}
-
-	TSharedRef<SHorizontalBox> OptionWidget = SNew(SHorizontalBox);
-	OptionWidget->AddSlot()
-		[
-			RetargetWidget
-		];
-
-	if (bShowDuplicateAssetOption)
-	{
-		TSharedRef<SVerticalBox> NameOptionWidget = SNew(SVerticalBox);
-
-		NameOptionWidget->AddSlot()
-		[	
-			SNew(SVerticalBox)
-			+SVerticalBox::Slot()
-			.AutoHeight()
-			.Padding(2, 3)
-			[
-				SNew(STextBlock)
-				.AutoWrapText(true)
-				.Font(FEditorStyle::GetFontStyle("Persona.RetargetManager.SmallBoldFont"))
-				.Text(LOCTEXT("RetargetBasePose_RenameLable", "New Asset Name"))
-			]
-
-			+SVerticalBox::Slot()
-			.AutoHeight()
-			.Padding(2, 1)
-			[
-				SNew(SHorizontalBox)
-				+SHorizontalBox::Slot()
-				[
-					SNew(STextBlock).Text(LOCTEXT("RemapSkeleton_DupeName_Prefix", "Prefix"))
-				]
-
-				+SHorizontalBox::Slot()
-				[
-					SNew(SEditableTextBox)
-						.Text(this, &SAnimationRemapSkeleton::GetPrefixName)
-						.MinDesiredWidth(100)
-						.OnTextChanged(this, &SAnimationRemapSkeleton::SetPrefixName)
-						.IsReadOnly(false)
-						.RevertTextOnEscape(true)
-				]
-
-			]
-
-			+SVerticalBox::Slot()
-			.AutoHeight()
-			.Padding(2, 1)
-			[
-				SNew(SHorizontalBox)
-				+SHorizontalBox::Slot()
-				[
-					SNew(STextBlock).Text(LOCTEXT("RemapSkeleton_DupeName_Suffix", "Suffix"))
-				]
-
-				+SHorizontalBox::Slot()
-				[
-					SNew(SEditableTextBox)
-						.Text(this, &SAnimationRemapSkeleton::GetSuffixName)
-						.MinDesiredWidth(100)
-						.OnTextChanged(this, &SAnimationRemapSkeleton::SetSuffixName)
-						.IsReadOnly(false)
-						.RevertTextOnEscape(true)
-				]
-
-			]
-
-			+SVerticalBox::Slot()
-			.AutoHeight()
-			.Padding(2, 1)
-			[
-				SNew(SHorizontalBox)
-				+SHorizontalBox::Slot()
-				[
-					SNew(STextBlock).Text(LOCTEXT("RemapSkeleton_DupeName_ReplaceFrom", "Replace "))
-				]
-
-				+SHorizontalBox::Slot()
-				[
-					SNew(SEditableTextBox)
-						.Text(this, &SAnimationRemapSkeleton::GetReplaceFrom)
-						.MinDesiredWidth(50)
-						.OnTextChanged(this, &SAnimationRemapSkeleton::SetReplaceFrom)
-						.IsReadOnly(false)
-						.RevertTextOnEscape(true)
-				]
-
-				+SHorizontalBox::Slot()
-				.Padding(5, 0)
-				.HAlign(HAlign_Center)
-				[
-					SNew(STextBlock).Text(LOCTEXT("RemapSkeleton_DupeName_ReplaceTo", "with "))
-				]
-
-				+SHorizontalBox::Slot()
-				[
-					SNew(SEditableTextBox)
-						.Text(this, &SAnimationRemapSkeleton::GetReplaceTo)
-						.MinDesiredWidth(50)
-						.OnTextChanged(this, &SAnimationRemapSkeleton::SetReplaceTo)
-						.IsReadOnly(false)
-						.RevertTextOnEscape(true)
-				]
-			]
-
-			+SVerticalBox::Slot()
-			.AutoHeight()
-			.Padding(2, 3)
-			[
-				SNew(SHorizontalBox)
-				+SHorizontalBox::Slot()
-				.Padding(5, 5)
-				[
-					SNew(STextBlock).Text(this,  &SAnimationRemapSkeleton::GetExampleText)
-						.Font(FEditorStyle::GetFontStyle("Persona.RetargetManager.ItalicFont"))
-				]
-			]
-
-			+SVerticalBox::Slot()
-			.AutoHeight()
-			.Padding(2, 3)
-			[
-				SNew(SHorizontalBox)
-
-				+SHorizontalBox::Slot()
-				.AutoWidth()
-				[
-					SNew(STextBlock)
-					.Text(LOCTEXT("RemapSkeleton_DupeName_Folder", "Folder "))
-						.Font(FEditorStyle::GetFontStyle("Persona.RetargetManager.SmallBoldFont"))
-				]
-
-				+SHorizontalBox::Slot()
-				.FillWidth(1)
-				.HAlign(HAlign_Center)
-				[
-					SNew(STextBlock).Text(this, &SAnimationRemapSkeleton::GetFolderPath)
-				]
-
-				+SHorizontalBox::Slot()
-				.AutoWidth()
-				[
-					SNew(SButton)
-					.HAlign(HAlign_Center)
-					.Text(LOCTEXT("RemapSkeleton_DupeName_ShowFolderOption", "Change..."))
-					.OnClicked(this, &SAnimationRemapSkeleton::ShowFolderOption)
-				]
-			]
-		];
-
-		OptionWidget->AddSlot()
-		[
-			NameOptionWidget
-		];
-	}
-
-	TSharedPtr<SToolTip> SkeletonTooltip = IDocumentation::Get()->CreateToolTip(FText::FromString("Pick a skeleton for this mesh"), nullptr, FString("Shared/Editors/Persona"), FString("Skeleton"));
-
-	this->ChildSlot
-		[
-			SNew (SHorizontalBox)
-			+SHorizontalBox::Slot()
-			.AutoWidth()
-			[
-				SNew(SVerticalBox)
-
-				+SVerticalBox::Slot()
-				.AutoHeight()
-				.Padding(2)
-				.HAlign(HAlign_Fill)
-				[
-					SNew(SHorizontalBox)
-
-					+SHorizontalBox::Slot()
-					.AutoWidth()
-					.VAlign(VAlign_Center)
-					+SHorizontalBox::Slot()
-					[
-						SNew(STextBlock)
-						.Text(LOCTEXT("CurrentlySelectedSkeletonLabel_SelectSkeleton", "Select Skeleton"))
-						.Font(FCoreStyle::GetDefaultFontStyle("Regular", 16))
-						.ToolTip(SkeletonTooltip)
-					]
-					+SHorizontalBox::Slot()
-					.FillWidth(1)
-					.HAlign(HAlign_Left)
-					.VAlign(VAlign_Center)
-					[
-						IDocumentation::Get()->CreateAnchor(FString("Engine/Animation/Skeleton"))
-					]
-				]
-
-				+SVerticalBox::Slot()
-				.AutoHeight()
-				.Padding(5)
-				[
-					SNew(STextBlock)
-					.Font(FEditorStyle::GetFontStyle("Persona.RetargetManager.FilterFont"))
-					.Text(InArgs._WarningMessage)
-				]
-
-				+SVerticalBox::Slot()
-				.AutoHeight()
-				.Padding(5)
-				[
-					SNew(SSeparator)
-				]
-
-				+SVerticalBox::Slot()
-				.MaxHeight(500)
-				[
-					SAssignNew(AssetPickerBox, SBox)
-					.WidthOverride(400)
-					.HeightOverride(300)
-				]
-
-				+SVerticalBox::Slot()
-				.AutoHeight()
-				.Padding(5)
-				[
-					SNew(SSeparator)
-				]
-
-				+SVerticalBox::Slot()
-				.AutoHeight()
-				.HAlign(HAlign_Right)
-				.VAlign(VAlign_Bottom)
-				[
-					SNew(SUniformGridPanel)
-					.SlotPadding(FEditorStyle::GetMargin("StandardDialog.SlotPadding"))
-					.MinDesiredSlotWidth(FEditorStyle::GetFloat("StandardDialog.MinDesiredSlotWidth"))
-					.MinDesiredSlotHeight(FEditorStyle::GetFloat("StandardDialog.MinDesiredSlotHeight"))
-					+SUniformGridPanel::Slot(0, 0)
-					[
-						SNew(SButton).HAlign(HAlign_Center)
-						.Text(LOCTEXT("RemapSkeleton_Apply", "Retarget"))
-						.IsEnabled(this, &SAnimationRemapSkeleton::CanApply)
-						.OnClicked(this, &SAnimationRemapSkeleton::OnApply)
-						.HAlign(HAlign_Center)
-						.ContentPadding(FEditorStyle::GetMargin("StandardDialog.ContentPadding"))
-					]
-					+SUniformGridPanel::Slot(1, 0)
-					[
-						SNew(SButton).HAlign(HAlign_Center)
-						.Text(LOCTEXT("RemapSkeleton_Cancel", "Cancel"))
-						.ContentPadding(FEditorStyle::GetMargin("StandardDialog.ContentPadding"))
-						.OnClicked(this, &SAnimationRemapSkeleton::OnCancel)
-					]
-				]
-			]
-
-			+SHorizontalBox::Slot()
-			.Padding(2)
-			.AutoWidth()
-			[
-				SNew(SSeparator)
-				.Orientation(Orient_Vertical)
-			]
-
-			+SHorizontalBox::Slot()
-			.HAlign(HAlign_Center)
-			.VAlign(VAlign_Center)
-			.AutoWidth()
-			[
-				SNew(SVerticalBox)
-				+SVerticalBox::Slot()
-				.AutoHeight()
-				.Padding(5, 5)
-				[
-					// put nice message here
-					SNew(STextBlock)
-					.AutoWrapText(true)
-					.Font(FEditorStyle::GetFontStyle("Persona.RetargetManager.FilterFont"))
-					.ColorAndOpacity(FLinearColor::Red)
-					.Text(LOCTEXT("RetargetBasePose_WarningMessage", "*Make sure you have the similar retarget base pose. \nIf they don't look alike here, you can edit your base pose in the Retarget Manager window."))
-				]
-
-				+SVerticalBox::Slot()
-				.FillHeight(1)
-				.Padding(0, 5)
-				[
-					SNew(SHorizontalBox)
-					+SHorizontalBox::Slot()
-					[
-						SNew(SVerticalBox)
-						+ SVerticalBox::Slot()
-						.AutoHeight()
-						[
-							SNew(STextBlock)
-							.Text(LOCTEXT("SourceSkeleteonTitle", "[Source]"))
-							.Font(FEditorStyle::GetFontStyle("Persona.RetargetManager.BoldFont"))
-							.AutoWrapText(true)
-						]
-
-						+ SVerticalBox::Slot()
-						.AutoHeight()
-						.Padding(5, 5)
-						[
-							SAssignNew(SourceViewport, SBasePoseViewport)
-							.Skeleton(OldSkeleton)
-						]
-					]
-
-					+SHorizontalBox::Slot()
-					[
-						SNew(SVerticalBox)
-						+ SVerticalBox::Slot()
-						.AutoHeight()
-						[
-							SNew(STextBlock)
-							.Text(LOCTEXT("TargetSkeleteonTitle", "[Target]"))
-							.Font(FEditorStyle::GetFontStyle("Persona.RetargetManager.BoldFont"))
-							.AutoWrapText(true)
-						]
-
-						+SVerticalBox::Slot()
-						.AutoHeight()
-						.Padding(5, 5)
-						[
-							SAssignNew(TargetViewport, SBasePoseViewport)
-							.Skeleton(nullptr)
-						]
-					]
-				]
-
-				+SVerticalBox::Slot()
-				.AutoHeight()
-				.Padding(0, 5)
-				[
-					SNew(SVerticalBox)
-					+SVerticalBox::Slot()
-					.AutoHeight()
-					.Padding(5, 5)
-					[
-						OptionWidget
-					]
-				]
-			]
-		];
-
-	UpdateAssetPicker();
-	UpdateExampleText();
-}
-
-FReply SAnimationRemapSkeleton::ShowFolderOption()
-{
-	TSharedRef<SSelectFolderDlg> NewAnimDlg =
-	SNew(SSelectFolderDlg);
-
-	if(NewAnimDlg->ShowModal() != EAppReturnType::Cancel)
-	{
-		NameDuplicateRule.FolderPath = NewAnimDlg->GetAssetPath();
-	}
-
-	if (WidgetWindow.IsValid())
-	{
-		WidgetWindow.Pin()->BringToFront(true);
-	}
-
-	return FReply::Handled();
-}
-
-void SAnimationRemapSkeleton::UpdateExampleText()
-{
-	FString ReplaceFrom = FString::Printf(TEXT("Old Name : ###%s###"), *NameDuplicateRule.ReplaceFrom);
-	FString ReplaceTo = FString::Printf(TEXT("New Name : %s###%s###%s"), *NameDuplicateRule.Prefix, *NameDuplicateRule.ReplaceTo, *NameDuplicateRule.Suffix);
-
-	ExampleText = FText::FromString(FString::Printf(TEXT("%s\n%s"), *ReplaceFrom, *ReplaceTo));
-}
-
-ECheckBoxState SAnimationRemapSkeleton::IsRemappingReferencedAssets() const
-{
-	return bRemapReferencedAssets ? ECheckBoxState::Checked : ECheckBoxState::Unchecked;
-}
-
-void SAnimationRemapSkeleton::OnRemappingReferencedAssetsChanged(ECheckBoxState InNewRadioState)
-{
-	bRemapReferencedAssets = (InNewRadioState == ECheckBoxState::Checked);
-}
-
-ECheckBoxState SAnimationRemapSkeleton::IsRemappingToExistingAssetsChecked() const
-{
-	return bAllowRemappingToExistingAssets ? ECheckBoxState::Checked : ECheckBoxState::Unchecked;
-}
-
-bool SAnimationRemapSkeleton::IsRemappingToExistingAssetsEnabled() const
-{
-	return bRemapReferencedAssets;
-}
-
-void SAnimationRemapSkeleton::OnRemappingToExistingAssetsChanged(ECheckBoxState InNewRadioState)
-{
-	bAllowRemappingToExistingAssets = (InNewRadioState == ECheckBoxState::Checked);
-}
-
-ECheckBoxState SAnimationRemapSkeleton::IsConvertSpacesChecked() const
-{
-	return bConvertSpaces ? ECheckBoxState::Checked : ECheckBoxState::Unchecked;
-}
-
-void SAnimationRemapSkeleton::OnConvertSpacesCheckChanged(ECheckBoxState InNewRadioState)
-{
-	bConvertSpaces = (InNewRadioState == ECheckBoxState::Checked);
-}
-
-bool SAnimationRemapSkeleton::IsShowOnlyCompatibleSkeletonsEnabled() const
-{
-	// if convert space is false, compatible skeletons won't matter either. 
-	return (bConvertSpaces == true);
-}
-
-ECheckBoxState SAnimationRemapSkeleton::IsShowOnlyCompatibleSkeletonsChecked() const
-{
-	return bShowOnlyCompatibleSkeletons ? ECheckBoxState::Checked : ECheckBoxState::Unchecked;
-}
-
-void SAnimationRemapSkeleton::OnShowOnlyCompatibleSkeletonsCheckChanged(ECheckBoxState InNewRadioState)
-{
-	bShowOnlyCompatibleSkeletons = (InNewRadioState == ECheckBoxState::Checked);
-
-	UpdateAssetPicker();
-}
-
-bool SAnimationRemapSkeleton::CanApply() const
-{
-	return (NewSkeleton!=nullptr && NewSkeleton!=OldSkeleton);
-}
-
-void SAnimationRemapSkeleton::OnAssetSelectedFromPicker(const FAssetData& AssetData)
-{
-	if (AssetData.GetAsset())
-	{
-		NewSkeleton = Cast<USkeleton>(AssetData.GetAsset());
-
-		TargetViewport->SetSkeleton(NewSkeleton);
-	}
-}
-
-FReply SAnimationRemapSkeleton::OnApply()
-{
-	if (OnRetargetAnimationDelegate.IsBound())
-	{
-		OnRetargetAnimationDelegate.Execute(OldSkeleton, NewSkeleton, bRemapReferencedAssets, bAllowRemappingToExistingAssets, bConvertSpaces, bShowDuplicateAssetOption ? &NameDuplicateRule : nullptr);
-	}
-
-	CloseWindow();
-	return FReply::Handled();
-}
-
-FReply SAnimationRemapSkeleton::OnCancel()
-{
-	NewSkeleton = nullptr;
-	CloseWindow();
-	return FReply::Handled();
-}
-
-void SAnimationRemapSkeleton::OnRemapDialogClosed(const TSharedRef<SWindow>& Window)
-{
-	NewSkeleton = nullptr;
-	DialogWindow = nullptr;
-}
-
-void SAnimationRemapSkeleton::CloseWindow()
-{
-	if ( WidgetWindow.IsValid() )
-	{
-		WidgetWindow.Pin()->RequestDestroyWindow();
-	}
-}
-
-void SAnimationRemapSkeleton::ShowWindow(USkeleton* OldSkeleton, const FText& WarningMessage, bool bDuplicateAssets, FOnRetargetAnimation RetargetDelegate)
-{
-	if(DialogWindow.IsValid())
-	{
-		FSlateApplication::Get().DestroyWindowImmediately(DialogWindow.ToSharedRef());
-	}
-
-	DialogWindow = SNew(SWindow)
-		.Title( LOCTEXT("RemapSkeleton", "Select Skeleton") )
-		.SupportsMinimize(false) .SupportsMaximize(false)
-		.SizingRule( ESizingRule::Autosized );
-
-	TSharedPtr<class SAnimationRemapSkeleton> DialogWidget;
-
-	TSharedPtr<SBorder> DialogWrapper = 
-		SNew(SBorder)
-		.BorderImage(FEditorStyle::GetBrush("ToolPanel.GroupBorder"))
-		.Padding(4.0f)
-		[
-			SAssignNew(DialogWidget, SAnimationRemapSkeleton)
-			.CurrentSkeleton(OldSkeleton)
-			.WidgetWindow(DialogWindow)
-			.WarningMessage(WarningMessage)
-			.ShowRemapOption(true)
-			.ShowExistingRemapOption(true)
-			.ShowConvertSpacesOption(OldSkeleton != nullptr)
-			.ShowCompatibleDisplayOption(OldSkeleton != nullptr)
-			.ShowDuplicateAssetOption(bDuplicateAssets)
-			.OnRetargetDelegate(RetargetDelegate)
-		];
-
-	DialogWindow->SetOnWindowClosed(FRequestDestroyWindowOverride::CreateSP(DialogWidget.Get(), &SAnimationRemapSkeleton::OnRemapDialogClosed));
-	DialogWindow->SetContent(DialogWrapper.ToSharedRef());
-
-	FSlateApplication::Get().AddWindow(DialogWindow.ToSharedRef());
-}
-
-////////////////////////////////////////////////////
-
-FDlgRemapSkeleton::FDlgRemapSkeleton( USkeleton* Skeleton )
-{
-	if (FSlateApplication::IsInitialized())
-	{
-		DialogWindow = SNew(SWindow)
-			.Title( LOCTEXT("RemapSkeleton", "Select Skeleton") )
-			.SupportsMinimize(false) .SupportsMaximize(false)
-			.SizingRule( ESizingRule::Autosized );
-
-		TSharedPtr<SBorder> DialogWrapper = 
-			SNew(SBorder)
-			.BorderImage(FEditorStyle::GetBrush("ToolPanel.GroupBorder"))
-			.Padding(4.0f)
-			[
-				SAssignNew(DialogWidget, SAnimationRemapSkeleton)
-				.CurrentSkeleton(Skeleton)
-				.WidgetWindow(DialogWindow)
-			];
-
-		DialogWindow->SetContent(DialogWrapper.ToSharedRef());
-	}
-}
-
-bool FDlgRemapSkeleton::ShowModal()
-{
-	GEditor->EditorAddModalWindow(DialogWindow.ToSharedRef());
-
-	NewSkeleton = DialogWidget.Get()->NewSkeleton;
-
-	return (NewSkeleton != nullptr && NewSkeleton != DialogWidget.Get()->OldSkeleton);
-}
-
-//////////////////////////////////////////////////////////////
-void SRemapFailures::Construct( const FArguments& InArgs )
-{
-	for ( auto RemapIt = InArgs._FailedRemaps.CreateConstIterator(); RemapIt; ++RemapIt )
-	{
-		FailedRemaps.Add( MakeShareable( new FText(*RemapIt) ) );
-	}
-
-	ChildSlot
-		[
-			SNew(SBorder)
-			.BorderImage( FEditorStyle::GetBrush("Docking.Tab.ContentAreaBrush") )
-			.Padding(FMargin(4, 8, 4, 4))
-			[
-				SNew(SVerticalBox)
-
-				// Title text
-				+SVerticalBox::Slot()
-				.AutoHeight()
-				[
-					SNew(STextBlock) .Text( LOCTEXT("RemapFailureTitle", "The following assets could not be Remaped.") )
-				]
-
-				// Failure list
-				+SVerticalBox::Slot()
-				.Padding(0, 8)
-				.FillHeight(1.f)
-				[
-					SNew(SBorder)
-					.BorderImage( FEditorStyle::GetBrush("ToolPanel.GroupBorder") )
-					[
-						SNew(SListView<TSharedRef<FText>>)
-						.ListItemsSource(&FailedRemaps)
-						.SelectionMode(ESelectionMode::None)
-						.OnGenerateRow(this, &SRemapFailures::MakeListViewWidget)
-					]
-				]
-
-				+SVerticalBox::Slot()
-				.AutoHeight()
-				.HAlign(HAlign_Right)
-				[
-					SNew(SUniformGridPanel)
-					.SlotPadding(FEditorStyle::GetMargin("StandardDialog.SlotPadding"))
-					.MinDesiredSlotWidth(FEditorStyle::GetFloat("StandardDialog.MinDesiredSlotWidth"))
-					.MinDesiredSlotHeight(FEditorStyle::GetFloat("StandardDialog.MinDesiredSlotHeight"))
-					+SUniformGridPanel::Slot(0,0)
-					[
-						SNew(SButton) 
-						.HAlign(HAlign_Center)
-						.ContentPadding(FEditorStyle::GetMargin("StandardDialog.ContentPadding"))
-						.Text(LOCTEXT("RemapFailuresCloseButton", "Close"))
-						.OnClicked(this, &SRemapFailures::CloseClicked)
-					]
-				]
-			]
-		];
-}
-
-void SRemapFailures::OpenRemapFailuresDialog(const TArray<FText>& InFailedRemaps)
-{
-	TSharedRef<SWindow> RemapWindow = SNew(SWindow)
-		.Title(LOCTEXT("FailedRemapsDialog", "Failed Remaps"))
-		.ClientSize(FVector2D(800,400))
-		.SupportsMaximize(false)
-		.SupportsMinimize(false)
-		[
-			SNew(SRemapFailures).FailedRemaps(InFailedRemaps)
-		];
-
-	IMainFrameModule& MainFrameModule = FModuleManager::LoadModuleChecked<IMainFrameModule>(TEXT("MainFrame"));
-
-	if ( MainFrameModule.GetParentWindow().IsValid() )
-	{
-		FSlateApplication::Get().AddWindowAsNativeChild(RemapWindow, MainFrameModule.GetParentWindow().ToSharedRef());
-	}
-	else
-	{
-		FSlateApplication::Get().AddWindow(RemapWindow);
-	}
-}
-
-TSharedRef<ITableRow> SRemapFailures::MakeListViewWidget(TSharedRef<FText> Item, const TSharedRef<STableViewBase>& OwnerTable)
-{
-	return
-		SNew(STableRow< TSharedRef<FText> >, OwnerTable)
-		[
-			SNew(STextBlock) .Text( Item.Get() )
-		];
-}
-
-FReply SRemapFailures::CloseClicked()
-{
-	TSharedPtr<SWindow> Window = FSlateApplication::Get().FindWidgetWindow(AsShared());
-
-	if ( Window.IsValid() )
-	{
-		Window->RequestDestroyWindow();
-	}
-
-	return FReply::Handled();
-}
-
 void SSkeletonBoneRemoval::Construct( const FArguments& InArgs )
 {
 	bShouldContinue = false;
@@ -1560,331 +757,6 @@ FString SSelectFolderDlg::GetAssetPath()
 	return AssetPath.ToString();
 }
 
-//////////////////////////////////////////////////////////////////////////
-// Window for remapping assets when retargetting
-TSharedPtr<SWindow> SAnimationRemapAssets::DialogWindow;
-
-void SAnimationRemapAssets::Construct(const FArguments& InArgs)
-{
-	NewSkeleton = InArgs._NewSkeleton;
-	RetargetContext = InArgs._RetargetContext;
-	
-	TSharedPtr<SVerticalBox> ListBox;
-
-	TArray<UObject*> Duplicates = RetargetContext->GetAllDuplicates();
-
-	for(UObject* Asset : Duplicates)
-	{
-		// We don't want to add anim blueprints here, just animation assets
-		if(Asset->GetClass() != UAnimBlueprint::StaticClass())
-		{
-			AssetListInfo.Add(FDisplayedAssetEntryInfo::Make(Asset, NewSkeleton));
-		}
-	}
-
-	this->ChildSlot
-	[
-		SNew(SVerticalBox)
-		+SVerticalBox::Slot()
-		.Padding(5.0f)
-		.AutoHeight()
-		[
-			SNew(STextBlock)
-			.Text(LOCTEXT("RemapAssetsDescription", "The assets shown below need to be duplicated or remapped for the new blueprint. Select a new animation to use in the new animation blueprint for each asset or leave blank to duplicate the existing asset."))
-			.AutoWrapText(true)
-		]
-		+ SVerticalBox::Slot()
-		.Padding(5.0f)
-		.AutoHeight()
-		.MaxHeight(500.0f)
-		[
-			SNew(SBorder)
-			.BorderImage(FEditorStyle::GetBrush("ToolPanel.DarkGroupBorder"))
-			[
-				SAssignNew(ListWidget, SRemapAssetEntryList)
-				.ItemHeight(20.0f)
-				.ListItemsSource(&AssetListInfo)
-				.OnGenerateRow(this, &SAnimationRemapAssets::OnGenerateMontageReferenceRow)
-				.SelectionMode(ESelectionMode::None)
-				.HeaderRow
-				(
-					SNew(SHeaderRow)
-					+ SHeaderRow::Column("AssetName")
-					.DefaultLabel(LOCTEXT("ColumnLabel_RemapAssetName", "Asset Name"))
-					+ SHeaderRow::Column("AssetType")
-					.DefaultLabel(LOCTEXT("ColumnLabel_RemapAssetType", "Asset Type"))
-					+ SHeaderRow::Column("AssetRemap")
-					.DefaultLabel(LOCTEXT("ColumnLabel_RemapAssetRemap", "Remapped Asset"))
-				)
-			]
-		]
-		+ SVerticalBox::Slot()
-		.Padding(5.0f)
-		.AutoHeight()
-		[
-			SNew(SHorizontalBox)
-			+SHorizontalBox::Slot()
-			.VAlign(VAlign_Center)
-			.HAlign(HAlign_Center)
-			.FillWidth(0.2f)
-			[
-				SNew(SButton)
-				.ContentPadding(2.0f)
-				.OnClicked(this, &SAnimationRemapAssets::OnBestGuessClicked)
-				.Content()
-				[
-					SNew(STextBlock)
-					.Text(LOCTEXT("BestGuessButton", "Auto-Fill Using Best Guess"))
-				]
-			]
-			+SHorizontalBox::Slot()
-			.VAlign(VAlign_Center)
-			.HAlign(HAlign_Center)
-			.Padding(FMargin(5.0f, 0.0f, 0.0f, 0.0f))
-			[
-				SNew(STextBlock)
-				.AutoWrapText(true)
-				.Text(LOCTEXT("BestGuessDescription", "Auto-Fill will look at the names of all compatible assets for the new skeleton and look for something similar to use for the remapped asset."))
-			]
-		]
-
-		+ SVerticalBox::Slot()
-		.Padding(5.0f)
-		.VAlign(VAlign_Center)
-		.HAlign(HAlign_Center)
-		.AutoHeight()
-		[
-			SNew(SBox)
-			.VAlign(VAlign_Center)
-			.HAlign(HAlign_Center)
-			.Content()
-			[
-				SNew(SButton)
-				.ContentPadding(FEditorStyle::GetMargin("StandardDialog.ContentPadding"))
-				.OnClicked(this, &SAnimationRemapAssets::OnOkClicked)
-				.Content()
-				[
-					SNew(STextBlock)
-					.Text(LOCTEXT("OkButton", "OK"))
-				]
-			]
-		]
-	];
-}
-
-void SAnimationRemapAssets::ShowWindow(FAnimationRetargetContext& RetargetContext, USkeleton* RetargetToSkeleton)
-{
-	if(DialogWindow.IsValid())
-	{
-		FSlateApplication::Get().DestroyWindowImmediately(DialogWindow.ToSharedRef());
-	}
-	
-	DialogWindow = SNew(SWindow)
-		.Title(LOCTEXT("RemapAssets", "Choose Assets to Remap"))
-		.SupportsMinimize(false)
-		.SupportsMaximize(false)
-		.HasCloseButton(false)
-		.MaxWidth(1024.0f)
-		.SizingRule(ESizingRule::Autosized);
-	
-	TSharedPtr<class SAnimationRemapAssets> DialogWidget;
-	
-	TSharedPtr<SBorder> DialogWrapper =
-		SNew(SBorder)
-		.BorderImage(FEditorStyle::GetBrush("ToolPanel.GroupBorder"))
-		.Padding(4.0f)
-		[
-			SAssignNew(DialogWidget, SAnimationRemapAssets)
-				.NewSkeleton(RetargetToSkeleton)
-				.RetargetContext(&RetargetContext)
-		];
-	
-	DialogWindow->SetOnWindowClosed(FRequestDestroyWindowOverride::CreateSP(DialogWidget.Get(), &SAnimationRemapAssets::OnDialogClosed));
-	DialogWindow->SetContent(DialogWrapper.ToSharedRef());
-	
-	FSlateApplication::Get().AddModalWindow(DialogWindow.ToSharedRef(), nullptr);
-}
-
-TSharedRef<ITableRow> SAnimationRemapAssets::OnGenerateMontageReferenceRow(TSharedPtr<FDisplayedAssetEntryInfo> Item, const TSharedRef<STableViewBase>& OwnerTable)
-{
-	return SNew(SAssetEntryRow, OwnerTable)
-		.DisplayedInfo(Item);
-}
-
-void SAnimationRemapAssets::OnDialogClosed(const TSharedRef<SWindow>& Window)
-{
-	DialogWindow = nullptr;
-}
-
-FReply SAnimationRemapAssets::OnOkClicked()
-{
-	for(TSharedPtr<FDisplayedAssetEntryInfo> AssetInfo : AssetListInfo)
-	{
-		if(AssetInfo->RemapAsset)
-		{
-			RetargetContext->AddRemappedAsset(Cast<UAnimationAsset>(AssetInfo->AnimAsset), Cast<UAnimationAsset>(AssetInfo->RemapAsset));
-		}
-	}
-
-	DialogWindow->RequestDestroyWindow();
-
-	return FReply::Handled();
-}
-
-FReply SAnimationRemapAssets::OnBestGuessClicked()
-{
-	// collect all compatible assets
-	FAssetRegistryModule& AssetRegistryModule = FModuleManager::LoadModuleChecked<FAssetRegistryModule>(TEXT("AssetRegistry"));
-	FString SkeletonName = FAssetData(NewSkeleton).GetExportTextName();
-
-	TArray<FAssetData> CompatibleAssets;
-	TArray<FAssetData> AssetDataList;
-
-	AssetRegistryModule.Get().GetAssetsByClass(UAnimationAsset::StaticClass()->GetFName(), AssetDataList, true);
-
-	for(const FAssetData& Data : AssetDataList)
-	{
-		if(Data.GetTagValueRef<FString>("Skeleton") == SkeletonName)
-		{
-			CompatibleAssets.Add(Data);
-		}
-	}
-
-	if(CompatibleAssets.Num() > 0)
-	{
-		// Do best guess analysis for the assets based on name.
-		for(TSharedPtr<FDisplayedAssetEntryInfo>& Info : AssetListInfo)
-		{
-			const FAssetData* BestMatchData = FindBestGuessMatch(FAssetData(Info->AnimAsset), CompatibleAssets);
-			Info->RemapAsset = BestMatchData ? BestMatchData->GetAsset() : nullptr;
-		}
-	}
-
-	ListWidget->RequestListRefresh();
-
-	return FReply::Handled();
-}
-
-const FAssetData* SAnimationRemapAssets::FindBestGuessMatch(const FAssetData& AssetData, const TArray<FAssetData>& PossibleAssets) const
-{
-	int32 LowestScore = MAX_int32;
-	int32 FoundIndex = INDEX_NONE;
-
-	for(int32 Idx = 0 ; Idx < PossibleAssets.Num() ; ++Idx)
-	{
-		const FAssetData& Data = PossibleAssets[Idx];
-
-		if(Data.AssetClass == AssetData.AssetClass)
-		{
-			int32 Distance = FAnimationRuntime::GetStringDistance(AssetData.AssetName.ToString(), Data.AssetName.ToString());
-
-			if(Distance < LowestScore)
-			{
-				LowestScore = Distance;
-				FoundIndex = Idx;
-			}
-		}
-	}
-
-	if(FoundIndex != INDEX_NONE)
-	{
-		return &PossibleAssets[FoundIndex];
-	}
-
-	return nullptr;
-}
-
-void SAssetEntryRow::Construct(const FArguments& InArgs, const TSharedRef<STableViewBase>& InOwnerTableView)
-{
-	check(InArgs._DisplayedInfo.IsValid());
-	DisplayedInfo = InArgs._DisplayedInfo;
-
-	SkeletonExportName = FAssetData(DisplayedInfo->NewSkeleton).GetExportTextName();
-
-	SMultiColumnTableRow<TSharedPtr<FDisplayedAssetEntryInfo>>::Construct(FSuperRowType::FArguments(), InOwnerTableView);
-}
-
-TSharedRef<SWidget> SAssetEntryRow::GenerateWidgetForColumn(const FName& ColumnName)
-{
-	if(ColumnName == TEXT("AssetName"))
-	{
-		return SNew(STextBlock)
-			.Text(FText::Format(LOCTEXT("AssetNameEntry", "{0}"), FText::FromString(DisplayedInfo->AnimAsset->GetName())));
-	}
-	else if(ColumnName == TEXT("AssetType"))
-	{
-		return SNew(STextBlock)
-			.Text(FText::Format(LOCTEXT("AssetTypeEntry", "{0}"), FText::FromString(DisplayedInfo->AnimAsset->GetClass()->GetName())));
-	}
-	else if(ColumnName == TEXT("AssetRemap"))
-	{
-		return SNew(SBox)
-			.Padding(2.0f)
-			[
-				SNew(SComboButton)
-				.ToolTipText(LOCTEXT("AssetRemapTooltip", "Select compatible asset to remap to."))
-				.ButtonStyle(FEditorStyle::Get(), "PropertyEditor.AssetComboStyle")
-				.ForegroundColor(FEditorStyle::GetColor("PropertyEditor.AssetName.ColorAndOpacity"))
-				.OnGetMenuContent(this, &SAssetEntryRow::GetRemapMenuContent)
-				.ContentPadding(2.0f)
-				.ButtonContent()
-				[
-					SNew(STextBlock)
-					.TextStyle(FEditorStyle::Get(), "PropertyEditor.AssetClass")
-					.Font(FEditorStyle::GetFontStyle("PropertyWindow.NormalFont"))
-					.Text(this, &SAssetEntryRow::GetRemapMenuButtonText)
-				]
-			];
-	}
-
-	return SNullWidget::NullWidget;
-}
-
-TSharedRef<SWidget> SAssetEntryRow::GetRemapMenuContent()
-{
-	FContentBrowserModule& ContentBrowserModule = FModuleManager::Get().LoadModuleChecked<FContentBrowserModule>(TEXT("ContentBrowser"));
-	
-	FAssetPickerConfig PickerConfig;
-	PickerConfig.SelectionMode = ESelectionMode::Single;
-	PickerConfig.Filter.ClassNames.Add(DisplayedInfo->AnimAsset->GetClass()->GetFName());
-	PickerConfig.OnAssetSelected = FOnAssetSelected::CreateSP(this, &SAssetEntryRow::OnAssetSelected);
-	PickerConfig.OnShouldFilterAsset = FOnShouldFilterAsset::CreateSP(this, &SAssetEntryRow::OnShouldFilterAsset);
-	PickerConfig.bAllowNullSelection = true;
-	
-	return SNew(SBox)
-		.WidthOverride(384)
-		.HeightOverride(768)
-		[
-			ContentBrowserModule.Get().CreateAssetPicker(PickerConfig)
-		];
-}
-
-FText SAssetEntryRow::GetRemapMenuButtonText() const
-{
-	FText NameText = (DisplayedInfo->RemapAsset) ? FText::FromString(*DisplayedInfo->RemapAsset->GetName()) : LOCTEXT("AssetRemapNone", "None");
-
-	return FText::Format(LOCTEXT("RemapButtonText", "{0}"), NameText);
-}
-
-void SAssetEntryRow::OnAssetSelected(const FAssetData& AssetData)
-{
-	// Close the asset picker menu
-	FSlateApplication::Get().DismissAllMenus();
-
-	RemapAsset = AssetData.GetAsset();
-	DisplayedInfo->RemapAsset = RemapAsset.Get();
-}
-
-bool SAssetEntryRow::OnShouldFilterAsset(const FAssetData& AssetData) const
-{
-	if(AssetData.GetTagValueRef<FString>("Skeleton") == SkeletonExportName)
-	{
-		return false;
-	}
-
-	return true;
-}
-
 TSharedRef<FDisplayedAssetEntryInfo> FDisplayedAssetEntryInfo::Make(UObject* InAsset, USkeleton* InNewSkeleton)
 {
 	return MakeShareable(new FDisplayedAssetEntryInfo(InAsset, InNewSkeleton));
diff --git a/Engine/Source/Editor/UnrealEd/Public/EditorAnimUtils.h b/Engine/Source/Editor/UnrealEd/Public/EditorAnimUtils.h
--- a/Engine/Source/Editor/UnrealEd/Public/EditorAnimUtils.h
+++ b/Engine/Source/Editor/UnrealEd/Public/EditorAnimUtils.h
@@ -59,9 +59,6 @@ namespace EditorAnimUtils
 
 		/** Duplicates the assets stored for retargetting, populating maps of original assets to new asset */
 		void DuplicateAssetsToRetarget(UPackage* DestinationPackage, const FNameDuplicationRule* NameRule);
-
-		/** Retarget the contained assets */
-		void RetargetAnimations(USkeleton* OldSkeleton, USkeleton* NewSkeleton);
 	
 		void AddRemappedAsset(UAnimationAsset* OriginalAsset, UAnimationAsset* NewAsset);
 
@@ -85,40 +82,7 @@ namespace EditorAnimUtils
 		/** Initialize the object, only to be called by constructors */
 		void Initialize(TArray<TWeakObjectPtr<UObject>> AssetsToRetarget, bool bRetargetReferredAssets);
 	};
-
-	/**
-	 * Retargets the supplied UObjects (as long as they are an animation asset), optionally duplicating them and retargetting their reference assets too
-	 *
-	 * @param NewSkeleton						The skeleton the supplied assets should be retargeted to
-	 * @param AssetsToRetarget					The animation assets to copy/retarget
-	 * @param bRetargetReferredAssets			If true retargets any assets referred to by assets in AssetsToRetarget. If false then the references are cleared.
-	 * @param bDuplicatedAssetsBeforeRetarget	If true the assets are duplicated and then the duplicates are retargetted
-	 * @param bConvertSpace						Do the conversion in component space of the animation to match new target
-	 */
-	UNREALED_API UObject* RetargetAnimations(USkeleton* OldSkeleton, USkeleton* NewSkeleton, TArray<TWeakObjectPtr<UObject>> AssetsToRetarget, bool bRetargetReferredAssets, const FNameDuplicationRule* NameRule, bool bConvertSpace);
-
-	/**
-	 * Retargets the supplied FAssetDatas (as long as they are an animation asset), optionally duplicating them and retargetting their reference assets too
-	 *
-	 * @param NewSkeleton						The skeleton the supplied assets should be retargeted to
-	 * @param AssetsToRetarget					The animation assets to copy/retarget
-	 * @param bRetargetReferredAssets			If true retargets any assets referred to by assets in AssetsToRetarget. If false then the references are cleared.
-	 * @param bDuplicatedAssetsBeforeRetarget	If true the assets are duplicated and then the duplicates are retargetted
-	 * @param bConvertSpace						Do the conversion in component space of the animation to match new target
-	 */
-	UNREALED_API UObject* RetargetAnimations(USkeleton* OldSkeleton, USkeleton* NewSkeleton, const TArray<FAssetData>& AssetsToRetarget, bool bRetargetReferredAssets, const FNameDuplicationRule* NameRule, bool bConvertSpace);
-
-	/**
-	 * Retargets the supplied FAnimationRetargetContext, optionally duplicating the assets and retargetting the assets reference assets too. Is called by other overloads of RetargetAnimations
-	 *
-	 * @param NewSkeleton						The skeleton the supplied assets should be retargeted to
-	 * @param AssetsToRetarget					The animation assets to copy/retarget
-	 * @param bRetargetReferredAssets			If true retargets any assets referred to by assets in AssetsToRetarget. If false then the references are cleared.
-	 * @param bDuplicatedAssetsBeforeRetarget	If true the assets are duplicated and then the duplicates are retargetted
-	 * @param bConvertSpace						Do the conversion in component space of the animation to match new target
-	 */
-	UNREALED_API UObject* RetargetAnimations(USkeleton* OldSkeleton, USkeleton* NewSkeleton, FAnimationRetargetContext& RetargetContext, bool bRetargetReferredAssets, const FNameDuplicationRule* NameRule);
-
+	
 	// Populates the supplied TArrays with any animation assets that this blueprint refers too
 	UNREALED_API void GetAllAnimationSequencesReferredInBlueprint(UAnimBlueprint* AnimBlueprint, TArray<UAnimationAsset*>& AnimationAsset);
 
diff --git a/Engine/Source/Editor/UnrealEd/Public/SSkeletonWidget.h b/Engine/Source/Editor/UnrealEd/Public/SSkeletonWidget.h
--- a/Engine/Source/Editor/UnrealEd/Public/SSkeletonWidget.h
+++ b/Engine/Source/Editor/UnrealEd/Public/SSkeletonWidget.h
@@ -301,207 +301,6 @@ private:
 	bool IsVisible() const override;
 };
 
-/////////////////////////////////////////////
-/** 
- * Slate panel for choose Skeleton for assets to relink
- */
-
-DECLARE_DELEGATE_SixParams(FOnRetargetAnimation, USkeleton* /*OldSkeleton*/, USkeleton* /*NewSkeleton*/, bool /*bRemapReferencedAssets*/, bool /*bAllowRemapToExisting*/, bool /*bConvertSpaces*/, const FNameDuplicationRule* /*NameRule*/)
-
-class SAnimationRemapSkeleton : public SCompoundWidget
-{
-public:
-
-	SLATE_BEGIN_ARGS(SAnimationRemapSkeleton)
-		: _CurrentSkeleton(NULL)
-		, _WidgetWindow(NULL)
-		, _WarningMessage()
-		, _ShowRemapOption(false)
-		, _ShowExistingRemapOption(false)
-		{}
-
-		/** The anim sequences to compress */
-		SLATE_ARGUMENT( USkeleton*, CurrentSkeleton )
-		SLATE_ARGUMENT( TSharedPtr<SWindow>, WidgetWindow )
-		SLATE_ARGUMENT( FText, WarningMessage )
-		SLATE_ARGUMENT( bool, ShowRemapOption )
-		SLATE_ARGUMENT( bool, ShowExistingRemapOption )
-		SLATE_ARGUMENT( bool, ShowConvertSpacesOption )
-		SLATE_ARGUMENT( bool, ShowCompatibleDisplayOption )
-		SLATE_ARGUMENT( bool, ShowDuplicateAssetOption )
-		SLATE_EVENT(FOnRetargetAnimation, OnRetargetDelegate)
-	SLATE_END_ARGS()	
-
-	/**
-	 * Constructs this widget
-	 *
-	 * @param	InArgs	The declaration data for this widget
-	 */
-	void Construct( const FArguments& InArgs );
-
-	/**
-	 * Old Skeleton that was mapped 
-	 * This data is needed to prevent users from selecting same skeleton
-	 */
-	USkeleton *OldSkeleton;
-	/**
-	 * New Skeleton that they would like to map to
-	 */
-	USkeleton *NewSkeleton;
-
-private:
-	/**
-	 * Whether we are remapping assets that are referenced by the assets the user selects to remap
-	 */
-	bool bRemapReferencedAssets;
-
-	/**
-	 * Whether we allow remapping to existing assets for the new skeleton
-	 */
-	bool bAllowRemappingToExistingAssets;
-
-	/**
-	 * Whether we are remapping assets that are referenced by the assets the user selects to remap
-	 */
-	bool bConvertSpaces;
-
-	/**
-	 * Whether to show skeletons with the same rig set up
-	*/
-	bool bShowOnlyCompatibleSkeletons;
-
-	TSharedPtr<SBasePoseViewport> SourceViewport;
-	TSharedPtr<SBasePoseViewport> TargetViewport;
-
-	TSharedPtr<SBox> AssetPickerBox;
-
-	TWeakPtr<SWindow> WidgetWindow;
-
-	FOnRetargetAnimation OnRetargetAnimationDelegate;
-
-	/** Handlers for check box for remapping assets option */
-	ECheckBoxState IsRemappingReferencedAssets() const;
-	void OnRemappingReferencedAssetsChanged(ECheckBoxState InNewRadioState);
-
-	/** Handlers for check box for remapping to existing assets */
-	ECheckBoxState IsRemappingToExistingAssetsChecked() const;
-	bool IsRemappingToExistingAssetsEnabled() const;
-	void OnRemappingToExistingAssetsChanged(ECheckBoxState InNewRadioState);
-
-	/** Handlers for check box for converting spaces*/
-	ECheckBoxState IsConvertSpacesChecked() const;
-	void OnConvertSpacesCheckChanged(ECheckBoxState InNewRadioState);
-
-	/** Handlers for check box for converting spaces*/
-	ECheckBoxState IsShowOnlyCompatibleSkeletonsChecked() const;
-	bool IsShowOnlyCompatibleSkeletonsEnabled() const;
-	void OnShowOnlyCompatibleSkeletonsCheckChanged(ECheckBoxState InNewRadioState);
-
-	/** should filter asset */
-	bool OnShouldFilterAsset(const struct FAssetData& AssetData);
-
-	/**
-	 * return true if it can apply 
-	 */
-	bool CanApply() const;
-
-	FReply OnApply();
-	FReply OnCancel();
-
-	void CloseWindow();
-
-	/** Handler for dialog window close button */
-	void OnRemapDialogClosed(const TSharedRef<SWindow>& Window);
-
-	/**
-	 * Handler for when asset is selected
-	 */
-	void OnAssetSelectedFromPicker(const FAssetData& AssetData);
-
-	/*
-	* Refreshes asset picker - call when asset picker option changes
-	*/
-	void UpdateAssetPicker();
-
-	/*
-	 * Duplicate Name Rule 
-	 */
-	FNameDuplicationRule NameDuplicateRule;
-	FText ExampleText;
-
-	FText GetPrefixName() const
-	{
-		return FText::FromString(NameDuplicateRule.Prefix);
-	}
-
-	void SetPrefixName(const FText &InText)
-	{
-		NameDuplicateRule.Prefix = InText.ToString();
-		UpdateExampleText();
-	}
-
-	FText GetSuffixName() const
-	{
-		return FText::FromString(NameDuplicateRule.Suffix);
-	}
-
-	void SetSuffixName(const FText &InText)
-	{
-		NameDuplicateRule.Suffix = InText.ToString();
-		UpdateExampleText();
-	}
-
-	FText GetReplaceFrom() const
-	{
-		return FText::FromString(NameDuplicateRule.ReplaceFrom);
-	}
-
-	void SetReplaceFrom(const FText &InText)
-	{
-		NameDuplicateRule.ReplaceFrom = InText.ToString();
-		UpdateExampleText();
-	}
-
-	FText GetReplaceTo() const
-	{
-		return FText::FromString(NameDuplicateRule.ReplaceTo);
-	}
-
-	void SetReplaceTo(const FText &InText)
-	{
-		NameDuplicateRule.ReplaceTo = InText.ToString();
-		UpdateExampleText();
-	}
-
-	FText GetExampleText() const
-	{
-		return ExampleText;
-	}
-
-	void UpdateExampleText();
-
-	FReply ShowFolderOption();
-	FText GetFolderPath() const
-	{
-		return FText::FromString(NameDuplicateRule.FolderPath);
-	}
-
-public:
-
-	/**
-	 *  Show window
-	 *
-	 * @param OldSkeleton		Old Skeleton to change from
-	 *
-	 * @return true if successfully selected new skeleton
-	 */
-	static UNREALED_API void ShowWindow(USkeleton* OldSkeleton, const FText& WarningMessage, bool bDuplicateAssets, FOnRetargetAnimation RetargetDelegate);
-
-	static TSharedPtr<SWindow> DialogWindow;
-
-	bool bShowDuplicateAssetOption;
-};
-
 //////////////////////////////////////////////////////////////////////////
 
 class FDisplayedAssetEntryInfo
@@ -519,130 +318,6 @@ protected:
 	FDisplayedAssetEntryInfo(UObject* InAsset, USkeleton* InNewSkeleton);
 };
 
-class SAssetEntryRow : public SMultiColumnTableRow<TSharedPtr<FDisplayedAssetEntryInfo>>
-{
-public:
-	SLATE_BEGIN_ARGS(SAssetEntryRow)
-	{}
-
-		SLATE_ARGUMENT(TSharedPtr<FDisplayedAssetEntryInfo>, DisplayedInfo)
-
-	SLATE_END_ARGS()
-
-	void Construct(const FArguments& InArgs, const TSharedRef<STableViewBase>& InOwnerTableView);
-
-	virtual TSharedRef<SWidget> GenerateWidgetForColumn(const FName& ColumnName) override;
-
-protected:
-
-	TSharedRef<SWidget> GetRemapMenuContent();
-	FText GetRemapMenuButtonText() const;
-	void OnAssetSelected(const FAssetData& AssetData);
-	bool OnShouldFilterAsset(const FAssetData& AssetData) const;
-
-	TSharedPtr<FDisplayedAssetEntryInfo> DisplayedInfo;
-
-	TWeakObjectPtr<UObject> RemapAsset;
-
-	FString SkeletonExportName;
-};
-
-typedef SListView<TSharedPtr<FDisplayedAssetEntryInfo>> SRemapAssetEntryList;
-
-class SAnimationRemapAssets : public SCompoundWidget
-{
-public:
-
-	SLATE_BEGIN_ARGS(SAnimationRemapAssets)
-		: _NewSkeleton(nullptr)
-		, _RetargetContext(nullptr)
-	{}
-
-		SLATE_ARGUMENT(USkeleton*, NewSkeleton)
-		SLATE_ARGUMENT(FAnimationRetargetContext*, RetargetContext)
-		SLATE_ARGUMENT(TSharedPtr<SWindow>, WidgetWindow)
-
-	SLATE_END_ARGS()
-	
-	void Construct(const FArguments& InArgs);
-
-	static UNREALED_API void ShowWindow(FAnimationRetargetContext& RetargetContext, USkeleton* RetargetToSkeleton);
-	static TSharedPtr<SWindow> DialogWindow;
-
-private:
-
-	TSharedRef<ITableRow> OnGenerateMontageReferenceRow( TSharedPtr<FDisplayedAssetEntryInfo> Item, const TSharedRef<STableViewBase>& OwnerTable );
-
-	void OnDialogClosed(const TSharedRef<SWindow>& Window);
-
-	/** Button Handlers */
-	FReply OnOkClicked();
-	FReply OnBestGuessClicked();
-
-	/** Best guess functions to try and match asset names */
-	const FAssetData* FindBestGuessMatch(const FAssetData& AssetName, const TArray<FAssetData>& PossibleAssets) const;
-
-	/** The retargetting context we're managing*/
-	FAnimationRetargetContext* RetargetContext;
-
-	/** List of entries to the remap list*/
-	TArray<TSharedPtr<FDisplayedAssetEntryInfo>> AssetListInfo;
-
-	/** Skeleton we're about to retarget to*/
-	USkeleton* NewSkeleton;
-
-	/** The list viewing AssetListInfo*/
-	TSharedPtr<SRemapAssetEntryList> ListWidget;
-};
-
-////////////////////////////////////////////////////
-/** 
-* FDlgRemapSkeleton
-* 
-* Wrapper class for SAnimationRemapSkeleton. This class creates and launches a dialog then awaits the
-* result to return to the user. 
-*/
-class UNREALED_API FDlgRemapSkeleton 
-{
-public:
-	FDlgRemapSkeleton( USkeleton* Skeleton );
-
-	/**  Shows the dialog box and waits for the user to respond. */
-	bool ShowModal();
-
-	/** New Skeleton that is chosen **/
-	USkeleton* NewSkeleton;
-
-	/** true if you'd like to retarget skeletal meshes as well **/
-	bool RetargetSkeletalMesh;
-
-private:
-
-	/** Cached pointer to the modal window */
-	TSharedPtr<SWindow> DialogWindow;
-
-	/** Cached pointer to the LOD Chain widget */
-	TSharedPtr<class SAnimationRemapSkeleton> DialogWidget;
-};
-
-
-class SRemapFailures : public SCompoundWidget
-{
-public:
-	SLATE_BEGIN_ARGS(SRemapFailures){}
-		SLATE_ARGUMENT(TArray<FText>, FailedRemaps)
-	SLATE_END_ARGS()
-
-	void Construct( const FArguments& InArgs );
-	static UNREALED_API void OpenRemapFailuresDialog(const TArray<FText>& InFailedRemaps);
-
-private:
-	TSharedRef<ITableRow> MakeListViewWidget(TSharedRef<FText> Item, const TSharedRef<STableViewBase>& OwnerTable);
-	FReply CloseClicked();
-
-private:
-	TArray< TSharedRef<FText> > FailedRemaps;
-};
 
 /////////////////////////////////////////////
 /** 
