- Add a UActorDescContainer::LoadAllActors
- Replace code in FileHelpers with a call to replace the current loading code

diff --git a/Engine/Source/Editor/UnrealEd/Private/FileHelpers.cpp b/Engine/Source/Editor/UnrealEd/Private/FileHelpers.cpp
--- a/Engine/Source/Editor/UnrealEd/Private/FileHelpers.cpp
+++ b/Engine/Source/Editor/UnrealEd/Private/FileHelpers.cpp
@@ -73,9 +73,8 @@
 #include "AnalyticsEventAttribute.h"
 #include "HierarchicalLOD.h"
 #include "WorldPartition/IWorldPartitionEditorModule.h"
+#include "WorldPartition/ActorDescContainer.h"
 #include "WorldPartition/WorldPartition.h"
-#include "WorldPartition/DataLayer/WorldDataLayers.h"
-#include "WorldPartition/DataLayer/DataLayer.h"
 #include "WorldPartition/WorldPartitionEditorPerProjectUserSettings.h"
 #include "PackageSourceControlHelper.h"
 #include "ActorFolder.h"
@@ -740,7 +739,7 @@ static bool SaveWorld(UWorld* World,
 		// Rename the package and the object, as necessary
 		UWorld* DuplicatedWorld = nullptr;
 		UWorldPartition* RenamedWorldPartition = nullptr;
-	
+		TArray<FWorldPartitionReference> ActorReferences;
 
 		// Save Loaded cells
 		TArray<FName> LoadedEditorCells;
@@ -760,22 +759,7 @@ static bool SaveWorld(UWorld* World,
 				if (RenamedWorldPartition)
 				{
 					LoadedEditorCells = RenamedWorldPartition->GetUserLoadedEditorGridCells();
-					if (AWorldDataLayers* WorldDataLayers = World->GetWorldDataLayers())
-					{
-						WorldDataLayers->ForEachDataLayer([](UDataLayer* DataLayer)
-						{
-							if (!DataLayer->IsLoadedInEditor())
-							{
-								DataLayer->SetIsLoadedInEditor(true, /*bFromUserChange*/false);
-							}
-							return true;
-						});
-					}
-					// Make sure AlwaysLoaded DL cells get loaded
-					RenamedWorldPartition->RefreshLoadedEditorCells(false);
-					// Load all the rest
-					const FBox All(FVector(-WORLD_MAX), FVector(WORLD_MAX));
-					RenamedWorldPartition->LoadEditorCells(All, false);
+					RenamedWorldPartition->LoadAllActors(ActorReferences);
 				}
 
 				// If we are doing a SaveAs on a world that already exists on disk, we need to duplicate it:
diff --git a/Engine/Source/Runtime/Engine/Private/WorldPartition/ActorDescContainer.cpp b/Engine/Source/Runtime/Engine/Private/WorldPartition/ActorDescContainer.cpp
--- a/Engine/Source/Runtime/Engine/Private/WorldPartition/ActorDescContainer.cpp
+++ b/Engine/Source/Runtime/Engine/Private/WorldPartition/ActorDescContainer.cpp
@@ -314,6 +314,16 @@ void UActorDescContainer::UnpinActor(const FGuid& ActorGuid)
 	PinnedActorRefs.Remove(ActorGuid);
 }
 
+void UActorDescContainer::LoadAllActors(TArray<FWorldPartitionReference>& OutReferences)
+{
+	OutReferences.Reserve(OutReferences.Num() + GetActorDescCount());
+	for (FActorDescList::TIterator<> ActorDescIterator(this); ActorDescIterator; ++ActorDescIterator)
+	{
+		FWorldPartitionActorDesc* ActorDesc = *ActorDescIterator;
+		OutReferences.Emplace(this, ActorDesc->GetGuid());
+	}
+}
+
 void UActorDescContainer::RegisterEditorDelegates()
 {
 	if (GEditor && !IsTemplate() && World && !World->IsGameWorld())
diff --git a/Engine/Source/Runtime/Engine/Public/WorldPartition/ActorDescContainer.h b/Engine/Source/Runtime/Engine/Public/WorldPartition/ActorDescContainer.h
--- a/Engine/Source/Runtime/Engine/Public/WorldPartition/ActorDescContainer.h
+++ b/Engine/Source/Runtime/Engine/Public/WorldPartition/ActorDescContainer.h
@@ -37,6 +37,8 @@ public:
 	void UnpinActor(const FGuid& ActorGuid);
 	bool IsActorPinned(const FGuid& ActorGuid) const { return PinnedActors.Contains(ActorGuid); }
 
+	void LoadAllActors(TArray<FWorldPartitionReference>& OutReferences);
+
 	/** Instancing support */
 	virtual bool GetInstancingContext(const FLinkerInstancingContext*& OutInstancingContext, FSoftObjectPathFixupArchive*& OutSoftObjectPathFixupArchive) const { return false; }
 
