-Disable static lighting computation on WorldPartition maps (until we make it work)
-Invert order of IsRunningCommandlet and CanAutoApplyLighting checks since CanAutoApplyLighting will error in a CommandLet context

diff --git a/Engine/Source/Editor/UnrealEd/Private/StaticLightingSystem/StaticLightingSystem.cpp b/Engine/Source/Editor/UnrealEd/Private/StaticLightingSystem/StaticLightingSystem.cpp
--- a/Engine/Source/Editor/UnrealEd/Private/StaticLightingSystem/StaticLightingSystem.cpp
+++ b/Engine/Source/Editor/UnrealEd/Private/StaticLightingSystem/StaticLightingSystem.cpp
@@ -291,6 +291,15 @@ void FStaticLightingManager::CreateStaticLightingSystem(const FLightingBuildOpti
 		bBuildReflectionCapturesOnFinish = !Options.bOnlyBuildVisibility;
 
 		UWorld* World = GWorld;
+
+		if (World->IsPartitionedWorld())
+		{
+			// No StaticLighting or Precomputed Visibility on WorldPartition Levels so we can early out here 
+			FStaticLightingManager::Get()->FailLightingBuild(
+				LOCTEXT("LightBuildWPFail", "Static Lighting on World Partition maps is not supported."));
+			return;
+		}
+		
 		for (ULevel* Level : World->GetLevels())
 		{
 			if (Level->bIsLightingScenario && Level->bIsVisible)
@@ -2338,7 +2347,7 @@ void FStaticLightingSystem::UpdateLightingBuild()
 	}
 	else if ( CurrentBuildStage == FStaticLightingSystem::AutoApplyingImport )
 	{
-		if ( CanAutoApplyLighting() || IsRunningCommandlet() )
+		if (IsRunningCommandlet() || CanAutoApplyLighting())
 		{
 			bool bAutoApplyFailed = false;
 			FStaticLightingManager::Get()->SendBuildDoneNotification(bAutoApplyFailed);
