Display all values of containers in Node Watch Bubbles

diff --git a/Engine/Source/Editor/GraphEditor/Private/KismetNodes/SGraphNodeK2Base.cpp b/Engine/Source/Editor/GraphEditor/Private/KismetNodes/SGraphNodeK2Base.cpp
--- a/Engine/Source/Editor/GraphEditor/Private/KismetNodes/SGraphNodeK2Base.cpp
+++ b/Engine/Source/Editor/GraphEditor/Private/KismetNodes/SGraphNodeK2Base.cpp
@@ -608,7 +608,7 @@ void SGraphNodeK2Base::GetNodeInfoPopups(FNodeInfoContext* Context, TArray<FGrap
 								WatchName += TEXT(" ");
 								WatchName += Schema->GetPinDisplayName(WatchPin).ToString();
 
-								WatchText = PinInfo->Value.ToString();
+								WatchText = PinInfo->GetWatchText();
 							}
 							else
 							{
@@ -618,7 +618,7 @@ void SGraphNodeK2Base::GetNodeInfoPopups(FNodeInfoContext* Context, TArray<FGrap
 									WatchName = UEdGraphSchema_K2::TypeToText(PropWatch->Property.Get()).ToString();
 									WatchName += TEXT(" ");
 
-									WatchText = PropWatch->Value.ToString();
+									WatchText = PropWatch->GetWatchText();
 								}
 								else
 								{
diff --git a/Engine/Source/Editor/UnrealEd/Private/Kismet2/KismetDebugUtilities.cpp b/Engine/Source/Editor/UnrealEd/Private/Kismet2/KismetDebugUtilities.cpp
--- a/Engine/Source/Editor/UnrealEd/Private/Kismet2/KismetDebugUtilities.cpp
+++ b/Engine/Source/Editor/UnrealEd/Private/Kismet2/KismetDebugUtilities.cpp
@@ -2273,7 +2273,25 @@ TSharedPtr<FPropertyInstanceInfo> FPropertyInstanceInfo::ResolvePathToProperty(c
 	return ThisDebugInfo;
 }
 
+FString FPropertyInstanceInfo::GetWatchText() const
+{
+	FProperty* Prop = Property.Get();
+	if (Prop)
+	{
+		if (Prop->IsA<FSetProperty>() || Prop->IsA<FArrayProperty>() || Prop->IsA<FMapProperty>())
+		{
+			FString WatchText;
+			for (const TSharedPtr<FPropertyInstanceInfo>& ContainerElement : Children)
+			{
+				WatchText.Append(FText::Format(LOCTEXT("WatchTextFmt", "{0} {1}\n"), ContainerElement->DisplayName, ContainerElement->Value).ToString());
+			}
 
+			return WatchText;
+		}
+	}
+
+	return Value.ToString();
+}
 
 FText FKismetDebugUtilities::GetAndClearLastExceptionMessage()
 {
diff --git a/Engine/Source/Editor/UnrealEd/Public/Kismet2/KismetDebugUtilities.h b/Engine/Source/Editor/UnrealEd/Public/Kismet2/KismetDebugUtilities.h
--- a/Engine/Source/Editor/UnrealEd/Public/Kismet2/KismetDebugUtilities.h
+++ b/Engine/Source/Editor/UnrealEd/Public/Kismet2/KismetDebugUtilities.h
@@ -80,6 +80,9 @@ struct FPropertyInstanceInfo
 
 	/** Resolves the PathToProperty treating this PropertyInstance as the head of the path */
 	UNREALED_API TSharedPtr<FPropertyInstanceInfo> ResolvePathToProperty(const TArray<FName>& InPathToProperty);
+
+	/** Returns the watch text for info popup bubbles on the graph */
+	UNREALED_API FString GetWatchText() const;
 	
 	FText Name;
 	FText DisplayName;
