Split BP project settings out into a separate header

diff --git a/Engine/Source/Developer/AssetTools/Private/AssetRenameManager.cpp b/Engine/Source/Developer/AssetTools/Private/AssetRenameManager.cpp
--- a/Engine/Source/Developer/AssetTools/Private/AssetRenameManager.cpp
+++ b/Engine/Source/Developer/AssetTools/Private/AssetRenameManager.cpp
@@ -44,9 +44,8 @@
 #include "Kismet2/KismetEditorUtilities.h"
 #include "Kismet2/BlueprintEditorUtils.h"
 #include "Misc/RedirectCollector.h"
-#include "Settings/EditorProjectSettings.h"
+#include "Settings/BlueprintEditorProjectSettings.h"
 #include "AssetToolsLog.h"
-#include "Settings/EditorProjectSettings.h"
 #include "Engine/World.h"
 #include "Engine/MapBuildDataRegistry.h"
 #include "GameMapsSettings.h"
diff --git a/Engine/Source/Editor/Kismet/Private/BlueprintCompilationManager.cpp b/Engine/Source/Editor/Kismet/Private/BlueprintCompilationManager.cpp
--- a/Engine/Source/Editor/Kismet/Private/BlueprintCompilationManager.cpp
+++ b/Engine/Source/Editor/Kismet/Private/BlueprintCompilationManager.cpp
@@ -27,7 +27,7 @@
 #include "ProfilingDebugging/ScopedTimers.h"
 #include "Serialization/ArchiveHasReferences.h"
 #include "Serialization/ArchiveReplaceObjectRef.h"
-#include "Settings/EditorProjectSettings.h"
+#include "Settings/BlueprintEditorProjectSettings.h"
 #include "TickableEditorObject.h"
 #include "UObject/MetaData.h"
 #include "UObject/ReferenceChainSearch.h"
diff --git a/Engine/Source/Editor/Kismet/Private/BlueprintDetailsCustomization.cpp b/Engine/Source/Editor/Kismet/Private/BlueprintDetailsCustomization.cpp
--- a/Engine/Source/Editor/Kismet/Private/BlueprintDetailsCustomization.cpp
+++ b/Engine/Source/Editor/Kismet/Private/BlueprintDetailsCustomization.cpp
@@ -91,7 +91,7 @@
 #include "SupportedRangeTypes.h"	// StructsSupportingRangeVisibility
 #include "IPropertyAccessEditor.h"
 #include "AssetRegistryModule.h"
-#include "Settings/EditorProjectSettings.h"
+#include "Settings/BlueprintEditorProjectSettings.h"
 
 #define LOCTEXT_NAMESPACE "BlueprintDetailsCustomization"
 
diff --git a/Engine/Source/Editor/Kismet/Private/BlueprintNamespaceHelper.cpp b/Engine/Source/Editor/Kismet/Private/BlueprintNamespaceHelper.cpp
--- a/Engine/Source/Editor/Kismet/Private/BlueprintNamespaceHelper.cpp
+++ b/Engine/Source/Editor/Kismet/Private/BlueprintNamespaceHelper.cpp
@@ -4,7 +4,7 @@
 #include "Engine/Blueprint.h"
 #include "BlueprintEditor.h"
 #include "BlueprintEditorSettings.h"
-#include "Settings/EditorProjectSettings.h"
+#include "Settings/BlueprintEditorProjectSettings.h"
 #include "Toolkits/ToolkitManager.h"
 #include "EdGraph/EdGraphPin.h"
 #include "EdGraphSchema_K2.h"
diff --git a/Engine/Source/Editor/SubobjectEditor/Private/SSubobjectEditor.cpp b/Engine/Source/Editor/SubobjectEditor/Private/SSubobjectEditor.cpp
--- a/Engine/Source/Editor/SubobjectEditor/Private/SSubobjectEditor.cpp
+++ b/Engine/Source/Editor/SubobjectEditor/Private/SSubobjectEditor.cpp
@@ -55,7 +55,7 @@
 #include "K2Node_Variable.h"		// Nodes for FSubobjectRowDragDropOp
 #include "K2Node_VariableGet.h"
 #include "AssetSelection.h"
-#include "Settings/EditorProjectSettings.h"
+#include "Settings/BlueprintEditorProjectSettings.h"
 #include "Misc/FeedbackContext.h"
 
 extern UNREALED_API UEditorEngine* GEditor;
diff --git a/Engine/Source/Editor/UnrealEd/Private/Kismet2/ChildActorComponentEditorUtils.cpp b/Engine/Source/Editor/UnrealEd/Private/Kismet2/ChildActorComponentEditorUtils.cpp
--- a/Engine/Source/Editor/UnrealEd/Private/Kismet2/ChildActorComponentEditorUtils.cpp
+++ b/Engine/Source/Editor/UnrealEd/Private/Kismet2/ChildActorComponentEditorUtils.cpp
@@ -1,7 +1,7 @@
 // Copyright Epic Games, Inc. All Rights Reserved.
 
 #include "Kismet2/ChildActorComponentEditorUtils.h"
-#include "Settings/EditorProjectSettings.h"
+#include "Settings/BlueprintEditorProjectSettings.h"
 #include "ToolMenus.h"
 #include "SSCSEditor.h"
 #include "SSCSEditorMenuContext.h"
diff --git a/Engine/Source/Editor/UnrealEd/Private/Kismet2/CompilerResultsLog.cpp b/Engine/Source/Editor/UnrealEd/Private/Kismet2/CompilerResultsLog.cpp
--- a/Engine/Source/Editor/UnrealEd/Private/Kismet2/CompilerResultsLog.cpp
+++ b/Engine/Source/Editor/UnrealEd/Private/Kismet2/CompilerResultsLog.cpp
@@ -13,7 +13,7 @@
 #include "EngineLogs.h"
 #include "Engine/Blueprint.h"
 #include "IMessageLogListing.h"
-#include "Settings/EditorProjectSettings.h"
+#include "Settings/BlueprintEditorProjectSettings.h"
 
 #if WITH_EDITOR
 
diff --git a/Engine/Source/Editor/UnrealEd/Private/Settings/BlueprintEditorProjectSettings.cpp b/Engine/Source/Editor/UnrealEd/Private/Settings/BlueprintEditorProjectSettings.cpp
new file mode 100644
--- /dev/null
+++ b/Engine/Source/Editor/UnrealEd/Private/Settings/BlueprintEditorProjectSettings.cpp
@@ -0,0 +1,55 @@
+// Copyright Epic Games, Inc. All Rights Reserved.
+
+#include "Settings/BlueprintEditorProjectSettings.h"
+#include "UObject/UnrealType.h"
+#include "Toolkits/ToolkitManager.h"
+#include "BlueprintEditor.h"
+
+/* UBlueprintEditorProjectSettings */
+
+UBlueprintEditorProjectSettings::UBlueprintEditorProjectSettings(const FObjectInitializer& ObjectInitializer)
+	: Super(ObjectInitializer)
+	, DefaultChildActorTreeViewMode(EChildActorComponentTreeViewVisualizationMode::ComponentOnly)
+{
+}
+
+void UBlueprintEditorProjectSettings::PostEditChangeProperty(struct FPropertyChangedEvent& PropertyChangedEvent)
+{
+	Super::PostEditChangeProperty(PropertyChangedEvent);
+
+	const FName Name = PropertyChangedEvent.Property ? PropertyChangedEvent.Property->GetFName() : NAME_None;
+	if (Name == GET_MEMBER_NAME_CHECKED(UBlueprintEditorProjectSettings, bEnableChildActorExpansionInTreeView))
+	{
+		if (!GEditor)
+		{
+			return;
+		}
+
+		// Find open blueprint editors and refresh them
+		if (UAssetEditorSubsystem* AssetEditorSubsystem = GEditor->GetEditorSubsystem<UAssetEditorSubsystem>())
+		{
+			TArray<UObject*> EditedAssets = AssetEditorSubsystem->GetAllEditedAssets();
+			for (UObject* Asset : EditedAssets)
+			{
+				if (Asset && Asset->IsA<UBlueprint>())
+				{
+					TSharedPtr<IToolkit> AssetEditorPtr = FToolkitManager::Get().FindEditorForAsset(Asset);
+					if (AssetEditorPtr.IsValid() && AssetEditorPtr->IsBlueprintEditor())
+					{
+						TSharedPtr<IBlueprintEditor> BlueprintEditorPtr = StaticCastSharedPtr<IBlueprintEditor>(AssetEditorPtr);
+						BlueprintEditorPtr->RefreshEditors();
+					}
+				}
+			}
+		}
+
+		// Deselect actors so we are forced to clear the current tree view
+		// @todo - Figure out how to update the tree view directly instead?
+		if (GEditor->GetSelectedActorCount() > 0)
+		{
+			const bool bNoteSelectionChange = true;
+			const bool bDeselectBSPSurfaces = true;
+			GEditor->SelectNone(bNoteSelectionChange, bDeselectBSPSurfaces);
+		}
+	}
+}
diff --git a/Engine/Source/Editor/UnrealEd/Private/Settings/EditorProjectSettings.cpp b/Engine/Source/Editor/UnrealEd/Private/Settings/EditorProjectSettings.cpp
--- a/Engine/Source/Editor/UnrealEd/Private/Settings/EditorProjectSettings.cpp
+++ b/Engine/Source/Editor/UnrealEd/Private/Settings/EditorProjectSettings.cpp
@@ -140,52 +140,3 @@ void ULevelEditor2DSettings::PostEditChangeProperty(struct FPropertyChangedEvent
 
 	Super::PostEditChangeProperty(PropertyChangedEvent);
 }
-
-/* UBlueprintEditorProjectSettings */
-
-UBlueprintEditorProjectSettings::UBlueprintEditorProjectSettings(const FObjectInitializer& ObjectInitializer)
-	: Super(ObjectInitializer)
-	, DefaultChildActorTreeViewMode(EChildActorComponentTreeViewVisualizationMode::ComponentOnly)
-{
-}
-
-void UBlueprintEditorProjectSettings::PostEditChangeProperty(struct FPropertyChangedEvent& PropertyChangedEvent)
-{
-	Super::PostEditChangeProperty(PropertyChangedEvent);
-
-	const FName Name = PropertyChangedEvent.Property ? PropertyChangedEvent.Property->GetFName() : NAME_None;
-	if (Name == GET_MEMBER_NAME_CHECKED(UBlueprintEditorProjectSettings, bEnableChildActorExpansionInTreeView))
-	{
-		if (!GEditor)
-		{
-			return;
-		}
-
-		// Find open blueprint editors and refresh them
-		if (UAssetEditorSubsystem* AssetEditorSubsystem = GEditor->GetEditorSubsystem<UAssetEditorSubsystem>())
-		{
-			TArray<UObject*> EditedAssets = AssetEditorSubsystem->GetAllEditedAssets();
-			for (UObject* Asset : EditedAssets)
-			{
-				if (Asset && Asset->IsA<UBlueprint>())
-				{
-					TSharedPtr<IToolkit> AssetEditorPtr = FToolkitManager::Get().FindEditorForAsset(Asset);
-					if (AssetEditorPtr.IsValid() && AssetEditorPtr->IsBlueprintEditor())
-					{
-						TSharedPtr<IBlueprintEditor> BlueprintEditorPtr = StaticCastSharedPtr<IBlueprintEditor>(AssetEditorPtr);
-						BlueprintEditorPtr->RefreshEditors();
-					}
-				}
-			}
-		}
-
-		// Deselect actors so we are forced to clear the current tree view
-		// @todo - Figure out how to update the tree view directly instead?
-		if (GEditor->GetSelectedActorCount() > 0)
-		{
-			const bool bNoteSelectionChange = true;
-			const bool bDeselectBSPSurfaces = true;
-			GEditor->SelectNone(bNoteSelectionChange, bDeselectBSPSurfaces);
-		}
-	}
-}
diff --git a/Engine/Source/Editor/UnrealEd/Public/Settings/BlueprintEditorProjectSettings.h b/Engine/Source/Editor/UnrealEd/Public/Settings/BlueprintEditorProjectSettings.h
new file mode 100644
--- /dev/null
+++ b/Engine/Source/Editor/UnrealEd/Public/Settings/BlueprintEditorProjectSettings.h
@@ -0,0 +1,66 @@
+// Copyright Epic Games, Inc. All Rights Reserved.
+
+#pragma once
+
+#include "CoreMinimal.h"
+#include "UObject/ObjectMacros.h"
+#include "Engine/DeveloperSettings.h"
+#include "Components/ChildActorComponent.h"
+
+#include "BlueprintEditorProjectSettings.generated.h"
+
+
+UCLASS(config=Editor, meta=(DisplayName="Blueprint Project Settings"), defaultconfig)
+class UNREALED_API UBlueprintEditorProjectSettings : public UDeveloperSettings
+{
+	GENERATED_UCLASS_BODY()
+
+public:
+	/**
+	 * Flag to disable faster compiles for individual blueprints if they have no function signature
+	 * changes. This flag is deprecated! In 4.21 there will be no way to force all dependencies to 
+	 * compile when no changes are detected. Report any issues immediately.
+	 */
+	UPROPERTY(EditAnywhere, config, Category=Blueprints, DisplayName = "Force All Dependencies To Recompile (DEPRECATED)")
+	uint8 bForceAllDependenciesToRecompile:1;
+
+	/** If enabled, the editor will load packages to look for soft references to actors when deleting/renaming them. This can be slow in large projects so disable this to improve performance but increase the chance of breaking blueprints/sequences that use soft actor references */
+	UPROPERTY(EditAnywhere, config, Category=Actors)
+	uint8 bValidateUnloadedSoftActorReferences : 1;
+
+	/**
+	 * Enable the option to expand child actor components within component tree views (experimental).
+	 */
+	UPROPERTY(EditAnywhere, config, Category = Experimental)
+	uint8 bEnableChildActorExpansionInTreeView : 1;
+
+	/** 
+	 * List of compiler messages that have been suppressed outside of full, interactive editor sessions for 
+	 * the current project - useful for silencing warnings that were added to the engine after 
+	 * project inception and are going to be addressed as they are found by content authors
+	 */
+	UPROPERTY(EditAnywhere, config, Category= Blueprints, DisplayName = "Compiler Messages Disabled Except in Editor")
+	TArray<FName> DisabledCompilerMessagesExceptEditor;
+	
+	/** 
+	 * List of compiler messages that have been suppressed completely - message suppression is only 
+	 * advisable when using blueprints that you cannot update and are raising innocuous warnings. 
+	 * If useless messages are being raised prefer to contact support rather than disabling messages
+	 */
+	UPROPERTY(EditAnywhere, config, Category= Blueprints, DisplayName = "Compiler Messages Disabled Entirely")
+	TArray<FName> DisabledCompilerMessages;
+
+	// The list of namespaces to always expose in any Blueprint (for all users of the game/project)
+	UPROPERTY(EditAnywhere, config, Category=Experimental)
+	TArray<FString> NamespacesToAlwaysInclude;
+	
+	/**
+	 * Default view mode to use for child actor components in a Blueprint actor's component tree hierarchy (experimental).
+	 */
+	UPROPERTY(EditAnywhere, config, Category=Experimental, meta=(EditCondition="bEnableChildActorExpansionInTreeView"))
+	EChildActorComponentTreeViewVisualizationMode DefaultChildActorTreeViewMode;
+
+	// UObject interface
+	virtual void PostEditChangeProperty(struct FPropertyChangedEvent& PropertyChangedEvent) override;
+	// End of UObject interface
+};
diff --git a/Engine/Source/Editor/UnrealEd/Public/Settings/EditorProjectSettings.h b/Engine/Source/Editor/UnrealEd/Public/Settings/EditorProjectSettings.h
--- a/Engine/Source/Editor/UnrealEd/Public/Settings/EditorProjectSettings.h
+++ b/Engine/Source/Editor/UnrealEd/Public/Settings/EditorProjectSettings.h
@@ -6,7 +6,6 @@
 #include "UObject/ObjectMacros.h"
 #include "Engine/DeveloperSettings.h"
 #include "Math/UnitConversion.h"
-#include "Components/ChildActorComponent.h"
 
 #include "EditorProjectSettings.generated.h"
 
@@ -159,61 +158,6 @@ public:
 };
 
 
-UCLASS(config=Editor, meta=(DisplayName="Blueprint Project Settings"), defaultconfig)
-class UNREALED_API UBlueprintEditorProjectSettings : public UDeveloperSettings
-{
-	GENERATED_UCLASS_BODY()
-
-public:
-	/**
-	 * Flag to disable faster compiles for individual blueprints if they have no function signature
-	 * changes. This flag is deprecated! In 4.21 there will be no way to force all dependencies to 
-	 * compile when no changes are detected. Report any issues immediately.
-	 */
-	UPROPERTY(EditAnywhere, config, Category=Blueprints, DisplayName = "Force All Dependencies To Recompile (DEPRECATED)")
-	uint8 bForceAllDependenciesToRecompile:1;
-
-	/** If enabled, the editor will load packages to look for soft references to actors when deleting/renaming them. This can be slow in large projects so disable this to improve performance but increase the chance of breaking blueprints/sequences that use soft actor references */
-	UPROPERTY(EditAnywhere, config, Category=Actors)
-	uint8 bValidateUnloadedSoftActorReferences : 1;
-
-	/**
-	 * Enable the option to expand child actor components within component tree views (experimental).
-	 */
-	UPROPERTY(EditAnywhere, config, Category = Experimental)
-	uint8 bEnableChildActorExpansionInTreeView : 1;
-
-	/** 
-	 * List of compiler messages that have been suppressed outside of full, interactive editor sessions for 
-	 * the current project - useful for silencing warnings that were added to the engine after 
-	 * project inception and are going to be addressed as they are found by content authors
-	 */
-	UPROPERTY(EditAnywhere, config, Category= Blueprints, DisplayName = "Compiler Messages Disabled Except in Editor")
-	TArray<FName> DisabledCompilerMessagesExceptEditor;
-	
-	/** 
-	 * List of compiler messages that have been suppressed completely - message suppression is only 
-	 * advisable when using blueprints that you cannot update and are raising innocuous warnings. 
-	 * If useless messages are being raised prefer to contact support rather than disabling messages
-	 */
-	UPROPERTY(EditAnywhere, config, Category= Blueprints, DisplayName = "Compiler Messages Disabled Entirely")
-	TArray<FName> DisabledCompilerMessages;
-
-	// The list of namespaces to always expose in any Blueprint (for all users of the game/project)
-	UPROPERTY(EditAnywhere, config, Category=Experimental)
-	TArray<FString> NamespacesToAlwaysInclude;
-	
-	/**
-	 * Default view mode to use for child actor components in a Blueprint actor's component tree hierarchy (experimental).
-	 */
-	UPROPERTY(EditAnywhere, config, Category=Experimental, meta=(EditCondition="bEnableChildActorExpansionInTreeView"))
-	EChildActorComponentTreeViewVisualizationMode DefaultChildActorTreeViewMode;
-
-	// UObject interface
-	virtual void PostEditChangeProperty(struct FPropertyChangedEvent& PropertyChangedEvent) override;
-	// End of UObject interface
-};
-
 UCLASS(config=Editor, meta=(DisplayName="Derived Data"), defaultconfig)
 class UNREALED_API UDDCProjectSettings : public UDeveloperSettings
 {
