More improvements to "Create Animation Blueprint" dialog

Rearranged UI to collapsed infrequently used UI elements (parent class)
Enabled thumbnails/list view for skeleton assets
Skeleton assets now generate a thumbnail (with skeleton debug drawing enabled for vis)

diff --git a/Engine/Config/BaseEditor.ini b/Engine/Config/BaseEditor.ini
--- a/Engine/Config/BaseEditor.ini
+++ b/Engine/Config/BaseEditor.ini
@@ -46,6 +46,7 @@ ColorFontAtlasSize=512
 +RenderableThumbnailTypes=(ClassNeedingThumbnailName="/Script/Engine.PhysicalMaterialMask",RendererClassName="/Script/UnrealEd.PhysicalMaterialMaskThumbnailRenderer")
 +RenderableThumbnailTypes=(ClassNeedingThumbnailName="/Script/Engine.StaticMesh",RendererClassName="/Script/UnrealEd.StaticMeshThumbnailRenderer")
 +RenderableThumbnailTypes=(ClassNeedingThumbnailName="/Script/Engine.SkeletalMesh",RendererClassName="/Script/UnrealEd.SkeletalMeshThumbnailRenderer")
++RenderableThumbnailTypes=(ClassNeedingThumbnailName="/Script/Engine.Skeleton",RendererClassName="/Script/UnrealEd.SkeletonThumbnailRenderer")
 +RenderableThumbnailTypes=(ClassNeedingThumbnailName="/Script/Engine.AnimSequenceBase",RendererClassName="/Script/UnrealEd.AnimSequenceThumbnailRenderer")
 +RenderableThumbnailTypes=(ClassNeedingThumbnailName="/Script/Engine.BlendSpace",RendererClassName="/Script/UnrealEd.BlendSpaceThumbnailRenderer")
 +RenderableThumbnailTypes=(ClassNeedingThumbnailName="/Script/Engine.AnimBlueprint",RendererClassName="/Script/UnrealEd.AnimBlueprintThumbnailRenderer")
diff --git a/Engine/Source/Editor/UnrealEd/Classes/ThumbnailRendering/SkeletonThumbnailRenderer.h b/Engine/Source/Editor/UnrealEd/Classes/ThumbnailRendering/SkeletonThumbnailRenderer.h
new file mode 100644
--- /dev/null
+++ b/Engine/Source/Editor/UnrealEd/Classes/ThumbnailRendering/SkeletonThumbnailRenderer.h
@@ -0,0 +1,38 @@
+// Copyright Epic Games, Inc. All Rights Reserved.
+
+/**
+ *
+ * This thumbnail renderer displays a given skeletal mesh
+ */
+
+#pragma once
+
+#include "CoreMinimal.h"
+#include "ThumbnailHelpers.h"
+#include "UObject/ObjectMacros.h"
+#include "ThumbnailRendering/DefaultSizedThumbnailRenderer.h"
+#include "SkeletonThumbnailRenderer.generated.h"
+
+class FCanvas;
+class FRenderTarget;
+
+UCLASS(config=Editor)
+class UNREALED_API USkeletonThumbnailRenderer : public UDefaultSizedThumbnailRenderer
+{
+	GENERATED_UCLASS_BODY()
+
+
+	// Begin UThumbnailRenderer Object
+	virtual void Draw(UObject* Object, int32 X, int32 Y, uint32 Width, uint32 Height, FRenderTarget*, FCanvas* Canvas, bool bAdditionalViewFamily) override;
+	virtual bool AllowsRealtimeThumbnails(UObject* Object) const override;
+	// End UThumbnailRenderer Object
+
+	// UObject implementation
+	virtual void BeginDestroy() override;
+
+	virtual void AddAdditionalPreviewSceneContent(UObject* Object, UWorld* PreviewWorld) {}
+
+protected:
+	TObjectInstanceThumbnailScene<FSkeletalMeshThumbnailScene, 128> ThumbnailSceneCache;
+};
+
diff --git a/Engine/Source/Editor/UnrealEd/Private/Factories/AnimBlueprintFactory.cpp b/Engine/Source/Editor/UnrealEd/Private/Factories/AnimBlueprintFactory.cpp
--- a/Engine/Source/Editor/UnrealEd/Private/Factories/AnimBlueprintFactory.cpp
+++ b/Engine/Source/Editor/UnrealEd/Private/Factories/AnimBlueprintFactory.cpp
@@ -36,6 +36,7 @@
 #include "ClassViewerFilter.h"
 #include "Animation/AnimLayerInterface.h"
 #include "Widgets/Input/SSegmentedControl.h"
+#include "Widgets/Layout/SExpandableArea.h"
 #include "Widgets/Layout/SWidgetSwitcher.h"
 
 #define LOCTEXT_NAMESPACE "AnimBlueprintFactory"
@@ -71,35 +72,27 @@ public:
 		bOkClicked = false;
 		ParentClass = UAnimInstance::StaticClass();
 
+		const FText TemplateDesc = LOCTEXT("TemplateDesc", "A Template Animation Blueprint has no target skeleton.\nTemplates cannot have direct references to animation assets placed inside of their animation graphs."); 
+		
 		ChildSlot
 		[
 			SNew(SBorder)
 			.Visibility(EVisibility::Visible)
 			.BorderImage(FEditorStyle::GetBrush("ChildWindow.Background"))
 			[
-				SNew(SBox)
-				.Visibility(EVisibility::Visible)
-				.WidthOverride(800.0f)
-				.HeightOverride(500.0f)
+				SNew(SVerticalBox)
+				+SVerticalBox::Slot()
+				.FillHeight(1.0f)
 				[
 					SNew(SVerticalBox)
 					+SVerticalBox::Slot()
-					.FillHeight(1.0f)
+					.AutoHeight()
+					.Padding(10.0f)
 					[
-						SNew(SHorizontalBox)
-						+SHorizontalBox::Slot()
-						.FillWidth(0.5f)
-						.Padding(10.0f)
-						[
-							SNew(SBorder)
-							.BorderImage(FEditorStyle::GetBrush("NewAnimBlueprintDialog.AreaBorder"))
-							[
-								MakeParentClassPicker()
-							]
-						]
-						+SHorizontalBox::Slot()
-						.FillWidth(0.5f)
-						.Padding(10.0f)
+						SNew(SBox)
+						.Visibility(EVisibility::Visible)
+						.WidthOverride(400.0f)
+						.HeightOverride(400.0f)
 						[
 							SNew(SBorder)
 							.BorderImage(FEditorStyle::GetBrush("NewAnimBlueprintDialog.AreaBorder"))
@@ -124,14 +117,16 @@ public:
 									+SSegmentedControl<bool>::Slot(false)
 									[
 										SNew(STextBlock)
-										.Text(LOCTEXT("SpecificSkeleton", "SPECIFIC SKELETON"))
+										.Text(LOCTEXT("SpecificSkeleton", "Specific Skeleton"))
 										.TextStyle( FEditorStyle::Get(), "NormalText" )
+										.ToolTipText(LOCTEXT("SpecoficSkeletonTooltip", "Choose a specific skeleton to bind your new Animation Blueprint to. The Blueprint will be able to use assets that are compatible with this skeleton."))
 									]
 									+SSegmentedControl<bool>::Slot(true)
 									[
 										SNew(STextBlock)
-										.Text(LOCTEXT("Template", "TEMPLATE"))
+										.Text(LOCTEXT("Template", "Template"))
 										.TextStyle(FEditorStyle::Get(), "NormalText")
+										.ToolTipText(TemplateDesc)
 									]
 								]
 								+SVerticalBox::Slot()
@@ -150,47 +145,58 @@ public:
 										SNew(STextBlock)
 										.Justification(ETextJustify::Center)
 										.AutoWrapText(true)
-										.Text(LOCTEXT("TemplateDesc", "A Template Animation Blueprint has no target skeleton.\nTemplates cannot have direct references to animation assets placed inside of their animation graphs."))
+										.Text(TemplateDesc)
 										.TextStyle(FEditorStyle::Get(), "NormalText")
 									]
 								]
 							]
 						]
 					]
-					
-					// Ok/Cancel buttons
 					+SVerticalBox::Slot()
 					.AutoHeight()
-					.HAlign(HAlign_Right)
-					.VAlign(VAlign_Bottom)
-					.Padding(8)
+					.Padding(10.0f, 0.0f, 10.0f, 0.0f)
 					[
-						SNew(SUniformGridPanel)
-						.SlotPadding(FEditorStyle::GetMargin("StandardDialog.SlotPadding"))
-						.MinDesiredSlotWidth(FEditorStyle::GetFloat("StandardDialog.MinDesiredSlotWidth"))
-						.MinDesiredSlotHeight(FEditorStyle::GetFloat("StandardDialog.MinDesiredSlotHeight"))
-						+SUniformGridPanel::Slot(0,0)
-						[
-							SNew(SButton)
-							.IsEnabled_Lambda([this]()
-							{
-								return ParentClass.Get() != nullptr && (bTemplate || TargetSkeleton.IsValid());
-							})
-							.HAlign(HAlign_Center)
-							.ContentPadding( FEditorStyle::GetMargin("StandardDialog.ContentPadding") )
-							.OnClicked(this, &SAnimBlueprintCreateDialog::OkClicked)
-							.Text(LOCTEXT("CreateAnimBlueprintCreate", "Create"))
-						]
-						+SUniformGridPanel::Slot(1,0)
+						SNew(SBorder)
+						.BorderImage(FEditorStyle::GetBrush("NewAnimBlueprintDialog.AreaBorder"))
 						[
-							SNew(SButton)
-							.HAlign(HAlign_Center)
-							.ContentPadding( FEditorStyle::GetMargin("StandardDialog.ContentPadding") )
-							.OnClicked(this, &SAnimBlueprintCreateDialog::CancelClicked)
-							.Text(LOCTEXT("CreateAnimBlueprintCancel", "Cancel"))
+							MakeParentClassPicker()
 						]
 					]
 				]
+
+				// Ok/Cancel buttons
+				+SVerticalBox::Slot()
+				.AutoHeight()
+				.HAlign(HAlign_Right)
+				.VAlign(VAlign_Bottom)
+				.Padding(10.0f)
+				[
+					SNew(SUniformGridPanel)
+					.SlotPadding(FEditorStyle::GetMargin("StandardDialog.SlotPadding"))
+					.MinDesiredSlotWidth(FEditorStyle::GetFloat("StandardDialog.MinDesiredSlotWidth"))
+					.MinDesiredSlotHeight(FEditorStyle::GetFloat("StandardDialog.MinDesiredSlotHeight"))
+					+SUniformGridPanel::Slot(0,0)
+					[
+						SNew(SButton)
+						.ToolTipText(LOCTEXT("CreateAnimBlueprintCreate_Tooltip", "Create a new animation Blueprint.\nSelect a target skeleton or whether the animation Blueprint should be a template.\nOptionally select a parent class."))
+						.IsEnabled_Lambda([this]()
+						{
+							return ParentClass.Get() != nullptr && (bTemplate || TargetSkeleton.IsValid());
+						})
+						.HAlign(HAlign_Center)
+						.ContentPadding( FEditorStyle::GetMargin("StandardDialog.ContentPadding") )
+						.OnClicked(this, &SAnimBlueprintCreateDialog::OkClicked)
+						.Text(LOCTEXT("CreateAnimBlueprintCreate", "Create"))
+					]
+					+SUniformGridPanel::Slot(1,0)
+					[
+						SNew(SButton)
+						.HAlign(HAlign_Center)
+						.ContentPadding( FEditorStyle::GetMargin("StandardDialog.ContentPadding") )
+						.OnClicked(this, &SAnimBlueprintCreateDialog::CancelClicked)
+						.Text(LOCTEXT("CreateAnimBlueprintCancel", "Cancel"))
+					]
+				]
 			]
 		];
 
@@ -204,8 +210,10 @@ public:
 
 		TSharedRef<SWindow> Window = SNew(SWindow)
 		.Title( LOCTEXT("CreateAnimBlueprintOptions", "Create Animation Blueprint") )
-		.ClientSize(FVector2D(800, 500))
-		.SupportsMinimize(false) .SupportsMaximize(false)
+		.ClientSize(FVector2D(500, 500))
+		.SizingRule(ESizingRule::Autosized)
+		.SupportsMinimize(false)
+		.SupportsMaximize(false)
 		[
 			AsShared()
 		];
@@ -259,19 +267,25 @@ private:
 		// All child child classes of UAnimInstance are valid.
 		Filter->AllowedChildrenOfClasses.Add(UAnimInstance::StaticClass());
 
-		return SNew(SVerticalBox)
-			+SVerticalBox::Slot()
-			.AutoHeight()
-			.Padding(5.0f)
-			.HAlign(HAlign_Center)
+		return
+			SNew(SExpandableArea)
+			.ToolTipText(LOCTEXT("ParentClass_Tooltip", "Optionally choose a parent class for your Animation Blueprint"))
+			.Padding(10.0f)
+			.InitiallyCollapsed(true)
+			.HeaderContent()
 			[
-				SNew(STextBlock)
-				.Text(LOCTEXT("ParentClass", "PARENT CLASS"))
-				.TextStyle( FEditorStyle::Get(), "NormalText" )
+				SNew(SBox)
+				.Padding(5.0f)
+				[
+					SNew(STextBlock)
+					.Text_Lambda([this]()
+					{
+						return FText::Format(LOCTEXT("ParentClassFormat", "Parent Class: {0}"), FText::FromString(ParentClass->GetName()));
+					})
+					.TextStyle( FEditorStyle::Get(), "NormalText" )
+				]
 			]
-			+SVerticalBox::Slot()
-			.FillHeight(1.0f)
-			.Padding(5.0f)
+			.BodyContent()
 			[
 				ClassViewerModule.CreateClassViewer(Options, FOnClassPicked::CreateSP(this, &SAnimBlueprintCreateDialog::OnClassPicked))
 			];
@@ -294,7 +308,7 @@ private:
 		AssetPickerConfig.OnAssetSelected = FOnAssetSelected::CreateSP(this, &SAnimBlueprintCreateDialog::OnSkeletonSelected);
 		AssetPickerConfig.OnShouldFilterAsset = FOnShouldFilterAsset::CreateSP(this, &SAnimBlueprintCreateDialog::FilterSkeletonBasedOnParentClass);
 		AssetPickerConfig.bAllowNullSelection = false;
-		AssetPickerConfig.InitialAssetViewType = EAssetViewType::Column;
+		AssetPickerConfig.InitialAssetViewType = EAssetViewType::List;
 		AssetPickerConfig.InitialAssetSelection = TargetSkeleton;
 		AssetPickerConfig.HiddenColumnNames =
 		{
diff --git a/Engine/Source/Editor/UnrealEd/Private/ThumbnailHelpers.cpp b/Engine/Source/Editor/UnrealEd/Private/ThumbnailHelpers.cpp
--- a/Engine/Source/Editor/UnrealEd/Private/ThumbnailHelpers.cpp
+++ b/Engine/Source/Editor/UnrealEd/Private/ThumbnailHelpers.cpp
@@ -425,6 +425,8 @@ void FSkeletalMeshThumbnailScene::SetSkeletalMesh(USkeletalMesh* InSkeletalMesh)
 {
 	PreviewActor->GetSkeletalMeshComponent()->OverrideMaterials.Empty();
 	PreviewActor->GetSkeletalMeshComponent()->SetSkeletalMesh(InSkeletalMesh, false);
+	PreviewActor->GetSkeletalMeshComponent()->SetDrawDebugSkeleton(bDrawDebugSkeleton);
+	PreviewActor->GetSkeletalMeshComponent()->SetDebugDrawColor(DrawDebugColor);
 
 	if ( InSkeletalMesh )
 	{
@@ -440,6 +442,15 @@ void FSkeletalMeshThumbnailScene::SetSkeletalMesh(USkeletalMesh* InSkeletalMesh)
 	}
 }
 
+void FSkeletalMeshThumbnailScene::SetDrawDebugSkeleton(bool bInDrawDebugSkeleton, const FLinearColor& InSkeletonColor)
+{
+ 	bDrawDebugSkeleton = bInDrawDebugSkeleton;
+	DrawDebugColor = InSkeletonColor;
+	PreviewActor->GetSkeletalMeshComponent()->SetDrawDebugSkeleton(bDrawDebugSkeleton);
+	PreviewActor->GetSkeletalMeshComponent()->SetDebugDrawColor(DrawDebugColor);
+	PreviewActor->GetSkeletalMeshComponent()->RecreateRenderState_Concurrent();
+}
+
 void FSkeletalMeshThumbnailScene::GetViewMatrixParameters(const float InFOVDegrees, FVector& OutOrigin, float& OutOrbitPitch, float& OutOrbitYaw, float& OutOrbitZoom) const
 {
 	check(PreviewActor->GetSkeletalMeshComponent());
diff --git a/Engine/Source/Editor/UnrealEd/Private/ThumbnailRendering/SkeletonThumbnailRenderer.cpp b/Engine/Source/Editor/UnrealEd/Private/ThumbnailRendering/SkeletonThumbnailRenderer.cpp
new file mode 100644
--- /dev/null
+++ b/Engine/Source/Editor/UnrealEd/Private/ThumbnailRendering/SkeletonThumbnailRenderer.cpp
@@ -0,0 +1,72 @@
+// Copyright Epic Games, Inc. All Rights Reserved.
+
+#include "ThumbnailRendering/SkeletonThumbnailRenderer.h"
+
+#include "AssetToolsModule.h"
+#include "Misc/App.h"
+#include "ShowFlags.h"
+#include "SceneView.h"
+#include "Animation/Skeleton.h"
+#include "Engine/SkeletalMesh.h"
+#include "ThumbnailHelpers.h"
+
+USkeletonThumbnailRenderer::USkeletonThumbnailRenderer(const FObjectInitializer& ObjectInitializer)
+	: Super(ObjectInitializer)
+{
+}
+
+void USkeletonThumbnailRenderer::Draw(UObject* Object, int32 X, int32 Y, uint32 Width, uint32 Height, FRenderTarget* RenderTarget, FCanvas* Canvas, bool bAdditionalViewFamily)
+{
+	USkeleton* Skeleton = Cast<USkeleton>(Object);
+	TSharedRef<FSkeletalMeshThumbnailScene> ThumbnailScene = ThumbnailSceneCache.EnsureThumbnailScene(Object);
+
+	FAssetToolsModule& AssetToolsModule = FModuleManager::LoadModuleChecked<FAssetToolsModule>(TEXT("AssetTools"));
+	const TSharedPtr<IAssetTypeActions> AssetTypeActions = AssetToolsModule.Get().GetAssetTypeActionsForClass(USkeleton::StaticClass()).Pin();
+	ThumbnailScene->SetDrawDebugSkeleton(true, AssetTypeActions->GetTypeColor());
+	
+	if(Skeleton)
+	{
+		constexpr bool bFindIfNotSet = true; 
+		if(USkeletalMesh* SkeletalMesh = Skeleton->GetPreviewMesh(bFindIfNotSet))
+		{
+			ThumbnailScene->SetSkeletalMesh(SkeletalMesh);
+		}
+	}
+	AddAdditionalPreviewSceneContent(Object, ThumbnailScene->GetWorld());
+
+	FSceneViewFamilyContext ViewFamily( FSceneViewFamily::ConstructionValues( RenderTarget, ThumbnailScene->GetScene(), FEngineShowFlags(ESFIM_Game) )
+		.SetTime(UThumbnailRenderer::GetTime())
+		.SetAdditionalViewFamily(bAdditionalViewFamily));
+
+	ViewFamily.EngineShowFlags.DisableAdvancedFeatures();
+	ViewFamily.EngineShowFlags.MotionBlur = 0;
+	ViewFamily.EngineShowFlags.LOD = 0;
+
+	ThumbnailScene->GetView(&ViewFamily, X, Y, Width, Height);
+	RenderViewFamily(Canvas,&ViewFamily);
+	ThumbnailScene->SetSkeletalMesh(nullptr);
+}
+
+bool USkeletonThumbnailRenderer::AllowsRealtimeThumbnails(UObject* Object) const
+{
+	if (!Super::AllowsRealtimeThumbnails(Object))
+	{
+		return false;
+	}
+	USkeleton* Skeleton = Cast<USkeleton>(Object);
+	if(Skeleton)
+	{
+		if(USkeletalMesh* SkeletalMesh = Skeleton->GetPreviewMesh())
+		{
+			return SkeletalMesh->GetResourceForRendering() != nullptr;
+		}
+	}
+	return false;
+}
+
+void USkeletonThumbnailRenderer::BeginDestroy()
+{
+	ThumbnailSceneCache.Clear();
+
+	Super::BeginDestroy();
+}
diff --git a/Engine/Source/Editor/UnrealEd/Public/ThumbnailHelpers.h b/Engine/Source/Editor/UnrealEd/Public/ThumbnailHelpers.h
--- a/Engine/Source/Editor/UnrealEd/Public/ThumbnailHelpers.h
+++ b/Engine/Source/Editor/UnrealEd/Public/ThumbnailHelpers.h
@@ -98,6 +98,9 @@ public:
 	/** Sets the skeletal mesh to use in the next GetView() */
 	void SetSkeletalMesh(class USkeletalMesh* InSkeletalMesh);
 
+	/** Set whether to draw debug skeleton */
+	void SetDrawDebugSkeleton(bool bInDrawDebugSkeleton, const FLinearColor& InSkeletonColor);
+	
 	/** Returns the preview actor within the scene */
 	class ASkeletalMeshActor* GetPreviewActor() { return PreviewActor; }
 
@@ -108,6 +111,12 @@ protected:
 private:
 	/** The skeletal mesh actor used to display all skeletal mesh thumbnails */
 	class ASkeletalMeshActor* PreviewActor;
+
+	/** The color used to draw the debug skeleton */
+	FLinearColor DrawDebugColor = FLinearColor::White;
+	
+	/** Whether to draw debug skeleton */
+	bool bDrawDebugSkeleton = false;
 };
 
 class UNREALED_API FStaticMeshThumbnailScene : public FThumbnailPreviewScene
