Make menu buttons in the toolbar and viewport toolbar hide their tooltips when opened, like the main menu does.

diff --git a/Engine/Source/Editor/UnrealEd/Private/SEditorViewportToolBarMenu.cpp b/Engine/Source/Editor/UnrealEd/Private/SEditorViewportToolBarMenu.cpp
--- a/Engine/Source/Editor/UnrealEd/Private/SEditorViewportToolBarMenu.cpp
+++ b/Engine/Source/Editor/UnrealEd/Private/SEditorViewportToolBarMenu.cpp
@@ -15,6 +15,8 @@ void SEditorViewportToolbarMenu::Construct( const FArguments& Declaration )
 {
 	const TAttribute<FText>& Label = Declaration._Label;
 
+	SetToolTipText(FText::GetEmpty());
+
 	const FName ImageName = Declaration._Image;
 	const FSlateBrush* ImageBrush = FEditorStyle::GetBrush( ImageName );
 
@@ -91,6 +93,7 @@ void SEditorViewportToolbarMenu::Construct( const FArguments& Declaration )
 		.ContentPadding( FMargin(0.0f, 0.0f) )
 		.VAlign(VAlign_Center)
 		.ButtonStyle(Declaration._MenuStyle)
+		.ToolTipText(this, &SEditorViewportToolbarMenu::GetFilteredToolTipText, Declaration._ToolTipText)
 		.OnClicked( this, &SEditorViewportToolbarMenu::OnMenuClicked )
 		.ForegroundColor(Declaration._ForegroundColor)
 		[
@@ -99,6 +102,19 @@ void SEditorViewportToolbarMenu::Construct( const FArguments& Declaration )
 	);
 }
 
+FText SEditorViewportToolbarMenu::GetFilteredToolTipText(TAttribute<FText> ToolTipText) const
+{
+	// If we're part of a toolbar that has a currently open menu, then we suppress our tool-tip
+	// as it will just get in the way
+	TWeakPtr<SMenuAnchor> OpenedMenu = ParentToolBar.Pin()->GetOpenMenu();
+	if (OpenedMenu.IsValid() && OpenedMenu.Pin()->IsOpen())
+	{
+		return FText::GetEmpty();
+	}
+
+	return ToolTipText.Get();
+}
+
 FReply SEditorViewportToolbarMenu::OnMenuClicked()
 {
 	// If the menu button is clicked toggle the state of the menu anchor which will open or close the menu
diff --git a/Engine/Source/Editor/UnrealEd/Private/STransformViewportToolbar.cpp b/Engine/Source/Editor/UnrealEd/Private/STransformViewportToolbar.cpp
--- a/Engine/Source/Editor/UnrealEd/Private/STransformViewportToolbar.cpp
+++ b/Engine/Source/Editor/UnrealEd/Private/STransformViewportToolbar.cpp
@@ -60,6 +60,7 @@ TSharedRef< SWidget > STransformViewportToolBar::MakeSurfaceSnappingButton()
 		SNew(SEditorViewportToolbarMenu)
 		.ParentToolBar(SharedThis(this))
 		.Image("EditorViewport.ToggleSurfaceSnapping")
+		.ToolTipText(LOCTEXT("SnapToSurfaceMenu_ToolTip", "Control how objects snap to surfaces"))
 		.OnGetMenuContent(this, &STransformViewportToolBar::GenerateSurfaceSnappingMenu)
 		.ForegroundColor(this, &STransformViewportToolBar::GetSurfaceSnappingForegroundColor);
 
diff --git a/Engine/Source/Editor/UnrealEd/Private/SViewportToolBarComboMenu.cpp b/Engine/Source/Editor/UnrealEd/Private/SViewportToolBarComboMenu.cpp
--- a/Engine/Source/Editor/UnrealEd/Private/SViewportToolBarComboMenu.cpp
+++ b/Engine/Source/Editor/UnrealEd/Private/SViewportToolBarComboMenu.cpp
@@ -49,7 +49,7 @@ void SViewportToolBarComboMenu::Construct( const FArguments& InArgs )
 			[
 				SNew(SEditorViewportToolBarMenuButton, MenuAnchor.ToSharedRef())
 				.ButtonStyle(&ButtonStyle)
-				.ToolTipText(InArgs._MenuButtonToolTip)
+				.ToolTipText(this, &SViewportToolBarComboMenu::GetFilteredToolTipText, InArgs._MenuButtonToolTip)
 				.OnClicked(this, &SViewportToolBarComboMenu::OnMenuClicked)
 				.VAlign(VAlign_Center)
 				.HAlign(HAlign_Center)
@@ -80,6 +80,19 @@ void SViewportToolBarComboMenu::Construct( const FArguments& InArgs )
 	];
 }
 
+FText SViewportToolBarComboMenu::GetFilteredToolTipText(TAttribute<FText> ToolTipText) const
+{
+	// If we're part of a toolbar that has a currently open menu, then we suppress our tool-tip
+	// as it will just get in the way
+	TWeakPtr<SMenuAnchor> OpenedMenu = ParentToolBar.Pin()->GetOpenMenu();
+	if (OpenedMenu.IsValid() && OpenedMenu.Pin()->IsOpen())
+	{
+		return FText::GetEmpty();
+	}
+
+	return ToolTipText.Get();
+}
+
 FReply SViewportToolBarComboMenu::OnMenuClicked()
 {
 	// If the menu button is clicked toggle the state of the menu anchor which will open or close the menu
diff --git a/Engine/Source/Editor/UnrealEd/Public/SEditorViewportToolBarMenu.h b/Engine/Source/Editor/UnrealEd/Public/SEditorViewportToolBarMenu.h
--- a/Engine/Source/Editor/UnrealEd/Public/SEditorViewportToolBarMenu.h
+++ b/Engine/Source/Editor/UnrealEd/Public/SEditorViewportToolBarMenu.h
@@ -91,6 +91,16 @@ protected:
 	/** Name of tool menu */
 	FName MenuName;
 
+	/**
+	 * Called to query the tool tip text for this widget, but will return an empty text for menu bar items
+	 * when a menu for that menu bar is already open
+	 *
+	 * @param	ToolTipText	Tool tip text to display, if possible
+	 *
+	 * @return	Tool tip text, or an empty text if filtered out
+	 */
+	FText GetFilteredToolTipText(TAttribute<FText> ToolTipText) const;
+
 private:
 	/** Our menus anchor */
 	TSharedPtr<SMenuAnchor> MenuAnchor;
diff --git a/Engine/Source/Editor/UnrealEd/Public/SViewportToolBarComboMenu.h b/Engine/Source/Editor/UnrealEd/Public/SViewportToolBarComboMenu.h
--- a/Engine/Source/Editor/UnrealEd/Public/SViewportToolBarComboMenu.h
+++ b/Engine/Source/Editor/UnrealEd/Public/SViewportToolBarComboMenu.h
@@ -66,7 +66,18 @@ public:
 	 * @param	InArgs	The declaration data for this widget
 	 */
 	void Construct( const FArguments& InArgs );
-	
+
+protected:
+	/**
+	 * Called to query the tool tip text for this widget, but will return an empty text for toolbar items
+	 * when a menu for that toolbar is already open
+	 *
+	 * @param	ToolTipText	Tool tip text to display, if possible
+	 *
+	 * @return	Tool tip text, or an empty text if filtered out
+	 */
+	FText GetFilteredToolTipText(TAttribute<FText> ToolTipText) const;
+
 private:
 	/**
 	 * Called when the menu button is clicked.  Will toggle the visibility of the menu content                   
diff --git a/Engine/Source/Runtime/Slate/Private/Widgets/Input/SComboButton.cpp b/Engine/Source/Runtime/Slate/Private/Widgets/Input/SComboButton.cpp
--- a/Engine/Source/Runtime/Slate/Private/Widgets/Input/SComboButton.cpp
+++ b/Engine/Source/Runtime/Slate/Private/Widgets/Input/SComboButton.cpp
@@ -36,6 +36,7 @@ void SComboButton::Construct( const FArguments& InArgs )
 			.ButtonStyle( OurButtonStyle )
 			.ClickMethod( EButtonClickMethod::MouseDown )
 			.OnClicked( this, &SComboButton::OnButtonClicked )
+			.ToolTipText( this, &SComboButton::GetFilteredToolTipText, InArgs._ToolTipText)
 			.ContentPadding( InArgs._ContentPadding.IsSet() ? InArgs._ContentPadding : InArgs._ComboButtonStyle->ContentPadding )
 			.ForegroundColor( InArgs._ForegroundColor )
 			.ButtonColorAndOpacity( InArgs._ButtonColorAndOpacity )
@@ -95,6 +96,16 @@ void SComboButton::Construct( const FArguments& InArgs )
 	SetMenuContent( InArgs._MenuContent.Widget );
 }
 
+FText SComboButton::GetFilteredToolTipText(TAttribute<FText> ToolTipText) const
+{
+	if (IsOpen())
+	{
+		return FText::GetEmpty();
+	}
+
+	return ToolTipText.Get();
+}
+
 FReply SComboButton::OnButtonClicked()
 {
 	// Button was clicked; show the popup.
diff --git a/Engine/Source/Runtime/Slate/Public/Widgets/Input/SComboButton.h b/Engine/Source/Runtime/Slate/Public/Widgets/Input/SComboButton.h
--- a/Engine/Source/Runtime/Slate/Public/Widgets/Input/SComboButton.h
+++ b/Engine/Source/Runtime/Slate/Public/Widgets/Input/SComboButton.h
@@ -97,6 +97,15 @@ protected:
 	virtual FReply OnButtonClicked();
 	virtual FReply OnKeyDown(const FGeometry& MyGeometry, const FKeyEvent& InKeyEvent) override;
 
+	/**
+	 * Called to query the tool tip text for this widget, but will return an empty text when the menu is already open
+	 *
+	 * @param	ToolTipText	Tool tip text to display, if possible
+	 *
+	 * @return	Tool tip text, or an empty text if filtered out
+	 */
+	FText GetFilteredToolTipText(TAttribute<FText> ToolTipText) const;
+
 protected:
 	/** Area where the button's content resides */
 	SHorizontalBox::FSlot* ButtonContentSlot;
