Fix crash when renaming level

diff --git a/Engine/Source/Editor/UnrealEd/Private/EditorEngine.cpp b/Engine/Source/Editor/UnrealEd/Private/EditorEngine.cpp
--- a/Engine/Source/Editor/UnrealEd/Private/EditorEngine.cpp
+++ b/Engine/Source/Editor/UnrealEd/Private/EditorEngine.cpp
@@ -4351,7 +4351,7 @@ FSavePackageResultStruct UEditorEngine::Save(UPackage* InOuter, UObject* InAsset
 	SlowTask.EnterProgressFrame(10);
 
 	const bool bAutosave = (SaveArgs.SaveFlags & SAVE_FromAutosave) != 0;
-	if (!bSavingConcurrent && !IsRunningCommandlet() && !bAutosave && (World || Asset->IsA<AActor>() || Asset->IsA<UActorFolder>()))
+	if (!bSavingConcurrent && !IsRunningCommandlet() && !bAutosave && Asset && (World || Asset->IsA<AActor>() || Asset->IsA<UActorFolder>()))
 	{
 		// Always reset the transaction buffer on level/actor save to avoid problems with deleted actors (marked pending kill) that gets marked transient by the saving code
 		ResetTransaction( World ? NSLOCTEXT("UnrealEd", "MapSaved", "Map Saved") : Asset->IsA<AActor>() ? NSLOCTEXT("UnrealEd", "ActorSaved", "Actor Saved") : NSLOCTEXT("UnrealEd", "ActorFolderSaved", "Actor Folder Saved"));
