Rename function on mode manager to prevent naming confusion with deactivate/delete.

diff --git a/Engine/Source/Editor/UnrealEd/Private/EditorModeManager.cpp b/Engine/Source/Editor/UnrealEd/Private/EditorModeManager.cpp
--- a/Engine/Source/Editor/UnrealEd/Private/EditorModeManager.cpp
+++ b/Engine/Source/Editor/UnrealEd/Private/EditorModeManager.cpp
@@ -105,7 +105,7 @@ FEditorModeTools::~FEditorModeTools()
 
 	RemoveAllDelegateHandlers();
 
-	DeactivateAllModesPendingDeletion();
+	ExitAllModesPendingDeactivate();
 	RecycledScriptableModes.Empty();
 
 	// We may be destroyed after the UObject system has already shutdown, 
@@ -613,7 +613,7 @@ void FEditorModeTools::ForEachEdMode(TFunctionRef<bool(UEdMode*)> InCalllback) c
 	}
 }
 
-void FEditorModeTools::DeactivateAllModesPendingDeletion()
+void FEditorModeTools::ExitAllModesPendingDeactivate()
 {
 	// Make a copy so we can modify the pending deactivate modes map
 	TMap<FEditorModeID, UEdMode*> PendingDeactivateModesCopy(PendingDeactivateModes);
@@ -699,7 +699,7 @@ void FEditorModeTools::OnWorldCleanup(UWorld* InWorld, bool bSessionEnded, bool
 	UWorld* World = GetWorld();
 	if (InWorld == World)
 	{
-		DeactivateAllModesPendingDeletion();
+		ExitAllModesPendingDeactivate();
 	}
 }
 
@@ -1195,7 +1195,7 @@ bool FEditorModeTools::ShouldDrawBrushVertices() const
 void FEditorModeTools::Tick( FEditorViewportClient* ViewportClient, float DeltaTime )
 {	
 	// Remove anything pending destruction
-	DeactivateAllModesPendingDeletion();
+	ExitAllModesPendingDeactivate();
 
 	if (ActiveScriptableModes.Num() == 0)
 	{
diff --git a/Engine/Source/Editor/UnrealEd/Public/EditorModeManager.h b/Engine/Source/Editor/UnrealEd/Public/EditorModeManager.h
--- a/Engine/Source/Editor/UnrealEd/Public/EditorModeManager.h
+++ b/Engine/Source/Editor/UnrealEd/Public/EditorModeManager.h
@@ -555,7 +555,7 @@ protected:
 		});
 	}
 
-	void DeactivateAllModesPendingDeletion();
+	void ExitAllModesPendingDeactivate();
 
 	/** List of default modes for this tool.  These must all be compatible with each other. */
 	TArray<FEditorModeID> DefaultModeIDs;
