Linux is assumed to always have a Compiler *but* it should not be assumed to always have an IDE.

This function complicate things by assumed having an IDE == having a compiler which is incorrect for Linux.

Assume we can always build for Linux

diff --git a/Engine/Source/Editor/UnrealEd/Private/Settings/SettingsClasses.cpp b/Engine/Source/Editor/UnrealEd/Private/Settings/SettingsClasses.cpp
--- a/Engine/Source/Editor/UnrealEd/Private/Settings/SettingsClasses.cpp
+++ b/Engine/Source/Editor/UnrealEd/Private/Settings/SettingsClasses.cpp
@@ -532,7 +532,7 @@ void FPlayScreenResolution::PostInitProperties()
 
 void ULevelEditorPlaySettings::PostEditChangeProperty(struct FPropertyChangedEvent& PropertyChangedEvent)
 {
-	if (BuildGameBeforeLaunch != EPlayOnBuildMode::PlayOnBuild_Always && !FSourceCodeNavigation::IsCompilerAvailable())
+	if (BuildGameBeforeLaunch != EPlayOnBuildMode::PlayOnBuild_Always && !FSourceCodeNavigation::IsCompilerAvailable() && !PLATFORM_LINUX)
 	{
 		BuildGameBeforeLaunch = EPlayOnBuildMode::PlayOnBuild_Never;
 	}
