Updated UActorFactory to honor AActor::GetDefaultActorLabel when placing class/BP assets

diff --git a/Engine/Source/Editor/UnrealEd/Classes/ActorFactories/ActorFactory.h b/Engine/Source/Editor/UnrealEd/Classes/ActorFactories/ActorFactory.h
--- a/Engine/Source/Editor/UnrealEd/Classes/ActorFactories/ActorFactory.h
+++ b/Engine/Source/Editor/UnrealEd/Classes/ActorFactories/ActorFactory.h
@@ -113,6 +113,12 @@ protected:
 	    update the blueprint's CDO properties with state from the asset for this factory.
 		IMPORTANT: If you override this, you should usually also override PostSpawnActor()! */
 	virtual void PostCreateBlueprint( UObject* Asset, AActor* CDO );
+
+	/**
+	 * Get the default label that should be used for the actor spawned by the given asset (does not include any numeric suffix).
+	 * For classes or BPs that inherit from AActor this will defer to AActor::GetDefaultActorLabel, and for everything else it will use the asset name.
+	 */
+	virtual FString GetDefaultActorLabel(UObject* Asset) const;
 };
 
 extern UNREALED_API FQuat FindActorAlignmentRotation(const FQuat& InActorRotation, const FVector& InModelAxis, const FVector& InWorldNormal, FQuat* OutDeltaRotation = nullptr);
\ No newline at end of file
diff --git a/Engine/Source/Editor/UnrealEd/Private/Factories/ActorFactory.cpp b/Engine/Source/Editor/UnrealEd/Private/Factories/ActorFactory.cpp
--- a/Engine/Source/Editor/UnrealEd/Private/Factories/ActorFactory.cpp
+++ b/Engine/Source/Editor/UnrealEd/Private/Factories/ActorFactory.cpp
@@ -325,12 +325,13 @@ void UActorFactory::PostPlaceAsset(TArrayView<const FTypedElementHandle> InEleme
 		{
 			if (AActor* CreatedActor = ObjectInterface.GetObjectAs<AActor>())
 			{
-				PostSpawnActor(InPlacementInfo.AssetToPlace.GetAsset(), CreatedActor);
+				UObject* Asset = InPlacementInfo.AssetToPlace.GetAsset();
+				PostSpawnActor(Asset, CreatedActor);
 
 				// Only do this if the actor wasn't already given a name
 				if (InPlacementInfo.NameOverride.IsNone())
 				{
-					FActorLabelUtilities::SetActorLabelUnique(CreatedActor, InPlacementInfo.AssetToPlace.AssetName.ToString());
+					FActorLabelUtilities::SetActorLabelUnique(CreatedActor, GetDefaultActorLabel(Asset));
 				}
 
 				CreatedActor->PostEditChange();
@@ -410,7 +411,7 @@ AActor* UActorFactory::CreateActor(UObject* InAsset, ULevel* InLevel, const FTra
 			// Only do this if the actor wasn't already given a name
 			if ((InAsset != nullptr) && (InSpawnParams.Name == NAME_None))
 			{
-				FActorLabelUtilities::SetActorLabelUnique(NewActor, InAsset->GetName());
+				FActorLabelUtilities::SetActorLabelUnique(NewActor, GetDefaultActorLabel(InAsset));
 			}
 
 			PostSpawnActor(InAsset, NewActor);
@@ -502,6 +503,26 @@ void UActorFactory::PostCreateBlueprint( UObject* Asset, AActor* CDO )
 	// Override this in derived actor factories to initialize the blueprint's CDO based on the asset assigned to the factory!
 }
 
+FString UActorFactory::GetDefaultActorLabel(UObject* Asset) const
+{
+	UClass* PotentialActorClass = nullptr;
+	if (UBlueprint* Blueprint = Cast<UBlueprint>(Asset))
+	{
+		PotentialActorClass = Blueprint->GeneratedClass;
+	}
+	else
+	{
+		PotentialActorClass = Cast<UClass>(Asset);
+	}
+
+	if (PotentialActorClass && PotentialActorClass->IsChildOf<AActor>())
+	{
+		return GetDefault<AActor>(PotentialActorClass)->GetDefaultActorLabel();
+	}
+
+	return Asset->GetName();
+}
+
 
 /*-----------------------------------------------------------------------------
 UActorFactoryStaticMesh
