Added support for importing animation transform attributes from FBX locators

diff --git a/Engine/Source/Editor/UnrealEd/Private/Fbx/FbxSkeletalMeshImport.cpp b/Engine/Source/Editor/UnrealEd/Private/Fbx/FbxSkeletalMeshImport.cpp
--- a/Engine/Source/Editor/UnrealEd/Private/Fbx/FbxSkeletalMeshImport.cpp
+++ b/Engine/Source/Editor/UnrealEd/Private/Fbx/FbxSkeletalMeshImport.cpp
@@ -9,6 +9,7 @@
 #include "EngineDefines.h"
 #include "Misc/MessageDialog.h"
 #include "Containers/IndirectArray.h"
+#include "Containers/UnrealString.h"
 #include "Stats/Stats.h"
 #include "Async/AsyncWork.h"
 #include "Misc/ConfigCacheIni.h"
@@ -25,6 +26,7 @@
 #include "AnimEncoding.h"
 #include "Factories/Factory.h"
 #include "Factories/FbxSkeletalMeshImportData.h"
+#include "Animation/AnimationSettings.h"
 #include "Animation/MorphTarget.h"
 #include "PhysicsAssetUtils.h"
 #include "Rendering/SkeletalMeshModel.h"
@@ -761,13 +763,35 @@ bool UnFbx::FFbxImporter::IsUnrealBone(FbxNode* Link)
 			AttrType == FbxNodeAttribute::eMesh ||
 			AttrType == FbxNodeAttribute::eNull )
 		{
-			return true;
+			return !IsUnrealTransformAttribute(Link);
 		}
 	}
 	
 	return false;
 }
 
+bool UnFbx::FFbxImporter::IsUnrealTransformAttribute(FbxNode* Link)
+{
+	FbxNodeAttribute* Attr = Link->GetNodeAttribute();
+	if (Attr)
+	{
+		FbxNodeAttribute::EType AttrType = Attr->GetAttributeType();
+		if (AttrType == FbxNodeAttribute::eNull)
+		{
+			FString AttributeName = FSkeletalMeshImportData::FixupBoneName(MakeName(Link->GetName()));
+
+			for (const FString& TransformAttributeName : UAnimationSettings::Get()->TransformAttributeNames)
+			{
+				if (AttributeName.MatchesWildcard(TransformAttributeName, ESearchCase::IgnoreCase))
+				{
+					return true;
+				}
+			}
+		}
+	}
+
+	return false;
+}
 
 void UnFbx::FFbxImporter::RecursiveBuildSkeleton(FbxNode* Link, TArray<FbxNode*>& OutSortedLinks)
 {
@@ -780,6 +804,10 @@ void UnFbx::FFbxImporter::RecursiveBuildSkeleton(FbxNode* Link, TArray<FbxNode*>
 			RecursiveBuildSkeleton(Link->GetChild(ChildIndex),OutSortedLinks);
 		}
 	}
+	else if (IsUnrealTransformAttribute(Link))
+	{
+		OutSortedLinks.Add(Link);
+	}
 }
 
 void UnFbx::FFbxImporter::BuildSkeletonSystem(TArray<FbxCluster*>& ClusterArray, TArray<FbxNode*>& OutSortedLinks)
diff --git a/Engine/Source/Editor/UnrealEd/Private/SkeletalMeshEdit.cpp b/Engine/Source/Editor/UnrealEd/Private/SkeletalMeshEdit.cpp
--- a/Engine/Source/Editor/UnrealEd/Private/SkeletalMeshEdit.cpp
+++ b/Engine/Source/Editor/UnrealEd/Private/SkeletalMeshEdit.cpp
@@ -466,7 +466,7 @@ void UnFbx::FFbxImporter::FillAndVerifyBoneNames(USkeleton* Skeleton, TArray<Fbx
 	for (int32 I = 0; I < TrackNum; ++I)
 	{
 		FName RawBoneName = OutRawBoneNames[I];
-		if ( RefSkeleton.FindBoneIndex(RawBoneName) == INDEX_NONE)
+		if (RefSkeleton.FindBoneIndex(RawBoneName) == INDEX_NONE && !IsUnrealTransformAttribute(SortedLinks[I]))
 		{
 			BoneNames += RawBoneName.ToString();
 			BoneNames += TEXT("  \n");
@@ -1916,7 +1916,7 @@ void UnFbx::FFbxImporter::ImportBoneTracks(USkeleton* Skeleton, FAnimCurveImport
 		const FText StatusUpate = FText::Format(LOCTEXT("ImportingAnimTrackDetail", "Importing Animation Track [{TrackName}] ({TrackIndex}/{TotalTracks}) - TotalKey {TotalKey}"), Args);
 		SlowTask.EnterProgressFrame(1, StatusUpate);
 
-		if (BoneTreeIndex != INDEX_NONE)
+		if (BoneTreeIndex != INDEX_NONE || IsUnrealTransformAttribute(AnimImportSettings.SortedLinks[SourceTrackIdx]))
 		{
 			bool bSuccess = true;
 
@@ -1925,6 +1925,8 @@ void UnFbx::FFbxImporter::ImportBoneTracks(USkeleton* Skeleton, FAnimCurveImport
 			RawTrack.RotKeys.Empty();
 			RawTrack.ScaleKeys.Empty();
 
+			TArray<float> TimeKeys;
+
 			AnimationTransformDebug::FAnimationTransformDebugData NewDebugData;
 
 			FbxNode* Link = AnimImportSettings.SortedLinks[SourceTrackIdx];
@@ -2009,21 +2011,89 @@ void UnFbx::FFbxImporter::ImportBoneTracks(USkeleton* Skeleton, FAnimCurveImport
 				RawTrack.PosKeys.Add(FVector3f(LocalTransform.GetTranslation()));
 				RawTrack.RotKeys.Add(FQuat4f(LocalTransform.GetRotation()));
 
+				TimeKeys.Add((CurTime - AnimTimeSpan.GetStart()).GetSecondDouble());
+
 				NewDebugData.RecalculatedLocalTransform.Add(LocalTransform);
 				++NumKeysForTrack;
 			}
 
 			if (bSuccess)
 			{
-				//add new track
-				Controller.AddBoneTrack(BoneName);
-				Controller.SetBoneTrackKeys(BoneName, RawTrack.PosKeys, RawTrack.RotKeys, RawTrack.ScaleKeys);
-				const int32 TrackIndex = DestSeq->GetDataModel()->GetBoneTrackIndexByName(BoneName);
+				check(RawTrack.ScaleKeys.Num() == NumKeysForTrack);
+				check(RawTrack.PosKeys.Num() == NumKeysForTrack);
+				check(RawTrack.RotKeys.Num() == NumKeysForTrack);
+				check(TimeKeys.Num() == NumKeysForTrack);
+
+				if (BoneTreeIndex != INDEX_NONE)
+				{
+					//add new track
+					Controller.AddBoneTrack(BoneName);
+					Controller.SetBoneTrackKeys(BoneName, RawTrack.PosKeys, RawTrack.RotKeys, RawTrack.ScaleKeys);
+					const int32 TrackIndex = DestSeq->GetDataModel()->GetBoneTrackIndexByName(BoneName);
+
+					NewDebugData.SetTrackData(TrackIndex, BoneTreeIndex, BoneName);
 
-				NewDebugData.SetTrackData(TrackIndex, BoneTreeIndex, BoneName);
+					// add mapping to skeleton bone track
+					TransformDebugData.Add(NewDebugData);
+				}
+				else if (NumKeysForTrack > 0) // add transform attribute
+				{
+					FbxNode* TargetBoneLink = LinkParent;
+					while (TargetBoneLink != nullptr && !IsUnrealBone(TargetBoneLink))
+					{
+						TargetBoneLink = TargetBoneLink->GetParent();
+					}
 
-				// add mapping to skeleton bone track
-				TransformDebugData.Add(NewDebugData);
+					if (TargetBoneLink)
+					{
+						int32 TargetBoneTrackIndex = AnimImportSettings.SortedLinks.Find(TargetBoneLink);
+						if (TargetBoneTrackIndex != INDEX_NONE)
+						{
+							FName TargetBoneName = AnimImportSettings.FbxRawBoneNames[TargetBoneTrackIndex];
+							if (RefSkeleton.FindBoneIndex(TargetBoneName) != INDEX_NONE)
+							{
+								FAnimationAttributeIdentifier AttributeIdentifier = UAnimationAttributeIdentifierExtensions::CreateAttributeIdentifier(DestSeq, FName(BoneName), TargetBoneName, FTransformAnimationAttribute::StaticStruct());
+								if (AttributeIdentifier.IsValid())
+								{
+									// remove any existing attribute with the same identifier
+									if (const UAnimDataModel* Model = Controller.GetModel())
+									{
+										if (Model->FindAttribute(AttributeIdentifier))
+										{
+											Controller.RemoveAttribute(AttributeIdentifier);
+										}
+									}
+
+									// pack the separate rot/pos/scale key arrays into a single array
+									TArray<FTransform> TransformValues;
+									TransformValues.Reserve(NumKeysForTrack);
+									for (int32 KeyIndex = 0; KeyIndex < NumKeysForTrack; ++KeyIndex)
+									{
+										const FQuat Q(RawTrack.RotKeys[KeyIndex]);
+										const FVector T(RawTrack.PosKeys[KeyIndex]);
+										const FVector S(RawTrack.ScaleKeys[KeyIndex]);
+										TransformValues.Add(FTransform(Q, T, S));
+									}
+
+									// reduce keys for the common case where all of the keys have the same values
+									bool bReduceKeys = true;
+									for (int32 KeyIndex = 1; KeyIndex < NumKeysForTrack; ++KeyIndex)
+									{
+										if (!TransformValues[KeyIndex].Equals(TransformValues[0]))
+										{
+											bReduceKeys = false;
+											break;
+										}
+									}
+
+									// create the attribute and add the transform keys
+									const int32 NumAttributeKeys = (bReduceKeys) ? 1 : NumKeysForTrack;
+									UE::Anim::AddTypedCustomAttribute<FTransformAnimationAttribute, FTransform>(FName(BoneName), TargetBoneName, DestSeq, MakeArrayView(TimeKeys.GetData(), NumAttributeKeys), MakeArrayView(TransformValues.GetData(), NumAttributeKeys));
+								}
+							}
+						}
+					}
+				}
 			}
 		}
 
diff --git a/Engine/Source/Editor/UnrealEd/Public/FbxImporter.h b/Engine/Source/Editor/UnrealEd/Public/FbxImporter.h
--- a/Engine/Source/Editor/UnrealEd/Public/FbxImporter.h
+++ b/Engine/Source/Editor/UnrealEd/Public/FbxImporter.h
@@ -1029,6 +1029,13 @@ public:
 	 */
 	bool IsUnrealBone(FbxNode* Link);
 
+	/**
+	 * Returns if the passed FbxNode can be used as a transform attribute in Unreal.
+	 * 
+	 * @return bool
+	 */
+	bool IsUnrealTransformAttribute(FbxNode* Link);
+
 	/**
 	* Fill FBX skeletons to OutSortedLinks recursively
 	*
diff --git a/Engine/Source/Runtime/Engine/Classes/Animation/AnimationSettings.h b/Engine/Source/Runtime/Engine/Classes/Animation/AnimationSettings.h
--- a/Engine/Source/Runtime/Engine/Classes/Animation/AnimationSettings.h
+++ b/Engine/Source/Runtime/Engine/Classes/Animation/AnimationSettings.h
@@ -87,6 +87,10 @@ class ENGINE_API UAnimationSettings : public UDeveloperSettings
 	UPROPERTY(config, EditAnywhere, Category = CustomAttributes)
 	ECustomAttributeBlendType DefaultAttributeBlendMode;
 
+	/** Names to match against when importing FBX node transform curves as attributes (can use ? and * wildcards) */
+	UPROPERTY(config, EditAnywhere, Category = CustomAttributes)
+	TArray<FString> TransformAttributeNames;
+
 	/** Find and Replace Expressions used for mirroring  */
 	UPROPERTY(config, EditAnywhere, Category = Mirroring)
 	TArray<FMirrorFindReplaceExpression> MirrorFindReplaceExpressions;
