Run Enterprise UnrealEditor Python Tests - "from_dataprep_to_unreal" crashes due to CreateSwapChainResult failing with error DXGI_ERROR_NOT_CURRENTLY_AVAILABLE

diff --git a/Engine/Plugins/Enterprise/DatasmithImporter/Source/DatasmithImporter/Private/DatasmithConsumer.cpp b/Engine/Plugins/Enterprise/DatasmithImporter/Source/DatasmithImporter/Private/DatasmithConsumer.cpp
--- a/Engine/Plugins/Enterprise/DatasmithImporter/Source/DatasmithImporter/Private/DatasmithConsumer.cpp
+++ b/Engine/Plugins/Enterprise/DatasmithImporter/Source/DatasmithImporter/Private/DatasmithConsumer.cpp
@@ -1069,7 +1069,7 @@ ULevel* UDatasmithConsumer::FindOrAddLevel(const FString& InLevelName)
 	// If it does not exist, create a new one and add it to the working world
 	else
 	{
-		StreamingLevel = EditorLevelUtils::CreateNewStreamingLevelForWorld( *WorkingWorld, ULevelStreamingAlwaysLoaded::StaticClass(), *PackageFilename );
+		StreamingLevel = EditorLevelUtils::CreateNewStreamingLevelForWorld( *WorkingWorld, ULevelStreamingAlwaysLoaded::StaticClass(), *PackageFilename, false, nullptr, false );
 
 		if( StreamingLevel )
 		{
diff --git a/Engine/Source/Editor/UnrealEd/Private/EditorLevelUtils.cpp b/Engine/Source/Editor/UnrealEd/Private/EditorLevelUtils.cpp
--- a/Engine/Source/Editor/UnrealEd/Private/EditorLevelUtils.cpp
+++ b/Engine/Source/Editor/UnrealEd/Private/EditorLevelUtils.cpp
@@ -686,7 +686,7 @@ UWorld* GetWorldForNewStreamingLevel(UWorld* InTemplateWorld)
 }
 };
 
-ULevelStreaming* UEditorLevelUtils::CreateNewStreamingLevelForWorld(UWorld& InWorld, TSubclassOf<ULevelStreaming> LevelStreamingClass, const FString& DefaultFilename /* = TEXT( "" ) */, bool bMoveSelectedActorsIntoNewLevel /* = false */, UWorld* InTemplateWorld /* = nullptr */)
+ULevelStreaming* UEditorLevelUtils::CreateNewStreamingLevelForWorld(UWorld& InWorld, TSubclassOf<ULevelStreaming> LevelStreamingClass, const FString& DefaultFilename /* = TEXT( "" ) */, bool bMoveSelectedActorsIntoNewLevel /* = false */, UWorld* InTemplateWorld /* = nullptr */, bool bInUseSaveAs /*= true*/)
 {
 	TArray<AActor*> ActorsToMove;
 	if (bMoveSelectedActorsIntoNewLevel)
@@ -701,10 +701,10 @@ ULevelStreaming* UEditorLevelUtils::CreateNewStreamingLevelForWorld(UWorld& InWo
 		}
 	}
 
-	return CreateNewStreamingLevelForWorld(InWorld, LevelStreamingClass, /*bUseExternalActors=*/false, DefaultFilename, &ActorsToMove, InTemplateWorld);
+	return CreateNewStreamingLevelForWorld(InWorld, LevelStreamingClass, /*bUseExternalActors=*/false, DefaultFilename, &ActorsToMove, InTemplateWorld, bInUseSaveAs);
 }
 
-ULevelStreaming* UEditorLevelUtils::CreateNewStreamingLevelForWorld(UWorld& InWorld, TSubclassOf<ULevelStreaming> LevelStreamingClass, bool bUseExternalActors, const FString& DefaultFilename /* = TEXT("") */, const TArray<AActor*>* ActorsToMove /* = nullptr */, UWorld* InTemplateWorld /* = nullptr */)
+ULevelStreaming* UEditorLevelUtils::CreateNewStreamingLevelForWorld(UWorld& InWorld, TSubclassOf<ULevelStreaming> LevelStreamingClass, bool bUseExternalActors, const FString& DefaultFilename /* = TEXT("") */, const TArray<AActor*>* ActorsToMove /* = nullptr */, UWorld* InTemplateWorld /* = nullptr */, bool bInUseSaveAs /*= true*/)
 {
 	// Editor modes cannot be active when any level saving occurs.
 	if (!IsRunningCommandlet())
@@ -728,7 +728,15 @@ ULevelStreaming* UEditorLevelUtils::CreateNewStreamingLevelForWorld(UWorld& InWo
 	bool bNewWorldSaved = false;
 	FString NewPackageName = DefaultFilename;
 
-	bNewWorldSaved = FEditorFileUtils::SaveLevelAs(PersistentLevel, &NewPackageName);
+	if (bInUseSaveAs)
+	{
+		bNewWorldSaved = FEditorFileUtils::SaveLevelAs(PersistentLevel, &NewPackageName);
+	}
+	else
+	{
+		bNewWorldSaved = FEditorFileUtils::SaveLevel(PersistentLevel, DefaultFilename, &NewPackageName);
+	}
+	
 	if (bNewWorldSaved && !NewPackageName.IsEmpty())
 	{
 		NewPackageName = FPackageName::FilenameToLongPackageName(NewPackageName);
diff --git a/Engine/Source/Editor/UnrealEd/Public/EditorLevelUtils.h b/Engine/Source/Editor/UnrealEd/Public/EditorLevelUtils.h
--- a/Engine/Source/Editor/UnrealEd/Public/EditorLevelUtils.h
+++ b/Engine/Source/Editor/UnrealEd/Public/EditorLevelUtils.h
@@ -143,10 +143,11 @@ public:
 	 * @param	DefaultFilename						Optional file name for level.  If empty, the user will be prompted during the save process.
 	 * @param	bMoveSelectedActorsIntoNewLevel		If true, move any selected actors into the new level.
 	 * @param	InTemplateWorld						If valid, the new level will be a copy of the template world.
-	 *
+	 * @param	bInUseSaveAs						If true, show SaveAs dialog instead of Save with DefaultFilename
+	 * 
 	 * @return	Returns the newly created level, or NULL on failure
 	 */
-	static UNREALED_API ULevelStreaming* CreateNewStreamingLevelForWorld(UWorld& World, TSubclassOf<ULevelStreaming> LevelStreamingClass, const FString& DefaultFilename = TEXT(""), bool bMoveSelectedActorsIntoNewLevel = false, UWorld* InTemplateWorld = nullptr);
+	static UNREALED_API ULevelStreaming* CreateNewStreamingLevelForWorld(UWorld& World, TSubclassOf<ULevelStreaming> LevelStreamingClass, const FString& DefaultFilename = TEXT(""), bool bMoveSelectedActorsIntoNewLevel = false, UWorld* InTemplateWorld = nullptr, bool bInUseSaveAs = true);
 
 	/**
 	 * Creates a new streaming level and adds it to a world
@@ -157,10 +158,11 @@ public:
 	 * @param	DefaultFilename						file name for level.  If empty, the user will be prompted during the save process.
 	 * @param	ActorsToMove						Optional, move provided actors into the new level.
 	 * @param	InTemplateWorld						If valid, the new level will be a copy of the template world.
+	 * @param	bInUseSaveAs						If true, show SaveAs dialog instead of Save with DefaultFilename
 	 *
 	 * @return	Returns the newly created level, or NULL on failure
 	 */
-	static UNREALED_API ULevelStreaming* CreateNewStreamingLevelForWorld(UWorld& World, TSubclassOf<ULevelStreaming> LevelStreamingClass, bool bUseExternalActors, const FString& DefaultFilename, const TArray<AActor*>* ActorsToMove = nullptr, UWorld* InTemplateWorld = nullptr);
+	static UNREALED_API ULevelStreaming* CreateNewStreamingLevelForWorld(UWorld& World, TSubclassOf<ULevelStreaming> LevelStreamingClass, bool bUseExternalActors, const FString& DefaultFilename, const TArray<AActor*>* ActorsToMove = nullptr, UWorld* InTemplateWorld = nullptr, bool bInUseSaveAs = true);
 
 	/**
 	 * Adds the named level packages to the world.  Does nothing if all the levels already exist in the world.
