Fix a crash crash when running garbage collection during cooks

#rnx

- It was possible for a package to be unloaded via garbage collection while it still had cooking work inflight in and async job. This has many potential problems but the most common issue to be displayed would be bulkdata being unable to load it's payload as it had already been informed by FLinkerLoad that it should detach from the source file on disk and no longer attempt to use it.
- To fix this, during ::PreGarbageCollect we now gather a list of all packages containing objects that are waiting on async cooking work and add them to the list of packages to not garbage collect

diff --git a/Engine/Source/Editor/UnrealEd/Private/CookOnTheFlyServer.cpp b/Engine/Source/Editor/UnrealEd/Private/CookOnTheFlyServer.cpp
--- a/Engine/Source/Editor/UnrealEd/Private/CookOnTheFlyServer.cpp
+++ b/Engine/Source/Editor/UnrealEd/Private/CookOnTheFlyServer.cpp
@@ -3844,6 +3844,29 @@ void UCookOnTheFlyServer::PreGarbageCollect()
 		}
 	}
 
+	// Find the packages that are waiting on async jobs to finish cooking data
+	// and make sure that they are not garbage collected until the jobs have
+	// completed.
+	{
+		TSet<UPackage*> UniquePendingPackages;
+		for (FPendingCookedPlatformData& PendingData : PackageDatas->GetPendingCookedPlatformDatas())
+		{
+			if (UObject* Object = PendingData.Object.Get())
+			{	
+				if (UPackage* Package = Object->GetPackage())
+				{
+					UniquePendingPackages.Add(Package);
+				}	
+			}
+		}
+
+		GCKeepPackages.Reserve(GCKeepPackages.Num() + UniquePendingPackages.Num());
+		for (UPackage* Package : UniquePendingPackages)
+		{
+			GCKeepPackages.Add(Package);
+		}	
+	}
+
 	const bool bPartialGC = IsCookFlagSet(ECookInitializationFlags::EnablePartialGC);
 	if (bPartialGC)
 	{
