[WP Navmesh] Expose Navigation Data Chunk Grid Size and build loading size into WorldSettings
Add a dialog on WP navigation data building to set additional building options

diff --git a/Engine/Source/Editor/UnrealEd/Private/EditorBuildUtils.cpp b/Engine/Source/Editor/UnrealEd/Private/EditorBuildUtils.cpp
--- a/Engine/Source/Editor/UnrealEd/Private/EditorBuildUtils.cpp
+++ b/Engine/Source/Editor/UnrealEd/Private/EditorBuildUtils.cpp
@@ -48,6 +48,9 @@
 #include "LandscapeSubsystem.h"
 #include "ShaderCompilerCore.h"
 #include "AssetRegistryModule.h"
+#include "Interfaces/IMainFrameModule.h"
+#include "WorldPartition/SWorldPartitionBuildNavigationDialog.h"
+#include "WorldPartition/WorldPartitionBuildNavigationOptions.h"
 
 DEFINE_LOG_CATEGORY_STATIC(LogEditorBuildUtils, Log, All);
 
@@ -880,87 +883,118 @@ void FEditorBuildUtils::TriggerNavigationBuilder(UWorld*& InOutWorld, FName Id)
 
 bool FEditorBuildUtils::WorldPartitionBuildNavigation(const FString& InLongPackageName)
 {
-	// Ask user to save dirty packages
-	if (!FEditorFileUtils::SaveDirtyPackages(/*bPromptUserToSave=*/true, /*bSaveMapPackages=*/true, /*bSaveContentPackages=*/false))
-	{
-		return false;
-	}
+	UWorldPartitionBuildNavigationOptions* DefaultBuildNavigationOptions = GetMutableDefault<UWorldPartitionBuildNavigationOptions>();
+	DefaultBuildNavigationOptions->bVerbose = false;
+	DefaultBuildNavigationOptions->bCleanPackages = false;
+
+	TSharedPtr<SWindow> DlgWindow =
+		SNew(SWindow)
+		.Title(LOCTEXT("BuildNavigationWindowTitle", "Build Navigation Settings"))
+		.ClientSize(SWorldPartitionBuildNavigationDialog::DEFAULT_WINDOW_SIZE)
+		.SupportsMinimize(false)
+		.SupportsMaximize(false)
+		.SizingRule(ESizingRule::FixedSize);
+
+	TSharedRef<SWorldPartitionBuildNavigationDialog> Dialog =
+		SNew(SWorldPartitionBuildNavigationDialog)
+		.ParentWindow(DlgWindow)
+		.BuildNavigationOptions(DefaultBuildNavigationOptions);
+
+	DlgWindow->SetContent(Dialog);
 
-	// Unload any loaded map
-	if (!UEditorLoadingAndSavingUtils::NewBlankMap(/*bSaveExistingMap*/false))
+	IMainFrameModule& MainFrameModule = FModuleManager::LoadModuleChecked<IMainFrameModule>(TEXT("MainFrame"));
+	FSlateApplication::Get().AddModalWindow(DlgWindow.ToSharedRef(), MainFrameModule.GetParentWindow());
+
+	if (Dialog->ClickedOk())
 	{
-		return false;
-	}
+		// Ask user to save dirty packages
+		if (!FEditorFileUtils::SaveDirtyPackages(/*bPromptUserToSave=*/true, /*bSaveMapPackages=*/true, /*bSaveContentPackages=*/false))
+		{
+			return false;
+		}
+
+		// Unload any loaded map
+		if (!UEditorLoadingAndSavingUtils::NewBlankMap(/*bSaveExistingMap*/false))
+		{
+			return false;
+		}
 	
-	FProcHandle ProcessHandle;
-	bool bCancelled = false;
+		FProcHandle ProcessHandle;
+		bool bCancelled = false;
 
-	// Task scope
-	{
-		FScopedSlowTask SlowTask(0, LOCTEXT("WorldPartitionBuildNavigationProgress", "Building navigation..."));
-		SlowTask.MakeDialog(true);
+		// Task scope
+		{
+			FScopedSlowTask SlowTask(0, LOCTEXT("WorldPartitionBuildNavigationProgress", "Building navigation..."));
+			SlowTask.MakeDialog(true);
 
-		const FString CurrentExecutableName = FPlatformProcess::ExecutablePath();
+			const FString CurrentExecutableName = FPlatformProcess::ExecutablePath();
 
-		// Try to provide complete Path, if we can't try with project name
-		const FString ProjectPath = FPaths::IsProjectFilePathSet() ? FPaths::GetProjectFilePath() : FApp::GetProjectName();
+			// Try to provide complete Path, if we can't try with project name
+			const FString ProjectPath = FPaths::IsProjectFilePathSet() ? FPaths::GetProjectFilePath() : FApp::GetProjectName();
 
-		uint32 ProcessID;
+			uint32 ProcessID;
 
-		const FString Arguments = FString::Printf(TEXT("\"%s\" -run=WorldPartitionBuilderCommandlet %s  %s"), *ProjectPath, *InLongPackageName, TEXT(" -AllowCommandletRendering -Builder=WorldPartitionNavigationDataBuilder -SCCProvider=None -log=WPNavigationBuilderLog.txt"));
-		ProcessHandle = FPlatformProcess::CreateProc(*CurrentExecutableName, *Arguments, true, false, false, &ProcessID, 0, nullptr, nullptr);
+			const FString Arguments = FString::Printf(TEXT("\"%s\" -run=WorldPartitionBuilderCommandlet %s %s %s %s"),
+				*ProjectPath,
+				*InLongPackageName,
+				TEXT(" -AllowCommandletRendering -Builder=WorldPartitionNavigationDataBuilder -SCCProvider=None -log=WPNavigationBuilderLog.txt"), 
+				DefaultBuildNavigationOptions->bVerbose ? TEXT("-Verbose") : TEXT(""),
+				DefaultBuildNavigationOptions->bCleanPackages ? TEXT("-CleanPackages") : TEXT(""));
+				
+			ProcessHandle = FPlatformProcess::CreateProc(*CurrentExecutableName, *Arguments, true, false, false, &ProcessID, 0, nullptr, nullptr);
 		
-		while (FPlatformProcess::IsProcRunning(ProcessHandle))
-		{
-			if (SlowTask.ShouldCancel())
+			while (FPlatformProcess::IsProcRunning(ProcessHandle))
 			{
-				bCancelled = true;
-				FPlatformProcess::TerminateProc(ProcessHandle);
-				break;
-			}
+				if (SlowTask.ShouldCancel())
+				{
+					bCancelled = true;
+					FPlatformProcess::TerminateProc(ProcessHandle);
+					break;
+				}
 
-			SlowTask.EnterProgressFrame(0);
-			FPlatformProcess::Sleep(0.1);
+				SlowTask.EnterProgressFrame(0);
+				FPlatformProcess::Sleep(0.1);
+			}
 		}
-	}
 
-	int32 Result = 0;
-	if (!bCancelled && FPlatformProcess::GetProcReturnCode(ProcessHandle, &Result))
-	{	
-		if (Result == 0)
-		{
-			// Unload any loaded map
-			if (!UEditorLoadingAndSavingUtils::NewBlankMap(/*bSaveExistingMap*/false))
+		int32 Result = 0;
+		if (!bCancelled && FPlatformProcess::GetProcReturnCode(ProcessHandle, &Result))
+		{	
+			if (Result == 0)
 			{
-				return false;
-			}
+				// Unload any loaded map
+				if (!UEditorLoadingAndSavingUtils::NewBlankMap(/*bSaveExistingMap*/false))
+				{
+					return false;
+				}
 			
-			// Force registry update before loading converted map
-			const FAssetRegistryModule& AssetRegistryModule = FModuleManager::LoadModuleChecked<FAssetRegistryModule>("AssetRegistry");
-			IAssetRegistry& AssetRegistry = AssetRegistryModule.Get();
+				// Force registry update before loading converted map
+				const FAssetRegistryModule& AssetRegistryModule = FModuleManager::LoadModuleChecked<FAssetRegistryModule>("AssetRegistry");
+				IAssetRegistry& AssetRegistry = AssetRegistryModule.Get();
 											
-			FString MapToLoad = InLongPackageName;
+				FString MapToLoad = InLongPackageName;
 			
-			AssetRegistry.ScanModifiedAssetFiles({ MapToLoad });
-			AssetRegistry.ScanPathsSynchronous(ULevel::GetExternalObjectsPaths(MapToLoad), true);
+				AssetRegistry.ScanModifiedAssetFiles({ MapToLoad });
+				AssetRegistry.ScanPathsSynchronous(ULevel::GetExternalObjectsPaths(MapToLoad), true);
 
-			// Force a directory watcher tick for the asset registry to get notified of the changes
-			FDirectoryWatcherModule& DirectoryWatcherModule = FModuleManager::Get().LoadModuleChecked<FDirectoryWatcherModule>(TEXT("DirectoryWatcher"));
-			DirectoryWatcherModule.Get()->Tick(-1.0f);
+				// Force a directory watcher tick for the asset registry to get notified of the changes
+				FDirectoryWatcherModule& DirectoryWatcherModule = FModuleManager::Get().LoadModuleChecked<FDirectoryWatcherModule>(TEXT("DirectoryWatcher"));
+				DirectoryWatcherModule.Get()->Tick(-1.0f);
 			
-			FEditorFileUtils::LoadMap(MapToLoad);
+				FEditorFileUtils::LoadMap(MapToLoad);
+			}
+		}
+		else if (bCancelled)
+		{
+			FMessageDialog::Open(EAppMsgType::Ok, LOCTEXT("WorldPartitionBuildNavigationCancelled", "Building navigation cancelled!"));
 		}
-	}
-	else if (bCancelled)
-	{
-		FMessageDialog::Open(EAppMsgType::Ok, LOCTEXT("WorldPartitionBuildNavigationCancelled", "Building navigation cancelled!"));
-	}
 	
-	if (Result != 0)
-	{
-		FMessageDialog::Open(EAppMsgType::Ok, LOCTEXT("WorldPartitionBuildNavigationFailed", "Building navigation failed!"));
+		if (Result != 0)
+		{
+			FMessageDialog::Open(EAppMsgType::Ok, LOCTEXT("WorldPartitionBuildNavigationFailed", "Building navigation failed!"));
+		}
 	}
-
+	
 	return false;
 }
 
diff --git a/Engine/Source/Editor/UnrealEd/Private/WorldPartition/SWorldPartitionBuildNavigationDialog.cpp b/Engine/Source/Editor/UnrealEd/Private/WorldPartition/SWorldPartitionBuildNavigationDialog.cpp
new file mode 100644
--- /dev/null
+++ b/Engine/Source/Editor/UnrealEd/Private/WorldPartition/SWorldPartitionBuildNavigationDialog.cpp
@@ -0,0 +1,95 @@
+﻿// Copyright Epic Games, Inc. All Rights Reserved.
+#include "WorldPartition/SWorldPartitionBuildNavigationDialog.h"
+#include "Widgets/Layout/SSpacer.h"
+#include "Widgets/Input/SButton.h"
+#include "Modules/ModuleManager.h"
+#include "PropertyEditorModule.h"
+#include "WorldPartition/WorldPartitionBuildNavigationOptions.h"
+
+#define LOCTEXT_NAMESPACE "WorldPartitionBuildNavigationDialog"
+
+const FVector2D SWorldPartitionBuildNavigationDialog::DEFAULT_WINDOW_SIZE = FVector2D(600, 350);
+
+void SWorldPartitionBuildNavigationDialog::Construct(const FArguments& InArgs)
+{
+	ParentWindowPtr = InArgs._ParentWindow.Get();
+	BuildNavigationOptions = InArgs._BuildNavigationOptions.Get();
+	bClickedOk = false;
+
+	FPropertyEditorModule& PropertyEditorModule = FModuleManager::GetModuleChecked<FPropertyEditorModule>("PropertyEditor");
+	FDetailsViewArgs DetailsViewArgs;
+	DetailsViewArgs.bAllowSearch = false;
+	DetailsViewArgs.NameAreaSettings = FDetailsViewArgs::HideNameArea;
+	DetailsViewArgs.bAllowMultipleTopLevelObjects = false;
+
+	DetailsView = PropertyEditorModule.CreateDetailView(DetailsViewArgs);
+	
+	// Set provided objects on SDetailsView
+	DetailsView->SetObject(BuildNavigationOptions.Get(), true);
+
+	this->ChildSlot
+		[
+			SNew(SBorder)
+			.BorderImage(FEditorStyle::GetBrush("ToolPanel.GroupBorder"))
+			[
+				SNew(SVerticalBox)
+				+ SVerticalBox::Slot()
+				.MaxHeight(500.0f)
+				[
+					DetailsView->AsShared()
+				]
+				+ SVerticalBox::Slot()
+				.AutoHeight()
+				[
+					SNew(SHorizontalBox)
+					+ SHorizontalBox::Slot()
+					.FillWidth(1.f)
+					[
+						SNew(SSpacer)
+					]
+					+ SHorizontalBox::Slot()
+					.AutoWidth()
+					.Padding(2.f)
+					[
+						SNew(SButton)
+						.HAlign(HAlign_Center)
+						.IsEnabled(this, &SWorldPartitionBuildNavigationDialog::IsOkEnabled)
+						.ContentPadding(FEditorStyle::GetMargin("StandardDialog.ContentPadding"))
+						.OnClicked(this, &SWorldPartitionBuildNavigationDialog::OnOkClicked)
+						.Text(LOCTEXT("OkButton", "Ok"))
+					]
+					+ SHorizontalBox::Slot()
+					.AutoWidth()
+					.Padding(2.f)
+					[
+						SNew(SButton)
+						.HAlign(HAlign_Center)
+						.ContentPadding(FEditorStyle::GetMargin("StandardDialog.ContentPadding"))
+						.OnClicked(this, &SWorldPartitionBuildNavigationDialog::OnCancelClicked)
+						.Text(LOCTEXT("CancelButton", "Cancel"))
+					]
+				]
+			]
+		];
+}
+
+bool SWorldPartitionBuildNavigationDialog::IsOkEnabled() const
+{
+	return true;
+}
+
+FReply SWorldPartitionBuildNavigationDialog::OnOkClicked()
+{
+	bClickedOk = true;
+	ParentWindowPtr.Pin()->RequestDestroyWindow();
+	return FReply::Handled();
+}
+
+FReply SWorldPartitionBuildNavigationDialog::OnCancelClicked()
+{
+	bClickedOk = false;
+	ParentWindowPtr.Pin()->RequestDestroyWindow();
+	return FReply::Handled();
+}
+
+#undef LOCTEXT_NAMESPACE
diff --git a/Engine/Source/Editor/UnrealEd/Private/WorldPartition/WorldPartitionNavigationDataBuilder.cpp b/Engine/Source/Editor/UnrealEd/Private/WorldPartition/WorldPartitionNavigationDataBuilder.cpp
--- a/Engine/Source/Editor/UnrealEd/Private/WorldPartition/WorldPartitionNavigationDataBuilder.cpp
+++ b/Engine/Source/Editor/UnrealEd/Private/WorldPartition/WorldPartitionNavigationDataBuilder.cpp
@@ -26,11 +26,12 @@ UWorldPartitionNavigationDataBuilder::UWorldPartitionNavigationDataBuilder(const
 bool UWorldPartitionNavigationDataBuilder::PreRun(UWorld* World, FPackageSourceControlHelper& PackageHelper)
 {
 	const TSubclassOf<APartitionActor>& NavigationDataActorClass = ANavigationDataChunkActor::StaticClass();
-	const uint32 GridSize = NavigationDataActorClass->GetDefaultObject<APartitionActor>()->GetDefaultGridSize(World);
-	
+	uint32 GridSize = NavigationDataActorClass->GetDefaultObject<APartitionActor>()->GetDefaultGridSize(World);
+	GridSize = FMath::Max(GridSize, 1u);
+
 	// Size of loaded cell. Set as big as your hardware can afford.
-	// @todo: move to a config file.
-	IterativeCellSize = 2*GridSize;
+	const uint32 LoadingCellSizeSetting = World->GetWorldSettings()->NavigationDataBuilderLoadingCellSize;
+	IterativeCellSize = (LoadingCellSizeSetting / GridSize) * GridSize;
 	
 	// Extra padding around loaded cell.
 	// @todo: set value programatically.
@@ -41,6 +42,11 @@ bool UWorldPartitionNavigationDataBuilder::PreRun(UWorld* World, FPackageSourceC
 
 	bCleanBuilderPackages = Switches.Contains(TEXT("CleanPackages"));
 
+	UE_LOG(LogWorldPartitionNavigationDataBuilder, Log, TEXT("Starting NavigationDataBuilder"));
+	UE_LOG(LogWorldPartitionNavigationDataBuilder, Log, TEXT("   ANavigationDataChunkActor GridSize: %8i"), GridSize);
+	UE_LOG(LogWorldPartitionNavigationDataBuilder, Log, TEXT("   IterativeCellSize:                  %8i (%ix%i navigation datachunk partition actor per loaded cell)"), IterativeCellSize, IterativeCellSize/GridSize, IterativeCellSize/GridSize);
+	UE_LOG(LogWorldPartitionNavigationDataBuilder, Log, TEXT("   IterativeCellOverlapSize:           %8i"), IterativeCellOverlapSize);
+	
 	return true;
 }
 
diff --git a/Engine/Source/Editor/UnrealEd/Public/WorldPartition/SWorldPartitionBuildNavigationDialog.h b/Engine/Source/Editor/UnrealEd/Public/WorldPartition/SWorldPartitionBuildNavigationDialog.h
new file mode 100644
--- /dev/null
+++ b/Engine/Source/Editor/UnrealEd/Public/WorldPartition/SWorldPartitionBuildNavigationDialog.h
@@ -0,0 +1,51 @@
+﻿// Copyright Epic Games, Inc. All Rights Reserved.
+
+#pragma once
+
+#include "Input/Reply.h"
+#include "Widgets/SWindow.h"
+#include "Widgets/DeclarativeSyntaxSupport.h"
+#include "Widgets/SCompoundWidget.h"
+
+class UWorldPartitionBuildNavigationOptions;
+
+class SWorldPartitionBuildNavigationDialog : public SCompoundWidget
+{
+public:
+
+	SLATE_BEGIN_ARGS(SWorldPartitionBuildNavigationDialog) {}
+	/** A pointer to the parent window */
+	SLATE_ATTRIBUTE(TSharedPtr<SWindow>, ParentWindow)
+	SLATE_ATTRIBUTE(UWorldPartitionBuildNavigationOptions*, BuildNavigationOptions)
+	SLATE_END_ARGS()
+
+	~SWorldPartitionBuildNavigationDialog() {}
+
+	/**
+	 * Construct this widget.  Called by the SNew() Slate macro.
+	 *
+	 * @param	InArgs				Declaration used by the SNew() macro to construct this widget
+	 * @param	InViewModel			The UI logic not specific to slate
+	 */
+	void Construct(const FArguments& InArgs);
+
+	static const FVector2D DEFAULT_WINDOW_SIZE;
+
+	bool ClickedOk() const { return bClickedOk; }
+	
+private:
+	FReply OnOkClicked();
+	bool IsOkEnabled() const;
+
+	FReply OnCancelClicked();
+		
+	/** Pointer to the parent window, so we know to destroy it when done */
+	TWeakPtr<SWindow> ParentWindowPtr;
+	TWeakObjectPtr<UWorldPartitionBuildNavigationOptions> BuildNavigationOptions;
+
+	/** Detailsview used to display SettingsObjects, and allowing user to change options */
+	TSharedPtr<class IDetailsView> DetailsView;
+		
+	/** Dialog Result */
+	bool bClickedOk;
+};
\ No newline at end of file
diff --git a/Engine/Source/Editor/UnrealEd/Public/WorldPartition/WorldPartitionBuildNavigationOptions.h b/Engine/Source/Editor/UnrealEd/Public/WorldPartition/WorldPartitionBuildNavigationOptions.h
new file mode 100644
--- /dev/null
+++ b/Engine/Source/Editor/UnrealEd/Public/WorldPartition/WorldPartitionBuildNavigationOptions.h
@@ -0,0 +1,20 @@
+﻿// Copyright Epic Games, Inc. All Rights Reserved.
+
+#pragma once
+
+#include "WorldPartitionBuildNavigationOptions.generated.h"
+
+UCLASS()
+class UWorldPartitionBuildNavigationOptions : public UObject
+{
+	GENERATED_BODY()
+
+public:
+	/** Use verbose logging. */
+	UPROPERTY(EditAnywhere, Category = "Build Navigation")
+	bool bVerbose = false;
+
+	/** When enabled, delete all NavigationDataChunkActors instead of generating them. */
+	UPROPERTY(EditAnywhere, Category = "Build Navigation")
+	bool bCleanPackages = false;
+};
diff --git a/Engine/Source/Runtime/Engine/Classes/GameFramework/WorldSettings.h b/Engine/Source/Runtime/Engine/Classes/GameFramework/WorldSettings.h
--- a/Engine/Source/Runtime/Engine/Classes/GameFramework/WorldSettings.h
+++ b/Engine/Source/Runtime/Engine/Classes/GameFramework/WorldSettings.h
@@ -559,13 +559,20 @@ protected:
 
 public:
 #if WITH_EDITORONLY_DATA
-	/** Size of the grid for instanced foliage actors*/
+	/** Size of the grid for instanced foliage actors */
 	UPROPERTY(EditAnywhere, Category = Foliage)
 	uint32 InstancedFoliageGridSize;
 
-	/** Size of the grid for navigation data chunk actors*/
+	/** Size of the grid for navigation data chunk actors */
 	UPROPERTY(EditAnywhere, Category = Navigation)
 	uint32 NavigationDataChunkGridSize;
+
+	/**
+	 * Loading cell size used when building navigation data iteratively.
+	 * The actual cell size used will be rounded using the NavigationDataChunkGridSize.
+	 */
+	UPROPERTY(EditAnywhere, Category = Navigation)
+	uint32 NavigationDataBuilderLoadingCellSize;
 	
 	/** Default size of the grid for placed elements from the editor */
 	UPROPERTY(EditAnywhere, Category=Editor)
diff --git a/Engine/Source/Runtime/Engine/Private/WorldSettings.cpp b/Engine/Source/Runtime/Engine/Private/WorldSettings.cpp
--- a/Engine/Source/Runtime/Engine/Private/WorldSettings.cpp
+++ b/Engine/Source/Runtime/Engine/Private/WorldSettings.cpp
@@ -93,6 +93,7 @@ PRAGMA_ENABLE_DEPRECATION_WARNINGS
 #if WITH_EDITORONLY_DATA
 	InstancedFoliageGridSize = DefaultPlacementGridSize = 25600;
 	NavigationDataChunkGridSize = 102400;
+	NavigationDataBuilderLoadingCellSize = 102400 * 4;
 	bIncludeGridSizeInNameForFoliageActors = false;
 	bIncludeGridSizeInNameForPartitionedActors = false;
 #endif
