When running automated test, initiate start location of PIE when there is no Player Start actor in the map

#rnx

diff --git a/Engine/Source/Editor/UnrealEd/Private/EditorEngine.cpp b/Engine/Source/Editor/UnrealEd/Private/EditorEngine.cpp
--- a/Engine/Source/Editor/UnrealEd/Private/EditorEngine.cpp
+++ b/Engine/Source/Editor/UnrealEd/Private/EditorEngine.cpp
@@ -187,6 +187,7 @@
 
 #if !UE_BUILD_SHIPPING
 #include "Tests/AutomationCommon.h"
+#include "Tests/AutomationEditorCommon.h"
 #endif
 
 #include "PhysicsPublic.h"
@@ -7365,6 +7366,12 @@ void UEditorEngine::AutomationLoadMap(const FString& MapName, bool bForceReload,
 		EditorPlaySettings->bLaunchSeparateServer = false;
 		RequestParams.EditorPlaySettings = EditorPlaySettings;
 
+		// Make sure the player start location is a valid location.
+		if (CheckForPlayerStart() == nullptr)
+		{
+			FAutomationEditorCommonUtils::SetPlaySessionStartToActiveViewport(RequestParams);
+		}
+
 		RequestPlaySession(RequestParams);
 
 		// Immediately launch the session 
diff --git a/Engine/Source/Editor/UnrealEd/Private/Tests/AutomationEditorCommon.cpp b/Engine/Source/Editor/UnrealEd/Private/Tests/AutomationEditorCommon.cpp
--- a/Engine/Source/Editor/UnrealEd/Private/Tests/AutomationEditorCommon.cpp
+++ b/Engine/Source/Editor/UnrealEd/Private/Tests/AutomationEditorCommon.cpp
@@ -25,6 +25,7 @@
 #include "ARFilter.h"
 #include "AssetRegistryModule.h"
 #include "Tests/AutomationCommon.h"
+#include "IAssetViewport.h"
 
 #include "LevelEditor.h"
 #include "Interfaces/IMainFrameModule.h"
@@ -502,6 +503,23 @@ bool FAutomationEditorCommonUtils::SetOrthoViewportView(const FVector& ViewLocat
 	return false;
 }
 
+bool FAutomationEditorCommonUtils::SetPlaySessionStartToActiveViewport(FRequestPlaySessionParams& OutParams)
+{
+	FLevelEditorModule& LevelEditorModule = FModuleManager::GetModuleChecked<FLevelEditorModule>(TEXT("LevelEditor"));
+	TSharedPtr<IAssetViewport> ActiveLevelViewport = LevelEditorModule.GetFirstActiveViewport();
+	// Make sure we can find a path to the view port.
+	if (ActiveLevelViewport.IsValid() &&
+		FSlateApplication::Get().FindWidgetWindow(ActiveLevelViewport->AsWidget()).IsValid())
+	{
+		// Start the player where the camera is if not forcing from player start
+		OutParams.StartLocation = ActiveLevelViewport->GetAssetViewportClient().GetViewLocation();
+		OutParams.StartRotation = ActiveLevelViewport->GetAssetViewportClient().GetViewRotation();
+		return true;
+	}
+
+	return false;
+}
+
 //////////////////////////////////////////////////////////////////////
 //Asset Path Commands
 
@@ -796,6 +814,12 @@ bool FStartPIECommand::Update()
 		Params.WorldType = EPlaySessionWorldType::SimulateInEditor;
 	}
 
+	// Make sure the player start location is a valid location.
+	if (GUnrealEd->CheckForPlayerStart() == nullptr)
+	{
+		FAutomationEditorCommonUtils::SetPlaySessionStartToActiveViewport(Params);
+	}
+
 	GUnrealEd->RequestPlaySession(Params);
 	return true;
 }
diff --git a/Engine/Source/Editor/UnrealEd/Public/Tests/AutomationEditorCommon.h b/Engine/Source/Editor/UnrealEd/Public/Tests/AutomationEditorCommon.h
--- a/Engine/Source/Editor/UnrealEd/Public/Tests/AutomationEditorCommon.h
+++ b/Engine/Source/Editor/UnrealEd/Public/Tests/AutomationEditorCommon.h
@@ -6,6 +6,7 @@
 #include "Misc/AutomationTest.h"
 #include "AssetData.h"
 #include "Tests/AutomationTestSettings.h"
+#include "PlayInEditorDataTypes.h"
 
 class UFactory;
 
@@ -134,6 +135,13 @@ public:
 	*/
 	static bool SetOrthoViewportView(const FVector& ViewLocation, const FRotator& ViewRotation);
 
+	/**
+	* Set Play Session parameters to Active Viewport location
+	* 
+	* @param OutParams - FRequestPlaySessionParams object to set
+	*/
+	static bool SetPlaySessionStartToActiveViewport(FRequestPlaySessionParams& OutParams);
+
 	/**
 	* Converts a package path to an asset path
 	*
