Fixed cooked anim BP assets not animating when placed

Also fixed up cooked anim BPs to display their specific class type.
Also re-instated fix for placing cooked anim BPs in the viewport. Added in CL 18465876, accidentally removed in CL 18589986.

diff --git a/Engine/Source/Developer/AssetTools/Private/AssetTools.cpp b/Engine/Source/Developer/AssetTools/Private/AssetTools.cpp
--- a/Engine/Source/Developer/AssetTools/Private/AssetTools.cpp
+++ b/Engine/Source/Developer/AssetTools/Private/AssetTools.cpp
@@ -42,6 +42,7 @@
 #include "AssetTypeActions/AssetTypeActions_VectorField.h"
 #include "AssetTypeActions/AssetTypeActions_AnimationAsset.h"
 #include "AssetTypeActions/AssetTypeActions_AnimBlueprint.h"
+#include "AssetTypeActions/AssetTypeActions_AnimBlueprintGeneratedClass.h"
 #include "AssetTypeActions/AssetTypeActions_AnimBoneCompressionSettings.h"
 #include "AssetTypeActions/AssetTypeActions_AnimComposite.h"
 #include "AssetTypeActions/AssetTypeActions_AnimStreamable.h"
@@ -222,6 +223,7 @@ UAssetToolsImpl::UAssetToolsImpl(const FObjectInitializer& ObjectInitializer)
 	RegisterAssetTypeActions(MakeShareable(new FAssetTypeActions_ActorFolder));
 	RegisterAssetTypeActions(MakeShareable(new FAssetTypeActions_AnimationAsset));
 	RegisterAssetTypeActions(MakeShareable(new FAssetTypeActions_AnimBlueprint));
+	RegisterAssetTypeActions(MakeShareable(new FAssetTypeActions_AnimBlueprintGeneratedClass));
 	RegisterAssetTypeActions(MakeShareable(new FAssetTypeActions_AnimBoneCompressionSettings));
 	RegisterAssetTypeActions(MakeShareable(new FAssetTypeActions_AnimComposite));
 	RegisterAssetTypeActions(MakeShareable(new FAssetTypeActions_AnimStreamable));
diff --git a/Engine/Source/Developer/AssetTools/Private/AssetTypeActions/AssetTypeActions_AnimBlueprintGeneratedClass.h b/Engine/Source/Developer/AssetTools/Private/AssetTypeActions/AssetTypeActions_AnimBlueprintGeneratedClass.h
new file mode 100644
--- /dev/null
+++ b/Engine/Source/Developer/AssetTools/Private/AssetTypeActions/AssetTypeActions_AnimBlueprintGeneratedClass.h
@@ -0,0 +1,24 @@
+// Copyright Epic Games, Inc. All Rights Reserved.
+
+#pragma once
+
+#include "CoreMinimal.h"
+#include "Toolkits/IToolkitHost.h"
+#include "AssetTypeActions/AssetTypeActions_BlueprintGeneratedClass.h"
+#include "Animation/AnimBlueprintGeneratedClass.h"
+
+class UBlueprintGeneratedClass;
+struct FAssetData;
+class IClassTypeActions;
+class UFactory;
+
+class FAssetTypeActions_AnimBlueprintGeneratedClass : public FAssetTypeActions_BlueprintGeneratedClass
+{
+public:
+	// IAssetTypeActions Implementation
+	virtual FText GetName() const override { return NSLOCTEXT("AssetTypeActions", "AssetTypeActions_AnimBlueprintGeneratedClass", "Compiled Anim Blueprint Class"); }
+	virtual FColor GetTypeColor() const override { return FColor(240, 156, 0); }
+	virtual uint32 GetCategories() override { return EAssetTypeCategories::Animation; }
+	virtual UClass* GetSupportedClass() const override { return UAnimBlueprintGeneratedClass::StaticClass(); }
+	// End IAssetTypeActions Implementation
+};
diff --git a/Engine/Source/Editor/UnrealEd/Private/Factories/ActorFactory.cpp b/Engine/Source/Editor/UnrealEd/Private/Factories/ActorFactory.cpp
--- a/Engine/Source/Editor/UnrealEd/Private/Factories/ActorFactory.cpp
+++ b/Engine/Source/Editor/UnrealEd/Private/Factories/ActorFactory.cpp
@@ -1184,7 +1184,7 @@ bool UActorFactorySkeletalMesh::CanCreateActorFrom( const FAssetData& AssetData,
 		SkeletalMeshData = AssetData;
 	}
 
-	if ( !SkeletalMeshData.IsValid() && AssetData.GetClass()->IsChildOf( UAnimBlueprint::StaticClass() ) )
+	if ( !SkeletalMeshData.IsValid() && (AssetData.GetClass()->IsChildOf(UAnimBlueprint::StaticClass()) || AssetData.GetClass()->IsChildOf(UAnimBlueprintGeneratedClass::StaticClass())))
 	{
 		const FString TargetSkeletonPath = AssetData.GetTagValueRef<FString>( "TargetSkeleton" );
 		if ( TargetSkeletonPath.IsEmpty() )
@@ -1271,14 +1271,26 @@ bool UActorFactorySkeletalMesh::CanCreateActorFrom( const FAssetData& AssetData,
 
 USkeletalMesh* UActorFactorySkeletalMesh::GetSkeletalMeshFromAsset( UObject* Asset )
 {
-	USkeletalMesh*SkeletalMesh = Cast<USkeletalMesh>( Asset );
-	UAnimBlueprint* AnimBlueprint = Cast<UAnimBlueprint>( Asset );
+	USkeletalMesh* SkeletalMesh = Cast<USkeletalMesh>( Asset );
 	USkeleton* Skeleton = Cast<USkeleton>( Asset );
 
-	if( SkeletalMesh == nullptr && AnimBlueprint != nullptr && AnimBlueprint->TargetSkeleton )
+	if(SkeletalMesh == nullptr)
 	{
 		// base it on preview skeletal mesh, just to have something
-		SkeletalMesh = AnimBlueprint->TargetSkeleton->GetPreviewMesh(true);
+		if(UAnimBlueprint* AnimBlueprint = Cast<UAnimBlueprint>(Asset))
+		{
+			if(AnimBlueprint->TargetSkeleton)
+			{
+				SkeletalMesh = AnimBlueprint->TargetSkeleton->GetPreviewMesh(true);
+			}
+		}
+		if(UAnimBlueprintGeneratedClass* AnimBlueprintGeneratedClass = Cast<UAnimBlueprintGeneratedClass>(Asset))
+		{
+			if(AnimBlueprintGeneratedClass->TargetSkeleton)
+			{
+				SkeletalMesh = AnimBlueprintGeneratedClass->TargetSkeleton->GetPreviewMesh(true);
+			}
+		}
 	}
 
 	if( SkeletalMesh == nullptr && Skeleton != nullptr )
@@ -1306,7 +1318,6 @@ USkeletalMesh* UActorFactorySkeletalMesh::GetSkeletalMeshFromAsset( UObject* Ass
 void UActorFactorySkeletalMesh::PostSpawnActor( UObject* Asset, AActor* NewActor )
 {
 	USkeletalMesh* SkeletalMesh = GetSkeletalMeshFromAsset(Asset);
-	UAnimBlueprint* AnimBlueprint = Cast<UAnimBlueprint>( Asset );
 	ASkeletalMeshActor* NewSMActor = CastChecked<ASkeletalMeshActor>(NewActor);
 
 	Super::PostSpawnActor(SkeletalMesh, NewActor);
@@ -1323,10 +1334,14 @@ void UActorFactorySkeletalMesh::PostSpawnActor( UObject* Asset, AActor* NewActor
 
 	// Init Component
 	NewSMActor->GetSkeletalMeshComponent()->RegisterComponent();
-	if( AnimBlueprint )
+	if( UAnimBlueprint* AnimBlueprint = Cast<UAnimBlueprint>( Asset ) )
 	{
 		NewSMActor->GetSkeletalMeshComponent()->SetAnimInstanceClass(AnimBlueprint->GeneratedClass);
 	}
+	else if( UAnimBlueprintGeneratedClass* AnimBlueprintGeneratedClass = Cast<UAnimBlueprintGeneratedClass>(Asset) )
+	{
+		NewSMActor->GetSkeletalMeshComponent()->SetAnimInstanceClass(AnimBlueprintGeneratedClass);
+	}
 
 	if (ClassUsedForDelegate != nullptr)
 	{
