Disable InEditorCooking by default for all projects. WorldPartition does not support it; a project should enable InEditorCooking only if it does not use WorldPartition.
Add comment explaining why WorldPartition does not support it.
[RN] Minor, Cooking

diff --git a/Engine/Config/BaseEditorPerProjectUserSettings.ini b/Engine/Config/BaseEditorPerProjectUserSettings.ini
--- a/Engine/Config/BaseEditorPerProjectUserSettings.ini
+++ b/Engine/Config/BaseEditorPerProjectUserSettings.ini
@@ -155,6 +155,9 @@ bEnableAsyncTextureCompilation=True
 bEnableAsyncStaticMeshCompilation=True
 bEnableAsyncSkeletalMeshCompilation=True
 
+; Cooking in the editor is disabled by default now; WorldPartition does not yet support it and we want to deprecate the mode in the future to reduce the number of code flows that have to be supported
+bDisableCookInEditor=true
+
 [/Script/UnrealEd.EditorLoadingSavingSettings]
 ; True if we should automatically load a default level at start up
 LoadLevelAtStartup=ProjectDefault
diff --git a/Engine/Source/Editor/UnrealEd/Private/CookOnTheFlyServer.cpp b/Engine/Source/Editor/UnrealEd/Private/CookOnTheFlyServer.cpp
--- a/Engine/Source/Editor/UnrealEd/Private/CookOnTheFlyServer.cpp
+++ b/Engine/Source/Editor/UnrealEd/Private/CookOnTheFlyServer.cpp
@@ -2273,9 +2273,12 @@ UE::Cook::FGeneratorPackage* UCookOnTheFlyServer::CreateGeneratorPackage(UE::Coo
 	check(Splitter);
 	check(!PackageData.GetGeneratorPackage());
 
-	// TODO: Add support for cooking in the editor
+	// TODO: Add support for cooking in the editor. Possibly moot since we plan to deprecate cooking in the editor.
 	if (IsCookingInEditor())
 	{
+		// CookPackageSplitters allow destructive changes to the generator package. e.g. moving UObjects out
+		// of it into the streaming packages. To allow its use in the editor, we will need to make it non-destructive
+		// (by e.g. copying to new packages), or restore the package after the changes have been made.
 		UE_LOG(LogCook, Error, TEXT("Cooking in editor doesn't support Cook Package Splitters."));
 		return nullptr;
 	}
diff --git a/Engine/Source/Editor/UnrealEd/Private/PlayLevelViaLauncher.cpp b/Engine/Source/Editor/UnrealEd/Private/PlayLevelViaLauncher.cpp
--- a/Engine/Source/Editor/UnrealEd/Private/PlayLevelViaLauncher.cpp
+++ b/Engine/Source/Editor/UnrealEd/Private/PlayLevelViaLauncher.cpp
@@ -249,7 +249,9 @@ void UEditorEngine::StartPlayUsingLauncherSession(FRequestPlaySessionParams& InR
 	if (CurrentLauncherCookMode == ELauncherProfileCookModes::OnTheFlyInEditor ||
 		CurrentLauncherCookMode == ELauncherProfileCookModes::ByTheBookInEditor)
 	{
-		// For now World Partition doesn't these cook modes
+		// For now World Partition doesn't support InEditor cooking because its cooking is destructive -
+		// it moves UObjects out of the generator package into the streaming packages. To allow cooking it in
+		// the editor process, we will need to make it non-destructive or restore the package afterwards.
 		FWorldContext& EditorContext = GetEditorWorldContext();
 		if (EditorContext.World()->IsPartitionedWorld())
 		{
