Fixed Editor startup crash on Mac OS X due to vieport size being 0;0

diff --git a/Engine/Source/Editor/UnrealEd/Private/EditorViewportClient.cpp b/Engine/Source/Editor/UnrealEd/Private/EditorViewportClient.cpp
--- a/Engine/Source/Editor/UnrealEd/Private/EditorViewportClient.cpp
+++ b/Engine/Source/Editor/UnrealEd/Private/EditorViewportClient.cpp
@@ -914,6 +914,8 @@ FSceneView* FEditorViewportClient::CalcSceneView(FSceneViewFamily* ViewFamily, c
 	}
 
 	FIntPoint ViewportSize = Viewport->GetSizeXY();
+	ViewportSize.X = FMath::Max(ViewportSize.X, 1);
+	ViewportSize.Y = FMath::Max(ViewportSize.Y, 1);
 	FIntPoint ViewportOffset(0, 0);
 
 	// We expect some size to avoid problems with the view rect manipulation
