Move event registration out of constructor to avoid CDO from processing events redundantly

diff --git a/Engine/Source/Editor/UnrealEd/Private/CookOnTheFlyServer.cpp b/Engine/Source/Editor/UnrealEd/Private/CookOnTheFlyServer.cpp
--- a/Engine/Source/Editor/UnrealEd/Private/CookOnTheFlyServer.cpp
+++ b/Engine/Source/Editor/UnrealEd/Private/CookOnTheFlyServer.cpp
@@ -471,13 +471,8 @@ UCookOnTheFlyServer::UCookOnTheFlyServer(const FObjectInitializer& ObjectInitial
 	CookByTheBookOptions(nullptr),
 	CookFlags(ECookInitializationFlags::None),
 	bIsSavingPackage(false),
-	AssetRegistry(nullptr),
-	PackageDatas(MakeUnique<UE::Cook::FPackageDatas>(*this))
+	AssetRegistry(nullptr)
 {
-	PlatformManager = MakeUnique<UE::Cook::FPlatformManager>();
-	ExternalRequests = MakeUnique<UE::Cook::FExternalRequests>();
-	PackageTracker = MakeUnique<UE::Cook::FPackageTracker>(*PackageDatas.Get());
-	DiffModeHelper = MakeUnique<FDiffModeCookServerUtils>();
 }
 
 UCookOnTheFlyServer::UCookOnTheFlyServer(FVTableHelper& Helper) :Super(Helper) {}
@@ -4447,6 +4442,12 @@ void UCookOnTheFlyServer::Initialize( ECookMode::Type DesiredCookMode, ECookInit
 {
 	TRACE_CPUPROFILER_EVENT_SCOPE(UCookOnTheFlyServer::Initialize);
 
+	PackageDatas = MakeUnique<UE::Cook::FPackageDatas>(*this);
+	PlatformManager = MakeUnique<UE::Cook::FPlatformManager>();
+	ExternalRequests = MakeUnique<UE::Cook::FExternalRequests>();
+	PackageTracker = MakeUnique<UE::Cook::FPackageTracker>(*PackageDatas.Get());
+	DiffModeHelper = MakeUnique<FDiffModeCookServerUtils>();
+
 	BuildDefinitions.Reset(new UE::Cook::FBuildDefinitions());
 	UE::Cook::InitializeTls();
 	UE::Cook::FPlatformManager::InitializeTls();
