Adds embedded FoliageType conversion to UWorldPartitionConvertCommandlet

#changelist validated

diff --git a/Engine/Source/Editor/FoliageEdit/Private/FoliageEditUtility.cpp b/Engine/Source/Editor/FoliageEdit/Private/FoliageEditUtility.cpp
--- a/Engine/Source/Editor/FoliageEdit/Private/FoliageEditUtility.cpp
+++ b/Engine/Source/Editor/FoliageEdit/Private/FoliageEditUtility.cpp
@@ -40,18 +40,7 @@ UFoliageType* FFoliageEditUtility::SaveFoliageTypeObject(UFoliageType* InFoliage
 		if (SaveFoliageTypeDialog->ShowModal() != EAppReturnType::Cancel)
 		{
 			PackageName = SaveFoliageTypeDialog->GetFullAssetPath().ToString();
-			UPackage* Package = CreatePackage( *PackageName);
-
-			// We should not save a copy of this duplicate into the transaction buffer as it's an asset
-			InFoliageType->ClearFlags(RF_Transactional);
-			TypeToSave = Cast<UFoliageType>(StaticDuplicateObject(InFoliageType, Package, *FPackageName::GetLongPackageAssetName(PackageName)));
-			InFoliageType->SetFlags(RF_Transactional);
-
-			TypeToSave->SetFlags(RF_Standalone | RF_Public | RF_Transactional);
-			TypeToSave->Modify();
-
-			// Notify the asset registry
-			FAssetRegistryModule::AssetCreated(TypeToSave);
+			TypeToSave = DuplicateFoliageTypeToNewPackage(PackageName, InFoliageType);
 		}
 	}
 	else
@@ -77,6 +66,27 @@ UFoliageType* FFoliageEditUtility::SaveFoliageTypeObject(UFoliageType* InFoliage
 	return TypeToSave;
 }
 
+UFoliageType* FFoliageEditUtility::DuplicateFoliageTypeToNewPackage(const FString& InPackageName, UFoliageType* InFoliageType)
+{
+	UPackage* Package = CreatePackage(*InPackageName);
+	UFoliageType* TypeToSave = nullptr;
+
+	// We should not save a copy of this duplicate into the transaction buffer as it's an asset. Save and restore Transactional flag
+	EObjectFlags Transactional = InFoliageType->HasAnyFlags(RF_Transactional) ? RF_Transactional : RF_NoFlags;
+	
+	InFoliageType->ClearFlags(Transactional);
+	TypeToSave = Cast<UFoliageType>(StaticDuplicateObject(InFoliageType, Package, *FPackageName::GetLongPackageAssetName(InPackageName)));
+	InFoliageType->SetFlags(Transactional);
+
+	TypeToSave->SetFlags(RF_Standalone | RF_Public | Transactional);
+	TypeToSave->Modify();
+
+	// Notify the asset registry
+	FAssetRegistryModule::AssetCreated(TypeToSave);
+	
+	return TypeToSave;
+}
+
 void FFoliageEditUtility::ReplaceFoliageTypeObject(UWorld* InWorld, UFoliageType* OldType, UFoliageType* NewType)
 {
 	FScopedTransaction Transaction(NSLOCTEXT("UnrealEd", "FoliageMode_ReplaceSettingsObject", "Foliage Editing: Replace Settings Object"));
diff --git a/Engine/Source/Editor/FoliageEdit/Public/FoliageEditUtility.h b/Engine/Source/Editor/FoliageEdit/Public/FoliageEditUtility.h
--- a/Engine/Source/Editor/FoliageEdit/Public/FoliageEditUtility.h
+++ b/Engine/Source/Editor/FoliageEdit/Public/FoliageEditUtility.h
@@ -18,7 +18,16 @@ class FFoliageEditUtility
 {
 public:
 	/** Save the foliage type object. If it isn't an asset, will prompt the user for a location to save the new asset. */
-	static FOLIAGEEDIT_API UFoliageType* SaveFoliageTypeObject(UFoliageType* Settings);
+	static FOLIAGEEDIT_API UFoliageType* SaveFoliageTypeObject(UFoliageType* Settings); 
+
+	/**
+	 * Copies the provided FoliageType to a new asset
+	 * @param 	InPackageName 	The name of the new asset
+	 * @param 	InFoliageType 	The FoliageType to copy
+	 * @return 	The new FoliageType
+	 */
+	static FOLIAGEEDIT_API UFoliageType* DuplicateFoliageTypeToNewPackage(const FString& InPackageName, UFoliageType* InFoliageType);
+
 	static FOLIAGEEDIT_API void ReplaceFoliageTypeObject(UWorld* InWorld, UFoliageType* OldType, UFoliageType* NewType);
 
 	static FOLIAGEEDIT_API void MoveActorFoliageInstancesToLevel(ULevel* InTargetLevel, AActor* InIFA = nullptr);
diff --git a/Engine/Source/Editor/UnrealEd/Classes/Commandlets/WorldPartitionConvertCommandlet.h b/Engine/Source/Editor/UnrealEd/Classes/Commandlets/WorldPartitionConvertCommandlet.h
--- a/Engine/Source/Editor/UnrealEd/Classes/Commandlets/WorldPartitionConvertCommandlet.h
+++ b/Engine/Source/Editor/UnrealEd/Classes/Commandlets/WorldPartitionConvertCommandlet.h
@@ -121,6 +121,9 @@ protected:
 	UPROPERTY(Config)
 	FString DefaultHLODLayerName;
 
+	UPROPERTY(Config)
+	FString FoliageTypePath;
+
 	UPROPERTY(Config)
 	TArray<FHLODLayerActorMapping> HLODLayersForActorClasses;
 
diff --git a/Engine/Source/Editor/UnrealEd/Private/Commandlets/WorldPartitionConvertCommandlet.cpp b/Engine/Source/Editor/UnrealEd/Private/Commandlets/WorldPartitionConvertCommandlet.cpp
--- a/Engine/Source/Editor/UnrealEd/Private/Commandlets/WorldPartitionConvertCommandlet.cpp
+++ b/Engine/Source/Editor/UnrealEd/Private/Commandlets/WorldPartitionConvertCommandlet.cpp
@@ -39,6 +39,7 @@
 #include "EdGraph/EdGraph.h"
 #include "ProfilingDebugging/ScopedTimers.h"
 #include "AssetRegistryModule.h"
+#include "FoliageEditUtility.h"
 #include "FoliageHelper.h"
 #include "Engine/WorldComposition.h"
 #include "ActorPartition/ActorPartitionSubsystem.h"
@@ -686,7 +687,8 @@ int32 UWorldPartitionConvertCommandlet::Main(const FString& Params)
 	UE_SCOPED_TIMER(TEXT("Conversion"), LogWorldPartitionConvertCommandlet, Display);
 
 	TArray<FString> Tokens, Switches;
-	ParseCommandLine(*Params, Tokens, Switches);
+	TMap<FString, FString> Arguments;
+	ParseCommandLine(*Params, Tokens, Switches, Arguments);
 
 	if (Tokens.Num() != 1)
 	{
@@ -706,9 +708,15 @@ int32 UWorldPartitionConvertCommandlet::Main(const FString& Params)
 	bGenerateIni = Switches.Contains(TEXT("GenerateIni"));
 	bReportOnly = bGenerateIni || Switches.Contains(TEXT("ReportOnly"));
 	bVerbose = Switches.Contains(TEXT("Verbose"));
-
 	ConversionSuffix = GetConversionSuffix(bOnlyMergeSubLevels);
 
+	FString* FoliageTypePathValue = Arguments.Find(TEXT("FoliageTypePath"));
+	
+	if (FoliageTypePathValue != nullptr)
+	{
+		FoliageTypePath = *FoliageTypePathValue;
+	}
+
 	if (!Switches.Contains(TEXT("AllowCommandletRendering")))
 	{
 		UE_LOG(LogWorldPartitionConvertCommandlet, Error, TEXT("The option \"-AllowCommandletRendering\" is required."));
@@ -843,7 +851,7 @@ int32 UWorldPartitionConvertCommandlet::Main(const FString& Params)
 
 	OnWorldLoaded(MainWorld);
 
-	auto PartitionFoliage = [this, MainWorld](AInstancedFoliageActor* IFA)
+	auto PartitionFoliage = [this, MainWorld](AInstancedFoliageActor* IFA) -> bool
 	{
 		TRACE_CPUPROFILER_EVENT_SCOPE(PartitionFoliage);
 
@@ -851,7 +859,7 @@ int32 UWorldPartitionConvertCommandlet::Main(const FString& Params)
 		int32 NumInstances = 0;
 		int32 NumInstancesProcessed = 0;
 
-		IFA->ForEachFoliageInfo([IFA, &FoliageToAdd, &NumInstances](UFoliageType* FoliageType, FFoliageInfo& FoliageInfo)
+		bool bAddFoliageSucceeded = IFA->ForEachFoliageInfo([IFA, &FoliageToAdd, &NumInstances, this](UFoliageType* FoliageType, FFoliageInfo& FoliageInfo) -> bool
 		{
 			if (FoliageInfo.Type == EFoliageImplType::Actor)
 			{
@@ -860,18 +868,51 @@ int32 UWorldPartitionConvertCommandlet::Main(const FString& Params)
 				return true;
 			}
 
+			UFoliageType* FoliageTypeToAdd = FoliageType;
+
+			if (FoliageType->GetTypedOuter<AInstancedFoliageActor>() != nullptr)
+			{
+				UFoliageType* NewFoliageType = nullptr;
+				
+				if (!FoliageTypePath.IsEmpty())
+				{
+					UObject* FoliageSource = FoliageType->GetSource();
+					const FString BaseAssetName = (FoliageSource != nullptr) ? FoliageSource->GetName() : FoliageType->GetName();
+					FString PackageName = FoliageTypePath / BaseAssetName + TEXT("_FoliageType");
+
+					NewFoliageType = FFoliageEditUtility::DuplicateFoliageTypeToNewPackage(PackageName, FoliageType);
+				}
+
+				if (NewFoliageType == nullptr)
+				{
+					UE_LOG(LogWorldPartitionConvertCommandlet, Error,
+						   TEXT("Level contains embedded FoliageType settings: please save the FoliageType setting assets, ")
+						   TEXT("use the SaveFoliageTypeToContentFolder switch, ")
+						   TEXT("specify FoliageTypePath in configuration file or the commandline."));
+					return false;
+				}
+
+				FoliageTypeToAdd = NewFoliageType;
+				PackagesToSave.Add(NewFoliageType->GetOutermost());
+			}
+
 			if (FoliageInfo.Instances.Num() > 0)
 			{
-				check(FoliageType->GetTypedOuter<AInstancedFoliageActor>() == nullptr);
-				FoliageToAdd.FindOrAdd(FoliageType).Append(FoliageInfo.Instances);
-			
-				NumInstances += FoliageInfo.Instances.Num();
+				check(FoliageTypeToAdd->GetTypedOuter<AInstancedFoliageActor>() == nullptr);
 
-				UE_LOG(LogWorldPartitionConvertCommandlet, Display, TEXT("FoliageType: %s Count: %d"), *FoliageType->GetName(), FoliageInfo.Instances.Num());
+				FoliageToAdd.FindOrAdd(FoliageTypeToAdd).Append(FoliageInfo.Instances);
+				NumInstances += FoliageInfo.Instances.Num();
+				UE_LOG(LogWorldPartitionConvertCommandlet, Display, TEXT("FoliageType: %s Count: %d"), *FoliageTypeToAdd->GetName(), FoliageInfo.Instances.Num());
 			}
+
 			return true;
 		});
 
+		if (!bAddFoliageSucceeded)
+		{
+			return false;
+		}
+
 		IFA->GetLevel()->GetWorld()->DestroyActor(IFA);
 
 		// Add Foliage to those actors
@@ -888,6 +929,8 @@ int32 UWorldPartitionConvertCommandlet::Main(const FString& Params)
 		}
 
 		check(NumInstances == NumInstancesProcessed);
+
+		return true;
 	};
 
 	auto PartitionLandscape = [this, MainWorld](ULandscapeInfo* LandscapeInfo)
@@ -950,7 +993,7 @@ int32 UWorldPartitionConvertCommandlet::Main(const FString& Params)
 		}
 	};
 
-	auto PrepareLevelActors = [this, PartitionFoliage, PartitionLandscape, MainWorldDataLayers](ULevel* Level, TArray<AActor*>& Actors, bool bMainLevel)
+	auto PrepareLevelActors = [this, PartitionFoliage, PartitionLandscape, MainWorldDataLayers](ULevel* Level, TArray<AActor*>& Actors, bool bMainLevel) -> bool
 	{
 		TRACE_CPUPROFILER_EVENT_SCOPE(PrepareLevelActors);
 
@@ -1022,7 +1065,10 @@ int32 UWorldPartitionConvertCommandlet::Main(const FString& Params)
 
 			for (AInstancedFoliageActor* IFA : IFAs)
 			{
-				PartitionFoliage(IFA);
+				if (!PartitionFoliage(IFA))
+				{
+					return false;
+				}
 			}
 		}
 
@@ -1035,6 +1081,8 @@ int32 UWorldPartitionConvertCommandlet::Main(const FString& Params)
 				PartitionLandscape(LandscapeInfo);
 			}
 		}
+
+		return true;
 	};
 
 	// Gather and load sublevels
@@ -1069,7 +1117,12 @@ int32 UWorldPartitionConvertCommandlet::Main(const FString& Params)
 
 	// Prepare levels for conversion
 	DetachDependantLevelPackages(MainLevel);
-	PrepareLevelActors(MainLevel, MainLevel->Actors, true);
+	
+	if (!PrepareLevelActors(MainLevel, MainLevel->Actors, true))
+	{
+		return 1;
+	}
+	
 	PackagesToSave.Add(MainLevel->GetPackage());
 
 	if (bConversionSuffix)
diff --git a/Engine/Source/Editor/WorldPartitionEditor/Private/WorldPartitionEditorModule.cpp b/Engine/Source/Editor/WorldPartitionEditor/Private/WorldPartitionEditorModule.cpp
--- a/Engine/Source/Editor/WorldPartitionEditor/Private/WorldPartitionEditorModule.cpp
+++ b/Engine/Source/Editor/WorldPartitionEditor/Private/WorldPartitionEditorModule.cpp
@@ -425,6 +425,11 @@ FString UWorldPartitionConvertOptions::ToCommandletArgs() const
 	{
 		CommandletArgsBuilder.Append(TEXT(" -OnlyMergeSubLevels"));
 	}
+
+	if (bSaveFoliageTypeToContentFolder)
+	{
+		CommandletArgsBuilder.Append(TEXT(" -FoliageTypePath=/Game/FoliageTypes"));
+	}
 	
 	return CommandletArgsBuilder.ToString();
 }
diff --git a/Engine/Source/Editor/WorldPartitionEditor/Public/WorldPartition/WorldPartitionConvertOptions.h b/Engine/Source/Editor/WorldPartitionEditor/Public/WorldPartition/WorldPartitionConvertOptions.h
--- a/Engine/Source/Editor/WorldPartitionEditor/Public/WorldPartition/WorldPartitionConvertOptions.h
+++ b/Engine/Source/Editor/WorldPartitionEditor/Public/WorldPartition/WorldPartitionConvertOptions.h
@@ -41,5 +41,8 @@ public:
 	UPROPERTY(EditAnywhere, AdvancedDisplay, Category = Convert)
 	bool bOnlyMergeSubLevels;
 
+	UPROPERTY(EditAnywhere, AdvancedDisplay, Category = Convert)
+	bool bSaveFoliageTypeToContentFolder = true;
+
 	FString LongPackageName;
 };
\ No newline at end of file
diff --git a/Engine/Source/Runtime/Foliage/Private/InstancedFoliage.cpp b/Engine/Source/Runtime/Foliage/Private/InstancedFoliage.cpp
--- a/Engine/Source/Runtime/Foliage/Private/InstancedFoliage.cpp
+++ b/Engine/Source/Runtime/Foliage/Private/InstancedFoliage.cpp
@@ -2902,15 +2902,17 @@ void AInstancedFoliageActor::GetOverlappingMeshCounts(const FSphere& Sphere, TMa
 	}
 }
 
-void AInstancedFoliageActor::ForEachFoliageInfo(TFunctionRef<bool(UFoliageType* FoliageType, FFoliageInfo& FoliageInfo)> InOperation)
+bool AInstancedFoliageActor::ForEachFoliageInfo(TFunctionRef<bool(UFoliageType* FoliageType, FFoliageInfo& FoliageInfo)> InOperation)
 {
 	for (auto& Pair : FoliageInfos)
 	{
 		if (!InOperation(Pair.Key, Pair.Value.Get()))
 		{
-			return;
+			return false;
 		}
 	}
+
+	return true;
 }
 
 TUniqueObj<FFoliageInfo>& AInstancedFoliageActor::AddFoliageInfo(UFoliageType* FoliageType)
diff --git a/Engine/Source/Runtime/Foliage/Public/InstancedFoliageActor.h b/Engine/Source/Runtime/Foliage/Public/InstancedFoliageActor.h
--- a/Engine/Source/Runtime/Foliage/Public/InstancedFoliageActor.h
+++ b/Engine/Source/Runtime/Foliage/Public/InstancedFoliageActor.h
@@ -39,7 +39,7 @@ private:
 	TMap<UFoliageType*, TUniqueObj<FFoliageInfo>> FoliageInfos;
 
 public:
-	FOLIAGE_API void ForEachFoliageInfo(TFunctionRef<bool(UFoliageType* FoliageType, FFoliageInfo& FoliageInfo)> InOperation);
+	FOLIAGE_API bool ForEachFoliageInfo(TFunctionRef<bool(UFoliageType* FoliageType, FFoliageInfo& FoliageInfo)> InOperation);
 	FOLIAGE_API const TMap<UFoliageType*, TUniqueObj<FFoliageInfo>>& GetFoliageInfos() const { return FoliageInfos; }
 	FOLIAGE_API TUniqueObj<FFoliageInfo>& AddFoliageInfo(UFoliageType* FoliageType);
 	FOLIAGE_API TUniqueObj<FFoliageInfo>& AddFoliageInfo(UFoliageType* FoliageType, TUniqueObj<FFoliageInfo>&& FoliageInfo);
