Fixed a bug when generating the URL for the gamemode of externally launched process.

The issue happens when users are calling GEditor->RequestPlaySession with a GamemodeOverride.
The URL need to contain the Unreal path for the game mode to be found.

diff --git a/Engine/Source/Editor/UnrealEd/Private/PlayLevelNewProcess.cpp b/Engine/Source/Editor/UnrealEd/Private/PlayLevelNewProcess.cpp
--- a/Engine/Source/Editor/UnrealEd/Private/PlayLevelNewProcess.cpp
+++ b/Engine/Source/Editor/UnrealEd/Private/PlayLevelNewProcess.cpp
@@ -71,7 +71,7 @@ void UEditorEngine::LaunchNewProcess(const FRequestPlaySessionParams& InParams,
 
 	if (InParams.GameModeOverride)
 	{
-		UnrealURLParams += FString::Printf(TEXT("?game=%s"), *InParams.GameModeOverride->GetName());
+		UnrealURLParams += FString::Printf(TEXT("?game=%s"), *InParams.GameModeOverride->GetPathName());
 	}
 
 	if (bIsDedicatedServer)
