UDataLayer data is split between UDataLayerAsset & UDataLayerInstance (change 1/2)
Underlying system still work with UDataLayer, deprecation of UDataLayer coming in a second CL.
UDataLayerAsset contains non-world specific data
UDataLayer Instance contains world-specifc data and referes a UDataLAsset
Actors refers UDataLayerAssets
AWorldDataLayers refers UDataLayerInstances
Use a Commandlet to convert existing UDataLayer to UDataInstance + UDataLayerAsset
Added a map to WorldDataLayers to allow retrieving a DataLayerAsset from the UDataLayer FName it was converted from (blueprint deprecation).

diff --git a/Engine/Source/Editor/DataLayerEditor/Private/DataLayer/DataLayerFactory.cpp b/Engine/Source/Editor/DataLayerEditor/Private/DataLayer/DataLayerFactory.cpp
new file mode 100644
--- /dev/null
+++ b/Engine/Source/Editor/DataLayerEditor/Private/DataLayer/DataLayerFactory.cpp
@@ -0,0 +1,20 @@
+// Copyright Epic Games, Inc. All Rights Reserved.
+
+#include "DataLayer/DataLayerFactory.h"
+
+#include "WorldPartition/DataLayer/DataLayerAsset.h"
+
+UDataLayerFactory::UDataLayerFactory(const FObjectInitializer& ObjectInitializer)
+	: Super(ObjectInitializer)
+
+{
+	SupportedClass = UDataLayerAsset::StaticClass();
+	bCreateNew = true;
+	bEditorImport = false;
+	bEditAfterNew = true;
+}
+
+UObject* UDataLayerFactory::FactoryCreateNew(UClass* Class, UObject* InParent, FName InName, EObjectFlags Flags, UObject* Context, FFeedbackContext* Warn)
+{
+	return NewObject<UDataLayerAsset>(InParent, InName, Flags);
+}
\ No newline at end of file
diff --git a/Engine/Source/Editor/DataLayerEditor/Public/DataLayer/DataLayerFactory.h b/Engine/Source/Editor/DataLayerEditor/Public/DataLayer/DataLayerFactory.h
new file mode 100644
--- /dev/null
+++ b/Engine/Source/Editor/DataLayerEditor/Public/DataLayer/DataLayerFactory.h
@@ -0,0 +1,22 @@
+// Copyright Epic Games, Inc. All Rights Reserved.
+
+#pragma once
+
+#include "CoreMinimal.h"
+#include "UObject/ObjectMacros.h"
+#include "UObject/UObjectGlobals.h"
+#include "UObject/Object.h"
+
+#include "Factories/Factory.h"
+
+#include "DataLayerFactory.generated.h"
+
+UCLASS(hidecategories = Object, MinimalAPI)
+class UDataLayerFactory : public UFactory
+{
+	GENERATED_UCLASS_BODY()
+
+	// Begin UFactory Interface
+	virtual UObject* FactoryCreateNew(UClass* Class, UObject* InParent, FName Name, EObjectFlags Flags, UObject* Context, FFeedbackContext* Warn) override;
+	// Begin UFactory Interface	
+};
\ No newline at end of file
diff --git a/Engine/Source/Editor/EditorStyle/Private/StarshipStyle.cpp b/Engine/Source/Editor/EditorStyle/Private/StarshipStyle.cpp
--- a/Engine/Source/Editor/EditorStyle/Private/StarshipStyle.cpp
+++ b/Engine/Source/Editor/EditorStyle/Private/StarshipStyle.cpp
@@ -6064,6 +6064,7 @@ void FStarshipEditorStyle::FStyle::SetupClassIconsAndThumbnails()
 			{TEXT("CurveBase")},
 			{TEXT("DataAsset")},
 			{TEXT("DataTable")},
+			{TEXT("DataLayerAsset")},
 			{TEXT("DecalActor")},
 			{TEXT("DefaultPawn")},
 			{TEXT("DialogueVoice")},
diff --git a/Engine/Source/Editor/UnrealEd/Classes/Commandlets/WorldPartitionDataLayerToAssetCommandLet.h b/Engine/Source/Editor/UnrealEd/Classes/Commandlets/WorldPartitionDataLayerToAssetCommandLet.h
new file mode 100644
--- /dev/null
+++ b/Engine/Source/Editor/UnrealEd/Classes/Commandlets/WorldPartitionDataLayerToAssetCommandLet.h
@@ -0,0 +1,149 @@
+// Copyright Epic Games, Inc. All Rights Reserved.
+
+#pragma once
+
+#include "CoreMinimal.h"
+#include "UObject/ObjectMacros.h"
+#include "UObject/UObjectAnnotation.h"
+#include "Commandlets/Commandlet.h"
+#include "Commandlets/WorldPartitionCommandletHelpers.h"
+#include "PackageSourceControlHelper.h"
+#include "WorldPartition/DataLayer/DataLayerAsset.h"
+#include "WorldPartition/DataLayer/DataLayerInstance.h"
+
+#include "WorldPartitionDataLayerToAssetCommandLet.generated.h"
+
+UNREALED_API DECLARE_LOG_CATEGORY_EXTERN(LogDataLayerToAssetCommandlet, Log, All);
+
+class UDataLayer;
+class UDataLayerAsset;
+class UDataLayerInstance;
+class UDataLayerFactory;
+
+UCLASS(Transient)
+class UDataLayerConversionInfo : public UObject
+{
+	GENERATED_BODY()
+
+	friend class UDataLayerToAssetCommandletContext;
+
+public:
+	bool IsConverting() const { return DataLayerToConvert != nullptr; }
+	bool IsAPreviousConversion() const { return CurrentConvertingInfo != nullptr; }
+	bool IsConverted() const { return DataLayerInstance != nullptr && PreviousConversionsInfo.IsEmpty(); }
+
+	TArray<TWeakObjectPtr<UDataLayerConversionInfo>> const& GetPreviousConversions() const { return PreviousConversionsInfo; }
+	const TWeakObjectPtr<UDataLayerConversionInfo>& GetCurrentConversion() const { return CurrentConvertingInfo; }
+
+	void SetDataLayerToConvert(const TObjectPtr<UDataLayer>& InDataLayerToConvert);
+	void SetDataLayerInstance(const TObjectPtr<UDataLayerInstance>& InDataLayerInstance);
+
+	UPROPERTY()
+	TObjectPtr<UDataLayerAsset> DataLayerAsset;
+
+	UPROPERTY()
+	TObjectPtr<UDataLayer> DataLayerToConvert;
+
+	UPROPERTY()
+	TObjectPtr<UDataLayerInstance> DataLayerInstance;
+
+private:
+	UPROPERTY()
+	TArray<TWeakObjectPtr<UDataLayerConversionInfo>> PreviousConversionsInfo;
+
+	UPROPERTY()
+	TWeakObjectPtr<UDataLayerConversionInfo> CurrentConvertingInfo;
+};
+
+UCLASS(Transient)
+class UDataLayerToAssetCommandletContext : public UObject
+{
+	GENERATED_UCLASS_BODY()
+
+public:
+	const TArray<TObjectPtr<UDataLayerConversionInfo>>& GetDataLayerConversionInfos() const { return DataLayerConversionInfo; }
+	const TArray<TWeakObjectPtr<UDataLayerConversionInfo>>& GetConvertingDataLayerConversionInfo() const { return ConvertingDataLayerInfo; }
+
+	UDataLayerConversionInfo* GetDataLayerConversionInfo(const TObjectPtr<UDataLayer>& DataLayer) const;
+	UDataLayerConversionInfo* GetDataLayerConversionInfo(const UDataLayerAsset* DataLayerAsset) const;
+	UDataLayerConversionInfo* GetDataLayerConversionInfo(const TObjectPtr<UDataLayerInstance>& DataLayerInstance) const;
+	UDataLayerConversionInfo* GetDataLayerConversionInfo(const FActorDataLayer& ActorDataLayer) const;
+
+	UDataLayerConversionInfo* StoreExistingDataLayer(FAssetData& AssetData);
+	UDataLayerConversionInfo* StoreDataLayerAssetConversion(const TObjectPtr<UDataLayer>& DataLayerToConvert, TObjectPtr<UDataLayerAsset>& NewDataLayerAsset);
+	UDataLayerConversionInfo* StoreDataLayerInstanceConversion(const UDataLayerAsset* DataLayerAsset, UDataLayerInstance* NewDataLayerInstance);
+	
+	bool SetPreviousConversions(UDataLayerConversionInfo* CurrentConversion, TArray<TWeakObjectPtr<UDataLayerConversionInfo>>&& PreviousConversions);
+
+	bool FindDataLayerConversionInfos(FName DataLayerAssetName, TArray<TWeakObjectPtr<UDataLayerConversionInfo>>& OutConversionInfos) const;
+
+	void LogConversionInfos() const;
+
+private:
+	UPROPERTY()
+	TArray<TObjectPtr<UDataLayerConversionInfo>> DataLayerConversionInfo;
+
+	UPROPERTY()
+	TArray<TWeakObjectPtr<UDataLayerConversionInfo>> ConvertingDataLayerInfo;
+};
+
+UCLASS(Config = Engine)
+class UNREALED_API UDataLayerToAssetCommandlet : public UCommandlet
+{
+	GENERATED_UCLASS_BODY()
+
+public:
+	enum EReturnCode
+	{
+		Success = 0,
+		CommandletInitializationError,
+		DataLayerConversionError,
+		ActorDataLayerRemappingError,
+	};
+
+	//~ Begin UCommandlet Interface
+	virtual int32 Main(const FString& Params) override;
+	//~ End UCommandlet Interface
+
+private:
+	bool InitializeFromCommandLine(TArray<FString>& Tokens, TArray<FString> const& Switches, TMap<FString, FString> const& Params);
+
+	bool BuildConversionInfos(TStrongObjectPtr<UDataLayerToAssetCommandletContext>& CommandletContext);
+	bool CreateConversionFromDataLayer(TStrongObjectPtr<UDataLayerToAssetCommandletContext>& CommandletContext, const TObjectPtr<UDataLayer>& DataLayer);
+	TObjectPtr<UDataLayerAsset> GetOrCreateDataLayerAssetForConversion(TStrongObjectPtr<UDataLayerToAssetCommandletContext>& CommandletContext, FName DataLayerName);
+
+	bool ResolvePreviousConversionsToCurrent(TStrongObjectPtr<UDataLayerToAssetCommandletContext>& CommandletContext);
+
+	bool RemapActorDataLayersToAssets(TStrongObjectPtr<UDataLayerToAssetCommandletContext>& CommandletContext);
+	uint32 RemapActorDataLayers(TStrongObjectPtr<UDataLayerToAssetCommandletContext>& CommandletContext, AActor* Actor);
+	uint32 RemapDataLayersAssetsFromPreviousConversions(TStrongObjectPtr<UDataLayerToAssetCommandletContext>& CommandletContext, AActor* Actor);
+
+	bool CreateDataLayerInstances(TStrongObjectPtr<UDataLayerToAssetCommandletContext>& CommandletContext);
+	
+	bool DeletePreviousConversionsData(TStrongObjectPtr<UDataLayerToAssetCommandletContext>& CommandletContext);
+
+	bool CommitConversion(TStrongObjectPtr<UDataLayerToAssetCommandletContext>& CommandletContext);
+
+	FString GetConversionFolder() const;
+	bool IsAssetInConversionFolder(const TObjectPtr<UDataLayerAsset> DataLayerAsset);
+
+	UPROPERTY(Config)
+	FString DestinationFolder;
+
+	UPROPERTY(Transient)
+	FString ConversionFolder;
+
+	UPROPERTY(Config)
+	bool bPerformSavePackages = true;
+
+	UPROPERTY(Config)
+	bool bIgnoreActorLoadingErrors = false;
+
+	UPROPERTY()
+	TObjectPtr<UDataLayerFactory> DataLayerFactory;
+
+	UPROPERTY(Transient)
+	TObjectPtr<UWorld> MainWorld;
+
+	FPackageSourceControlHelper PackageHelper;
+};
diff --git a/Engine/Source/Editor/UnrealEd/Private/Commandlets/WorldPartitionCommandletHelpers.cpp b/Engine/Source/Editor/UnrealEd/Private/Commandlets/WorldPartitionCommandletHelpers.cpp
new file mode 100644
--- /dev/null
+++ b/Engine/Source/Editor/UnrealEd/Private/Commandlets/WorldPartitionCommandletHelpers.cpp
@@ -0,0 +1,185 @@
+// Copyright Epic Games, Inc. All Rights Reserved.
+
+#include "Commandlets/WorldPartitionCommandletHelpers.h"
+
+#include "PackageSourceControlHelper.h"
+#include "SourceControlHelpers.h"
+#include "UObject/SavePackage.h"
+
+namespace WorldPartitionCommandletHelpers
+{
+	DEFINE_LOG_CATEGORY(LogWorldParitionCommandletUtils)
+
+		UWorld* LoadAndInitWorld(const FString& LevelToLoad)
+	{
+		UWorld* World = LoadWorld(LevelToLoad);
+		if (World != nullptr)
+		{
+			if (InitLevel(World) != nullptr)
+			{
+				return World;
+			}
+		}
+
+		return nullptr;
+	}
+
+	UWorld* LoadWorld(const FString& LevelToLoad)
+	{
+		TRACE_CPUPROFILER_EVENT_SCOPE(UWorldPartitionConvertCommandlet::LoadWorld);
+
+		SET_WARN_COLOR(COLOR_WHITE);
+		UE_LOG(LogWorldParitionCommandletUtils, Log, TEXT("Loading level %s."), *LevelToLoad);
+		CLEAR_WARN_COLOR();
+
+		// This will convert incomplete package name to a fully qualified path, avoiding calling it several times (takes ~50s)
+		FString FullLevelToLoadPath;
+		if (!FPackageName::SearchForPackageOnDisk(LevelToLoad, &FullLevelToLoadPath))
+		{
+			UE_LOG(LogWorldParitionCommandletUtils, Error, TEXT("Unknown level '%s'"), *LevelToLoad);
+			return nullptr;
+		}
+
+		UPackage* MapPackage = LoadPackage(nullptr, *FullLevelToLoadPath, LOAD_None);
+		if (!MapPackage)
+		{
+			UE_LOG(LogWorldParitionCommandletUtils, Error, TEXT("Error loading %s."), *FullLevelToLoadPath);
+			return nullptr;
+		}
+
+		return UWorld::FindWorldInPackage(MapPackage);
+	}
+
+	ULevel* InitLevel(UWorld* World)
+	{
+		TRACE_CPUPROFILER_EVENT_SCOPE(UWorldPartitionConvertCommandlet::InitWorld);
+
+		SET_WARN_COLOR(COLOR_WHITE);
+		UE_LOG(LogWorldParitionCommandletUtils, Log, TEXT("Initializing level %s."), *World->GetName());
+		CLEAR_WARN_COLOR();
+
+		// Setup the world.
+		World->WorldType = EWorldType::Editor;
+		World->AddToRoot();
+		if (!World->bIsWorldInitialized)
+		{
+			UWorld::InitializationValues IVS;
+			IVS.RequiresHitProxies(false);
+			IVS.ShouldSimulatePhysics(false);
+			IVS.EnableTraceCollision(false);
+			IVS.CreateNavigation(false);
+			IVS.CreateAISystem(false);
+			IVS.AllowAudioPlayback(false);
+			IVS.CreatePhysicsScene(true);
+
+			World->InitWorld(IVS);
+			World->PersistentLevel->UpdateModelComponents();
+			World->UpdateWorldComponents(true, false);
+
+			World->FlushLevelStreaming(EFlushLevelStreamingType::Full);
+		}
+
+		return World->PersistentLevel;
+	}
+
+	bool Checkout(const TArray<UPackage*>& PackagesToCheckout, FPackageSourceControlHelper& SCHelper)
+	{
+		TRACE_CPUPROFILER_EVENT_SCOPE(CheckoutPackages);
+
+		UE_LOG(LogWorldParitionCommandletUtils, Log, TEXT("Checking out %d Packages."), PackagesToCheckout.Num());
+		for (UPackage* Package : PackagesToCheckout)
+		{
+			if (!Checkout(Package, SCHelper))
+			{
+				return false;
+			}
+		}
+
+		return true;
+	}
+
+	bool Checkout(UPackage* PackagesToCheckout, FPackageSourceControlHelper& SCHelper)
+	{
+		UE_LOG(LogWorldParitionCommandletUtils, Verbose, TEXT("Adding Package %s."), *SourceControlHelpers::PackageFilename(PackagesToCheckout));
+
+		FString PackageFileName = SourceControlHelpers::PackageFilename(PackagesToCheckout);
+		if (FPlatformFileManager::Get().GetPlatformFile().FileExists(*PackageFileName))
+		{
+			if (!SCHelper.Checkout(PackagesToCheckout))
+			{
+				return false;
+			}
+		}
+
+		return true;
+	}
+
+	bool Save(const TArray<UPackage*>& PackagesToSave, FPackageSourceControlHelper& SCHelper)
+	{
+		TRACE_CPUPROFILER_EVENT_SCOPE(SavePackages);
+
+		UE_LOG(LogWorldParitionCommandletUtils, Log, TEXT("Saving %d packages."), PackagesToSave.Num());
+		for (UPackage* PackageToSave : PackagesToSave)
+		{
+			if (!Save(PackagesToSave, SCHelper))
+			{
+				return false;
+			}
+		}	
+		
+		return true;
+	}
+
+	bool Save(UPackage* PackageToSave, FPackageSourceControlHelper& SCHelper)
+	{
+		FString PackageFileName = SourceControlHelpers::PackageFilename(PackageToSave);
+		FSavePackageArgs SaveArgs;
+		SaveArgs.TopLevelFlags = RF_Standalone;
+
+		UE_LOG(LogWorldParitionCommandletUtils, Verbose, TEXT("Saving Package %s."), *PackageFileName);
+		return UPackage::SavePackage(PackageToSave, nullptr, *PackageFileName, SaveArgs);
+	}
+
+	bool AddToSourceControl(const TArray<UPackage*>& PackagesToAdd, FPackageSourceControlHelper& SCHelper)
+	{
+		TRACE_CPUPROFILER_EVENT_SCOPE(AddPackagesToSourceControl);
+
+		for (UPackage* PackageToAdd : PackagesToAdd)
+		{
+			if (!AddToSourceControl(PackageToAdd, SCHelper))
+			{
+				return false;
+			}
+		}
+
+		return true;
+	}
+
+	bool AddToSourceControl(UPackage* PackageToAdd, FPackageSourceControlHelper& SCHelper)
+	{
+		UE_LOG(LogWorldParitionCommandletUtils, Verbose, TEXT("Adding Package %s."), *SourceControlHelpers::PackageFilename(PackageToAdd));
+
+		return SCHelper.AddToSourceControl(PackageToAdd);
+	}
+
+	bool Delete(const TArray<UPackage*>& PackagesToDelete, FPackageSourceControlHelper& SCHelper)
+	{
+		TRACE_CPUPROFILER_EVENT_SCOPE(DeletePackages);
+
+		for (UPackage* PackageToDelete : PackagesToDelete)
+		{
+			if (!Delete(PackageToDelete, SCHelper))
+			{
+				return false;
+			}
+		}
+		return true;
+	}
+
+	bool Delete(UPackage* PackageToDelete, FPackageSourceControlHelper& SCHelper)
+	{
+		UE_LOG(LogWorldParitionCommandletUtils, Verbose, TEXT("Deleting Package %s."), *SourceControlHelpers::PackageFilename(PackageToDelete));
+
+		return SCHelper.Delete(PackageToDelete);
+	}
+}
diff --git a/Engine/Source/Editor/UnrealEd/Private/Commandlets/WorldPartitionDataLayerToAssetCommandlet.cpp b/Engine/Source/Editor/UnrealEd/Private/Commandlets/WorldPartitionDataLayerToAssetCommandlet.cpp
new file mode 100644
--- /dev/null
+++ b/Engine/Source/Editor/UnrealEd/Private/Commandlets/WorldPartitionDataLayerToAssetCommandlet.cpp
@@ -0,0 +1,654 @@
+// Copyright Epic Games, Inc. All Rights Reserved.
+
+/*=============================================================================
+ UDataLayerToAssetCommandlet.cpp: Commandlet used to convert a partionned ULevel's data layers to assets
+=============================================================================*/
+
+#include "Commandlets/WorldPartitionDataLayerToAssetCommandLet.h"
+
+#include "Algo/Copy.h"
+#include "Algo/Accumulate.h"
+#include "AssetToolsModule.h"
+#include "DataLayer/DataLayerFactory.h"
+#include "Logging/LogVerbosity.h"
+#include "ProfilingDebugging/ScopedTimers.h"
+#include "WorldPartition/DataLayer/WorldDataLayers.h"
+#include "WorldPartition/WorldPartitionActorDesc.h"
+#include "WorldPartition/WorldPartitionHelpers.h"
+
+DEFINE_LOG_CATEGORY(LogDataLayerToAssetCommandlet);
+
+void UDataLayerConversionInfo::SetDataLayerToConvert(const TObjectPtr<UDataLayer>& InDataLayerToConvert)
+{
+	DataLayerToConvert = InDataLayerToConvert;
+
+	DataLayerAsset->DataLayerType = DataLayerToConvert->IsRuntime() ? EDataLayerType::Runtime : EDataLayerType::Editor;
+	DataLayerAsset->DebugColor = DataLayerToConvert->GetDebugColor();
+}
+
+void UDataLayerConversionInfo::SetDataLayerInstance(const TObjectPtr<UDataLayerInstance>& InDataLayerInstance)
+{
+	DataLayerInstance = InDataLayerInstance;
+
+	if (DataLayerToConvert != nullptr)
+	{
+		DataLayerInstance->DataLayerAsset = DataLayerAsset;
+		DataLayerInstance->bIsVisible = DataLayerToConvert->IsVisible();
+		DataLayerInstance->bIsInitiallyVisible = DataLayerToConvert->IsInitiallyVisible();
+		DataLayerInstance->bIsInitiallyLoadedInEditor = DataLayerToConvert->IsInitiallyLoadedInEditor();
+		DataLayerInstance->bIsLocked = DataLayerToConvert->IsLocked();
+		DataLayerInstance->InitialRuntimeState = DataLayerToConvert->GetInitialRuntimeState();
+	}
+	else
+	{
+		// Check the DataLayerInstance was already properly converted
+		check(DataLayerInstance->DataLayerAsset == DataLayerAsset)
+	}
+}
+
+UDataLayerToAssetCommandletContext::UDataLayerToAssetCommandletContext(const FObjectInitializer& ObjectInitializer)
+	: Super(ObjectInitializer)
+{
+
+}
+
+UDataLayerConversionInfo* UDataLayerToAssetCommandletContext::GetDataLayerConversionInfo(const TObjectPtr<UDataLayer>& DataLayer) const
+{
+	const TObjectPtr<UDataLayerConversionInfo>* entry = DataLayerConversionInfo.FindByPredicate([&DataLayer](const TObjectPtr<UDataLayerConversionInfo>& Other) { return DataLayer == Other->DataLayerToConvert; });
+	return entry != nullptr ? entry->Get() : nullptr;
+}
+
+UDataLayerConversionInfo* UDataLayerToAssetCommandletContext::GetDataLayerConversionInfo(const UDataLayerAsset* DataLayerAsset) const
+{
+	const TObjectPtr<UDataLayerConversionInfo>* entry = DataLayerConversionInfo.FindByPredicate([&DataLayerAsset](const TObjectPtr<UDataLayerConversionInfo>& Other) { return DataLayerAsset == Other->DataLayerAsset; });
+	return entry != nullptr ? entry->Get() : nullptr;
+}
+
+UDataLayerConversionInfo* UDataLayerToAssetCommandletContext::GetDataLayerConversionInfo(const TObjectPtr<UDataLayerInstance>& DataLayerInstance) const
+{
+	const TObjectPtr<UDataLayerConversionInfo>* entry = DataLayerConversionInfo.FindByPredicate([&DataLayerInstance](const TObjectPtr<UDataLayerConversionInfo>& Other) { return DataLayerInstance == Other->DataLayerInstance; });
+	return entry != nullptr ? entry->Get() : nullptr;
+}
+
+UDataLayerConversionInfo* UDataLayerToAssetCommandletContext::GetDataLayerConversionInfo(const FActorDataLayer& ActorDataLayer) const
+{
+	const TObjectPtr<UDataLayerConversionInfo>* entry = DataLayerConversionInfo.FindByPredicate([&ActorDataLayer](const TObjectPtr<UDataLayerConversionInfo>& Other)
+	{ 
+		return Other->DataLayerToConvert != nullptr && ActorDataLayer.Name == Other->DataLayerToConvert->GetFName();
+	});
+	
+	return entry != nullptr ? entry->Get() : nullptr;
+}
+
+UDataLayerConversionInfo* UDataLayerToAssetCommandletContext::StoreExistingDataLayer(FAssetData& AssetData)
+{
+	UDataLayerConversionInfo* ConversionInfo = NewObject<UDataLayerConversionInfo>();
+	ConversionInfo->DataLayerAsset = CastChecked<UDataLayerAsset>(AssetData.GetAsset());
+
+	UE_LOG(LogDataLayerToAssetCommandlet, Verbose, TEXT("Data Layer Asset %s discovered."),
+		*ConversionInfo->DataLayerAsset->GetFullName());
+	return DataLayerConversionInfo.Add_GetRef(ConversionInfo).Get();
+}
+
+UDataLayerConversionInfo* UDataLayerToAssetCommandletContext::StoreDataLayerAssetConversion(const TObjectPtr<UDataLayer>& DataLayerToConvert, TObjectPtr<UDataLayerAsset>& NewDataLayerAsset)
+{
+	if (UDataLayerConversionInfo* ConversionInfo = GetDataLayerConversionInfo(DataLayerToConvert))
+	{
+		UE_LOG(LogDataLayerToAssetCommandlet, Error, TEXT("Failed to assign asset %s to data Layer %s. The data layer is already associated to asset %s."), 
+			*NewDataLayerAsset->GetFullName() , *DataLayerToConvert->GetDataLayerLabel().ToString() , *ConversionInfo->DataLayerAsset->GetFullName());
+		return nullptr;
+	}
+
+	UDataLayerConversionInfo* ConversionInfo = GetDataLayerConversionInfo(NewDataLayerAsset);
+	if (ConversionInfo != nullptr)
+	{
+		if (ConversionInfo->DataLayerToConvert != nullptr)
+		{
+			UE_LOG(LogDataLayerToAssetCommandlet, Error, TEXT("Failed to assign asset %s to data Layer %s. The asset is already associated to data layer %s."),
+				*NewDataLayerAsset->GetFullName(), *DataLayerToConvert->GetDataLayerLabel().ToString(), *ConversionInfo->DataLayerToConvert->GetDataLayerLabel().ToString());
+			return nullptr;
+		}
+	}
+	else
+	{
+		ConversionInfo = NewObject<UDataLayerConversionInfo>();
+		ConversionInfo->DataLayerAsset = NewDataLayerAsset;
+		DataLayerConversionInfo.Add(ConversionInfo);
+	}
+
+	ConversionInfo->SetDataLayerToConvert(DataLayerToConvert);
+
+	UE_LOG(LogDataLayerToAssetCommandlet, Log, TEXT("Data Layer Asset %s is associated to Data Layer %s for conversion"),
+		*ConversionInfo->DataLayerAsset->GetFullName(), *ConversionInfo->DataLayerToConvert->GetDataLayerLabel().ToString());
+	return ConvertingDataLayerInfo.Add_GetRef(ConversionInfo).Get();
+}
+
+UDataLayerConversionInfo* UDataLayerToAssetCommandletContext::StoreDataLayerInstanceConversion(const UDataLayerAsset* DataLayerAsset, UDataLayerInstance* NewDataLayerInstance)
+{
+	if (UDataLayerConversionInfo* ConversionInfo = GetDataLayerConversionInfo(NewDataLayerInstance))
+	{
+		UE_LOG(LogDataLayerToAssetCommandlet, Error, TEXT("Failed to assign asset %s to data layer instance %s. The instance is already associated to asset %s."),
+			*DataLayerAsset->GetFullName(), *NewDataLayerInstance->GetName(), *ConversionInfo->DataLayerAsset->GetFullName());
+		return nullptr;
+	}
+
+	if (UDataLayerConversionInfo* ConversionInfo = GetDataLayerConversionInfo(DataLayerAsset)) 
+	{
+		check(ConversionInfo->DataLayerInstance == nullptr);
+		ConversionInfo->SetDataLayerInstance(NewDataLayerInstance);
+
+		UE_LOG(LogDataLayerToAssetCommandlet, Log, TEXT("Data Layer Instance %s is associated to Data Layer Asset %s for conversion"),
+			*ConversionInfo->DataLayerInstance->GetName(), *ConversionInfo->DataLayerAsset->GetFullName());
+		return ConversionInfo;
+	}
+
+	UE_LOG(LogDataLayerToAssetCommandlet, Error, TEXT("Failed to assign asset %s to data layer instance %s. The asset was not retrieved in conversion data."),
+		*DataLayerAsset->GetFullName(), *NewDataLayerInstance->GetName());
+	return nullptr;
+}
+
+bool UDataLayerToAssetCommandletContext::SetPreviousConversions(UDataLayerConversionInfo* CurrentConversion, TArray<TWeakObjectPtr<UDataLayerConversionInfo>>&& PreviousConversions)
+{
+	check(CurrentConversion->CurrentConvertingInfo == nullptr);
+	check(CurrentConversion->PreviousConversionsInfo.IsEmpty());
+
+	CurrentConversion->PreviousConversionsInfo = MoveTemp(PreviousConversions);
+
+	uint32 ErrorCount = 0;
+	for (TWeakObjectPtr<UDataLayerConversionInfo>& PreviousConversion : CurrentConversion->PreviousConversionsInfo)
+	{
+		check(PreviousConversion->CurrentConvertingInfo == nullptr);
+		check(PreviousConversion->PreviousConversionsInfo.IsEmpty());
+
+		if (PreviousConversion->DataLayerInstance != nullptr)
+		{
+			ErrorCount++;
+			UE_LOG(LogDataLayerToAssetCommandlet, Error,
+				TEXT("DataLayer %s was already converted but is still to be converted. Re-Sync Data to a clean conversion or pre-conversion state and re-run the commandlet"),
+				*CurrentConversion->DataLayerAsset->GetFullName());
+		}
+
+		PreviousConversion->CurrentConvertingInfo = CurrentConversion;
+	}
+
+	ConvertingDataLayerInfo.AddUnique(CurrentConversion);
+
+	return ErrorCount == 0;
+}
+
+bool UDataLayerToAssetCommandletContext::FindDataLayerConversionInfos(FName DataLayerAssetName, TArray<TWeakObjectPtr<UDataLayerConversionInfo>>& OutConversionInfos) const
+{
+	OutConversionInfos.Empty();
+
+	for (TObjectPtr<UDataLayerConversionInfo> const&  ConversionInfo : DataLayerConversionInfo)
+	{
+		if (ConversionInfo->DataLayerAsset->GetFName() == DataLayerAssetName)
+		{
+			OutConversionInfos.Add(ConversionInfo.Get());
+		}
+	}
+
+	return !OutConversionInfos.IsEmpty();
+}
+
+void UDataLayerToAssetCommandletContext::LogConversionInfos() const
+{
+	if (LogDataLayerToAssetCommandlet.GetVerbosity() >= ELogVerbosity::Verbose)
+	{
+		for (const TObjectPtr<UDataLayerConversionInfo>& info : DataLayerConversionInfo)
+		{
+			FString ConflictingConversionString = TEXT("");
+			if(!info->PreviousConversionsInfo.IsEmpty())
+			{
+				ConflictingConversionString  = FString::JoinBy(info->PreviousConversionsInfo, TEXT(", "), [](const TWeakObjectPtr<UDataLayerConversionInfo>& ConflictingInfo)
+				{ 
+						return ConflictingInfo->DataLayerAsset->GetFullName(); 
+				});
+			}
+			
+
+			UE_LOG(LogDataLayerToAssetCommandlet, Verbose, TEXT("[Conversion Info] Data Layer %s\t\tData Layer Asset: %s\t\t\t\t\t\tData Layer Instance: %s\t\tConverting By: %s\t\t\t\t\t\tConflicting Previous Conversion: %s"),
+				info->DataLayerToConvert != nullptr ? *info->DataLayerToConvert->GetDataLayerLabel().ToString() : TEXT("None"),
+				info->DataLayerAsset != nullptr ? *info->DataLayerAsset->GetFullName() : TEXT("None"),
+				info->DataLayerInstance != nullptr ? *info->DataLayerInstance->GetName() : TEXT("None"),
+				info->CurrentConvertingInfo != nullptr ? *info->CurrentConvertingInfo->DataLayerAsset->GetFullName() : TEXT("None"),
+				*ConflictingConversionString);
+
+		}
+	}
+}
+
+UDataLayerToAssetCommandlet::UDataLayerToAssetCommandlet(const FObjectInitializer& ObjectInitializer)
+	: Super(ObjectInitializer),
+	DataLayerFactory(NewObject<UDataLayerFactory>())
+{}
+
+int32 UDataLayerToAssetCommandlet::Main(const FString& Params)
+{
+	TRACE_CPUPROFILER_EVENT_SCOPE(UWorldPartitionConvertCommandlet::Main);
+
+	UE_SCOPED_TIMER(TEXT("Data Layer Conversion"), LogDataLayerToAssetCommandlet, Display);
+
+	ON_SCOPE_EXIT
+	{
+		if (MainWorld != nullptr)
+		{
+			const bool bBroadcastWorldDestroyedEvent = false;
+			MainWorld->DestroyWorld(bBroadcastWorldDestroyedEvent);
+		}
+	};
+
+	TArray<FString> Tokens, Switches;
+	TMap<FString, FString> CommandLineParams;
+	ParseCommandLine(*Params, Tokens, Switches, CommandLineParams);
+	if (!InitializeFromCommandLine(Tokens, Switches, CommandLineParams))
+	{
+		return EReturnCode::CommandletInitializationError;
+	}
+
+	ConversionFolder = DestinationFolder + "/" + MainWorld->GetName();
+
+	TStrongObjectPtr<UDataLayerToAssetCommandletContext> Context(NewObject<UDataLayerToAssetCommandletContext>());
+
+	if (!BuildConversionInfos(Context))
+	{
+		return EReturnCode::DataLayerConversionError;
+	}
+
+	if (!ResolvePreviousConversionsToCurrent(Context))
+	{
+		return EReturnCode::DataLayerConversionError;
+	}
+
+	Context->LogConversionInfos();
+
+	if (!RemapActorDataLayersToAssets(Context))
+	{
+		return EReturnCode::ActorDataLayerRemappingError;
+	}
+
+	if (!CreateDataLayerInstances(Context))
+	{
+		return EReturnCode::DataLayerConversionError;
+	}
+
+	if (!DeletePreviousConversionsData(Context))
+	{
+		return EReturnCode::DataLayerConversionError;
+	}
+
+	if (!CommitConversion(Context))
+	{
+		return EReturnCode::DataLayerConversionError;
+	}
+
+	return EReturnCode::Success;
+}
+
+bool UDataLayerToAssetCommandlet::InitializeFromCommandLine(TArray<FString>& Tokens, TArray<FString> const& Switches, TMap<FString, FString> const& Params)
+{
+	constexpr TCHAR DestinationFolderSwitch[] = TEXT("DestinationFolder");
+	if (FString const* DestFolderPtr = Params.Find(DestinationFolderSwitch))
+	{
+		DestinationFolder = *DestFolderPtr;
+	}
+	else
+	{
+		UE_LOG(LogDataLayerToAssetCommandlet, Error, TEXT("No \"%s\" switch specified"), DestinationFolderSwitch);
+		return false;
+	}
+
+	if (Switches.Contains(TEXT("Verbose")))
+	{
+		LogDataLayerToAssetCommandlet.SetVerbosity(ELogVerbosity::Verbose);
+		WorldPartitionCommandletHelpers::LogWorldParitionCommandletUtils.SetVerbosity(ELogVerbosity::Verbose);
+	}
+
+	bPerformSavePackages = Switches.Contains(TEXT("NoSave")) == false;
+	bIgnoreActorLoadingErrors = Switches.Contains(TEXT("IgnoreActorLoadingErrors"));
+
+	MainWorld = WorldPartitionCommandletHelpers::LoadAndInitWorld(Tokens[0]);
+	if (!MainWorld)
+	{
+		UE_LOG(LogDataLayerToAssetCommandlet, Error, TEXT("Failed to Load %s, Conversion will abort"), *Tokens[0]);
+		return false;
+	}
+
+	return true;
+}
+
+bool UDataLayerToAssetCommandlet::BuildConversionInfos(TStrongObjectPtr<UDataLayerToAssetCommandletContext>& CommandletContext)
+{
+	UE_SCOPED_TIMER(TEXT("Retrieving Already Converted Data Layers"), LogDataLayerToAssetCommandlet, Display);
+
+	TArray<FAssetData> ExistingDataLayerAssets;
+	IAssetRegistry& AssetRegistry = FAssetRegistryModule::GetRegistry();
+	AssetRegistry.GetAssetsByClass(UDataLayerAsset::StaticClass()->GetFName(), ExistingDataLayerAssets);
+	for (FAssetData& AssetData : ExistingDataLayerAssets)
+	{
+		CommandletContext->StoreExistingDataLayer(AssetData);
+	}
+
+	AWorldDataLayers* WorldDataLayers = MainWorld->GetWorldDataLayers();
+	TArray<TObjectPtr<UDataLayerInstance>> DataLayerInstances;
+	Algo::Copy(WorldDataLayers->DataLayerInstances, DataLayerInstances);
+	for (int32 i = DataLayerInstances.Num() -1; i >= 0; --i)
+	{
+		TObjectPtr<UDataLayerInstance>& DataLayerInstance = DataLayerInstances[i];
+		if (UDataLayerConversionInfo* ConversionInfo = CommandletContext->GetDataLayerConversionInfo(DataLayerInstance->GetAsset()))
+		{
+			CommandletContext->StoreDataLayerInstanceConversion(ConversionInfo->DataLayerAsset, DataLayerInstance);
+		}
+	}
+
+	uint32 ErrorCount = 0;
+	for (const TObjectPtr<UDataLayer>& DataLayer : WorldDataLayers->WorldDataLayers)
+	{
+		if (!CreateConversionFromDataLayer(CommandletContext, DataLayer))
+		{
+			ErrorCount++;
+		}
+	}
+
+	return ErrorCount == 0;
+}
+
+bool UDataLayerToAssetCommandlet::CreateConversionFromDataLayer(TStrongObjectPtr<UDataLayerToAssetCommandletContext>& CommandletContext, const TObjectPtr<UDataLayer> & DataLayer)
+{
+	if(TObjectPtr<UDataLayerAsset> DataLayerAsset = GetOrCreateDataLayerAssetForConversion(CommandletContext, DataLayer->GetDataLayerLabel()))
+	{
+		if (UDataLayerConversionInfo* ConversionInfo = CommandletContext->StoreDataLayerAssetConversion(DataLayer, DataLayerAsset))
+		{
+			if (bPerformSavePackages)
+			{
+				return WorldPartitionCommandletHelpers::CheckoutSaveAdd(ConversionInfo->DataLayerAsset->GetPackage(), PackageHelper);
+			}
+
+			return true;
+		}
+	}
+	
+	return false;
+}
+
+TObjectPtr<UDataLayerAsset> UDataLayerToAssetCommandlet::GetOrCreateDataLayerAssetForConversion(TStrongObjectPtr<UDataLayerToAssetCommandletContext>& CommandletContext, FName DataLayerAssetName)
+{
+	TObjectPtr<UDataLayerAsset> DataLayerAsset = nullptr;
+
+	TArray<TWeakObjectPtr<UDataLayerConversionInfo>> ConversionInfos;
+	if (CommandletContext->FindDataLayerConversionInfos(DataLayerAssetName, ConversionInfos))
+	{
+		for (TWeakObjectPtr<UDataLayerConversionInfo>& ConversionInfo : ConversionInfos)
+		{
+			if (IsAssetInConversionFolder(ConversionInfo->DataLayerAsset))
+			{
+				check(DataLayerAsset == nullptr);
+				return CastChecked<UDataLayerAsset>(ConversionInfo->DataLayerAsset);
+			}
+		}
+	}
+
+	UE_LOG(LogDataLayerToAssetCommandlet, Log, TEXT("Creating new Data Layer Asset %s in folder %s"),
+		*DataLayerAssetName.ToString(), *ConversionFolder);
+
+	IAssetTools& AssetTools = FModuleManager::GetModuleChecked<FAssetToolsModule>("AssetTools").Get();
+	if (UObject* Asset = AssetTools.CreateAsset(DataLayerAssetName.ToString(), ConversionFolder, UDataLayerAsset::StaticClass(), DataLayerFactory))
+	{
+		DataLayerAsset = CastChecked<UDataLayerAsset>(Asset);
+	}
+	else
+	{
+		UE_LOG(LogDataLayerToAssetCommandlet, Error, TEXT("Failed to create asset for %s in folder %s. Consult log for more details"),
+			*DataLayerAssetName.ToString(), *ConversionFolder);
+	}
+
+	return DataLayerAsset;
+}
+
+bool UDataLayerToAssetCommandlet::ResolvePreviousConversionsToCurrent(TStrongObjectPtr<UDataLayerToAssetCommandletContext>& CommandletContext)
+{
+	UE_SCOPED_TIMER(TEXT("Resolving Conflicting Conversions"), LogDataLayerToAssetCommandlet, Display);
+	UE_LOG(LogDataLayerToAssetCommandlet, Log, TEXT("Resolving Conflicting Conversions"));
+
+	uint32 ErrorCount = 0;
+	for (const TObjectPtr<UDataLayerConversionInfo>& ConversionInfo : CommandletContext->GetDataLayerConversionInfos())
+	{
+		if(IsAssetInConversionFolder(ConversionInfo->DataLayerAsset))
+		{
+			TArray<TWeakObjectPtr<UDataLayerConversionInfo>> PreviousConversionInfos;
+			if (CommandletContext->FindDataLayerConversionInfos(ConversionInfo->DataLayerAsset->GetFName(), PreviousConversionInfos))
+			{
+				PreviousConversionInfos.Remove(ConversionInfo);
+
+				if (!PreviousConversionInfos.IsEmpty())
+				{
+					if (!CommandletContext->SetPreviousConversions(ConversionInfo, MoveTemp(PreviousConversionInfos)))
+					{
+						ErrorCount++;
+					}
+				}
+			}
+		}
+	}
+
+	return ErrorCount == 0;
+}
+
+bool UDataLayerToAssetCommandlet::RemapActorDataLayersToAssets(TStrongObjectPtr<UDataLayerToAssetCommandletContext>& CommandletContext)
+{
+	UE_SCOPED_TIMER(TEXT("Remapping Actors Data Layers"), LogDataLayerToAssetCommandlet, Display);
+	UE_LOG(LogDataLayerToAssetCommandlet, Log, TEXT("Starting Actor Data Layer Remapping To Data Layer Asset. This can take a while."));
+
+	AWorldDataLayers* WorldDataLayersActor = MainWorld->GetWorldDataLayers();
+
+	uint32 ErrorCount = 0;
+	FWorldPartitionHelpers::ForEachActorWithLoading(MainWorld->GetWorldPartition(), AActor::StaticClass(),
+		[WorldDataLayersActor, &ErrorCount, &CommandletContext, this](const FWorldPartitionActorDesc* ActorDesc)
+		{
+			uint32 ActorConversionErrors = 0;
+			if (AActor* Actor = ActorDesc->GetActor())
+			{
+				ActorConversionErrors += RemapDataLayersAssetsFromPreviousConversions(CommandletContext, Actor);
+				ActorConversionErrors += RemapActorDataLayers(CommandletContext, Actor);
+
+				if (ActorConversionErrors == 0 && bPerformSavePackages)
+				{
+					if (bPerformSavePackages && Actor->GetExternalPackage()->IsDirty())
+					{
+						if (!WorldPartitionCommandletHelpers::CheckoutSaveAdd(Actor->GetExternalPackage(), PackageHelper))
+						{
+							ActorConversionErrors++;
+						}
+					}
+				}
+			}
+			else
+			{
+				const TArray<FName>& ActDescDataLayers = ActorDesc->GetDataLayers();
+				if (!ActDescDataLayers.IsEmpty())
+				{
+					FString DataLayerString = FString::JoinBy(ActDescDataLayers, TEXT(", "), [](const FName& DataLayerName) { return DataLayerName.ToString(); });
+
+					UE_LOG(LogDataLayerToAssetCommandlet, Error, TEXT("Actor %s failed to load. Its data layers %s will not be remapped to a data layer asset."),
+						*ActorDesc->GetActorName().ToString(), *DataLayerString);
+					if (!bIgnoreActorLoadingErrors)
+					{
+						ActorConversionErrors++;
+					}
+				}
+			}
+
+			ErrorCount += ActorConversionErrors;
+			return true;
+		});
+
+	UE_LOG(LogDataLayerToAssetCommandlet, Log, TEXT("Actor Data Layer Remapping To Data Layer Asset Done."));
+
+	return ErrorCount == 0;;
+}
+
+uint32 UDataLayerToAssetCommandlet::RemapActorDataLayers(TStrongObjectPtr<UDataLayerToAssetCommandletContext>& CommandletContext, AActor* Actor)
+{
+	uint32 ErrorCount = 0;
+
+	const TArray<FActorDataLayer>& ActorDataLayers = Actor->GetActorDataLayers();
+	for (int32 i = ActorDataLayers.Num() - 1; i >= 0; --i)
+	{
+		const FActorDataLayer& ActorDataLayer = ActorDataLayers[i];
+		if (UDataLayerConversionInfo* DataLayerConversionInfo = CommandletContext->GetDataLayerConversionInfo(ActorDataLayer))
+		{
+			UE_LOG(LogDataLayerToAssetCommandlet, Verbose, TEXT("Data layer %s in Actor %s remapped to data layer asset %s"),
+				*ActorDataLayer.Name.ToString(), *Actor->GetName(), *DataLayerConversionInfo->DataLayerAsset->GetName());
+
+			if (Actor->AddDataLayer(DataLayerConversionInfo->DataLayerAsset))
+			{
+				if (!Actor->RemoveDataLayer(DataLayerConversionInfo->DataLayerToConvert))
+				{
+					ErrorCount++;
+					UE_LOG(LogDataLayerToAssetCommandlet, Error, TEXT("Failed to remove data layer %s in Actor %s"),
+						*DataLayerConversionInfo->DataLayerToConvert->GetDataLayerLabel().ToString(), *Actor->GetName());
+				}
+			}
+			else
+			{
+				ErrorCount++;
+				UE_LOG(LogDataLayerToAssetCommandlet, Error, TEXT("Failed to add data layer asset %s in Actor %s"),
+					*DataLayerConversionInfo->DataLayerAsset->GetFullName(), *Actor->GetName());
+			}
+
+		}
+		else
+		{
+			ErrorCount++;
+			UE_LOG(LogDataLayerToAssetCommandlet, Error, TEXT("Actor %s is referencing data layer %s but it does not match any data layer asset."),
+				*Actor->GetName(), *ActorDataLayer.Name.ToString());
+		}
+	}
+
+	return ErrorCount;
+}
+
+// Multiple run of the commandlet can lead to Actors referencing Data Layer Assets in different folders. Always remap to the newly created asset.
+uint32 UDataLayerToAssetCommandlet::RemapDataLayersAssetsFromPreviousConversions(TStrongObjectPtr<UDataLayerToAssetCommandletContext>& CommandletContext, AActor* Actor)
+{
+	uint32 ErrorCount = 0;
+
+	const TArray<TObjectPtr<const UDataLayerAsset>>& ActorDataLayerAssets = Actor->GetDataLayerAssets();
+	for (int32 i = ActorDataLayerAssets.Num() - 1; i >= 0; --i)
+	{
+		const TObjectPtr<const UDataLayerAsset>& ActorDataLayerAsset = ActorDataLayerAssets[i];
+		if (UDataLayerConversionInfo* DataLayerConversionInfo = CommandletContext->GetDataLayerConversionInfo(ActorDataLayerAsset.Get()))
+		{
+			if (DataLayerConversionInfo->IsAPreviousConversion())
+			{
+				if (Actor->AddDataLayer(DataLayerConversionInfo->GetCurrentConversion()->DataLayerAsset))
+				{
+					if (!Actor->RemoveDataLayer(DataLayerConversionInfo->DataLayerAsset))
+					{
+						ErrorCount++;
+						UE_LOG(LogDataLayerToAssetCommandlet, Error, TEXT("Failed to remove deprecated data layer asset %s in Actor %s"),
+							*DataLayerConversionInfo->DataLayerAsset->GetFullName(), *Actor->GetName());
+					}
+				}
+				else
+				{
+					ErrorCount++;
+					UE_LOG(LogDataLayerToAssetCommandlet, Error, TEXT("Failed to replace data layer asset %s in Actor %s"),
+						*DataLayerConversionInfo->GetCurrentConversion()->DataLayerAsset->GetFullName(), *Actor->GetName());
+				}
+			}
+		}
+	}
+
+	return ErrorCount;
+}
+
+bool UDataLayerToAssetCommandlet::CreateDataLayerInstances(TStrongObjectPtr<UDataLayerToAssetCommandletContext>& CommandletContext)
+{
+	UE_SCOPED_TIMER(TEXT("Updating WorldDataLayers Actor"), LogDataLayerToAssetCommandlet, Display);
+
+	uint32 ErrorCount = 0;
+	AWorldDataLayers* WorldDataLayers = MainWorld->GetWorldDataLayers();
+	for (const TWeakObjectPtr<UDataLayerConversionInfo>& ConvertingInfo : CommandletContext->GetConvertingDataLayerConversionInfo())
+	{
+		if (ConvertingInfo->DataLayerInstance == nullptr)
+		{
+			UDataLayerInstance* DataLayerInstance = WorldDataLayers->CreateDataLayer(ConvertingInfo->DataLayerAsset.Get());
+			UE_LOG(LogDataLayerToAssetCommandlet, Log, TEXT("Created new Data Layer Instance %s"),
+				*DataLayerInstance->GetName());
+
+			CommandletContext->StoreDataLayerInstanceConversion(ConvertingInfo->DataLayerAsset, DataLayerInstance);
+		}
+
+		WorldDataLayers->DeprecatedDataLayerNameToDataLayerInstance.Add(ConvertingInfo->DataLayerToConvert->GetFName(), ConvertingInfo->DataLayerInstance);
+	}
+
+	return ErrorCount == 0;
+}
+
+bool UDataLayerToAssetCommandlet::DeletePreviousConversionsData(TStrongObjectPtr<UDataLayerToAssetCommandletContext>& CommandletContext)
+{
+	UE_SCOPED_TIMER(TEXT("Delete Conflicting Assets"), LogDataLayerToAssetCommandlet, Display);
+
+	uint32 ErrorCount = 0;
+	AWorldDataLayers* WorldDataLayers = MainWorld->GetWorldDataLayers();
+	for (const TWeakObjectPtr<UDataLayerConversionInfo>& ConvertingInfo : CommandletContext->GetConvertingDataLayerConversionInfo())
+	{
+		for (const TWeakObjectPtr<UDataLayerConversionInfo>& PreviousConversion : ConvertingInfo->GetPreviousConversions())
+		{
+			if (bPerformSavePackages)
+			{
+				if (WorldPartitionCommandletHelpers::Delete(PreviousConversion->DataLayerAsset->GetPackage(), PackageHelper))
+				{
+					PreviousConversion->DataLayerAsset = nullptr;
+				}
+				else
+				{
+					UE_LOG(LogDataLayerToAssetCommandlet, Error, TEXT("Failed to delete Data Layer Asset %s from previous conversion. (Conflicting with %s)"),
+						*PreviousConversion->DataLayerAsset->GetFullName(), *ConvertingInfo->DataLayerAsset->GetFullName());
+					ErrorCount++;
+				}
+			}
+		}
+	}
+
+	return ErrorCount == 0;
+}
+
+bool UDataLayerToAssetCommandlet::CommitConversion(TStrongObjectPtr<UDataLayerToAssetCommandletContext>& CommandletContext)
+{
+	uint32 ErrorCount = 0;
+	AWorldDataLayers* WorldDataLayers = MainWorld->GetWorldDataLayers();
+	for (const TWeakObjectPtr<UDataLayerConversionInfo>& ConversionInfo : CommandletContext->GetConvertingDataLayerConversionInfo())
+	{
+		if (ConversionInfo->DataLayerToConvert != nullptr)
+		{
+			if (WorldDataLayers->RemoveDataLayer(ConversionInfo->DataLayerToConvert))
+			{
+				UE_LOG(LogDataLayerToAssetCommandlet, Log, TEXT("Deleted old data layer %s, it is now converted to Data Layer Asset %s and Data Layer Instance %s"),
+					*ConversionInfo->DataLayerToConvert->GetDataLayerLabel().ToString(), *ConversionInfo->DataLayerAsset->GetFullName(), *ConversionInfo->DataLayerInstance->GetName());
+			}
+			else
+			{
+				UE_LOG(LogDataLayerToAssetCommandlet, Error, TEXT("Failed to delete converting data layer %s."),
+					*ConversionInfo->DataLayerToConvert->GetDataLayerLabel().ToString());
+				ErrorCount++;
+			}
+		}
+	}
+
+	if (ErrorCount == 0 && bPerformSavePackages)
+	{
+		if (!WorldPartitionCommandletHelpers::CheckoutSaveAdd(WorldDataLayers->GetExternalPackage(), PackageHelper))
+		{
+			return false;
+		}
+	}
+
+	return true;
+}
+
+bool UDataLayerToAssetCommandlet::IsAssetInConversionFolder(const TObjectPtr<UDataLayerAsset> DataLayerAsset)
+{
+	return FCString::Stristr(*DataLayerAsset->GetPackage()->GetName(), *ConversionFolder) != nullptr;
+}
\ No newline at end of file
diff --git a/Engine/Source/Editor/UnrealEd/Public/Commandlets/WorldPartitionCommandletHelpers.h b/Engine/Source/Editor/UnrealEd/Public/Commandlets/WorldPartitionCommandletHelpers.h
new file mode 100644
--- /dev/null
+++ b/Engine/Source/Editor/UnrealEd/Public/Commandlets/WorldPartitionCommandletHelpers.h
@@ -0,0 +1,41 @@
+// Copyright Epic Games, Inc. All Rights Reserved.
+
+#pragma once
+
+#include "CoreMinimal.h"
+#include "UObject/ObjectMacros.h"
+#include "UObject/UObjectAnnotation.h"
+
+class FPackageSourceControlHelper;
+
+namespace WorldPartitionCommandletHelpers
+{
+	UNREALED_API DECLARE_LOG_CATEGORY_EXTERN(LogWorldParitionCommandletUtils, Log, All);
+
+	UWorld* LoadAndInitWorld(const FString& LevelToLoad);
+	UWorld* LoadWorld(const FString& LevelToLoad);
+	ULevel* InitLevel(UWorld* World);
+
+	bool Checkout(const TArray<UPackage*>& PackagesToCheckout, FPackageSourceControlHelper& SCHelper);
+	bool Checkout(UPackage* PackagesToCheckout, FPackageSourceControlHelper& SCHelper);
+	bool Save(const TArray<UPackage*>& PackagesToSave, FPackageSourceControlHelper& SCHelper);
+	bool Save(UPackage* PackagesToSave, FPackageSourceControlHelper& SCHelper);
+	bool AddToSourceControl(const TArray<UPackage*>& PackagesToAdd, FPackageSourceControlHelper& SCHelper);
+	bool AddToSourceControl(UPackage* PackageToAdd, FPackageSourceControlHelper& SCHelper);
+	bool Delete(const TArray<UPackage*>& PackagesToDelete, FPackageSourceControlHelper& SCHelper);
+	bool Delete(UPackage* PackagesToDelete, FPackageSourceControlHelper& SCHelper);
+
+	template<class T>
+	bool CheckoutSaveAdd(T&& ToSave, FPackageSourceControlHelper& SCHelper)
+	{
+		if (Checkout(ToSave, SCHelper))
+		{
+			if (Save(ToSave, SCHelper))
+			{
+				return AddToSourceControl(ToSave, SCHelper);
+			}
+		}
+
+		return false;
+	}
+}
\ No newline at end of file
diff --git a/Engine/Source/Editor/UnrealEd/UnrealEd.Build.cs b/Engine/Source/Editor/UnrealEd/UnrealEd.Build.cs
--- a/Engine/Source/Editor/UnrealEd/UnrealEd.Build.cs
+++ b/Engine/Source/Editor/UnrealEd/UnrealEd.Build.cs
@@ -136,6 +136,7 @@ public class UnrealEd : ModuleRules
 				"CinematicCamera",
 				"CurveEditor",
 				"ContentBrowserData",
+				"DataLayerEditor",
 				"DerivedDataCache",
 				"Zen",
 				"DesktopPlatform",
diff --git a/Engine/Source/Runtime/Engine/Classes/GameFramework/Actor.h b/Engine/Source/Runtime/Engine/Classes/GameFramework/Actor.h
--- a/Engine/Source/Runtime/Engine/Classes/GameFramework/Actor.h
+++ b/Engine/Source/Runtime/Engine/Classes/GameFramework/Actor.h
@@ -39,6 +39,8 @@ struct FAttachedActorInfo;
 struct FNetViewer;
 struct FNetworkObjectInfo;
 class UDataLayer;
+class UDataLayerAsset;
+class UDataLayerInstance;
 class AWorldDataLayers;
 class UActorFolder;
 
@@ -823,6 +825,10 @@ protected:
 	UPROPERTY(EditAnywhere, AdvancedDisplay, Category = DataLayers)
 	TArray<FActorDataLayer> DataLayers;
 
+	/** DataLayers the actor belongs to.*/
+	UPROPERTY()
+	TArray<TObjectPtr<const UDataLayerAsset>> DataLayerAssets;
+
 	TArray<FActorDataLayer> PreEditChangeDataLayers;
 
 public:
@@ -1153,6 +1159,13 @@ public:
 	bool IsValidForDataLayer() const;
 	void FixupDataLayers(bool bRevertChangesOnLockedDataLayer = false);
 	static const FName GetDataLayersPropertyName() { return GET_MEMBER_NAME_CHECKED(AActor, DataLayers); }
+
+	bool AddDataLayer(const UDataLayerAsset* DataLayerAsset);
+	bool RemoveDataLayer(const UDataLayerAsset* DataLayerAsset);
+	bool ContainsDataLayer(const UDataLayerAsset* DataLayerAsset) const;
+	const TArray<TObjectPtr<const UDataLayerAsset>>& GetDataLayerAssets() const { return DataLayerAssets; }
+
+	TArray<FActorDataLayer> const& GetActorDataLayers() const { return DataLayers; };
 #endif
 
 	//~=============================================================================
diff --git a/Engine/Source/Runtime/Engine/Private/ActorEditor.cpp b/Engine/Source/Runtime/Engine/Private/ActorEditor.cpp
--- a/Engine/Source/Runtime/Engine/Private/ActorEditor.cpp
+++ b/Engine/Source/Runtime/Engine/Private/ActorEditor.cpp
@@ -15,6 +15,7 @@
 #include "WorldPartition/WorldPartitionSubsystem.h"
 #include "WorldPartition/DataLayer/WorldDataLayers.h"
 #include "WorldPartition/DataLayer/DataLayer.h"
+#include "WorldPartition/DataLayer/DataLayerAsset.h"
 #include "EditorSupportDelegates.h"
 #include "Logging/TokenizedMessage.h"
 #include "Logging/MessageLog.h"
@@ -1319,10 +1320,34 @@ bool AActor::AddDataLayer(const UDataLayer* DataLayer)
 	return bActorWasModified;
 }
 
+bool AActor::AddDataLayer(const UDataLayerAsset* DataLayerAsset)
+{
+	bool bActorWasModified = false;
+	if (SupportsDataLayer() && DataLayerAsset && !ContainsDataLayer(DataLayerAsset))
+	{
+		Modify();
+		bActorWasModified = true;
+		DataLayerAssets.Add(DataLayerAsset);
+	}
+	return bActorWasModified;
+}
+
 bool AActor::RemoveDataLayer(const UDataLayer* DataLayer)
 {
 	bool bActorWasModified = false;
 	if (ContainsDataLayer(DataLayer))
+	{
+		Modify();
+		bActorWasModified = true;
+		DataLayers.Remove(FActorDataLayer(DataLayer->GetFName()));
+	}
+	return bActorWasModified;
+}
+
+bool AActor::RemoveDataLayer(const UDataLayerAsset* DataLayerAsset)
+{
+	bool bActorWasModified = false;
+	if (ContainsDataLayer(DataLayerAsset))
 	{
 		if (!bActorWasModified)
 		{
@@ -1330,7 +1355,7 @@ bool AActor::RemoveDataLayer(const UDataLayer* DataLayer)
 			bActorWasModified = true;
 		}
 
-		DataLayers.Remove(FActorDataLayer(DataLayer->GetFName()));
+		DataLayerAssets.Remove(DataLayerAsset);
 	}
 	return bActorWasModified;
 }
@@ -1351,6 +1376,11 @@ bool AActor::ContainsDataLayer(const UDataLayer* DataLayer) const
 	return DataLayer && DataLayers.Contains(FActorDataLayer(DataLayer->GetFName()));
 }
 
+bool AActor::ContainsDataLayer(const UDataLayerAsset* DataLayerAsset) const
+{
+	return DataLayerAsset && DataLayerAssets.Contains(DataLayerAsset);
+}
+
 bool AActor::HasDataLayers() const
 {
 	return DataLayers.Num() > 0;
diff --git a/Engine/Source/Runtime/Engine/Private/WorldPartition/DataLayer/DataLayerAsset.cpp b/Engine/Source/Runtime/Engine/Private/WorldPartition/DataLayer/DataLayerAsset.cpp
new file mode 100644
--- /dev/null
+++ b/Engine/Source/Runtime/Engine/Private/WorldPartition/DataLayer/DataLayerAsset.cpp
@@ -0,0 +1,35 @@
+// Copyright Epic Games, Inc. All Rights Reserved.
+
+#include "WorldPartition/DataLayer/DataLayerAsset.h"
+
+UDataLayerAsset::UDataLayerAsset(const FObjectInitializer& ObjectInitializer)
+	: Super(ObjectInitializer)
+
+{
+
+}
+
+void UDataLayerAsset::PostLoad()
+{
+	Super::PostLoad();
+
+#if WITH_EDITORONLY_DATA
+	if (DebugColor == FColor::Black)
+	{
+		FRandomStream RandomStream(GetFName());
+		const uint8 R = (uint8)(RandomStream.GetFraction() * 255.f);
+		const uint8 G = (uint8)(RandomStream.GetFraction() * 255.f);
+		const uint8 B = (uint8)(RandomStream.GetFraction() * 255.f);
+		DebugColor = FColor(R, G, B);
+	}
+#endif
+}
+
+#if WITH_EDITOR
+
+const TCHAR* UDataLayerAsset::GetDataLayerIconName() const
+{
+	return IsRuntime() ? TEXT("DataLayer.Runtime") : TEXT("DataLayer.Editor");
+}
+
+#endif
\ No newline at end of file
diff --git a/Engine/Source/Runtime/Engine/Private/WorldPartition/DataLayer/DataLayerInstance.cpp b/Engine/Source/Runtime/Engine/Private/WorldPartition/DataLayer/DataLayerInstance.cpp
new file mode 100644
--- /dev/null
+++ b/Engine/Source/Runtime/Engine/Private/WorldPartition/DataLayer/DataLayerInstance.cpp
@@ -0,0 +1,222 @@
+// Copyright Epic Games, Inc. All Rights Reserved.
+
+#include "WorldPartition/DataLayer/DataLayerInstance.h"
+
+#include "WorldPartition/DataLayer/WorldDataLayers.h"
+#include "Misc/UObjectToken.h"
+#include "Misc/MapErrors.h"
+
+
+#define LOCTEXT_NAMESPACE "DataLayer"
+
+UDataLayerInstance::UDataLayerInstance(const FObjectInitializer& ObjectInitializer)
+	: Super(ObjectInitializer)
+
+{
+
+}
+
+void UDataLayerInstance::PostLoad()
+{
+	Super::PostLoad();
+
+#if WITH_EDITORONLY_DATA
+	// Initialize bIsVisible with persistent flag bIsInitiallyVisible
+	bIsVisible = bIsInitiallyVisible;
+#endif
+
+	if (Parent)
+	{
+		Parent->AddChild(this);
+	}
+}
+
+#if WITH_EDITOR
+void UDataLayerInstance::SetVisible(bool bInIsVisible)
+{
+	if (bIsVisible != bInIsVisible)
+	{
+		Modify(/*bAlwaysMarkDirty*/false);
+		bIsVisible = bInIsVisible;
+	}
+}
+
+void UDataLayerInstance::SetIsInitiallyVisible(bool bInIsInitiallyVisible)
+{
+	if (bIsInitiallyVisible != bInIsInitiallyVisible)
+	{
+		Modify();
+		bIsInitiallyVisible = bInIsInitiallyVisible;
+	}
+}
+
+void UDataLayerInstance::SetIsLoadedInEditor(bool bInIsLoadedInEditor, bool bInFromUserChange)
+{
+	if (bIsLoadedInEditor != bInIsLoadedInEditor)
+	{
+		Modify(false);
+		bIsLoadedInEditor = bInIsLoadedInEditor;
+		bIsLoadedInEditorChangedByUserOperation |= bInFromUserChange;
+	}
+}
+
+void UDataLayerInstance::OnCreated(const UDataLayerAsset* Asset)
+{
+	Modify(/*bAlwaysMarkDirty*/false);
+
+	check(DataLayerAsset == nullptr);
+	DataLayerAsset = Asset;
+
+	bIsVisible = true;
+}
+
+bool UDataLayerInstance::IsEffectiveLoadedInEditor() const
+{
+	bool bResult = IsLoadedInEditor();
+
+	const UDataLayerInstance* ParentDataLayer = GetParent();
+	while (ParentDataLayer && bResult)
+	{
+		bResult = bResult && ParentDataLayer->IsLoadedInEditor();
+		ParentDataLayer = ParentDataLayer->GetParent();
+	}
+	return bResult;
+}
+
+bool UDataLayerInstance::IsLocked() const
+{
+	if (bIsLocked)
+	{
+		return true;
+	}
+
+	return IsRuntime() && !GetOuterAWorldDataLayers()->GetAllowRuntimeDataLayerEditing();
+}
+
+bool UDataLayerInstance::CanParent(const UDataLayerInstance* InParent) const
+{
+	return (this != InParent) && (Parent != InParent) && IsRuntime() == InParent->IsRuntime();
+}
+
+void UDataLayerInstance::SetParent(UDataLayerInstance* InParent)
+{
+	if (!CanParent(InParent))
+	{
+		return;
+	}
+
+	Modify();
+	if (Parent)
+	{
+		Parent->RemoveChild(this);
+	}
+	Parent = InParent;
+	if (Parent)
+	{
+		Parent->AddChild(this);
+	}
+}
+
+void UDataLayerInstance::SetChildParent(UDataLayerInstance* InParent)
+{
+	if (this == InParent)
+	{
+		return;
+	}
+
+	Modify();
+	while (Children.Num())
+	{
+		Children[0]->SetParent(InParent);
+	};
+}
+
+void UDataLayerInstance::RemoveChild(UDataLayerInstance* InDataLayer)
+{
+	Modify();
+	check(Children.Contains(InDataLayer));
+	Children.RemoveSingle(InDataLayer);
+}
+
+void UDataLayerInstance::CheckForErrors() const
+{
+	auto CheckRuntime = [this](UDataLayerInstance const* Child)
+	{
+		if (Child->IsRuntime() != IsRuntime())
+		{
+			TSharedRef<FTokenizedMessage> Error = FMessageLog("MapCheck").Error()
+				->AddToken(FTextToken::Create(LOCTEXT("MapCheck_DataLayers_DataLayer", "Data layer")))
+				->AddToken(FUObjectToken::Create(this))
+				->AddToken(FTextToken::Create(LOCTEXT("MapCheck_DataLayers_Asset", "asset")))
+				->AddToken(FAssetNameToken::Create(GetAsset()->GetPathName()))
+				->AddToken(FTextToken::Create(LOCTEXT("MapCheck_DataLayers_is runtime", "is Runtime but its child data layer")))
+				->AddToken(FUObjectToken::Create(Child))
+				->AddToken(FTextToken::Create(LOCTEXT("MapCheck_DataLayers_Asset", "asset")))
+				->AddToken(FAssetNameToken::Create(Child->GetAsset()->GetPathName()))
+				->AddToken(FTextToken::Create(LOCTEXT("MapCheck_DataLayers_isnot", "is not")))
+				->AddToken(FMapErrorToken::Create(FName(TEXT("DataLayers_CheckForErrors"))));
+		}
+
+		return true;
+	};
+
+	ForEachChild(CheckRuntime);
+	ForEachChild([](UDataLayerInstance const* Child) { Child->CheckForErrors(); return true; });
+}
+
+#endif
+
+bool UDataLayerInstance::IsInitiallyVisible() const
+{
+#if WITH_EDITOR
+	return bIsInitiallyVisible;
+#else
+	return false;
+#endif
+}
+
+bool UDataLayerInstance::IsVisible() const
+{
+#if WITH_EDITOR
+	return bIsVisible;
+#else
+	return false;
+#endif
+}
+
+bool UDataLayerInstance::IsEffectiveVisible() const
+{
+#if WITH_EDITOR
+	bool bResult = IsVisible();
+	const UDataLayerInstance* ParentDataLayer = GetParent();
+	while (ParentDataLayer && bResult)
+	{
+		bResult = bResult && ParentDataLayer->IsVisible();
+		ParentDataLayer = ParentDataLayer->GetParent();
+	}
+	return bResult && IsEffectiveLoadedInEditor();
+#else
+	return false;
+#endif
+}
+
+void UDataLayerInstance::ForEachChild(TFunctionRef<bool(const UDataLayerInstance*)> Operation) const
+{
+	for (const UDataLayerInstance* Child : Children)
+	{
+		if (!Operation(Child))
+		{
+			break;
+		}
+	}
+}
+
+void UDataLayerInstance::AddChild(UDataLayerInstance* InDataLayer)
+{
+	Modify();
+	checkSlow(!Children.Contains(InDataLayer));
+	check(InDataLayer->IsRuntime() == IsRuntime());
+	Children.Add(InDataLayer);
+}
+
+#undef LOCTEXT_NAMESPACE
\ No newline at end of file
diff --git a/Engine/Source/Runtime/Engine/Private/WorldPartition/DataLayer/WorldDataLayers.cpp b/Engine/Source/Runtime/Engine/Private/WorldPartition/DataLayer/WorldDataLayers.cpp
--- a/Engine/Source/Runtime/Engine/Private/WorldPartition/DataLayer/WorldDataLayers.cpp
+++ b/Engine/Source/Runtime/Engine/Private/WorldPartition/DataLayer/WorldDataLayers.cpp
@@ -6,6 +6,7 @@
 
 #include "WorldPartition/DataLayer/WorldDataLayers.h"
 #include "WorldPartition/DataLayer/DataLayer.h"
+#include "WorldPartition/DataLayer/DataLayerInstance.h"
 #include "WorldPartition/DataLayer/DataLayerSubsystem.h"
 #include "EngineUtils.h"
 #include "Engine/CoreSettings.h"
@@ -605,6 +606,18 @@ UDataLayer* AWorldDataLayers::CreateDataLayer(FName InName, EObjectFlags InObjec
 	return NewDataLayer;
 }
 
+UDataLayerInstance* AWorldDataLayers::CreateDataLayer(const UDataLayerAsset* InDataLayerAsset)
+{
+	Modify();
+
+	const FName DataLayerUniqueName = *FString::Format(TEXT("{0}_{1}"), { InDataLayerAsset->GetName(), FGuid::NewGuid().ToString() });
+	UDataLayerInstance* NewDataLayer = NewObject<UDataLayerInstance>(this, DataLayerUniqueName, RF_Transactional | RF_NoFlags);
+	check(NewDataLayer != NULL);
+	NewDataLayer->OnCreated(InDataLayerAsset);
+	DataLayerInstances.Add(NewDataLayer);
+	return NewDataLayer;
+}
+
 bool AWorldDataLayers::RemoveDataLayers(const TArray<UDataLayer*>& InDataLayers)
 {
 	bool bIsModified = false;
@@ -632,6 +645,17 @@ bool AWorldDataLayers::RemoveDataLayer(UDataLayer* InDataLayer)
 	return false;
 }
 
+bool AWorldDataLayers::RemoveDataLayer(UDataLayerInstance* InDataLayerInstance)
+{
+	if (ContainsDataLayer(InDataLayerInstance))
+	{
+		Modify();
+		DataLayerInstances.Remove(InDataLayerInstance);
+		return true;
+	}
+	return false;
+}
+
 void AWorldDataLayers::SetAllowRuntimeDataLayerEditing(bool bInAllowRuntimeDataLayerEditing)
 {
 	if (bAllowRuntimeDataLayerEditing != bInAllowRuntimeDataLayerEditing)
@@ -647,6 +671,11 @@ bool AWorldDataLayers::ContainsDataLayer(const UDataLayer* InDataLayer) const
 	return WorldDataLayers.Contains(InDataLayer);
 }
 
+bool AWorldDataLayers::ContainsDataLayer(const UDataLayerInstance* InDataLayerInstance) const 
+{
+	return DataLayerInstances.Contains(InDataLayerInstance);
+}
+
 const UDataLayer* AWorldDataLayers::GetDataLayerFromName(const FName& InDataLayerName) const
 {
 #if WITH_EDITOR	
diff --git a/Engine/Source/Runtime/Engine/Public/WorldPartition/DataLayer/DataLayerAsset.h b/Engine/Source/Runtime/Engine/Public/WorldPartition/DataLayer/DataLayerAsset.h
new file mode 100644
--- /dev/null
+++ b/Engine/Source/Runtime/Engine/Public/WorldPartition/DataLayer/DataLayerAsset.h
@@ -0,0 +1,46 @@
+// Copyright Epic Games, Inc. All Rights Reserved.
+
+#pragma once
+
+#include "CoreMinimal.h"
+#include "UObject/ObjectMacros.h"
+#include "UObject/UObjectGlobals.h"
+#include "UObject/Object.h"
+
+#include "DataLayerAsset.generated.h"
+
+UENUM()
+enum class EDataLayerType : uint8
+{
+	Runtime,
+	Editor
+};
+
+UCLASS(editinlinenew)
+class ENGINE_API UDataLayerAsset : public UObject
+{
+	GENERATED_UCLASS_BODY()
+
+	friend class UDataLayerConversionInfo;
+
+public:
+	virtual void PostLoad() override;
+
+#if WITH_EDITOR
+	const TCHAR* GetDataLayerIconName() const;
+#endif
+
+	UFUNCTION(Category = "Data Layer|Runtime", BlueprintCallable)
+	bool IsRuntime() const { return DataLayerType == EDataLayerType::Runtime; }
+
+	UFUNCTION(Category = "Data Layer|Runtime", BlueprintCallable)
+	FColor GetDebugColor() const { return DebugColor; }
+
+private:
+	/** Whether the Data Layer affects actor runtime loading */
+	UPROPERTY(Category = "Data Layer", EditAnywhere)
+	EDataLayerType DataLayerType;
+
+	UPROPERTY(Category = "Data Layer|Runtime", EditAnywhere, meta = (EditConditionHides, EditCondition = "DataLayerType == EDataLayerType::Runtime"))
+	FColor DebugColor;
+};
\ No newline at end of file
diff --git a/Engine/Source/Runtime/Engine/Public/WorldPartition/DataLayer/DataLayerInstance.h b/Engine/Source/Runtime/Engine/Public/WorldPartition/DataLayer/DataLayerInstance.h
new file mode 100644
--- /dev/null
+++ b/Engine/Source/Runtime/Engine/Public/WorldPartition/DataLayer/DataLayerInstance.h
@@ -0,0 +1,124 @@
+// Copyright Epic Games, Inc. All Rights Reserved.
+
+#pragma once
+
+#include "CoreMinimal.h"
+#include "UObject/ObjectMacros.h"
+#include "UObject/UObjectGlobals.h"
+#include "UObject/Object.h"
+#include "WorldPartition/DataLayer/DataLayerAsset.h"
+#include "WorldPartition/DataLayer/DataLayer.h"
+
+#include "DataLayerInstance.generated.h"
+
+UCLASS(Config = Engine, PerObjectConfig, Within = WorldDataLayers, BlueprintType, AutoCollapseCategories = ("Data Layer|Advanced"), AutoExpandCategories = ("Data Layer|Editor", "Data Layer|Advanced|Runtime"))
+class ENGINE_API UDataLayerInstance : public UObject
+{
+	GENERATED_UCLASS_BODY()
+
+	friend class UDataLayerConversionInfo;
+
+public:
+	virtual void PostLoad() override;
+
+#if WITH_EDITOR
+	void OnCreated(const UDataLayerAsset* Asset);
+
+	void SetVisible(bool bIsVisible);
+	void SetIsInitiallyVisible(bool bIsInitiallyVisible);
+	void SetIsLoadedInEditor(bool bIsLoadedInEditor, bool bFromUserChange);
+	void SetIsLocked(bool bInIsLocked) { bIsLocked = bInIsLocked; }
+
+	bool IsLocked() const;
+	bool IsInitiallyLoadedInEditor() const { return bIsInitiallyLoadedInEditor; }
+	bool IsLoadedInEditor() const { return bIsLoadedInEditor; }
+	bool IsEffectiveLoadedInEditor() const;
+	bool IsLoadedInEditorChangedByUserOperation() const { return bIsLoadedInEditorChangedByUserOperation; }
+	void ClearLoadedInEditorChangedByUserOperation() { bIsLoadedInEditorChangedByUserOperation = false; }
+
+	const TCHAR* GetDataLayerIconName() const { return DataLayerAsset->GetDataLayerIconName(); }
+
+	bool CanParent(const UDataLayerInstance* InParent) const;
+	void SetParent(UDataLayerInstance* InParent);
+
+	void SetChildParent(UDataLayerInstance* InParent);
+
+	void CheckForErrors() const;
+#endif
+
+	const UDataLayerAsset* GetAsset() const { return DataLayerAsset.Get(); }
+
+	UFUNCTION(Category = "Data Layer", BlueprintCallable)
+	bool Equals(const FActorDataLayer& ActorDataLayer) const { return ActorDataLayer.Name == GetFName(); }
+
+	UFUNCTION(Category = "Data Layer|Editor", BlueprintCallable)
+	bool IsInitiallyVisible() const;
+
+	UFUNCTION(Category = "Data Layer|Editor", BlueprintCallable)
+	bool IsVisible() const;
+
+	UFUNCTION(Category = "Data Layer|Editor", BlueprintCallable)
+	bool IsEffectiveVisible() const;
+
+	UFUNCTION(Category = "Data Layer", BlueprintCallable)
+	FName GetDataLayerLabel() const { return DataLayerAsset->GetFName(); }
+
+	UFUNCTION(Category = "Data Layer|Runtime", BlueprintCallable)
+	bool IsRuntime() const { return DataLayerAsset->IsRuntime(); }
+
+	UFUNCTION(Category = "Data Layer|Runtime", BlueprintCallable)
+	FColor GetDebugColor() const { return DataLayerAsset->GetDebugColor(); }
+
+	UFUNCTION(Category = "Data Layer|Runtime", BlueprintCallable)
+	EDataLayerRuntimeState GetInitialRuntimeState() const { return IsRuntime() ? InitialRuntimeState : EDataLayerRuntimeState::Unloaded; }
+
+	const TArray<TObjectPtr<UDataLayerInstance>>& GetChildren() const { return Children; }
+	void ForEachChild(TFunctionRef<bool(const UDataLayerInstance*)> Operation) const;
+
+	const UDataLayerInstance* GetParent() const { return Parent; }
+	UDataLayerInstance* GetParent() { return Parent; }
+
+private:
+	void AddChild(UDataLayerInstance* DataLayer);
+#if WITH_EDITOR
+	void RemoveChild(UDataLayerInstance* DataLayer);
+#endif
+
+#if WITH_EDITORONLY_DATA
+	/** Whether actors associated with the DataLayer are visible in the viewport */
+	UPROPERTY(Transient)
+	uint32 bIsVisible : 1;
+
+	/** Whether actors associated with the Data Layer should be initially visible in the viewport when loading the map */
+	UPROPERTY(Category = "Data Layer|Editor", EditAnywhere)
+	uint32 bIsInitiallyVisible : 1;
+
+	/** Determines the default value of the data layer's loaded state in editor if it hasn't been changed in data layer outliner by the user */
+	UPROPERTY(Category = "Data Layer|Editor", EditAnywhere, meta = (DisplayName = "Is Initially Loaded"))
+	uint32 bIsInitiallyLoadedInEditor : 1;
+
+	/** Wheter the data layer is loaded in editor (user setting) */
+	UPROPERTY(Transient)
+	uint32 bIsLoadedInEditor : 1;
+
+	/** Whether this data layer editor visibility was changed by a user operation */
+	UPROPERTY(Transient)
+	uint32 bIsLoadedInEditorChangedByUserOperation : 1;
+
+	/** Whether this data layer is locked, which means the user can't change actors assignation, remove or rename it */
+	UPROPERTY()
+	uint32 bIsLocked : 1;
+#endif
+
+	UPROPERTY(Category = "Data Layer", VisibleAnywhere)
+	TObjectPtr<const UDataLayerAsset> DataLayerAsset;
+
+	UPROPERTY(Category = "Data Layer|Advanced|Runtime", EditAnywhere, meta = (EditConditionHides, EditCondition = "bIsRuntime"))
+	EDataLayerRuntimeState InitialRuntimeState;
+
+	UPROPERTY()
+	TObjectPtr<UDataLayerInstance> Parent;
+
+	UPROPERTY(Transient)
+	TArray<TObjectPtr<UDataLayerInstance>> Children;
+};
diff --git a/Engine/Source/Runtime/Engine/Public/WorldPartition/DataLayer/WorldDataLayers.h b/Engine/Source/Runtime/Engine/Public/WorldPartition/DataLayer/WorldDataLayers.h
--- a/Engine/Source/Runtime/Engine/Public/WorldPartition/DataLayer/WorldDataLayers.h
+++ b/Engine/Source/Runtime/Engine/Public/WorldPartition/DataLayer/WorldDataLayers.h
@@ -10,6 +10,8 @@
 #include "WorldDataLayers.generated.h"
 
 class UDataLayer;
+class UDataLayerInstance;
+class UDataLayerAsset;
 
 /**
  * Actor containing all data layers for a world
@@ -19,6 +21,8 @@ class ENGINE_API AWorldDataLayers : public AInfo
 {
 	GENERATED_UCLASS_BODY()
 
+	friend class UDataLayerToAssetCommandlet;
+
 public:
 	virtual void PostLoad() override;
 	virtual void RewindForReplay() override;
@@ -32,12 +36,15 @@ public:
 
 	static AWorldDataLayers* Create(UWorld* World);
 	UDataLayer* CreateDataLayer(FName InName = TEXT("DataLayer"), EObjectFlags InObjectFlags = RF_NoFlags);
+	UDataLayerInstance* CreateDataLayer(const UDataLayerAsset* InDataLayerAsset);
 	bool RemoveDataLayer(UDataLayer* InDataLayer);
 	bool RemoveDataLayers(const TArray<UDataLayer*>& InDataLayers);
 	FName GenerateUniqueDataLayerLabel(const FName& InDataLayerLabel) const;
 	void SetAllowRuntimeDataLayerEditing(bool bInAllowRuntimeDataLayerEditing);
 	bool GetAllowRuntimeDataLayerEditing() const { return bAllowRuntimeDataLayerEditing; }
 
+	bool RemoveDataLayer(UDataLayerInstance* InDataLayerInstance);
+
 	//~ Begin Helper Functions
 	TArray<const UDataLayer*> GetDataLayerObjects(const TArray<FActorDataLayer>& DataLayers) const;
 	TArray<const UDataLayer*> GetDataLayerObjects(const TArray<FName>& InDataLayerNames) const;
@@ -58,6 +65,8 @@ public:
 	void ForEachDataLayer(TFunctionRef<bool(UDataLayer*)> Func);
 	void ForEachDataLayer(TFunctionRef<bool(UDataLayer*)> Func) const;
 
+	bool ContainsDataLayer(const UDataLayerInstance* InDataLayerInstance) const;
+
 	// DataLayer Runtime State
 	void SetDataLayerRuntimeState(FActorDataLayer InDataLayer, EDataLayerRuntimeState InState, bool bIsRecursive = false);
 	EDataLayerRuntimeState GetDataLayerRuntimeStateByName(FName InDataLayerName) const;
@@ -118,6 +127,12 @@ private:
 	bool bAllowRuntimeDataLayerEditing;
 #endif
 
+	UPROPERTY()
+	TSet<TObjectPtr<UDataLayerInstance>> DataLayerInstances;
+
+	UPROPERTY()
+	TMap<FName, TWeakObjectPtr<UDataLayerInstance>> DeprecatedDataLayerNameToDataLayerInstance;
+
 	UPROPERTY()
 	TSet<TObjectPtr<UDataLayer>> WorldDataLayers;
 
