[WP Navmesh] Fix removed actors still showing in the editor actor list after navmesh build
Enable usage of source control provider when building WP navmesh
Fix issue with deleted packages being added to source control
Make sure WP navmesh actor in WP levels remains empty (data is provided by navmesh chunk actors)

diff --git a/Engine/Source/Editor/UnrealEd/Private/EditorBuildUtils.cpp b/Engine/Source/Editor/UnrealEd/Private/EditorBuildUtils.cpp
--- a/Engine/Source/Editor/UnrealEd/Private/EditorBuildUtils.cpp
+++ b/Engine/Source/Editor/UnrealEd/Private/EditorBuildUtils.cpp
@@ -934,10 +934,13 @@ bool FEditorBuildUtils::WorldPartitionBuildNavigation(const FString& InLongPacka
 
 			uint32 ProcessID;
 
-			const FString Arguments = FString::Printf(TEXT("\"%s\" -run=WorldPartitionBuilderCommandlet %s %s %s %s"),
+			ISourceControlProvider& SCCProvider = ISourceControlModule::Get().GetProvider();
+
+			const FString Arguments = FString::Printf(TEXT("\"%s\" -run=WorldPartitionBuilderCommandlet %s %s -SCCProvider=%s %s %s"),
 				*ProjectPath,
 				*InLongPackageName,
-				TEXT(" -AllowCommandletRendering -Builder=WorldPartitionNavigationDataBuilder -SCCProvider=None -log=WPNavigationBuilderLog.txt"), 
+				TEXT(" -AllowCommandletRendering -Builder=WorldPartitionNavigationDataBuilder -log=WPNavigationBuilderLog.txt"),
+				*SCCProvider.GetName().ToString(),
 				DefaultBuildNavigationOptions->bVerbose ? TEXT("-Verbose") : TEXT(""),
 				DefaultBuildNavigationOptions->bCleanPackages ? TEXT("-CleanPackages") : TEXT(""));
 				
@@ -962,6 +965,10 @@ bool FEditorBuildUtils::WorldPartitionBuildNavigation(const FString& InLongPacka
 		{	
 			if (Result == 0)
 			{
+				// Force a directory watcher tick for the asset registry to get notified of the changes
+				FDirectoryWatcherModule& DirectoryWatcherModule = FModuleManager::Get().LoadModuleChecked<FDirectoryWatcherModule>(TEXT("DirectoryWatcher"));
+				DirectoryWatcherModule.Get()->Tick(-1.0f);
+				
 				// Unload any loaded map
 				if (!UEditorLoadingAndSavingUtils::NewBlankMap(/*bSaveExistingMap*/false))
 				{
@@ -976,10 +983,6 @@ bool FEditorBuildUtils::WorldPartitionBuildNavigation(const FString& InLongPacka
 			
 				AssetRegistry.ScanModifiedAssetFiles({ MapToLoad });
 				AssetRegistry.ScanPathsSynchronous(ULevel::GetExternalObjectsPaths(MapToLoad), true);
-
-				// Force a directory watcher tick for the asset registry to get notified of the changes
-				FDirectoryWatcherModule& DirectoryWatcherModule = FModuleManager::Get().LoadModuleChecked<FDirectoryWatcherModule>(TEXT("DirectoryWatcher"));
-				DirectoryWatcherModule.Get()->Tick(-1.0f);
 			
 				FEditorFileUtils::LoadMap(MapToLoad);
 			}
diff --git a/Engine/Source/Editor/UnrealEd/Private/WorldPartition/WorldPartitionNavigationDataBuilder.cpp b/Engine/Source/Editor/UnrealEd/Private/WorldPartition/WorldPartitionNavigationDataBuilder.cpp
--- a/Engine/Source/Editor/UnrealEd/Private/WorldPartition/WorldPartitionNavigationDataBuilder.cpp
+++ b/Engine/Source/Editor/UnrealEd/Private/WorldPartition/WorldPartitionNavigationDataBuilder.cpp
@@ -75,7 +75,6 @@ bool UWorldPartitionNavigationDataBuilder::RunInternal(UWorld* World, const FCel
 	// Destroy any existing navigation data chunk actors within bounds we are generating, we will make new ones.
 	int32 Count = 0;
 	const FBox GeneratingBounds = InCellInfo.Bounds.ExpandBy(-IterativeCellOverlapSize);
-
 	UE_LOG(LogWorldPartitionNavigationDataBuilder, Verbose, TEXT("   GeneratingBounds %s"), *GeneratingBounds.ToString());
 	
 	for (TActorIterator<ANavigationDataChunkActor> It(World); It; ++It)
@@ -149,6 +148,8 @@ bool UWorldPartitionNavigationDataBuilder::RunInternal(UWorld* World, const FCel
 	}
 
 	TArray<UPackage*> PackagesToSave;
+	
+	TArray<UPackage*> PackagesToAdd;
 	TArray<UPackage*> PackagesToDelete;
 
 	for (UPackage* ActorPackage : NavigationDataChunkActorPackages)
@@ -160,6 +161,10 @@ bool UWorldPartitionNavigationDataBuilder::RunInternal(UWorld* World, const FCel
 			{
 				PackagesToDelete.Add(ActorPackage);
 			}
+			else
+			{
+				PackagesToAdd.Add(ActorPackage);
+			}
 
 			// Save all packages (we need to also save the ones we are deleting).
 			PackagesToSave.Add(ActorPackage);
@@ -186,7 +191,7 @@ bool UWorldPartitionNavigationDataBuilder::RunInternal(UWorld* World, const FCel
 	if (!PackagesToSave.IsEmpty())
 	{
 		{
-			// Checkout packages to save
+			// Checkout or remove read-only for packages to save
 			TRACE_CPUPROFILER_EVENT_SCOPE(CheckoutPackages);
 			UE_LOG(LogWorldPartitionNavigationDataBuilder, Log, TEXT("Checking out %d packages."), PackagesToSave.Num());
 
@@ -222,7 +227,7 @@ bool UWorldPartitionNavigationDataBuilder::RunInternal(UWorld* World, const FCel
 			TRACE_CPUPROFILER_EVENT_SCOPE(AddingToSourceControl);
 			UE_LOG(LogWorldPartitionNavigationDataBuilder, Log, TEXT("Adding packages to source control."));
 
-			for (UPackage* Package : PackagesToSave)
+			for (UPackage* Package : PackagesToAdd)
 			{
 				if (!PackageHelper.AddToSourceControl(Package))
 				{
@@ -238,14 +243,14 @@ bool UWorldPartitionNavigationDataBuilder::RunInternal(UWorld* World, const FCel
 	return true;
 }
 
-FName GetCellName(UWorldPartition* WorldPartition, const UActorPartitionSubsystem::FCellCoord& InCellCoord)
+FName GetCellName(const UWorldPartition* WorldPartition, const UActorPartitionSubsystem::FCellCoord& InCellCoord)
 {
 	const FString PackageName = FPackageName::GetShortName(WorldPartition->GetPackage());
 	const FString PackageNameNoPIEPrefix = UWorld::RemovePIEPrefix(PackageName);
-	return FName(*FString::Printf(TEXT("%s_%i_%i"), *PackageNameNoPIEPrefix, InCellCoord.X, InCellCoord.Y));
+	return FName(*FString::Printf(TEXT("%s_%lld_%lld"), *PackageNameNoPIEPrefix, InCellCoord.X, InCellCoord.Y));
 }
 
-bool UWorldPartitionNavigationDataBuilder::GenerateNavigationData(UWorldPartition* WorldPartition, const FBox& LoadedBounds, const FBox& GeneratingBounds)
+bool UWorldPartitionNavigationDataBuilder::GenerateNavigationData(UWorldPartition* WorldPartition, const FBox& LoadedBounds, const FBox& GeneratingBounds) const
 {
 	TRACE_CPUPROFILER_EVENT_SCOPE(UWorldPartitionNavigationDataBuilder::GenerateNavigationData);
 
@@ -289,8 +294,6 @@ bool UWorldPartitionNavigationDataBuilder::GenerateNavigationData(UWorldPartitio
 	}
 
 	// For each cell, gather navmesh and generate a datachunk actor
-	const FBox WorldBounds = WorldPartition->GetWorldBounds();
-
 	int32 ActorCount = 0;
 
 	// Keep track of all valid navigation data chunk actors
@@ -298,7 +301,7 @@ bool UWorldPartitionNavigationDataBuilder::GenerateNavigationData(UWorldPartitio
 
 	// A DataChunkActor will be generated for each tile touching the generating bounds.
 	const TSubclassOf<APartitionActor>& NavigationDataActorClass = ANavigationDataChunkActor::StaticClass();
-	FIntRect GeneratingBounds2D(GeneratingBounds.Min.X, GeneratingBounds.Min.Y, GeneratingBounds.Max.X, GeneratingBounds.Max.Y);
+	const FIntRect GeneratingBounds2D(GeneratingBounds.Min.X, GeneratingBounds.Min.Y, GeneratingBounds.Max.X, GeneratingBounds.Max.Y);
 	FActorPartitionGridHelper::ForEachIntersectingCell(NavigationDataActorClass, GeneratingBounds2D, World->PersistentLevel,
 		[&WorldPartition, &ActorCount, World, &ValidNavigationDataChunkActors, &NavDataBounds, this](const UActorPartitionSubsystem::FCellCoord& InCellCoord, const FIntRect& InCellBounds)->bool
 		{
@@ -354,7 +357,7 @@ bool UWorldPartitionNavigationDataBuilder::GenerateNavigationData(UWorldPartitio
 	return true;
 }
 
-bool UWorldPartitionNavigationDataBuilder::SavePackages(const TArray<UPackage*>& PackagesToSave)
+bool UWorldPartitionNavigationDataBuilder::SavePackages(const TArray<UPackage*>& PackagesToSave) const
 {
 	// Save packages
 	TRACE_CPUPROFILER_EVENT_SCOPE(SavingPackages);
diff --git a/Engine/Source/Editor/UnrealEd/Public/WorldPartition/WorldPartitionNavigationDataBuilder.h b/Engine/Source/Editor/UnrealEd/Public/WorldPartition/WorldPartitionNavigationDataBuilder.h
--- a/Engine/Source/Editor/UnrealEd/Public/WorldPartition/WorldPartitionNavigationDataBuilder.h
+++ b/Engine/Source/Editor/UnrealEd/Public/WorldPartition/WorldPartitionNavigationDataBuilder.h
@@ -22,9 +22,9 @@ protected:
 	virtual bool RunInternal(UWorld* World, const FCellInfo& InCellInfo, FPackageSourceControlHelper& PackageHelper) override;
 	// UWorldPartitionBuilder interface end
 
-	bool GenerateNavigationData(UWorldPartition* WorldPartition, const FBox& LoadedBounds, const FBox& GeneratingBounds);
+	bool GenerateNavigationData(UWorldPartition* WorldPartition, const FBox& LoadedBounds, const FBox& GeneratingBounds) const;
 
-	bool SavePackages(const TArray<UPackage*>& PackagesToSave);
+	bool SavePackages(const TArray<UPackage*>& PackagesToSave) const;
 
 	bool bCleanBuilderPackages = false;
 };
diff --git a/Engine/Source/Runtime/NavigationSystem/Private/NavMesh/PImplRecastNavMesh.cpp b/Engine/Source/Runtime/NavigationSystem/Private/NavMesh/PImplRecastNavMesh.cpp
--- a/Engine/Source/Runtime/NavigationSystem/Private/NavMesh/PImplRecastNavMesh.cpp
+++ b/Engine/Source/Runtime/NavigationSystem/Private/NavMesh/PImplRecastNavMesh.cpp
@@ -619,25 +619,34 @@ void FPImplRecastNavMesh::Serialize( FArchive& ArWrapped, int32 NavMeshVersion )
 		
 		const UNavigationSystemV1* NavSys = FNavigationSystem::GetCurrent<const UNavigationSystemV1>(NavMeshOwner->GetWorld());
 
-		// Need to keep the check !IsRunningCommandlet() for the case where maps are cooked and saved from UCookCommandlet.
-		// In that flow the nav bounds are not set (no bounds means no tiles to save and the navmesh would be saved without tiles).
-		// This flow would benefit to be revisited since navmesh serialization should not be different whether it was run or not by a commandlet.
-		// Fixes missing navmesh regression (UE-103604).
-		if (NavMeshOwner->SupportsStreaming() && NavSys && !IsRunningCommandlet())
+		if (NavMeshOwner->bIsWorldPartitioned)
 		{
-			// We save only tiles that belongs to this level
-			GetNavMeshTilesIn(NavMeshOwner->GetNavigableBoundsInLevel(NavMeshOwner->GetLevel()), TilesToSave);
+			// Ignore (leave TilesToSave empty so no tiles are saved).
+			// Navmesh data are stored in ANavigationDataChunkActor.
+			UE_LOG(LogNavigation, VeryVerbose, TEXT("%s Ar.IsSaving() no tiles are being saved because bIsWorldPartitioned=true in %s."), ANSI_TO_TCHAR(__FUNCTION__), *NavMeshOwner->GetFullName());
 		}
 		else
 		{
-			// Otherwise all valid tiles
-			dtNavMesh const* ConstNavMesh = DetourNavMesh;
-			for (int i = 0; i < ConstNavMesh->getMaxTiles(); ++i)
+			// Need to keep the check !IsRunningCommandlet() for the case where maps are cooked and saved from UCookCommandlet.
+			// In that flow the nav bounds are not set (no bounds means no tiles to save and the navmesh would be saved without tiles).
+			// This flow would benefit to be revisited since navmesh serialization should not be different whether it was run or not by a commandlet.
+			// Fixes missing navmesh regression (UE-103604).
+			if (NavMeshOwner->SupportsStreaming() && NavSys && !IsRunningCommandlet())
 			{
-				const dtMeshTile* Tile = ConstNavMesh->getTile(i);
-				if (Tile != NULL && Tile->header != NULL && Tile->dataSize > 0)
+				// We save only tiles that belongs to this level
+				GetNavMeshTilesIn(NavMeshOwner->GetNavigableBoundsInLevel(NavMeshOwner->GetLevel()), TilesToSave);
+			}
+			else
+			{
+				// Otherwise all valid tiles
+				dtNavMesh const* ConstNavMesh = DetourNavMesh;
+				for (int i = 0; i < ConstNavMesh->getMaxTiles(); ++i)
 				{
-					TilesToSave.Add(i);
+					const dtMeshTile* Tile = ConstNavMesh->getTile(i);
+					if (Tile != NULL && Tile->header != NULL && Tile->dataSize > 0)
+					{
+						TilesToSave.Add(i);
+					}
 				}
 			}
 		}
diff --git a/Engine/Source/Runtime/NavigationSystem/Private/NavMesh/RecastNavMesh.cpp b/Engine/Source/Runtime/NavigationSystem/Private/NavMesh/RecastNavMesh.cpp
--- a/Engine/Source/Runtime/NavigationSystem/Private/NavMesh/RecastNavMesh.cpp
+++ b/Engine/Source/Runtime/NavigationSystem/Private/NavMesh/RecastNavMesh.cpp
@@ -324,6 +324,7 @@ FRecastNavMeshGenerationProperties::FRecastNavMeshGenerationProperties()
 	bFilterLowSpanSequences = false;
 	bFilterLowSpanFromTileCache = false;
 	bFixedTilePoolSize = false;
+	bIsWorldPartitioned = false;
 }
 
 FRecastNavMeshGenerationProperties::FRecastNavMeshGenerationProperties(const ARecastNavMesh& RecastNavMesh)
@@ -351,6 +352,7 @@ FRecastNavMeshGenerationProperties::FRecastNavMeshGenerationProperties(const ARe
 	bFilterLowSpanSequences = RecastNavMesh.bFilterLowSpanSequences;
 	bFilterLowSpanFromTileCache = RecastNavMesh.bFilterLowSpanFromTileCache;
 	bFixedTilePoolSize = RecastNavMesh.bFixedTilePoolSize;
+	bIsWorldPartitioned = RecastNavMesh.bIsWorldPartitioned;
 }
 
 FRecastNavMeshTileGenerationDebug::FRecastNavMeshTileGenerationDebug()
@@ -387,6 +389,7 @@ ARecastNavMesh::ARecastNavMesh(const FObjectInitializer& ObjectInitializer)
 	, DefaultMaxSearchNodes(RECAST_MAX_SEARCH_NODES)
 	, DefaultMaxHierarchicalSearchNodes(RECAST_MAX_SEARCH_NODES)
 	, bSortNavigationAreasByCost(true)
+	, bIsWorldPartitioned(false)
 	, bPerformVoxelFiltering(true)	
 	, bMarkLowHeightAreas(false)
 	, bUseExtraTopCellWhenMarkingAreas(true)
@@ -2786,6 +2789,7 @@ void ARecastNavMesh::UpdateGenerationProperties(const FRecastNavMeshGenerationPr
 	bFilterLowSpanSequences = GenerationProps.bFilterLowSpanSequences;
 	bFilterLowSpanFromTileCache = GenerationProps.bFilterLowSpanFromTileCache;
 	bFixedTilePoolSize = GenerationProps.bFixedTilePoolSize;
+	bIsWorldPartitioned = GenerationProps.bIsWorldPartitioned;
 }
 
 bool ARecastNavMesh::IsVoxelCacheEnabled()
diff --git a/Engine/Source/Runtime/NavigationSystem/Public/NavMesh/RecastNavMesh.h b/Engine/Source/Runtime/NavigationSystem/Public/NavMesh/RecastNavMesh.h
--- a/Engine/Source/Runtime/NavigationSystem/Public/NavMesh/RecastNavMesh.h
+++ b/Engine/Source/Runtime/NavigationSystem/Public/NavMesh/RecastNavMesh.h
@@ -373,6 +373,10 @@ struct NAVIGATIONSYSTEM_API FRecastNavMeshGenerationProperties
 	UPROPERTY(EditAnywhere, Category = Generation)
 	uint32 bFixedTilePoolSize : 1;
 
+	/* In a world partitioned map, is this navmesh using world partitioning */
+	UPROPERTY(EditAnywhere, Category = Generation)
+	uint32 bIsWorldPartitioned : 1;
+	
 	FRecastNavMeshGenerationProperties();
 	FRecastNavMeshGenerationProperties(const ARecastNavMesh& RecastNavMesh);
 };
@@ -697,6 +701,10 @@ class NAVIGATIONSYSTEM_API ARecastNavMesh : public ANavigationData
 	UPROPERTY(EditAnywhere, Category=Generation, config)
 	uint32 bSortNavigationAreasByCost:1;
 
+	/* In a world partitioned map, is this navmesh using world partitioning */
+	UPROPERTY(EditAnywhere, Category=Generation, config)
+	uint32 bIsWorldPartitioned : 1;
+	
 	/** controls whether voxel filtering will be applied (via FRecastTileGenerator::ApplyVoxelFilter). 
 	 *	Results in generated navmesh better fitting navigation bounds, but hits (a bit) generation performance */
 	UPROPERTY(EditAnywhere, Category=Generation, config, AdvancedDisplay)
