PackageSourceControlHelper: Batched Add

diff --git a/Engine/Source/Editor/UnrealEd/Private/PackageSourceControlHelper.cpp b/Engine/Source/Editor/UnrealEd/Private/PackageSourceControlHelper.cpp
--- a/Engine/Source/Editor/UnrealEd/Private/PackageSourceControlHelper.cpp
+++ b/Engine/Source/Editor/UnrealEd/Private/PackageSourceControlHelper.cpp
@@ -199,17 +199,70 @@ bool FPackageSourceControlHelper::AddToSourceControl(UPackage* Package) const
 	if (UseSourceControl())
 	{
 		FString PackageFilename = SourceControlHelpers::PackageFilename(Package);
-		FSourceControlStatePtr SourceControlState = GetSourceControlProvider().GetState(PackageFilename, EStateCacheUsage::ForceUpdate);
+		return AddToSourceControl({ PackageFilename });
+	}
+
+	return true;
+}
+
+bool FPackageSourceControlHelper::AddToSourceControl(const TArray<FString>& PackageNames) const
+{
+	if (!UseSourceControl())
+	{
+		return true;
+	}
+
+	// Convert package names to package filenames
+	TArray<FString> PackageFilenames = SourceControlHelpers::PackageFilenames(PackageNames);
+
+	// Two-pass checkout mechanism
+	TArray<FString> PackagesToAdd;
+	PackagesToAdd.Reserve(PackageFilenames.Num());
+	bool bSuccess = true;
+	
+	TArray<FSourceControlStateRef> SourceControlStates;
+	ECommandResult::Type UpdateState = GetSourceControlProvider().GetState(PackageFilenames, SourceControlStates, EStateCacheUsage::ForceUpdate);
+
+	if (UpdateState != ECommandResult::Succeeded)
+	{
+		UE_LOG(LogCommandletPackageHelper, Error, TEXT("Could not get source control state for packages"));
+		return false;
+	}
+	
+	for (FSourceControlStateRef& SourceControlState : SourceControlStates)
+	{
+		const FString& PackageFilename = SourceControlState->GetFilename();
 
-		if (SourceControlState.IsValid() && !SourceControlState->IsSourceControlled())
+		FString OtherCheckedOutUser;
+		if (SourceControlState->IsCheckedOutOther(&OtherCheckedOutUser))
 		{
-			UE_LOG(LogCommandletPackageHelper, Log, TEXT("Adding package %s to source control"), *PackageFilename);
-			if (GetSourceControlProvider().Execute(ISourceControlOperation::Create<FMarkForAdd>(), Package) != ECommandResult::Succeeded)
-			{
-				UE_LOG(LogCommandletPackageHelper, Error, TEXT("Error adding %s to source control."), *PackageFilename);
-				return false;
-			}
+			UE_LOG(LogCommandletPackageHelper, Error, TEXT("Overwriting package %s already checked out by %s, will not add"), *PackageFilename, *OtherCheckedOutUser);
+			bSuccess = false;
 		}
+		else if (!SourceControlState->IsCurrent())
+		{
+			UE_LOG(LogCommandletPackageHelper, Error, TEXT("Overwriting package %s (not at head revision), will not add"), *PackageFilename);
+			bSuccess = false;
+		}
+		else if (SourceControlState->IsAdded())
+		{
+			// Nothing to do
+		}
+		else if (!SourceControlState->IsSourceControlled())
+		{
+			PackagesToAdd.Add(PackageFilename);
+		}
+	}
+
+	// Any error up to here will be an early out
+	if (!bSuccess)
+	{
+		return false;
+	}
+
+	if (PackagesToAdd.Num())
+	{
+		return (GetSourceControlProvider().Execute(ISourceControlOperation::Create<FMarkForAdd>(), PackagesToAdd) == ECommandResult::Succeeded);
 	}
 
 	return true;
diff --git a/Engine/Source/Editor/UnrealEd/Public/PackageSourceControlHelper.h b/Engine/Source/Editor/UnrealEd/Public/PackageSourceControlHelper.h
--- a/Engine/Source/Editor/UnrealEd/Public/PackageSourceControlHelper.h
+++ b/Engine/Source/Editor/UnrealEd/Public/PackageSourceControlHelper.h
@@ -15,6 +15,7 @@ public:
 	bool Delete(UPackage* Package) const;
 	bool Delete(const TArray<UPackage*>& Packages) const;
 	bool AddToSourceControl(UPackage* Package) const;
+	bool AddToSourceControl(const TArray<FString>& PackageNames) const;
 	bool Checkout(UPackage* Package) const;
 	bool Checkout(const TArray<FString>& PackageNames) const;
 
