WorldPartitionRenameDuplicateBuilder: Allows Renaming or Duplicating a WorldPartition map without loading it all at once (iterates through actor reference clusters)

diff --git a/Engine/Source/Editor/UnrealEd/Classes/Commandlets/WorldPartitionRenameCommandlet.h b/Engine/Source/Editor/UnrealEd/Classes/Commandlets/WorldPartitionRenameCommandlet.h
deleted file mode 100644
diff --git a/Engine/Source/Editor/UnrealEd/Private/Commandlets/WorldPartitionRenameCommandlet.cpp b/Engine/Source/Editor/UnrealEd/Private/Commandlets/WorldPartitionRenameCommandlet.cpp
deleted file mode 100644
diff --git a/Engine/Source/Editor/UnrealEd/Private/WorldPartition/WorldPartitionBuilder.cpp b/Engine/Source/Editor/UnrealEd/Private/WorldPartition/WorldPartitionBuilder.cpp
--- a/Engine/Source/Editor/UnrealEd/Private/WorldPartition/WorldPartitionBuilder.cpp
+++ b/Engine/Source/Editor/UnrealEd/Private/WorldPartition/WorldPartitionBuilder.cpp
@@ -58,51 +58,58 @@ bool UWorldPartitionBuilder::RunBuilder(UWorld* World)
 		return false;
 	}
 
-	// Setup the world
-	UWorld::InitializationValues IVS;
-	IVS.RequiresHitProxies(false);
-	IVS.ShouldSimulatePhysics(false);
-	IVS.EnableTraceCollision(false);
-	IVS.CreateNavigation(false);
-	IVS.CreateAISystem(false);
-	IVS.AllowAudioPlayback(false);
-	IVS.CreatePhysicsScene(true);
-	FScopedEditorWorld EditorWorld(World, IVS);
-
-	// Make sure the world is partitioned
 	bool bResult = true;
-	if (World->HasSubsystem<UWorldPartitionSubsystem>())
+	// Setup the world
 	{
-		// Ensure the world has a valid world partition.
-		UWorldPartition* WorldPartition = World->GetWorldPartition();
-		check(WorldPartition);
+		UWorld::InitializationValues IVS;
+		IVS.RequiresHitProxies(false);
+		IVS.ShouldSimulatePhysics(false);
+		IVS.EnableTraceCollision(false);
+		IVS.CreateNavigation(false);
+		IVS.CreateAISystem(false);
+		IVS.AllowAudioPlayback(false);
+		IVS.CreatePhysicsScene(true);
+		FScopedEditorWorld EditorWorld(World, IVS);
+
+		// Make sure the world is partitioned
+		if (World->HasSubsystem<UWorldPartitionSubsystem>())
+		{
+			// Ensure the world has a valid world partition.
+			UWorldPartition* WorldPartition = World->GetWorldPartition();
+			check(WorldPartition);
 
-		FWorldContext& WorldContext = GEditor->GetEditorWorldContext(true /*bEnsureIsGWorld*/);
-		WorldContext.SetCurrentWorld(World);
-		UWorld* PrevGWorld = GWorld;
-		GWorld = World;
+			FWorldContext& WorldContext = GEditor->GetEditorWorldContext(true /*bEnsureIsGWorld*/);
+			WorldContext.SetCurrentWorld(World);
+			UWorld* PrevGWorld = GWorld;
+			GWorld = World;
 
-		// Run builder
-		bResult = Run(World, SCCHelper);
+			// Run builder
+			bResult = Run(World, SCCHelper);
 
-		// Restore previous world
-		WorldContext.SetCurrentWorld(PrevGWorld);
-		GWorld = PrevGWorld;
+			// Restore previous world
+			WorldContext.SetCurrentWorld(PrevGWorld);
+			GWorld = PrevGWorld;
 
-		// Save default configuration
-		if (bResult)
-		{
-			if (!FPlatformFileManager::Get().GetPlatformFile().FileExists(*WorldConfigFilename) ||
-				!FPlatformFileManager::Get().GetPlatformFile().IsReadOnly(*WorldConfigFilename))
+			// Save default configuration
+			if (bResult)
 			{
-				SaveConfig(CPF_Config, *WorldConfigFilename);
+				if (!FPlatformFileManager::Get().GetPlatformFile().FileExists(*WorldConfigFilename) ||
+					!FPlatformFileManager::Get().GetPlatformFile().IsReadOnly(*WorldConfigFilename))
+				{
+					SaveConfig(CPF_Config, *WorldConfigFilename);
+				}
 			}
 		}
+		else
+		{
+			UE_LOG(LogWorldPartitionBuilder, Error, TEXT("WorldPartition builders only works on partitioned maps."));
+			bResult = false;
+		}
 	}
-	else
+
+	if (bResult)
 	{
-		UE_LOG(LogWorldPartitionBuilder, Error, TEXT("WorldPartition builders only works on partitioned maps."));
-		bResult = false;
+		bResult = PostWorldTeardown(SCCHelper);
 	}
 
 	return bResult;
diff --git a/Engine/Source/Editor/UnrealEd/Private/WorldPartition/WorldPartitionRenameDuplicateBuilder.cpp b/Engine/Source/Editor/UnrealEd/Private/WorldPartition/WorldPartitionRenameDuplicateBuilder.cpp
new file mode 100644
--- /dev/null
+++ b/Engine/Source/Editor/UnrealEd/Private/WorldPartition/WorldPartitionRenameDuplicateBuilder.cpp
@@ -0,0 +1,557 @@
+// Copyright Epic Games, Inc. All Rights Reserved.
+
+#include "WorldPartition/WorldPartitionRenameDuplicateBuilder.h"
+
+#include "WorldPartition/WorldPartition.h"
+#include "WorldPartition/WorldPartitionHelpers.h"
+#include "WorldPartition/WorldPartitionActorDescView.h"
+#include "WorldPartition/WorldPartitionActorCluster.h"
+#include "PackageSourceControlHelper.h"
+#include "SourceControlHelpers.h"
+#include "Engine/World.h"
+#include "GameFramework/Actor.h"
+#include "EditorWorldUtils.h"
+#include "UObject/SavePackage.h"
+#include "UObject/MetaData.h"
+#include "ProfilingDebugging/ScopedTimers.h"
+#include "HAL/Filemanager.h"
+
+DEFINE_LOG_CATEGORY_STATIC(LogWorldPartitionCopyWorldBuilder, All, All);
+
+class FReplaceObjectRefsArchive : public FArchiveUObject
+{
+public:
+	FReplaceObjectRefsArchive(UObject* InRoot, const TMap<TObjectPtr<UObject>,TObjectPtr<UObject>>& InObjectsToReplace)
+		: Root(InRoot)
+		, ObjectsToReplace(InObjectsToReplace)
+	{
+		// Don't gather transient actor references
+		SetIsPersistent(true);
+
+		// Don't trigger serialization of compilable assets
+		SetShouldSkipCompilingAssets(true);
+
+		ArIgnoreOuterRef = true;
+		ArIsObjectReferenceCollector = true;
+		ArShouldSkipBulkData = true;
+
+		SubObjects.Add(Root);
+
+		Root->Serialize(*this);
+	}
+
+	virtual FArchive& operator<<(UObject*& Obj) override
+	{
+		if (Obj && !Obj->IsTemplate() && !Obj->HasAnyFlags(RF_Transient))
+		{
+			bool bWasAlreadyInSet;
+			SubObjects.Add(Obj, &bWasAlreadyInSet);
+
+			if (!bWasAlreadyInSet)
+			{
+				if (const TObjectPtr<UObject>* ReplacementObject = ObjectsToReplace.Find(Obj))
+				{
+					Obj = *ReplacementObject;
+				}
+				else if (Obj->IsInOuter(Root))
+				{
+					Obj->Serialize(*this);
+				}
+			}
+		}
+		return *this;
+	}
+
+private:
+	UObject* Root;
+	const TMap<TObjectPtr<UObject>, TObjectPtr<UObject>>& ObjectsToReplace;
+	TSet<UObject*> SubObjects;
+};
+
+static bool DeleteExistingMapPackages(const FString& ExistingPackageName, FPackageSourceControlHelper& PackageHelper)
+{
+	UE_SCOPED_TIMER(TEXT("Delete existing destination packages"), LogWorldPartitionCopyWorldBuilder, Display);
+	TArray<FString> PackagesToDelete;
+
+	FString ExistingMapPackageFilePath = FPackageName::LongPackageNameToFilename(ExistingPackageName, FPackageName::GetMapPackageExtension());
+	if (IFileManager::Get().FileExists(*ExistingMapPackageFilePath))
+	{
+		PackagesToDelete.Add(ExistingMapPackageFilePath);
+	}
+
+	FString BuildDataPackageName = ExistingPackageName + TEXT("_BuiltData");
+	FString ExistingBuildDataPackageFilePath = FPackageName::LongPackageNameToFilename(BuildDataPackageName, FPackageName::GetAssetPackageExtension());
+	if (IFileManager::Get().FileExists(*ExistingBuildDataPackageFilePath))
+	{
+		PackagesToDelete.Add(ExistingBuildDataPackageFilePath);
+	}
+
+	// Search for external object packages
+	const TArray<FString> ExternalPackagesPaths = ULevel::GetExternalObjectsPaths(ExistingPackageName);
+	for (const FString& ExternalPackagesPath : ExternalPackagesPaths)
+	{
+		FString ExternalPackagesFilePath = FPackageName::LongPackageNameToFilename(ExternalPackagesPath);
+		if (IFileManager::Get().DirectoryExists(*ExternalPackagesFilePath))
+		{
+			const bool bSuccess = IFileManager::Get().IterateDirectoryRecursively(*ExternalPackagesFilePath, [&PackagesToDelete](const TCHAR* FilenameOrDirectory, bool bIsDirectory)
+			{
+				if (!bIsDirectory)
+				{
+					FString Filename(FilenameOrDirectory);
+					if (Filename.EndsWith(FPackageName::GetAssetPackageExtension()))
+					{
+						PackagesToDelete.Add(Filename);
+					}
+				}
+				// Continue Directory Iteration
+				return true;
+			});
+
+			if (!bSuccess)
+			{
+				UE_LOG(LogWorldPartitionCopyWorldBuilder, Error, TEXT("Failed to iterate existing external actors path: %s"), *ExternalPackagesPath);
+				return false;
+			}
+		}
+	}
+
+	UE_LOG(LogWorldPartitionCopyWorldBuilder, Display, TEXT("Deleting %d packages..."), PackagesToDelete.Num());
+	if (PackagesToDelete.Num() > 0)
+	{
+		if (!PackageHelper.Delete(PackagesToDelete))
+		{
+			UE_LOG(LogWorldPartitionCopyWorldBuilder, Error, TEXT("Failed to delete existing destination packages:"));
+			for (const FString& PackageToDelete : PackagesToDelete)
+			{
+				UE_LOG(LogWorldPartitionCopyWorldBuilder, Error, TEXT("    Package: %s"), *PackageToDelete);
+			}
+			return false;
+		}
+	}
+
+	return true;
+}
+
+static bool SavePackages(const TArray<UPackage*>& Packages, FPackageSourceControlHelper& PackageHelper)
+{
+	const TArray<FString> PackageFilenames = SourceControlHelpers::PackageFilenames(Packages);
+	TArray<FString> PackagesToCheckout;
+	TArray<FString> PackagesToAdd;
+
+	for (int PackageIndex = 0; PackageIndex < Packages.Num(); ++PackageIndex)
+	{
+		const FString& PackageFilename = PackageFilenames[PackageIndex];
+		const bool bFileExists = IPlatformFile::GetPlatformPhysical().FileExists(*PackageFilename);
+
+		if (bFileExists)
+		{
+			PackagesToCheckout.Add(PackageFilename);
+		}
+		else
+		{
+			PackagesToAdd.Add(PackageFilename);
+		}
+	}
+
+	if (PackagesToCheckout.Num())
+	{
+		if (!PackageHelper.Checkout(PackagesToCheckout))
+		{
+			return false;
+		}
+	}
+
+	for(int PackageIndex = 0; PackageIndex < Packages.Num(); ++PackageIndex)
+	{
+		// Save package
+		FSavePackageArgs SaveArgs;
+		SaveArgs.TopLevelFlags = RF_Standalone;
+		if (!UPackage::SavePackage(Packages[PackageIndex], nullptr, *PackageFilenames[PackageIndex], SaveArgs))
+		{
+			UE_LOG(LogWorldPartitionCopyWorldBuilder, Error, TEXT("Error saving package %s."), *Packages[PackageIndex]->GetName());
+			return false;
+		}
+	}
+
+	if (PackagesToAdd.Num())
+	{
+		if (!PackageHelper.AddToSourceControl(PackagesToAdd))
+		{
+			return false;
+		}
+	}
+
+	return true;
+}
+
+UWorldPartitionRenameDuplicateBuilder::UWorldPartitionRenameDuplicateBuilder(const FObjectInitializer& ObjectInitializer)
+	: Super(ObjectInitializer)
+{
+	FParse::Value(FCommandLine::Get(), TEXT("NewPackage="), NewPackageName);
+	bRename = FParse::Param(FCommandLine::Get(), TEXT("Rename"));
+}
+
+bool UWorldPartitionRenameDuplicateBuilder::RunInternal(UWorld* World, const FCellInfo& CellInfo, FPackageSourceControlHelper& PackageHelper)
+{
+	UWorldPartition* WorldPartition = World->GetWorldPartition();
+	if (!WorldPartition)
+	{
+		UE_LOG(LogWorldPartitionCopyWorldBuilder, Error, TEXT("Failed to retrieve WorldPartition."));
+		return false;
+	}
+		
+	TMap<FGuid, FWorldPartitionActorDescView> ActorDescViewMap;
+	for (UActorDescContainer::TConstIterator<> ActorDescIterator(WorldPartition); ActorDescIterator; ++ActorDescIterator)
+	{
+		const FWorldPartitionActorDesc* ActorDesc = *ActorDescIterator;
+		FWorldPartitionActorDescView ActorDescView(ActorDesc);
+		// Invalidate data layers to avoid clustering errors. It doesn't matter here as we need clustering only to gather references
+		ActorDescView.SetInvalidDataLayers();
+		ActorDescViewMap.Add(ActorDesc->GetGuid(), ActorDescView);
+	}
+
+	TArray<FActorCluster> ActorClusters;
+	{
+		UE_SCOPED_TIMER(TEXT("Create actor clusters"), LogWorldPartitionCopyWorldBuilder, Display);
+		FActorClusterContext::CreateActorClusters(World, ActorDescViewMap, ActorClusters);
+	}
+
+	UPackage* OriginalPackage = World->GetPackage();
+	OriginalWorldName = World->GetName();
+	OriginalPackageName = OriginalPackage->GetName();
+	
+	TSet<FString> PackagesToDelete;
+	if (bRename)
+	{
+		if (World->PersistentLevel->MapBuildData)
+		{
+			FString BuildDataPackageName = OriginalPackageName + TEXT("_BuiltData");
+			PackagesToDelete.Add(BuildDataPackageName);
+		}
+
+		for (UPackage* ExternalPackage : OriginalPackage->GetExternalPackages())
+		{
+			PackagesToDelete.Add(ExternalPackage->GetName());
+			ResetLoaders(ExternalPackage);
+		}
+	}
+	
+	const FString NewWorldName = FPackageName::GetLongPackageAssetName(NewPackageName);
+
+	// Delete destination if it exists
+	if (!DeleteExistingMapPackages(NewPackageName, PackageHelper))
+	{
+		UE_LOG(LogWorldPartitionCopyWorldBuilder, Error, TEXT("Failed to delete existing destination package."));
+		return false;
+	}
+				
+	UPackage* NewPackage = CreatePackage(*NewPackageName);
+	TMap<FGuid, FGuid> DuplicatedActorGuids;
+	TArray<UPackage*> DuplicatedPackagesToSave;
+	UWorld* NewWorld = nullptr;
+	{
+		UE_SCOPED_TIMER(TEXT("Duplicating world"), LogWorldPartitionCopyWorldBuilder, Display);
+		FObjectDuplicationParameters DuplicationParameters(World, NewPackage);
+		DuplicationParameters.DuplicateMode = EDuplicateMode::World;
+
+		DuplicatedObjects.Empty();
+		TMap<UObject*, UObject*> DuplicatedObjectPtrs;
+		DuplicationParameters.CreatedObjects = &DuplicatedObjectPtrs;
+		
+		NewWorld = Cast<UWorld>(StaticDuplicateObjectEx(DuplicationParameters));
+		check(NewWorld);
+	
+		// Copy Object pointers to Property so that GC doesn't try and collect any of them
+		for (const auto& Pair : DuplicatedObjectPtrs)
+		{
+			DuplicatedObjects.Add(Pair.Key, Pair.Value);
+
+			// Keep list of duplicated actor guids to skip processing them
+			if (Pair.Value->IsPackageExternal())
+			{
+				if (AActor* SourceActor = Cast<AActor>(Pair.Key))
+				{
+					AActor* DuplicatedActor = CastChecked<AActor>(Pair.Value);
+
+					DuplicatedActorGuids.Add(SourceActor->GetActorGuid(), DuplicatedActor->GetActorGuid());
+				}
+			}
+		}
+
+		DuplicatedPackagesToSave.Append(NewWorld->GetPackage()->GetExternalPackages());
+		DuplicatedPackagesToSave.Add(NewWorld->GetPackage());
+	}
+		
+	// World Scope
+	{
+		UWorld::InitializationValues IVS;
+		IVS.RequiresHitProxies(false);
+		IVS.ShouldSimulatePhysics(false);
+		IVS.EnableTraceCollision(false);
+		IVS.CreateNavigation(false);
+		IVS.CreateAISystem(false);
+		IVS.AllowAudioPlayback(false);
+		IVS.CreatePhysicsScene(true);
+		FScopedEditorWorld NewEditorWorld(NewWorld, IVS);
+
+		// Fixup SoftPath archive
+		FSoftObjectPathFixupArchive SoftObjectPathFixupArchive(OriginalPackageName + TEXT(".") + OriginalWorldName, NewPackageName + TEXT(".") + NewWorldName);
+
+		{
+			UE_SCOPED_TIMER(TEXT("Saving actors"), LogWorldPartitionCopyWorldBuilder, Display);
+								
+			auto ProcessLoadedActors = [&](TArray<FWorldPartitionReference>& ActorReferences) -> bool
+			{
+				if (!ActorReferences.Num())
+				{
+					return true;
+				}
+
+				TArray<UPackage*> ActorPackages;
+				
+				for (const FWorldPartitionReference& ActorReference : ActorReferences)
+				{
+					AActor* Actor = ActorReference->GetActor();
+					UPackage* PreviousActorPackage = Actor->GetExternalPackage();
+
+					// Rename Actor first so new package gets created
+					Actor->Rename(nullptr, NewWorld->PersistentLevel, REN_NonTransactional | REN_DontCreateRedirectors | REN_ForceNoResetLoaders | REN_DoNotDirty);
+
+					TArray<UObject*> DependantObjects;
+					ForEachObjectWithPackage(PreviousActorPackage, [&DependantObjects](UObject* Object)
+					{
+						if (!Cast<UMetaData>(Object))
+						{
+							DependantObjects.Add(Object);
+						}
+						return true;
+					}, false);
+
+					// Move dependant objects into the new actor package
+					for (UObject* DependantObject : DependantObjects)
+					{
+						DependantObject->Rename(nullptr, Actor->GetExternalPackage(), REN_NonTransactional | REN_DontCreateRedirectors | REN_ForceNoResetLoaders | REN_DoNotDirty);
+					}
+
+					// Releases file handle so it can be deleted
+					if (bRename)
+					{
+						ResetLoaders(PreviousActorPackage);
+					}
+
+					// Patch SoftObject Paths
+					SoftObjectPathFixupArchive.Fixup(Actor);
+					// Patch Duplicated Object Refs
+					FReplaceObjectRefsArchive(Actor, DuplicatedObjects);
+
+					ActorPackages.Add(Actor->GetPackage());
+				}
+
+				UE_LOG(LogWorldPartitionCopyWorldBuilder, Display, TEXT("Saving %d actor(s)"), ActorPackages.Num());
+				if (!SavePackages(ActorPackages, PackageHelper))
+				{
+					UE_LOG(LogWorldPartitionCopyWorldBuilder, Error, TEXT("Failed to save actor packages:"));
+					for (UPackage* ActorPackage : ActorPackages)
+					{
+						UE_LOG(LogWorldPartitionCopyWorldBuilder, Error, TEXT("    Package: %s"), *ActorPackage->GetName());
+					}
+					return false;
+				}
+				ActorPackages.Empty();
+
+				// Rename Actor(s) back into their original Outer so that they stay valid until the next GC. 
+				// This is to prevent failures when some non-serialized references get taken by loaded actors this makes sure those references will resolve. (example that no longer exists: Landscape SplineHandles)
+				for (FWorldPartitionReference ActorReference : ActorReferences)
+				{
+					AActor* Actor = ActorReference->GetActor();
+					UPackage* NewActorPackage = Actor->GetExternalPackage();
+					check(Actor);
+					Actor->Rename(nullptr, World->PersistentLevel, REN_NonTransactional | REN_DontCreateRedirectors | REN_ForceNoResetLoaders | REN_DoNotDirty);
+
+					TArray<UObject*> DependantObjects;
+					ForEachObjectWithPackage(NewActorPackage, [&DependantObjects](UObject* Object)
+					{
+						if (!Cast<UMetaData>(Object))
+						{
+							DependantObjects.Add(Object);
+						}
+						return true;
+					}, false);
+
+					// Move back dependant objects into the previous actor package
+					for (UObject* DependantObject : DependantObjects)
+					{
+						DependantObject->Rename(nullptr, Actor->GetExternalPackage(), REN_NonTransactional | REN_DontCreateRedirectors | REN_ForceNoResetLoaders | REN_DoNotDirty);
+					}
+				}
+				ActorReferences.Empty();
+
+				return true;
+			};
+
+			TArray<FWorldPartitionReference> ActorReferences;
+			for (const FActorCluster& ActorCluster : ActorClusters)
+			{
+				UE_LOG(LogWorldPartitionCopyWorldBuilder, Display, TEXT("Processing cluster with %d actor(s)"), ActorCluster.Actors.Num());
+				for (const FGuid& ActorGuid : ActorCluster.Actors)
+				{
+					// Duplicated actors don't need to be processed
+					if (!DuplicatedActorGuids.Contains(ActorGuid))
+					{
+						FWorldPartitionReference ActorReference(WorldPartition, ActorGuid);
+						check(*ActorReference);
+						AActor* Actor = ActorReference->GetActor();
+						check(Actor);
+						ActorReferences.Add(ActorReference);
+					}
+
+					// If we are renaming add package to delete
+					if (bRename)
+					{
+						FWorldPartitionActorDescView& ActorDescView = ActorDescViewMap.FindChecked(ActorGuid);
+						PackagesToDelete.Add(ActorDescView.GetActorPackage().ToString());
+					}
+				}
+						
+				if (FWorldPartitionHelpers::HasExceededMaxMemory())
+				{
+					if (!ProcessLoadedActors(ActorReferences))
+					{
+						return false;
+					}
+					FWorldPartitionHelpers::DoCollectGarbage();
+				}
+			}
+
+			// Call one last time
+			if (!ProcessLoadedActors(ActorReferences))
+			{
+				return false;
+			}
+		}
+
+		{
+			// Save all duplicated packages
+			UE_SCOPED_TIMER(TEXT("Saving new map packages"), LogWorldPartitionCopyWorldBuilder, Display);
+			if (!SavePackages(DuplicatedPackagesToSave, PackageHelper))
+			{
+				return false;
+			}
+		}
+				
+
+		{
+			// Validate results
+			UE_SCOPED_TIMER(TEXT("Validating actors"), LogWorldPartitionCopyWorldBuilder, Display);
+			for (UActorDescContainer::TConstIterator<> ActorDescIterator(WorldPartition); ActorDescIterator; ++ActorDescIterator)
+			{
+				const FWorldPartitionActorDesc* SourceActorDesc = *ActorDescIterator;
+				FGuid* DuplicatedGuid = DuplicatedActorGuids.Find(SourceActorDesc->GetGuid());
+				const FWorldPartitionActorDesc* NewActorDesc = NewWorld->GetWorldPartition()->GetActorDesc(DuplicatedGuid ? *DuplicatedGuid : SourceActorDesc->GetGuid());
+				if (!NewActorDesc)
+				{
+					UE_LOG(LogWorldPartitionCopyWorldBuilder, Warning, TEXT("Failed to find source actor for Actor: %s"), *SourceActorDesc->GetActorPath().ToString());
+				}
+				else
+				{
+					if (NewActorDesc->GetReferences().Num() != SourceActorDesc->GetReferences().Num())
+					{
+						UE_LOG(LogWorldPartitionCopyWorldBuilder, Warning, TEXT("Actor: %s and Source Actor: %s have mismatching reference count"), *NewActorDesc->GetActorPath().ToString(), *SourceActorDesc->GetActorPath().ToString());
+					}
+					else
+					{
+						for (const FGuid& ReferenceGuid : SourceActorDesc->GetReferences())
+						{
+							FGuid* DuplicateReferenceGuid = DuplicatedActorGuids.Find(ReferenceGuid);
+							if (!NewActorDesc->GetReferences().Contains(DuplicateReferenceGuid ? *DuplicateReferenceGuid : ReferenceGuid))
+							{
+								UE_LOG(LogWorldPartitionCopyWorldBuilder, Warning, TEXT("Actor: %s and Source Actor: %s have mismatching reference"), *NewActorDesc->GetActorPath().ToString(), *SourceActorDesc->GetActorPath().ToString());
+							}
+						}
+					}
+				}
+			}
+		}
+
+
+		DuplicatedObjects.Empty();
+	}
+		
+	if (PackagesToDelete.Num() > 0)
+	{
+		UE_SCOPED_TIMER(TEXT("Delete source packages (-rename switch)"), LogWorldPartitionCopyWorldBuilder, Display);
+		
+		UE_LOG(LogWorldPartitionCopyWorldBuilder, Display, TEXT("Deleting %d packages"), PackagesToDelete.Num());
+		if (!PackageHelper.Delete(PackagesToDelete.Array()))
+		{
+			UE_LOG(LogWorldPartitionCopyWorldBuilder, Error, TEXT("Failed to delete source packages:"));
+			for (const FString& PackageToDelete : PackagesToDelete)
+			{
+				UE_LOG(LogWorldPartitionCopyWorldBuilder, Error, TEXT("    Package: %s"), *PackageToDelete);
+			}
+			return false;
+		}
+	}
+			
+	return true;
+}
+
+bool UWorldPartitionRenameDuplicateBuilder::PostWorldTeardown(FPackageSourceControlHelper& PackageHelper)
+{
+	if (!UWorldPartitionBuilder::PostWorldTeardown(PackageHelper))
+	{
+		return false;
+	}
+
+	// Create redirector
+	if (bRename)
+	{
+		// Make sure to release handle on original package if a redirector needs to be saved
+		UPackage* OriginalPackage = FindPackage(nullptr, *OriginalPackageName);
+		ResetLoaders(OriginalPackage);
+		FWorldPartitionHelpers::DoCollectGarbage();
+		check(!FindPackage(nullptr, *OriginalPackageName));
+		if(!PackageHelper.Delete(*OriginalPackageName))
+		{
+			UE_LOG(LogWorldPartitionCopyWorldBuilder, Error, TEXT("Failed to delete package: %s"), *OriginalPackageName);
+			return false;
+		}
+				
+		UPackage* RedirectorPackage = CreatePackage(*OriginalPackageName);
+		RedirectorPackage->ThisContainsMap();
+
+		UObjectRedirector* Redirector = NewObject<UObjectRedirector>(RedirectorPackage, *OriginalWorldName, RF_Standalone | RF_Public);
+		FSoftObjectPath RedirectorPath(Redirector);
+
+		UPackage* NewWorldPackage = LoadPackage(nullptr, *NewPackageName, LOAD_None);
+		Redirector->DestinationObject = UWorld::FindWorldInPackage(NewWorldPackage);
+		check(Redirector->DestinationObject);
+
+		// Saving the NewPackage will save the duplicated external packages
+		UE_SCOPED_TIMER(TEXT("Saving new redirector"), LogWorldPartitionCopyWorldBuilder, Display);
+		if (!SavePackages({ RedirectorPackage }, PackageHelper))
+		{
+			UE_LOG(LogWorldPartitionCopyWorldBuilder, Error, TEXT("Failed to save redirector package: %s"), *RedirectorPackage->GetName());
+			return false;
+		}
+
+		// Validate Redirector
+		ResetLoaders(RedirectorPackage);
+		ForEachObjectWithPackage(RedirectorPackage, [](UObject* Object)
+		{
+			Object->ClearFlags(RF_Standalone);
+			return true;
+		}, false);
+		FWorldPartitionHelpers::DoCollectGarbage();
+		check(!FindPackage(nullptr, *OriginalPackageName));
+
+		UWorld* RedirectedWorld = CastChecked<UWorld>(RedirectorPath.TryLoad());
+		if (!RedirectedWorld)
+		{
+			UE_LOG(LogWorldPartitionCopyWorldBuilder, Error, TEXT("Failed to validate redirector package: %s"), *RedirectorPackage->GetName());
+			return false;
+		}
+	}
+
+	return true;
+}
diff --git a/Engine/Source/Editor/UnrealEd/Public/WorldPartition/WorldPartitionBuilder.h b/Engine/Source/Editor/UnrealEd/Public/WorldPartition/WorldPartitionBuilder.h
--- a/Engine/Source/Editor/UnrealEd/Public/WorldPartition/WorldPartitionBuilder.h
+++ b/Engine/Source/Editor/UnrealEd/Public/WorldPartition/WorldPartitionBuilder.h
@@ -60,7 +60,7 @@ public:
 
 protected:
 	/**
-	 * Overridable method for derived classed to perform operations when world builder process starts.
+	 * Overridable method for derived classes to perform operations when world builder process starts.
 	 * This is called before loading data (e.g. data layers, editor cells) and before calling `RunInternal`.
 	 */
 	virtual bool PreRun(UWorld* World, FPackageSourceControlHelper& PackageHelper) { return true; }
@@ -68,11 +68,16 @@ protected:
 	virtual bool RunInternal(UWorld* World, const FCellInfo& InCellInfo, FPackageSourceControlHelper& PackageHelper) PURE_VIRTUAL(UWorldPartition::RunInternal, return false;);
 
 	/**
-	 * Overridable method for derived classed to perform operations when world builder process completes.
+	 * Overridable method for derived classes to perform operations when world builder process completes.
 	 * This is called after loading all data (e.g. data layers, editor cells) and after calling `RunInternal` for all editor cells.
 	 */
 	virtual bool PostRun(UWorld* World, FPackageSourceControlHelper& PackageHelper, const bool bInRunSuccess) { return true; }
 
+	/**
+	 * Overridable method for derived classes to perform operations when world builder has unloaded the world.
+	 */
+	virtual bool PostWorldTeardown(FPackageSourceControlHelper& PackageHelper) { return true; }
+
 	int32 IterativeCellSize = 102400;
 	int32 IterativeCellOverlapSize = 0;
 	TSet<FName> DataLayerLabels;
diff --git a/Engine/Source/Editor/UnrealEd/Public/WorldPartition/WorldPartitionRenameDuplicateBuilder.h b/Engine/Source/Editor/UnrealEd/Public/WorldPartition/WorldPartitionRenameDuplicateBuilder.h
new file mode 100644
--- /dev/null
+++ b/Engine/Source/Editor/UnrealEd/Public/WorldPartition/WorldPartitionRenameDuplicateBuilder.h
@@ -0,0 +1,31 @@
+// Copyright Epic Games, Inc. All Rights Reserved.
+
+#pragma once
+
+#include "WorldPartition/WorldPartitionBuilder.h"
+#include "WorldPartitionRenameDuplicateBuilder.generated.h"
+
+// Example Command Line: ProjectName MapName -run=WorldPartitionBuilderCommandlet -SCCProvider=Perforce -Builder=WorldPartitionRenameDuplicateBuilder -NewPackage=NewPackage [Optional: -Rename]
+UCLASS()
+class UWorldPartitionRenameDuplicateBuilder : public UWorldPartitionBuilder
+{
+	GENERATED_UCLASS_BODY()
+
+public:
+	// UWorldPartitionBuilder interface begin
+	virtual bool RequiresCommandletRendering() const override { return false; }
+	virtual ELoadingMode GetLoadingMode() const override { return ELoadingMode::Custom; }
+	virtual bool RunInternal(UWorld* World, const FCellInfo& CellInfo, FPackageSourceControlHelper& PackageHelper) override;
+	virtual bool PostWorldTeardown(FPackageSourceControlHelper& PackageHelper) override;
+	// UWorldPartitionBuilder interface end
+
+private:
+	bool bRename;
+	FString NewPackageName;
+	FString OriginalPackageName;
+	FString OriginalWorldName;
+
+	// Keep duplicated objects around through GC calls.
+	UPROPERTY(transient)
+	TMap<TObjectPtr<UObject>, TObjectPtr<UObject>> DuplicatedObjects;
+};
\ No newline at end of file
diff --git a/Engine/Source/Runtime/Engine/Private/Level.cpp b/Engine/Source/Runtime/Engine/Private/Level.cpp
--- a/Engine/Source/Runtime/Engine/Private/Level.cpp
+++ b/Engine/Source/Runtime/Engine/Private/Level.cpp
@@ -1079,6 +1079,7 @@ void ULevel::PreDuplicate(FObjectDuplicationParameters& DupParams)
 				}
 				UPackage* DupPackage = CreateActorPackage(DstPackage, GetActorPackagingScheme(), Path);
 				DupPackage->MarkAsFullyLoaded();
+				DupPackage->MarkPackageDirty();
 				DupParams.DuplicationSeed.Add(Package, DupPackage);
 			}
 		}
@@ -1094,6 +1095,7 @@ void ULevel::PreDuplicate(FObjectDuplicationParameters& DupParams)
 				}
 				UPackage* DupPackage = FExternalPackageHelper::CreateExternalPackage(DstPackage, Path, UActorFolder::GetExternalPackageFlags());
 				DupPackage->MarkAsFullyLoaded();
+				DupPackage->MarkPackageDirty();
 				DupParams.DuplicationSeed.Add(Package, DupPackage);
 			}
 			return true;
diff --git a/Engine/Source/Runtime/Engine/Private/WorldPartition/WorldPartition.cpp b/Engine/Source/Runtime/Engine/Private/WorldPartition/WorldPartition.cpp
--- a/Engine/Source/Runtime/Engine/Private/WorldPartition/WorldPartition.cpp
+++ b/Engine/Source/Runtime/Engine/Private/WorldPartition/WorldPartition.cpp
@@ -407,7 +407,9 @@ void UWorldPartition::Initialize(UWorld* InWorld, const FTransform& InTransform)
 	if (bIsEditor || bIsGame || bPIEWorldTravel || bIsDedicatedServer)
 	{
 		UPackage* LevelPackage = OuterWorld->PersistentLevel->GetOutermost();
-		const FName PackageName = LevelPackage->GetLoadedPath().GetPackageFName();
+
+		// Duplicated worlds (ex: WorldPartitionRenameDuplicateBuilder) will not have a loaded path 
+		const FName PackageName = LevelPackage->GetLoadedPath().GetPackageFName().IsNone() ? LevelPackage->GetFName() : LevelPackage->GetLoadedPath().GetPackageFName();
 
 		// Currently known Instancing use cases:
 		// - World Partition map template (New Level)
