Fix cooked ActorComponent BP not showing up in component class list
Cooked packages don't have any metadata (yet; they might become available via the sidecar editor data)
However, all uncooked BPs that derive from ActorComponent have the BlueprintSpawnableComponent metadata set on them (see FBlueprintEditorUtils::RecreateClassMetaData), so include any ActorComponent BP that comes from a cooked package

diff --git a/Engine/Source/Editor/UnrealEd/Private/Kismet2/Kismet2.cpp b/Engine/Source/Editor/UnrealEd/Private/Kismet2/Kismet2.cpp
--- a/Engine/Source/Editor/UnrealEd/Private/Kismet2/Kismet2.cpp
+++ b/Engine/Source/Editor/UnrealEd/Private/Kismet2/Kismet2.cpp
@@ -2706,9 +2706,12 @@ bool FKismetEditorUtilities::CanBlueprintImplementInterface(UBlueprint const* Bl
 
 bool FKismetEditorUtilities::IsClassABlueprintSpawnableComponent(const UClass* Class)
 {
+	// @fixme: Cooked packages don't have any metadata (yet; they might become available via the sidecar editor data)
+	// However, all uncooked BPs that derive from ActorComponent have the BlueprintSpawnableComponent metadata set on them
+	// (see FBlueprintEditorUtils::RecreateClassMetaData), so include any ActorComponent BP that comes from a cooked package
 	return (!Class->HasAnyClassFlags(CLASS_Abstract) &&
-	        Class->IsChildOf<UActorComponent>() &&
-	        Class->HasMetaData(FBlueprintMetadata::MD_BlueprintSpawnableComponent));
+			Class->IsChildOf<UActorComponent>() &&
+			(Class->HasMetaData(FBlueprintMetadata::MD_BlueprintSpawnableComponent) || Class->GetPackage()->bIsCookedForEditor));
 }
 
 bool FKismetEditorUtilities::IsClassABlueprintSkeleton(const UClass* Class)
