Fix incorrect header case

diff --git a/Engine/Source/Editor/UnrealEd/Private/WorldPartition/WorldPartitionRenameDuplicateBuilder.cpp b/Engine/Source/Editor/UnrealEd/Private/WorldPartition/WorldPartitionRenameDuplicateBuilder.cpp
--- a/Engine/Source/Editor/UnrealEd/Private/WorldPartition/WorldPartitionRenameDuplicateBuilder.cpp
+++ b/Engine/Source/Editor/UnrealEd/Private/WorldPartition/WorldPartitionRenameDuplicateBuilder.cpp
@@ -14,7 +14,7 @@
 #include "UObject/SavePackage.h"
 #include "UObject/MetaData.h"
 #include "ProfilingDebugging/ScopedTimers.h"
-#include "HAL/Filemanager.h"
+#include "HAL/FileManager.h"
 
 DEFINE_LOG_CATEGORY_STATIC(LogWorldPartitionCopyWorldBuilder, All, All);
 
