WorldPartitionRenameDuplicateBuilder: Fix Redirector SCC State

diff --git a/Engine/Source/Editor/UnrealEd/Private/WorldPartition/WorldPartitionRenameDuplicateBuilder.cpp b/Engine/Source/Editor/UnrealEd/Private/WorldPartition/WorldPartitionRenameDuplicateBuilder.cpp
--- a/Engine/Source/Editor/UnrealEd/Private/WorldPartition/WorldPartitionRenameDuplicateBuilder.cpp
+++ b/Engine/Source/Editor/UnrealEd/Private/WorldPartition/WorldPartitionRenameDuplicateBuilder.cpp
@@ -511,12 +511,7 @@ bool UWorldPartitionRenameDuplicateBuilder::PostWorldTeardown(FPackageSourceCont
 		ResetLoaders(OriginalPackage);
 		FWorldPartitionHelpers::DoCollectGarbage();
 		check(!FindPackage(nullptr, *OriginalPackageName));
-		if(!PackageHelper.Delete(*OriginalPackageName))
-		{
-			UE_LOG(LogWorldPartitionCopyWorldBuilder, Error, TEXT("Failed to delete package: %s"), *OriginalPackageName);
-			return false;
-		}
-				
+
 		UPackage* RedirectorPackage = CreatePackage(*OriginalPackageName);
 		RedirectorPackage->ThisContainsMap();
 
@@ -526,6 +521,7 @@ bool UWorldPartitionRenameDuplicateBuilder::PostWorldTeardown(FPackageSourceCont
 		UPackage* NewWorldPackage = LoadPackage(nullptr, *NewPackageName, LOAD_None);
 		Redirector->DestinationObject = UWorld::FindWorldInPackage(NewWorldPackage);
 		check(Redirector->DestinationObject);
+		RedirectorPackage->MarkAsFullyLoaded();
 
 		// Saving the NewPackage will save the duplicated external packages
 		UE_SCOPED_TIMER(TEXT("Saving new redirector"), LogWorldPartitionCopyWorldBuilder, Display);
