Iterative cooking: Fix bug causing previously-cooked packages that failed to save (e.g editoronly packages) to never be iteratively skipped.
#rnx

diff --git a/Engine/Source/Editor/UnrealEd/Private/Cooker/LooseCookedPackageWriter.cpp b/Engine/Source/Editor/UnrealEd/Private/Cooker/LooseCookedPackageWriter.cpp
--- a/Engine/Source/Editor/UnrealEd/Private/Cooker/LooseCookedPackageWriter.cpp
+++ b/Engine/Source/Editor/UnrealEd/Private/Cooker/LooseCookedPackageWriter.cpp
@@ -436,39 +436,47 @@ TUniquePtr<FAssetRegistryState> FLooseCookedPackageWriter::LoadPreviousAssetRegi
 	// we assume they are packaged in the pak file (which we don't want to extract to confirm) and keep them all.
 	if (!bIterateSharedBuild)
 	{
-		// For regular iteration, remove every file from the PreviousState that no longer exists in the cooked directory
-		// and remove every cooked file from disk that is not present in the AssetRegistry
+		// For regular iteration, remove each SuccessfulSave cooked file from the PreviousState if it no longer exists
+		// in the cooked directory. Keep the FailedSave previous cook packages; we don't expect them to exist on disk.
+		// Also, remove from the registry and from disk the cooked files that no longer exist in the editor.
 		GetAllCookedFiles();
-		TSet<FName> ExistsOnlyInRegistry;
-		TSet<FName> ExistsOnlyOnDisk;
-		ExistsOnlyOnDisk.Reserve(UncookedPathToCookedPath.Num());
+		TSet<FName> RemoveFromRegistry;
+		TSet<FName> RemoveFromDisk;
+		RemoveFromDisk.Reserve(UncookedPathToCookedPath.Num());
 		for (TPair<FName, FName>& Pair : UncookedPathToCookedPath)
 		{
-			ExistsOnlyOnDisk.Add(Pair.Key);
+			RemoveFromDisk.Add(Pair.Key);
 		}
 		for (const TPair<FName, const FAssetPackageData*>& Pair : PreviousState->GetAssetPackageDataMap())
 		{
 			FName PackageName = Pair.Key;
-			const FAssetPackageData* PackageData = Pair.Value;
 			const FName UncookedFilename = PackageDatas.GetFileNameByPackageName(PackageName);
-			bool bExistsOnDisk = false;
-			if (!UncookedFilename.IsNone())
+			if (UncookedFilename.IsNone())
 			{
-				bExistsOnDisk = (ExistsOnlyOnDisk.Remove(UncookedFilename) == 1);
+				// Remove package from both disk and registry
+				// Keep its RemoveFromDisk entry
+				RemoveFromRegistry.Add(PackageName); // Add a RemoveFromRegistry entry
 			}
-			if (!bExistsOnDisk)
+			else
 			{
-				ExistsOnlyInRegistry.Add(PackageName);
+				// Keep package on disk if it exists. Keep package in registry if it exists on disk or was a FailedSave.
+				bool bExistsOnDisk = (RemoveFromDisk.Remove(UncookedFilename) == 1); // Remove its RemoveFromDisk entry
+				const FAssetPackageData* PackageData = Pair.Value;
+				if (!bExistsOnDisk && PackageData->DiskSize >= 0)
+				{
+					// Add RemoveFromRegistry entry if it didn't exist on disk and is a SuccessfulSave package
+					RemoveFromRegistry.Add(PackageName);
+				}
 			}
 		}
 
-		if (ExistsOnlyInRegistry.Num())
+		if (RemoveFromRegistry.Num())
 		{
-			PreviousState->PruneAssetData(TSet<FName>(), ExistsOnlyInRegistry, FAssetRegistrySerializationOptions());
+			PreviousState->PruneAssetData(TSet<FName>(), RemoveFromRegistry, FAssetRegistrySerializationOptions());
 		}
-		if (ExistsOnlyOnDisk.Num())
+		if (RemoveFromDisk.Num())
 		{
-			RemoveCookedPackagesByUncookedFilename(ExistsOnlyOnDisk.Array());
+			RemoveCookedPackagesByUncookedFilename(RemoveFromDisk.Array());
 		}
 	}
 
