Fix PCH compilation

diff --git a/Engine/Source/Editor/UnrealEd/Classes/Commandlets/WorldPartitionDataLayerToAssetCommandLet.h b/Engine/Source/Editor/UnrealEd/Classes/Commandlets/WorldPartitionDataLayerToAssetCommandLet.h
--- a/Engine/Source/Editor/UnrealEd/Classes/Commandlets/WorldPartitionDataLayerToAssetCommandLet.h
+++ b/Engine/Source/Editor/UnrealEd/Classes/Commandlets/WorldPartitionDataLayerToAssetCommandLet.h
@@ -5,6 +5,7 @@
 #include "CoreMinimal.h"
 #include "UObject/ObjectMacros.h"
 #include "UObject/UObjectAnnotation.h"
+#include "UObject/StrongObjectPtr.h"
 #include "Commandlets/Commandlet.h"
 #include "Commandlets/WorldPartitionCommandletHelpers.h"
 #include "PackageSourceControlHelper.h"
