Fixed crash opening collapsed area in "Create Animation Blueprint" dialog

Widget no longer gets destroyed and recreated on class pick (which occurs on the first tick of the class picker, once the collapsed area is opened)

diff --git a/Engine/Source/Editor/UnrealEd/Private/Factories/AnimBlueprintFactory.cpp b/Engine/Source/Editor/UnrealEd/Private/Factories/AnimBlueprintFactory.cpp
--- a/Engine/Source/Editor/UnrealEd/Private/Factories/AnimBlueprintFactory.cpp
+++ b/Engine/Source/Editor/UnrealEd/Private/Factories/AnimBlueprintFactory.cpp
@@ -199,8 +199,6 @@ public:
 				]
 			]
 		];
-
-		RefreshSkeletonPicker();
 	}
 	
 	/** Sets properties for the supplied AnimBlueprintFactory */
@@ -298,12 +296,18 @@ private:
 		RefreshSkeletonPicker();
 	}
 
-	/** Creates the combo menu for the target skeleton */
 	void RefreshSkeletonPicker()
+	{
+		RefreshSkeletonViewDelegate.ExecuteIfBound(true);
+	}
+	
+	/** Creates the widgets for the target skeleton area */
+	TSharedRef<SWidget> MakeSkeletonPickerArea()
 	{
 		FContentBrowserModule& ContentBrowserModule = FModuleManager::Get().LoadModuleChecked<FContentBrowserModule>(TEXT("ContentBrowser"));
 
 		FAssetPickerConfig AssetPickerConfig;
+		AssetPickerConfig.RefreshAssetViewDelegates.Add(&RefreshSkeletonViewDelegate);
 		AssetPickerConfig.Filter.ClassNames.Add(USkeleton::StaticClass()->GetFName());
 		AssetPickerConfig.OnAssetSelected = FOnAssetSelected::CreateSP(this, &SAnimBlueprintCreateDialog::OnSkeletonSelected);
 		AssetPickerConfig.OnShouldFilterAsset = FOnShouldFilterAsset::CreateSP(this, &SAnimBlueprintCreateDialog::FilterSkeletonBasedOnParentClass);
@@ -319,23 +323,13 @@ private:
 		AssetPickerConfig.bShowPathInColumnView = false;
 		AssetPickerConfig.bShowTypeInColumnView = false;
 		AssetPickerConfig.bFocusSearchBoxWhenOpened = false;
-		
-		SkeletonContainer->ClearChildren();
-		SkeletonContainer->AddSlot()
-		[
-			ContentBrowserModule.Get().CreateAssetPicker(AssetPickerConfig)
-		];
-	}
 	
-	/** Creates the widgets for the target skeleton area */
-	TSharedRef<SWidget> MakeSkeletonPickerArea()
-	{
 		return SNew(SVerticalBox)
 			+SVerticalBox::Slot()
 			.FillHeight(1.0f)
 			.Padding(5.0f)
 			[
-				SAssignNew(SkeletonContainer, SVerticalBox)
+				ContentBrowserModule.Get().CreateAssetPicker(AssetPickerConfig)
 			];
 	}
 
@@ -416,6 +410,9 @@ private:
 	/** The selected skeleton */
 	FAssetData TargetSkeleton;
 
+	/** Delegate called to refresh the skeleton view */
+	FRefreshAssetViewDelegate RefreshSkeletonViewDelegate;
+	
 	/** Whether we have a template selected or not */
 	bool bTemplate = false;
 	
